[mutable] Convert the following nodes to modifiers: EditMaterial, ExtendMaterial, RemoveMesh, RemoveMeshBlocks, MorphMaterial.

diff --git a/Engine/Plugins/Experimental/Mutable/Source/CustomizableObjectEditor/Private/MuCOE/Nodes/CustomizableObjectNode.cpp b/Engine/Plugins/Experimental/Mutable/Source/CustomizableObjectEditor/Private/MuCOE/Nodes/CustomizableObjectNode.cpp
--- a/Engine/Plugins/Experimental/Mutable/Source/CustomizableObjectEditor/Private/MuCOE/Nodes/CustomizableObjectNode.cpp
+++ b/Engine/Plugins/Experimental/Mutable/Source/CustomizableObjectEditor/Private/MuCOE/Nodes/CustomizableObjectNode.cpp
@@ -413,6 +413,18 @@ bool UCustomizableObjectNode::CanConnect( const UEdGraphPin* InOwnedInputPin, co
 }
 
 
+bool UCustomizableObjectNode::IsAffectedByLOD() const
+{ 
+	return true; 
+}
+
+
+TArray<FString>* UCustomizableObjectNode::GetEnableTags()
+{ 
+	return nullptr; 
+}
+
+
 UCustomizableObjectNodeRemapPins* UCustomizableObjectNode::CreateRemapPinsDefault() const
 {
 	return CreateRemapPinsByName();
@@ -579,55 +591,6 @@ int32 UCustomizableObjectNode::GetLOD() const
 }
 
 
-TArray<UCustomizableObjectNodeObject*> UCustomizableObjectNode::GetParentObjectNodes(const int LOD) const
-{
-	// Search recursively all parent nodes. Add all Object nodes found.
-	check(LOD >= 0);
-
-	TArray<UCustomizableObjectNodeObject*> Result;
-
-	TQueue<const UCustomizableObjectNode*> PotentialCustomizableNodeObjects;
-	PotentialCustomizableNodeObjects.Enqueue(this);
-
-	const UCustomizableObjectNode* CurrentElement;
-	while (PotentialCustomizableNodeObjects.Dequeue(CurrentElement))
-	{
-		for (UEdGraphPin* Pin : CurrentElement->GetAllNonOrphanPins())
-		{
-			if (Pin->Direction == EEdGraphPinDirection::EGPD_Output)
-			{
-				for (UEdGraphPin* LinkedPin : FollowOutputPinArray(*Pin))
-				{
-					UCustomizableObjectNode* Node = Cast<UCustomizableObjectNode>(LinkedPin->GetOwningNode());
-					check(Node); // All nodes inherit from UCustomizableObjectNode.
-					PotentialCustomizableNodeObjects.Enqueue(Node);
-
-					if (UCustomizableObjectNodeObject* CurrentCustomizableObjectNode = Cast<UCustomizableObjectNodeObject>(Node))
-					{
-						const int32 LODPin = CurrentCustomizableObjectNode->GetLOD(LinkedPin);
-						if (LODPin == -1 || LODPin == LOD) // Add to parents if it is not connected to a LOD pin or it is connected to the given LOD pin.
-						{
-							Result.Add(CurrentCustomizableObjectNode);
-
-							// Explore external CO.
-							if (CurrentCustomizableObjectNode->bIsBase)
-							{
-								if (UCustomizableObjectNodeObjectGroup* ObjectGroupNode = GetCustomizableObjectExternalNode<UCustomizableObjectNodeObjectGroup>(CurrentCustomizableObjectNode->ParentObject, CurrentCustomizableObjectNode->ParentObjectGroupId))
-								{
-									Result.Append(ObjectGroupNode->GetParentObjectNodes(LOD)); // Recursive call.
-								}
-							}
-						}
-					}
-				}
-			}
-		}
-	}
-
-	return Result;
-}
-
-
 void UCustomizableObjectNode::SetPinHidden(UEdGraphPin& Pin, bool bHidden)
 {
 	Pin.SafeSetHidden(bHidden && CanPinBeHidden(Pin));
diff --git a/Engine/Plugins/Experimental/Mutable/Source/CustomizableObjectEditor/Private/MuCOE/Nodes/CustomizableObjectNodeCopyMaterial.cpp b/Engine/Plugins/Experimental/Mutable/Source/CustomizableObjectEditor/Private/MuCOE/Nodes/CustomizableObjectNodeCopyMaterial.cpp
--- a/Engine/Plugins/Experimental/Mutable/Source/CustomizableObjectEditor/Private/MuCOE/Nodes/CustomizableObjectNodeCopyMaterial.cpp
+++ b/Engine/Plugins/Experimental/Mutable/Source/CustomizableObjectEditor/Private/MuCOE/Nodes/CustomizableObjectNodeCopyMaterial.cpp
@@ -231,15 +231,15 @@ FName UCustomizableObjectNodeCopyMaterial::GetMeshComponentName() const
 }
 
 
-TArray<FString> UCustomizableObjectNodeCopyMaterial::GetTags() const
+TArray<FString>* UCustomizableObjectNodeCopyMaterial::GetEnableTags()
 {
 	if (UCustomizableObjectNodeMaterial* NodeMaterial = GetMaterialNode())
 	{
-		return NodeMaterial->GetTags();
+		return NodeMaterial->GetEnableTags();
 	}
 	else
 	{
-		return {};
+		return nullptr;
 	}
 }
 
diff --git a/Engine/Plugins/Experimental/Mutable/Source/CustomizableObjectEditor/Private/MuCOE/Nodes/CustomizableObjectNodeEditLayoutBlocks.cpp b/Engine/Plugins/Experimental/Mutable/Source/CustomizableObjectEditor/Private/MuCOE/Nodes/CustomizableObjectNodeEditLayoutBlocks.cpp
--- a/Engine/Plugins/Experimental/Mutable/Source/CustomizableObjectEditor/Private/MuCOE/Nodes/CustomizableObjectNodeEditLayoutBlocks.cpp
+++ b/Engine/Plugins/Experimental/Mutable/Source/CustomizableObjectEditor/Private/MuCOE/Nodes/CustomizableObjectNodeEditLayoutBlocks.cpp
@@ -1,12 +1,12 @@
 // Copyright Epic Games, Inc. All Rights Reserved.
 
-#include "MuCOE/Nodes/CustomizableObjectNodeEditLayoutBlocks.h"
+#include "MuCOE/Nodes/CustomizableObjectNodeModifierEditLayoutBlocks.h"
 
 
 #define LOCTEXT_NAMESPACE "CustomizableObjectEditor"
 
 
-UCustomizableObjectNodeEditLayoutBlocks::UCustomizableObjectNodeEditLayoutBlocks()
+UCustomizableObjectNodeModifierEditLayoutBlocks::UCustomizableObjectNodeModifierEditLayoutBlocks()
 	: Super()
 {
 	Layout = CreateDefaultSubobject<UCustomizableObjectLayout>(FName("CustomizableObjectLayout"));
diff --git a/Engine/Plugins/Experimental/Mutable/Source/CustomizableObjectEditor/Private/MuCOE/Nodes/CustomizableObjectNodeMaterial.cpp b/Engine/Plugins/Experimental/Mutable/Source/CustomizableObjectEditor/Private/MuCOE/Nodes/CustomizableObjectNodeMaterial.cpp
--- a/Engine/Plugins/Experimental/Mutable/Source/CustomizableObjectEditor/Private/MuCOE/Nodes/CustomizableObjectNodeMaterial.cpp
+++ b/Engine/Plugins/Experimental/Mutable/Source/CustomizableObjectEditor/Private/MuCOE/Nodes/CustomizableObjectNodeMaterial.cpp
@@ -999,11 +999,11 @@ FName UCustomizableObjectNodeMaterial::GetParameterName(const EMaterialParameter
 }
 
 
-int32 UCustomizableObjectNodeMaterial::GetParameterLayerIndex(const EMaterialParameterType Type, const int32 ParameterIndex) const
+int32 UCustomizableObjectNodeMaterial::GetParameterLayerIndex(const UMaterialInterface* InMaterial, const EMaterialParameterType Type, const int32 ParameterIndex)
 {
-	check(Material);
+	check(InMaterial);
 
-	const FMaterialCachedParameterEntry& Entry = Material->GetCachedExpressionData().GetParameterTypeEntry(Type);
+	const FMaterialCachedParameterEntry& Entry = InMaterial->GetCachedExpressionData().GetParameterTypeEntry(Type);
 
 	for (TSet<FMaterialParameterInfo>::TConstIterator It(Entry.ParameterInfoSet); It; ++It)
 	{
@@ -1022,6 +1022,12 @@ int32 UCustomizableObjectNodeMaterial::GetParameterLayerIndex(const EMaterialPar
 }
 
 
+int32 UCustomizableObjectNodeMaterial::GetParameterLayerIndex(const EMaterialParameterType Type, const int32 ParameterIndex) const
+{
+	return GetParameterLayerIndex(Material.Get(), Type, ParameterIndex);
+}
+
+
 FText UCustomizableObjectNodeMaterial::GetParameterLayerName(const EMaterialParameterType Type, const int32 ParameterIndex) const
 {
 	check(Material)
@@ -1036,16 +1042,16 @@ FText UCustomizableObjectNodeMaterial::GetParameterLayerName(const EMaterialPara
 }
 
 
-bool UCustomizableObjectNodeMaterial::HasParameter(const FNodeMaterialParameterId& ParameterId) const
+bool UCustomizableObjectNodeMaterial::HasParameter(const UMaterialInterface* InMaterial, const FNodeMaterialParameterId& ParameterId)
 {
-	if (!Material)
+	if (!InMaterial)
 	{
 		return false;
 	}
 
 	for (const EMaterialParameterType Type : ParameterTypes)
 	{
-		const FMaterialCachedExpressionData& Data = Material->GetCachedExpressionData();
+		const FMaterialCachedExpressionData& Data = InMaterial->GetCachedExpressionData();
 		const FMaterialCachedParameterEntry& Entry = Data.GetParameterTypeEntry(Type);
 
 		if (!Data.EditorOnlyData || Data.EditorOnlyData->EditorEntries[(int32)Type].EditorInfo.IsEmpty())
@@ -1058,7 +1064,7 @@ bool UCustomizableObjectNodeMaterial::HasParameter(const FNodeMaterialParameterI
 			const int32 IteratorIndex = It.GetId().AsInteger();
 			
 			const FGuid& ParamGuid = Data.EditorOnlyData->EditorEntries[(int32)Type].EditorInfo[IteratorIndex].ExpressionGuid;
-			const int32 LayerIndex = GetParameterLayerIndex(Type, IteratorIndex);
+			const int32 LayerIndex = GetParameterLayerIndex(InMaterial, Type, IteratorIndex);
 			const FNodeMaterialParameterId ParamId = { ParamGuid, LayerIndex };
 
 			if (ParamId == ParameterId)
@@ -1072,6 +1078,12 @@ bool UCustomizableObjectNodeMaterial::HasParameter(const FNodeMaterialParameterI
 }
 
 
+bool UCustomizableObjectNodeMaterial::HasParameter(const FNodeMaterialParameterId& ParameterId) const
+{
+	return HasParameter( Material.Get(), ParameterId );
+}
+
+
 UEdGraphPin* UCustomizableObjectNodeMaterial::GetParameterPin(const EMaterialParameterType Type, const int32 ParameterIndex) const
 {
 	const FNodeMaterialParameterId ParameterId = GetParameterId(Type, ParameterIndex);
@@ -1196,9 +1208,9 @@ FName UCustomizableObjectNodeMaterial::GetMeshComponentName() const
 }
 
 
-TArray<FString> UCustomizableObjectNodeMaterial::GetTags() const
+TArray<FString>* UCustomizableObjectNodeMaterial::GetEnableTags()
 {
-	return Tags;
+	return &Tags;
 }
 
 
diff --git a/Engine/Plugins/Experimental/Mutable/Source/CustomizableObjectEditor/Private/MuCOE/Nodes/CustomizableObjectNodeMeshClipDeform.cpp b/Engine/Plugins/Experimental/Mutable/Source/CustomizableObjectEditor/Private/MuCOE/Nodes/CustomizableObjectNodeMeshClipDeform.cpp
--- a/Engine/Plugins/Experimental/Mutable/Source/CustomizableObjectEditor/Private/MuCOE/Nodes/CustomizableObjectNodeMeshClipDeform.cpp
+++ b/Engine/Plugins/Experimental/Mutable/Source/CustomizableObjectEditor/Private/MuCOE/Nodes/CustomizableObjectNodeMeshClipDeform.cpp
@@ -1,6 +1,6 @@
 // Copyright Epic Games, Inc. All Rights Reserved.
 
-#include "MuCOE/Nodes/CustomizableObjectNodeMeshClipDeform.h"
+#include "MuCOE/Nodes/CustomizableObjectNodeModifierClipDeform.h"
 
 #include "MuCOE/EdGraphSchema_CustomizableObject.h"
 
@@ -9,7 +9,7 @@ class UCustomizableObjectNodeRemapPins;
 #define LOCTEXT_NAMESPACE "CustomizableObjectEditor"
 
 
-void UCustomizableObjectNodeMeshClipDeform::AllocateDefaultPins(UCustomizableObjectNodeRemapPins* RemapPins)
+void UCustomizableObjectNodeModifierClipDeform::AllocateDefaultPins(UCustomizableObjectNodeRemapPins* RemapPins)
 {
 	const UEdGraphSchema_CustomizableObject* Schema = GetDefault<UEdGraphSchema_CustomizableObject>();
 	
@@ -21,25 +21,19 @@ void UCustomizableObjectNodeMeshClipDeform::AllocateDefaultPins(UCustomizableObj
 }
 
 
-UEdGraphPin* UCustomizableObjectNodeMeshClipDeform::ClipShapePin() const
+UEdGraphPin* UCustomizableObjectNodeModifierClipDeform::ClipShapePin() const
 {
 	return FindPin(TEXT("Clip Shape"), EGPD_Input);
 }
 
 
-UEdGraphPin* UCustomizableObjectNodeMeshClipDeform::OutputPin() const
-{
-	return FindPin(TEXT("Modifier"));
-}
-
-
-FText UCustomizableObjectNodeMeshClipDeform::GetNodeTitle(ENodeTitleType::Type TitleType) const
+FText UCustomizableObjectNodeModifierClipDeform::GetNodeTitle(ENodeTitleType::Type TitleType) const
 {
 	return LOCTEXT("Clip_Deform_Mesh", "Clip Deform Mesh");
 }
 
 
-FText UCustomizableObjectNodeMeshClipDeform::GetTooltipText() const
+FText UCustomizableObjectNodeModifierClipDeform::GetTooltipText() const
 {
 	return LOCTEXT("Clip_Deform_Tooltip", "Defines a clip with mesh deformation based on a shape mesh and blend weights.");
 
diff --git a/Engine/Plugins/Experimental/Mutable/Source/CustomizableObjectEditor/Private/MuCOE/Nodes/CustomizableObjectNodeMeshClipMorph.cpp b/Engine/Plugins/Experimental/Mutable/Source/CustomizableObjectEditor/Private/MuCOE/Nodes/CustomizableObjectNodeMeshClipMorph.cpp
--- a/Engine/Plugins/Experimental/Mutable/Source/CustomizableObjectEditor/Private/MuCOE/Nodes/CustomizableObjectNodeMeshClipMorph.cpp
+++ b/Engine/Plugins/Experimental/Mutable/Source/CustomizableObjectEditor/Private/MuCOE/Nodes/CustomizableObjectNodeMeshClipMorph.cpp
@@ -1,6 +1,6 @@
 // Copyright Epic Games, Inc. All Rights Reserved.
 
-#include "MuCOE/Nodes/CustomizableObjectNodeMeshClipMorph.h"
+#include "MuCOE/Nodes/CustomizableObjectNodeModifierClipMorph.h"
 
 #include "MuCO/CustomizableObjectCustomVersion.h"
 #include "MuCOE/EdGraphSchema_CustomizableObject.h"
@@ -14,7 +14,7 @@ class UEdGraphPin;
 #define LOCTEXT_NAMESPACE "CustomizableObjectEditor"
 
 
-UCustomizableObjectNodeMeshClipMorph::UCustomizableObjectNodeMeshClipMorph()
+UCustomizableObjectNodeModifierClipMorph::UCustomizableObjectNodeModifierClipMorph()
 	: Super()
 {
 	StartOffset = FVector::ZeroVector;
@@ -31,7 +31,7 @@ UCustomizableObjectNodeMeshClipMorph::UCustomizableObjectNodeMeshClipMorph()
 }
 
 
-FVector UCustomizableObjectNodeMeshClipMorph::GetOriginWithOffset() const
+FVector UCustomizableObjectNodeModifierClipMorph::GetOriginWithOffset() const
 {
 	FVector NewOrigin;
 
@@ -50,7 +50,7 @@ FVector UCustomizableObjectNodeMeshClipMorph::GetOriginWithOffset() const
 	return NewOrigin;
 }
 
-void UCustomizableObjectNodeMeshClipMorph::FindLocalAxes(FVector& XAxis, FVector& YAxis, FVector& ZAxis) const
+void UCustomizableObjectNodeModifierClipMorph::FindLocalAxes(FVector& XAxis, FVector& YAxis, FVector& ZAxis) const
 {
 	YAxis = FVector(0.f, 1.f, 0.f);
 
@@ -69,7 +69,7 @@ void UCustomizableObjectNodeMeshClipMorph::FindLocalAxes(FVector& XAxis, FVector
 }
 
 
-void UCustomizableObjectNodeMeshClipMorph::ChangeStartOffsetTransform()
+void UCustomizableObjectNodeModifierClipMorph::ChangeStartOffsetTransform()
 {
 	// Local Offset
 	FVector XAxis, YAxis, ZAxis;
@@ -87,13 +87,13 @@ void UCustomizableObjectNodeMeshClipMorph::ChangeStartOffsetTransform()
 }
 
 
-UEdGraphPin* UCustomizableObjectNodeMeshClipMorph::OutputPin() const
+UEdGraphPin* UCustomizableObjectNodeModifierClipMorph::OutputPin() const
 {
 	return FindPin(TEXT("Modifier"));
 }
 
 
-void UCustomizableObjectNodeMeshClipMorph::PostEditChangeProperty(FPropertyChangedEvent& PropertyChangedEvent)
+void UCustomizableObjectNodeModifierClipMorph::PostEditChangeProperty(FPropertyChangedEvent& PropertyChangedEvent)
 {
 	const FName PropertyName = PropertyChangedEvent.GetPropertyName();
 
@@ -108,7 +108,7 @@ void UCustomizableObjectNodeMeshClipMorph::PostEditChangeProperty(FPropertyChang
 }
 
 
-void UCustomizableObjectNodeMeshClipMorph::Serialize(FArchive& Ar)
+void UCustomizableObjectNodeModifierClipMorph::Serialize(FArchive& Ar)
 {
 	Super::Serialize(Ar);
 
@@ -133,7 +133,7 @@ void UCustomizableObjectNodeMeshClipMorph::Serialize(FArchive& Ar)
 }
 
 
-void UCustomizableObjectNodeMeshClipMorph::AllocateDefaultPins(UCustomizableObjectNodeRemapPins* RemapPins)
+void UCustomizableObjectNodeModifierClipMorph::AllocateDefaultPins(UCustomizableObjectNodeRemapPins* RemapPins)
 {
 	const UEdGraphSchema_CustomizableObject* Schema = GetDefault<UEdGraphSchema_CustomizableObject>();
 
@@ -141,13 +141,13 @@ void UCustomizableObjectNodeMeshClipMorph::AllocateDefaultPins(UCustomizableObje
 }
 
 
-FText UCustomizableObjectNodeMeshClipMorph::GetNodeTitle(ENodeTitleType::Type TitleType) const
+FText UCustomizableObjectNodeModifierClipMorph::GetNodeTitle(ENodeTitleType::Type TitleType) const
 {
 	return LOCTEXT("Clip_Morph_Mesh", "Clip Morph Mesh");
 }
 
 
-void UCustomizableObjectNodeMeshClipMorph::PinConnectionListChanged(UEdGraphPin* Pin)
+void UCustomizableObjectNodeModifierClipMorph::PinConnectionListChanged(UEdGraphPin* Pin)
 {
 	Super::PinConnectionListChanged(Pin);
 
@@ -163,7 +163,7 @@ void UCustomizableObjectNodeMeshClipMorph::PinConnectionListChanged(UEdGraphPin*
 }
 
 
-FText UCustomizableObjectNodeMeshClipMorph::GetTooltipText() const
+FText UCustomizableObjectNodeModifierClipMorph::GetTooltipText() const
 {
 	return LOCTEXT("Clip_Mesh_Morph_Tooltip", "Defines a cutting plane on a bone to cut tagged Materials that go past it, while morphing the mesh after the cut to blend in more naturally.\nIt only cuts and morphs mesh that receives some influence of that bone or other descendant bones.");
 }
diff --git a/Engine/Plugins/Experimental/Mutable/Source/CustomizableObjectEditor/Private/MuCOE/Nodes/CustomizableObjectNodeMeshClipWithMesh.cpp b/Engine/Plugins/Experimental/Mutable/Source/CustomizableObjectEditor/Private/MuCOE/Nodes/CustomizableObjectNodeMeshClipWithMesh.cpp
--- a/Engine/Plugins/Experimental/Mutable/Source/CustomizableObjectEditor/Private/MuCOE/Nodes/CustomizableObjectNodeMeshClipWithMesh.cpp
+++ b/Engine/Plugins/Experimental/Mutable/Source/CustomizableObjectEditor/Private/MuCOE/Nodes/CustomizableObjectNodeMeshClipWithMesh.cpp
@@ -1,6 +1,6 @@
 // Copyright Epic Games, Inc. All Rights Reserved.
 
-#include "MuCOE/Nodes/CustomizableObjectNodeMeshClipWithMesh.h"
+#include "MuCOE/Nodes/CustomizableObjectNodeModifierClipWithMesh.h"
 
 #include "MuCOE/CustomizableObjectGraph.h"
 #include "MuCOE/EdGraphSchema_CustomizableObject.h"
@@ -11,19 +11,19 @@ class UCustomizableObjectNodeRemapPins;
 #define LOCTEXT_NAMESPACE "CustomizableObjectEditor"
 
 
-UCustomizableObjectNodeMeshClipWithMesh::UCustomizableObjectNodeMeshClipWithMesh() :
+UCustomizableObjectNodeModifierClipWithMesh::UCustomizableObjectNodeModifierClipWithMesh() :
 	CustomizableObjectToClipWith(nullptr)
 {
 	// We allow this node to be shared
 	Transform = FTransform::Identity;
 
-	//When initialize we don't use materials neither tags
+	// When initialize we don't use materials neither tags
 	bUseMaterials = false;
 	bUseTags = false;
 }
 
 
-void UCustomizableObjectNodeMeshClipWithMesh::BeginPostDuplicate(bool bDuplicateForPIE)
+void UCustomizableObjectNodeModifierClipWithMesh::BeginPostDuplicate(bool bDuplicateForPIE)
 {
 	Super::BeginPostDuplicate(bDuplicateForPIE);
 
@@ -45,7 +45,7 @@ void UCustomizableObjectNodeMeshClipWithMesh::BeginPostDuplicate(bool bDuplicate
 }
 
 
-void UCustomizableObjectNodeMeshClipWithMesh::AllocateDefaultPins(UCustomizableObjectNodeRemapPins* RemapPins)
+void UCustomizableObjectNodeModifierClipWithMesh::AllocateDefaultPins(UCustomizableObjectNodeRemapPins* RemapPins)
 {
 	const UEdGraphSchema_CustomizableObject* Schema = GetDefault<UEdGraphSchema_CustomizableObject>();
 
@@ -56,7 +56,7 @@ void UCustomizableObjectNodeMeshClipWithMesh::AllocateDefaultPins(UCustomizableO
 }
 
 
-void UCustomizableObjectNodeMeshClipWithMesh::ReconstructNode(UCustomizableObjectNodeRemapPins* RemapPinsMode)
+void UCustomizableObjectNodeModifierClipWithMesh::ReconstructNode(UCustomizableObjectNodeRemapPins* RemapPinsMode)
 {
 	Super::ReconstructNode(RemapPinsMode);
 
@@ -67,13 +67,13 @@ void UCustomizableObjectNodeMeshClipWithMesh::ReconstructNode(UCustomizableObjec
 }
 
 
-FText UCustomizableObjectNodeMeshClipWithMesh::GetNodeTitle(ENodeTitleType::Type TitleType) const
+FText UCustomizableObjectNodeModifierClipWithMesh::GetNodeTitle(ENodeTitleType::Type TitleType) const
 {
 	return LOCTEXT("Clip_Mesh_With_Mesh", "Clip Mesh With Mesh");
 }
 
 
-void UCustomizableObjectNodeMeshClipWithMesh::PinConnectionListChanged(UEdGraphPin* Pin)
+void UCustomizableObjectNodeModifierClipWithMesh::PinConnectionListChanged(UEdGraphPin* Pin)
 {
 	Super::PinConnectionListChanged(Pin);
 
@@ -86,25 +86,25 @@ void UCustomizableObjectNodeMeshClipWithMesh::PinConnectionListChanged(UEdGraphP
 }
 
 
-UEdGraphPin* UCustomizableObjectNodeMeshClipWithMesh::OutputPin() const
+UEdGraphPin* UCustomizableObjectNodeModifierClipWithMesh::OutputPin() const
 {
 	return FindPin(TEXT("Modifier"));
 }
 
 
-UEdGraphPin* UCustomizableObjectNodeMeshClipWithMesh::ClipMeshPin() const
+UEdGraphPin* UCustomizableObjectNodeModifierClipWithMesh::ClipMeshPin() const
 {
 	return FindPin(TEXT("Clip Mesh"));
 }
 
 
-void UCustomizableObjectNodeMeshClipWithMesh::UpdateReferencedNodeId(const FGuid& NewGuid)
+void UCustomizableObjectNodeModifierClipWithMesh::UpdateReferencedNodeId(const FGuid& NewGuid)
 {
 	ArrayMaterialNodeToClipWithID.Add(NewGuid);
 }
 
 
-FText UCustomizableObjectNodeMeshClipWithMesh::GetTooltipText() const
+FText UCustomizableObjectNodeModifierClipWithMesh::GetTooltipText() const
 {
 	return LOCTEXT("Clip_Mesh_Mesh_Tooltip", "Removes the part of a material that is completely enclosed in a mesh volume.\nIt only removes the faces that fall completely inside the cutting volume, along with the vertices and edges that define only faces that are deleted.");
 }
diff --git a/Engine/Plugins/Experimental/Mutable/Source/CustomizableObjectEditor/Private/MuCOE/Nodes/CustomizableObjectNodeModifierBase.cpp b/Engine/Plugins/Experimental/Mutable/Source/CustomizableObjectEditor/Private/MuCOE/Nodes/CustomizableObjectNodeModifierBase.cpp
--- a/Engine/Plugins/Experimental/Mutable/Source/CustomizableObjectEditor/Private/MuCOE/Nodes/CustomizableObjectNodeModifierBase.cpp
+++ b/Engine/Plugins/Experimental/Mutable/Source/CustomizableObjectEditor/Private/MuCOE/Nodes/CustomizableObjectNodeModifierBase.cpp
@@ -4,8 +4,10 @@
 
 #include "MuCOE/EdGraphSchema_CustomizableObject.h"
 #include "MuCOE/CustomizableObjectEditorLogger.h"
-#include "MuCO/CustomizableObjectCustomVersion.h"
 #include "MuCOE/Nodes/CustomizableObjectNodeObject.h"
+#include "MuCOE/GraphTraversal.h"
+#include "MuCO/CustomizableObjectCustomVersion.h"
+#include "MuT/NodeModifier.h"
 
 
 FLinearColor UCustomizableObjectNodeModifierBase::GetNodeTitleColor() const
@@ -15,20 +17,104 @@ FLinearColor UCustomizableObjectNodeModifierBase::GetNodeTitleColor() const
 }
 
 
+UEdGraphPin* UCustomizableObjectNodeModifierBase::OutputPin() const
+{
+	return FindPin(TEXT("Modifier"));
+}
+
+
+bool UCustomizableObjectNodeModifierBase::IsApplicableTo(UCustomizableObjectNode* Candidate) const
+{
+	if (!Candidate)
+	{
+		return false;
+	}
+
+	const TArray<FString>* EnabledTags = Candidate->GetEnableTags();
+	const TArray<FString>* RequiredTags = GetRequiredTags();
+	if (EnabledTags && RequiredTags)
+	{
+		switch ( GetMultipleTagsPolicy() )
+		{
+		case EMutableMultipleTagPolicy::OnlyOneRequired:
+		{
+			for (const FString& RequiredTag : *RequiredTags)
+			{
+				if (EnabledTags->Contains(RequiredTag))
+				{
+					return true;
+				}
+			}
+			break;
+		}
+
+		case EMutableMultipleTagPolicy::AllRequired:
+		{
+			for (const FString& RequiredTag : *RequiredTags)
+			{
+				if (!EnabledTags->Contains(RequiredTag))
+				{
+					return false;
+				}
+			}
+			return true;
+			break;
+		}
+
+		default:
+			// Not implemented
+			check(false);
+		}
+	}
+
+	return false;
+}
+
+
+void UCustomizableObjectNodeModifierBase::GetPossiblyModifiedNodes(TArray<UCustomizableObjectNode*>& CandidateNodes) const
+{
+	// Scan all potential receivers
+	UCustomizableObject* ThisNodeObject = GetRootObject(*this);
+	UCustomizableObject* RootObject = GetRootObject(ThisNodeObject);
+
+	TSet<UCustomizableObject*> AllCustomizableObject;
+	GetAllObjectsInGraph(RootObject, AllCustomizableObject);
+
+	for (const UCustomizableObject* CustObject : AllCustomizableObject)
+	{
+		if (CustObject)
+		{
+			for (const TObjectPtr<UEdGraphNode>& CandidateNode : CustObject->GetPrivate()->GetSource()->Nodes)
+			{
+				UCustomizableObjectNode* Typed = Cast<UCustomizableObjectNode>(CandidateNode);
+				if (IsApplicableTo(Typed))
+				{
+					CandidateNodes.Add(Typed);
+				}
+			}
+		}
+	}
+}
+
 void UCustomizableObjectNodeModifierBase::BackwardsCompatibleFixup(int32 CustomizableObjectCustomVersion)
 {
 	Super::BackwardsCompatibleFixup(CustomizableObjectCustomVersion);
 	
 	const UEdGraphSchema_CustomizableObject* Schema = GetDefault<UEdGraphSchema_CustomizableObject>();
 
-	// Remove "Material" pin and add "Modifier"
+	// Remove "Material" pin and add "Modifier", fix the other side of the connections, if possible
 	if (CustomizableObjectCustomVersion == FCustomizableObjectCustomVersion::AddModifierPin)
 	{
 		// Replace the output pin
 		UEdGraphPin* OldPin = FindPin(TEXT("Material"));
-		UEdGraphPin* NewPin = CustomCreatePin(EGPD_Output, Schema->PC_Modifier, FName("Modifier"));
-		if (OldPin && NewPin)
+		if (OldPin)
 		{
+			UEdGraphPin* NewPin = FindPin(TEXT("Modifier")); 
+			if (!NewPin)
+			{
+				NewPin = CustomCreatePin(EGPD_Output, Schema->PC_Modifier, FName("Modifier"));
+			}
+
 			FGuid PinId = NewPin->PinId;
 			NewPin->CopyPersistentDataFromOldPin(*OldPin);
 			NewPin->PinId = PinId;
@@ -67,10 +153,6 @@ void UCustomizableObjectNodeModifierBase::BackwardsCompatibleFixup(int32 Customi
 						LinksToRemove.Add(LinkedToPin);
 						NewPin->MakeLinkTo(ToModifierPin);
 					}
-					else
-					{
-						ensure(false);
-					}
 				}
 			}
 
@@ -142,6 +224,8 @@ void UCustomizableObjectNodeModifierBase::BackwardsCompatibleFixup(int32 Customi
 
 void UCustomizableObjectNodeModifierBase::PostBackwardsCompatibleFixup()
 {
+	Super::PostBackwardsCompatibleFixup();
+
 	const UEdGraphSchema_CustomizableObject* Schema = GetDefault<UEdGraphSchema_CustomizableObject>();
 
 	// Check for old legacy connections that need manual update
@@ -178,5 +262,35 @@ void UCustomizableObjectNodeModifierBase::PostBackwardsCompatibleFixup()
 			}
 		}
 	}
+
+	// Apply backwards compatibility auto-generated tags to external objects.
+	const TArray<FString>* RequiredTags = GetRequiredTags();
+	if (RequiredTags)
+	{
+		for (const FLegacyTag& Tag : LegacyBackportsRequiredTags)
+		{
+			// If we still have it in this node
+			if (RequiredTags->Contains(Tag.Tag))
+			{
+				UCustomizableObjectNode* ParentNode = GetCustomizableObjectExternalNode<UCustomizableObjectNode>(Tag.ParentObject.Get(), Tag.ParentNode);
+				if (ParentNode)
+				{
+					TArray<FString>* NodeEnableTags = ParentNode->GetEnableTags();
+					if (NodeEnableTags)
+					{
+						NodeEnableTags->AddUnique(Tag.Tag);
+					}
+				}
+			}
+		}
+	}
 }
 
+
+FString UCustomizableObjectNodeModifierBase::MakeNodeAutoTag( UEdGraphNode* Node )
+{
+	const FString PackageName = Node->GetOutermost()->GetPathName();
+	const FString ShortName = FPackageName::GetShortName(PackageName);
+	FString NewLegacyTag = FString::Printf(TEXT("%s_%s"), *ShortName, *Node->NodeGuid.ToString());
+	return NewLegacyTag;
+}
diff --git a/Engine/Plugins/Experimental/Mutable/Source/CustomizableObjectEditor/Private/MuCOE/Nodes/CustomizableObjectNodeModifierClipMorphDetails.cpp b/Engine/Plugins/Experimental/Mutable/Source/CustomizableObjectEditor/Private/MuCOE/Nodes/CustomizableObjectNodeModifierClipMorphDetails.cpp
new file mode 100644
--- /dev/null
+++ b/Engine/Plugins/Experimental/Mutable/Source/CustomizableObjectEditor/Private/MuCOE/Nodes/CustomizableObjectNodeModifierClipMorphDetails.cpp
@@ -0,0 +1,352 @@
+// Copyright Epic Games, Inc. All Rights Reserved.
+
+#include "MuCOE/Nodes/CustomizableObjectNodeModifierClipMorphDetails.h"
+
+#include "DetailLayoutBuilder.h"
+#include "Engine/SkeletalMesh.h"
+#include "IDetailsView.h"
+#include "MuCOE/CustomizableObjectEditorUtilities.h"
+#include "MuCOE/CustomizableObjectGraph.h"
+#include "MuCOE/GraphTraversal.h"
+#include "MuCOE/Nodes/CustomizableObjectNodeModifierClipMorph.h"
+#include "MuCOE/UnrealEditorPortabilityHelpers.h"
+#include "PropertyCustomizationHelpers.h"
+#include "MuCO/CustomizableObjectPrivate.h"
+#include "Widgets/Input/SCheckBox.h"
+#include "Widgets/Input/STextComboBox.h"
+
+class UCustomizableObjectNodeObject;
+
+
+#define LOCTEXT_NAMESPACE "MeshClipMorphDetails"
+
+
+TSharedRef<IDetailCustomization> FCustomizableObjectNodeModifierClipMorphDetails::MakeInstance()
+{
+	return MakeShareable(new FCustomizableObjectNodeModifierClipMorphDetails);
+}
+
+
+void FCustomizableObjectNodeModifierClipMorphDetails::CustomizeDetails(IDetailLayoutBuilder& DetailBuilder)
+{
+	FCustomizableObjectNodeDetails::CustomizeDetails(DetailBuilder);
+
+	Node = nullptr;
+	DetailBuilderPtr = &DetailBuilder;
+
+	const IDetailsView* DetailsView = DetailBuilder.GetDetailsView();
+	if (DetailsView->GetSelectedObjects().Num())
+	{
+		Node = Cast<UCustomizableObjectNodeModifierClipMorph>(DetailsView->GetSelectedObjects()[0].Get());
+	}
+
+	IDetailCategoryBuilder& BlocksCategory = DetailBuilder.EditCategory("MeshToClipAndMorph");
+	DetailBuilder.HideProperty("BoneName");
+
+	IDetailCategoryBuilder& MeshClipParametersCategory = DetailBuilder.EditCategory("MeshClipParameters");
+	DetailBuilder.HideProperty("bInvertNormal");
+
+	TSharedPtr<IPropertyHandle> ReferenceSkeletonIndexPropertyHandle = DetailBuilder.GetProperty("ReferenceSkeletonIndex");
+	const FSimpleDelegate OnReferenceSkeletonIndexChangedDelegate = 
+		FSimpleDelegate::CreateSP(this, &FCustomizableObjectNodeModifierClipMorphDetails::OnReferenceSkeletonIndexChanged);
+	ReferenceSkeletonIndexPropertyHandle->SetOnPropertyValueChanged(OnReferenceSkeletonIndexChangedDelegate);
+
+	if (Node)
+	{
+		SkeletalMesh = nullptr;
+		TSharedPtr<FString> BoneToSelect;
+
+		UCustomizableObject* CustomizableObject = Cast<UCustomizableObject>(Node->GetCustomizableObjectGraph()->GetOuter());
+
+		if (CustomizableObject)
+		{
+			if (CustomizableObject->IsChildObject())
+			{
+				bool bMultipleBaseObjects;
+				UCustomizableObjectNodeObject* RootNode = GetRootNode(CustomizableObject, bMultipleBaseObjects);
+
+				if (RootNode && !bMultipleBaseObjects)
+				{
+					TArray<UCustomizableObject*> VisitedObjects;
+					UCustomizableObject* Parent = GetFullGraphRootObject(RootNode, VisitedObjects);
+
+					if (Parent)
+					{
+						SkeletalMesh = VisitedObjects.Last()->GetRefSkeletalMesh(Node->ReferenceSkeletonIndex);
+					}
+				}
+			}
+			else
+			{
+				SkeletalMesh = CustomizableObject->GetRefSkeletalMesh(Node->ReferenceSkeletonIndex);
+			}
+		}
+
+		if (SkeletalMesh)
+		{
+			BoneComboOptions.Empty();
+				
+			for (int32 i = 0; i < SkeletalMesh->GetRefSkeleton().GetRawBoneNum(); ++i)
+			{
+				FName BoneName = SkeletalMesh->GetRefSkeleton().GetBoneName(i);
+				BoneComboOptions.Add(MakeShareable(new FString(BoneName.ToString())));
+
+				if (BoneName == Node->BoneName)
+				{
+					BoneToSelect = BoneComboOptions.Last();
+				}
+			}
+
+            BoneComboOptions.Sort(CompareNames);
+
+			// Add them to the parent combo box
+			TSharedRef<IPropertyHandle> BoneProperty = DetailBuilder.GetProperty("BoneName");
+
+			BlocksCategory.AddCustomRow(LOCTEXT("ClipMorphDetails_BoneName", "Bone Name"))
+			[
+				SNew(SProperty, BoneProperty)
+				.ShouldDisplayName(false)
+				.CustomWidget()
+				[
+					SNew(SBorder)
+					.BorderImage(UE_MUTABLE_GET_BRUSH("NoBorder"))
+					.Padding(FMargin(0.0f, 0.0f, 10.0f, 0.0f))
+					.HAlign(EHorizontalAlignment::HAlign_Fill)
+					[
+						SNew(SHorizontalBox)
+						+SHorizontalBox::Slot().HAlign(HAlign_Left).VAlign(VAlign_Center)
+						[
+							SNew(STextBlock)
+							.Text( LOCTEXT("ClipMorphDetails_BoneName", "Bone Name"))
+							.Font(IDetailLayoutBuilder::GetDetailFont())
+						]
+						+SHorizontalBox::Slot().HAlign(HAlign_Fill)
+						[
+							SNew(STextComboBox)
+							.OptionsSource(&BoneComboOptions)
+							.InitiallySelectedItem(BoneToSelect)
+							.OnSelectionChanged(this, &FCustomizableObjectNodeModifierClipMorphDetails::OnBoneComboBoxSelectionChanged, BoneProperty)
+							.Font(IDetailLayoutBuilder::GetDetailFont())
+						]
+					]
+				]
+			];
+
+			MeshClipParametersCategory.AddCustomRow(LOCTEXT("ClipMorphDetails_PlaneNormal", "Invert plane normal"))
+			[
+				SNew(SProperty, BoneProperty)
+				.ShouldDisplayName(false)
+				.CustomWidget()
+				[
+					SNew(SBorder)
+					.BorderImage(UE_MUTABLE_GET_BRUSH("NoBorder"))
+					.Padding(FMargin(0.0f, 0.0f, 10.0f, 0.0f))
+					.HAlign(EHorizontalAlignment::HAlign_Fill)
+					[
+						SNew(SHorizontalBox)
+						+SHorizontalBox::Slot().HAlign(HAlign_Left).VAlign(VAlign_Center)
+						[
+							SNew(STextBlock)
+							.Text( LOCTEXT("ClipMorphDetails_PlaneNormal", "Invert plane normal"))
+						]
+						+SHorizontalBox::Slot().HAlign(HAlign_Left)
+						[
+							SNew(SCheckBox)
+							.OnCheckStateChanged(this, &FCustomizableObjectNodeModifierClipMorphDetails::OnInvertNormalCheckboxChanged)
+							.IsChecked(this, &FCustomizableObjectNodeModifierClipMorphDetails::GetInvertNormalCheckBoxState)
+							.ToolTipText(LOCTEXT("ClipMorphDetails_InvertNormal_Tooltip", "Invert normal direction of the clip plane"))
+						]
+					]
+				]
+			];
+		}
+	}
+	else
+	{
+		BlocksCategory.AddCustomRow(LOCTEXT("Node", "Node"))
+		[
+			SNew(STextBlock)
+			.Text(LOCTEXT("ClipMorphDetails_InvertNormal_NodeNotFound", "Node not found"))
+		];
+	}
+}
+
+
+FVector FindBoneLocation(int32 BoneIndex, USkeletalMesh* SkeletalMesh)
+{
+	const TArray<FTransform>& BoneArray = SkeletalMesh->GetRefSkeleton().GetRefBonePose();
+
+	int32 ParentIndex = BoneIndex;
+	FVector Location = FVector::ZeroVector;
+
+	while (ParentIndex >= 0)
+	{
+		Location = BoneArray[ParentIndex].TransformPosition(Location);
+		ParentIndex = SkeletalMesh->GetRefSkeleton().GetParentIndex(ParentIndex);
+	}
+
+	return Location;
+}
+
+
+void FCustomizableObjectNodeModifierClipMorphDetails::OnBoneComboBoxSelectionChanged(TSharedPtr<FString> Selection, ESelectInfo::Type SelectInfo, TSharedRef<IPropertyHandle> BoneProperty)
+{
+	for (int OptionIndex = 0; OptionIndex < BoneComboOptions.Num(); ++OptionIndex)
+	{
+		if (BoneComboOptions[OptionIndex] == Selection)
+		{
+			//TArray<FVector> AffectedBoneLocations;
+
+			//if (SkeletalMesh)
+			//{
+			//	const TArray<FTransform>& BoneArray = SkeletalMesh->RefSkeleton.GetRefBonePose();
+			//	int32 BoneIndex = SkeletalMesh->RefSkeleton.FindBoneIndex(FName(**Selection));
+			//	
+			//	if (BoneIndex >= 0)
+			//	{
+			//		// Get tranform from skeleton root to selected bone
+			//		int32 ParentIndex = BoneIndex;
+			//		FTransform RootToBoneTransform = FTransform::Identity;
+
+			//		while (ParentIndex >= 0)
+			//		{
+			//			RootToBoneTransform = RootToBoneTransform * BoneArray[ParentIndex];
+			//			ParentIndex = SkeletalMesh->RefSkeleton.GetParentIndex(ParentIndex);
+			//		}
+
+			//		// Now that we have the transform from root to the selected bone, go down from this bone to the end of the bone chain hanging from it building a bounding box
+			//		int32 CurrentBone = BoneIndex;
+
+			//		FVector SelectedBoneLocationWithTransform = RootToBoneTransform.TransformPosition(FVector::ZeroVector);
+			//		FVector SelectedBoneLocationFromSelection = BoneLocations[OptionIndex];
+			//		check(FVector::DistSquared(SelectedBoneLocationWithTransform, SelectedBoneLocationFromSelection) < 0.01f);
+
+			//		while (CurrentBone >= 0)
+			//		{
+			//			FVector CurrentBoneLocation = RootToBoneTransform.TransformPosition(FVector::ZeroVector);
+			//			AffectedBoneLocations.Add(CurrentBoneLocation);
+
+			//			check(FVector::DistSquared(CurrentBoneLocation, FindBoneLocation(CurrentBone, SkeletalMesh)) < 0.01f);
+
+			//			int32 NextBone = -1;
+
+			//			for (int32 i = 0; i < SkeletalMesh->RefSkeleton.GetRawBoneNum(); ++i)
+			//			{
+			//				if (SkeletalMesh->RefSkeleton.GetParentIndex(i) == CurrentBone)
+			//				{
+			//					NextBone = i;
+			//					break;
+			//				}
+			//			}
+
+			//			CurrentBone = NextBone;
+
+			//			if (CurrentBone >= 0)
+			//			{
+			//				RootToBoneTransform = BoneArray[CurrentBone] * RootToBoneTransform;
+			//			}
+			//		}
+			//	}
+			//}
+
+			FVector Location = FVector::ZeroVector;
+			FVector Direction = FVector::ForwardVector;
+
+			if (SkeletalMesh)
+			{
+				const TArray<FTransform>& BoneArray = SkeletalMesh->GetRefSkeleton().GetRefBonePose();
+				int32 ParentIndex = SkeletalMesh->GetRefSkeleton().FindBoneIndex(FName(**Selection));
+
+				FVector ChildLocation = FVector::ForwardVector;
+
+				for (int32 i = 0; i < SkeletalMesh->GetRefSkeleton().GetRawBoneNum(); ++i)
+				{
+					if (SkeletalMesh->GetRefSkeleton().GetParentIndex(i) == ParentIndex)
+					{
+						ChildLocation = BoneArray[i].TransformPosition(FVector::ZeroVector);
+						break;
+					}
+				}
+
+				//Direction = ChildLocation;
+
+				while (ParentIndex >= 0)
+				{
+					Location = BoneArray[ParentIndex].TransformPosition(Location);
+					ChildLocation = BoneArray[ParentIndex].TransformPosition(ChildLocation);
+					//Direction = BoneArray[ParentIndex].TransformVector(Direction);
+					ParentIndex = SkeletalMesh->GetRefSkeleton().GetParentIndex(ParentIndex);
+				}
+
+				Direction = (ChildLocation - Location).GetSafeNormal();
+			}
+			 
+			BoneProperty->SetValue(*BoneComboOptions[OptionIndex].Get());
+			
+			if (Node)
+			{
+				Node->Origin = Location;
+				Node->Normal = Direction;
+			}
+
+			return;
+		}
+	}
+
+	BoneProperty->SetValue(FString());
+}
+
+
+void FCustomizableObjectNodeModifierClipMorphDetails::OnInvertNormalCheckboxChanged(ECheckBoxState CheckBoxState)
+{
+	if (Node == nullptr)
+	{
+		return;
+	}
+
+	switch (CheckBoxState)
+	{
+		case ECheckBoxState::Checked:
+		{
+			Node->bInvertNormal = true;
+			Node->Normal *= -1.0f;
+			break;
+		}
+		case ECheckBoxState::Unchecked:
+		{
+			Node->bInvertNormal = false;
+			Node->Normal *= -1.0f;
+			break;
+		}
+	}
+
+	if (Node->bLocalStartOffset)
+	{
+		Node->StartOffset.Z *= -1;
+		Node->StartOffset.X *= -1;
+	}
+}
+
+
+ECheckBoxState FCustomizableObjectNodeModifierClipMorphDetails::GetInvertNormalCheckBoxState() const
+{
+	if (Node == nullptr)
+	{
+		return ECheckBoxState::Unchecked;
+	}
+
+	return Node->bInvertNormal ? ECheckBoxState::Checked : ECheckBoxState::Unchecked;
+}
+
+
+void FCustomizableObjectNodeModifierClipMorphDetails::OnReferenceSkeletonIndexChanged()
+{
+	if (Node)
+	{
+		Node->BoneName = FName();
+	}
+	DetailBuilderPtr->ForceRefreshDetails();
+}
+
+
+#undef LOCTEXT_NAMESPACE
diff --git a/Engine/Plugins/Experimental/Mutable/Source/CustomizableObjectEditor/Private/MuCOE/Nodes/CustomizableObjectNodeModifierClipMorphDetails.h b/Engine/Plugins/Experimental/Mutable/Source/CustomizableObjectEditor/Private/MuCOE/Nodes/CustomizableObjectNodeModifierClipMorphDetails.h
new file mode 100644
--- /dev/null
+++ b/Engine/Plugins/Experimental/Mutable/Source/CustomizableObjectEditor/Private/MuCOE/Nodes/CustomizableObjectNodeModifierClipMorphDetails.h
@@ -0,0 +1,40 @@
+// Copyright Epic Games, Inc. All Rights Reserved.
+
+#pragma once
+
+#include "CustomizableObjectNodeDetails.h"
+#include "IDetailCustomization.h"
+#include "UObject/ObjectMacros.h"
+
+enum class ECheckBoxState : uint8;
+namespace ESelectInfo { enum Type : int; }
+
+class FString;
+class IDetailLayoutBuilder;
+class IPropertyHandle;
+
+
+class FCustomizableObjectNodeModifierClipMorphDetails : public FCustomizableObjectNodeDetails
+{
+public:
+	/** Makes a new instance of this detail layout class for a specific detail view requesting it */
+	static TSharedRef<IDetailCustomization> MakeInstance();
+
+	/** ILayoutDetails interface */
+	virtual void CustomizeDetails( IDetailLayoutBuilder& DetailBuilder ) override;
+
+private:
+	class UCustomizableObjectNodeModifierClipMorph* Node;
+	class IDetailLayoutBuilder* DetailBuilderPtr = nullptr;
+	TArray< TSharedPtr<FString> > BoneComboOptions;
+
+	UPROPERTY()
+	class USkeletalMesh* SkeletalMesh;
+
+	void OnBoneComboBoxSelectionChanged(TSharedPtr<FString> Selection, ESelectInfo::Type SelectInfo, TSharedRef<IPropertyHandle> BoneProperty);
+
+	void OnInvertNormalCheckboxChanged(ECheckBoxState CheckBoxState);
+	ECheckBoxState GetInvertNormalCheckBoxState() const;
+
+	void OnReferenceSkeletonIndexChanged();
+};
diff --git a/Engine/Plugins/Experimental/Mutable/Source/CustomizableObjectEditor/Private/MuCOE/Nodes/CustomizableObjectNodeModifierClipWithMeshDetails.cpp b/Engine/Plugins/Experimental/Mutable/Source/CustomizableObjectEditor/Private/MuCOE/Nodes/CustomizableObjectNodeModifierClipWithMeshDetails.cpp
new file mode 100644
--- /dev/null
+++ b/Engine/Plugins/Experimental/Mutable/Source/CustomizableObjectEditor/Private/MuCOE/Nodes/CustomizableObjectNodeModifierClipWithMeshDetails.cpp
@@ -0,0 +1,337 @@
+// Copyright Epic Games, Inc. All Rights Reserved.
+
+#include "MuCOE/Nodes/CustomizableObjectNodeModifierClipWithMeshDetails.h"
+
+#include "DetailLayoutBuilder.h"
+#include "IDetailsView.h"
+#include "MuCOE/CustomizableObjectEditorUtilities.h"
+#include "MuCOE/GraphTraversal.h"
+#include "MuCOE/UnrealEditorPortabilityHelpers.h"
+#include "MuCOE/Nodes/CustomizableObjectNodeMaterial.h"
+#include "MuCOE/Nodes/CustomizableObjectNodeModifierClipWithMesh.h"
+#include "MuCOE/Nodes/CustomizableObjectNodeStaticMesh.h"
+#include "PropertyCustomizationHelpers.h"
+#include "ScopedTransaction.h"
+#include "MuCO/CustomizableObjectPrivate.h"
+#include "Widgets/Input/SButton.h"
+#include "Widgets/Input/SEditableTextBox.h"
+#include "Widgets/Input/STextComboBox.h"
+
+class IPropertyHandle;
+
+
+#define LOCTEXT_NAMESPACE "CustomizableObjectNodeModifierClipWithMeshDetails"
+
+
+TSharedRef<IDetailCustomization> FCustomizableObjectNodeModifierClipWithMeshDetails::MakeInstance()
+{
+	return MakeShareable(new FCustomizableObjectNodeModifierClipWithMeshDetails);
+}
+
+
+void FCustomizableObjectNodeModifierClipWithMeshDetails::CustomizeDetails(IDetailLayoutBuilder& DetailBuilder)
+{
+	FCustomizableObjectNodeDetails::CustomizeDetails(DetailBuilder);
+
+	Node = nullptr;
+	DetailBuilderPtr = &DetailBuilder;
+
+	const IDetailsView* DetailsView = DetailBuilder.GetDetailsView();
+	if (DetailsView->GetSelectedObjects().Num())
+	{
+		Node = Cast<UCustomizableObjectNodeModifierClipWithMesh>(DetailsView->GetSelectedObjects()[0].Get());
+	}
+
+	IDetailCategoryBuilder& CustomizableObjectToClipCategory = DetailBuilder.EditCategory("CustomizableObjectToClip");
+
+	if(Node != nullptr)
+	{
+		SelectedCO = Node->CustomizableObjectToClipWith;
+
+		TSharedRef<IPropertyHandle> ParentObjectProperty = DetailBuilder.GetProperty("ParentObject");
+
+		// Hidding unnecessary properties
+		DetailBuilder.HideProperty("CustomizableObjectToClipWith");
+		DetailBuilder.HideProperty("ArrayMaterialNodeToClipWithID");
+
+		if (const UEdGraphPin* ConnectedPin = FollowInputPin(*Node->ClipMeshPin());
+			ConnectedPin && !Cast<UCustomizableObjectNodeStaticMesh>(ConnectedPin->GetOwningNode()))
+		{
+			DetailBuilder.HideProperty("Transform");
+		}
+
+		ArrayMaterialNodeOption.Empty();
+		ArrayMaterialNodeName.Empty();
+
+		TArray<FGuid> ArrayMaterialNodeGuid;
+		TSharedPtr<FString> SelectedMethod;
+
+		CustomizableObjectToClipCategory.AddCustomRow(LOCTEXT("FCustomizableObjectNodeModifierClipWithMeshDetails", "Blocks"))
+		[
+			SNew(SHorizontalBox)
+			+ SHorizontalBox::Slot()
+			.Padding(FMargin(0.0f, 3.0f, 0.0f, 10.0f))
+			.FillWidth(0.15f)
+			[
+				SNew(STextBlock)
+				.Text(LOCTEXT("CustomizableObjectText", "Object:"))
+			]
+			+ SHorizontalBox::Slot()
+			.FillWidth(0.85f)
+			[
+				SNew(SObjectPropertyEntryBox)
+				.AllowedClass(UCustomizableObject::StaticClass())
+				.OnObjectChanged(this, &FCustomizableObjectNodeModifierClipWithMeshDetails::ClipperCustomizableObjectSelectionChanged)
+				.ObjectPath(SelectedCO != nullptr ? SelectedCO->GetPathName() : "")
+				.ForceVolatile(true)
+			]
+		];
+
+		SelectAMethod = MakeShareable(new FString("Select a Method"));
+
+		// Filling ComboBox options
+		ClippingMethods.Empty();
+		ClippingMethods.Add(MakeShareable(new FString("Material")));
+		ClippingMethods.Add(MakeShareable(new FString("Tags")));
+		ClippingMethods.Add(MakeShareable(new FString("Material & Tags")));
+
+		// Initializing the inital value of the ComboBox
+		SelectedMethod = SetInitialClippingMethod();
+		
+		// This flag is now up to date and can be used to hide ui.
+		if (!Node->bUseTags)
+		{
+			DetailBuilder.HideProperty("Tags");
+		}
+
+		// ComboBox to select the Clipping method (material or/and tags)
+		CustomizableObjectToClipCategory.AddCustomRow(LOCTEXT("FCustomizableObjectNodeModifierClipWithMeshDetails", "Blocks"))
+		[
+			SNew(SHorizontalBox)
+			+SHorizontalBox::Slot()
+			.AutoWidth()
+			.Padding(FMargin(0.0f, 8.0f, 0.0f, 5.0f))
+			[
+				SNew(STextBlock)
+				.Text(LOCTEXT("ClippingMethod", "Clipping Method:"))
+			]
+			
+			+ SHorizontalBox::Slot()
+			.AutoWidth()
+			.Padding(FMargin(4.0f, 5.0f, 0.0f, 5.0f))
+			[
+				SNew(STextComboBox)
+				.OptionsSource(&ClippingMethods)
+				.InitiallySelectedItem(SelectedMethod)
+				.OnSelectionChanged(this, &FCustomizableObjectNodeModifierClipWithMeshDetails::OnClippingMethodComboBoxSelectionChanged)
+			]
+		];
+
+
+		// Filling materials ComboBox with all the possible materials
+		int32 NumElementBeforeAdd;
+		int32 NumElementAfterAdd;
+		TSharedPtr<FString> ItemToSelect;
+		TArray<UCustomizableObjectNodeMaterial*> ArrayMaterialNodes;
+		if (SelectedCO != nullptr)
+		{
+			SelectedCO->GetPrivate()->GetSource()->GetNodesOfClass<UCustomizableObjectNodeMaterial>(ArrayMaterialNodes);
+		}
+
+		for (const UCustomizableObjectNodeMaterial* MaterialNode : ArrayMaterialNodes)
+		{
+			if (MaterialNode->GetMaterial())
+			{
+				NumElementBeforeAdd = ArrayMaterialNodeName.Num();
+				ArrayMaterialNodeName.AddUnique(MaterialNode->GetMaterial()->GetName());
+				NumElementAfterAdd = ArrayMaterialNodeName.Num();
+
+				if (NumElementBeforeAdd != NumElementAfterAdd)
+				{
+					ArrayMaterialNodeGuid.Add(MaterialNode->NodeGuid);
+				}
+			}
+		}
+
+		// Adding a none option
+		ArrayMaterialNodeOption.Add(MakeShareable(new FString("None")));
+		ItemToSelect = ArrayMaterialNodeOption.Last();
+
+		// Initializing the Selected material if it was already selected
+		const int32 MaxIndex = ArrayMaterialNodeName.Num();
+		for (int32 i = 0; i < MaxIndex; ++i)
+		{
+			ArrayMaterialNodeOption.Add(MakeShareable(new FString(ArrayMaterialNodeName[i])));
+
+			if (Node->ArrayMaterialNodeToClipWithID.Find(ArrayMaterialNodeGuid[i]) != INDEX_NONE)
+			{
+				ItemToSelect = ArrayMaterialNodeOption.Last();
+			}
+		}
+
+        ArrayMaterialNodeOption.Sort(CompareNames);
+
+		// Widget that shows the Tags widget or/and the materials ComboBox
+		if (Node->bUseMaterials)
+		{
+			CustomizableObjectToClipCategory.AddCustomRow(LOCTEXT("FCustomizableObjectNodeModifierClipWithMeshDetails", "Blocks"))
+			[
+				SNew(SVerticalBox)
+				+SVerticalBox::Slot()
+				.AutoHeight()
+				[
+					SNew(SHorizontalBox)
+					+ SHorizontalBox::Slot()
+					.FillWidth(0.15f)
+					.Padding(0.0f, 3.0f, 0.0f, 0.0f)
+					[
+						SNew(STextBlock)
+						.Text(LOCTEXT("MaterialsText", "Material:"))
+					]
+					+ SHorizontalBox::Slot()
+					.FillWidth(0.85f)
+					.Padding(2.0f, 0.0f, 0.0f, 0.0f)
+					[
+						SNew(STextComboBox)
+						.OptionsSource(&ArrayMaterialNodeOption)
+						.InitiallySelectedItem(ItemToSelect)
+						.OnSelectionChanged(this, &FCustomizableObjectNodeModifierClipWithMeshDetails::OnMeshClipWithMeshNodeComboBoxSelectionChanged)
+					]
+				]
+			];
+		}
+	}
+	else
+	{
+		CustomizableObjectToClipCategory.AddCustomRow(LOCTEXT("Node", "Node"))
+			[
+				SNew(STextBlock)
+				.Text(LOCTEXT("Node not found", "Node not found"))
+			];
+	}
+}
+
+
+void FCustomizableObjectNodeModifierClipWithMeshDetails::ClipperCustomizableObjectSelectionChanged(const FAssetData& AssetData)
+{
+	UCustomizableObject* NewClipperCO = Cast<UCustomizableObject>(AssetData.GetAsset());
+
+	if (Node != nullptr)
+	{
+		Node->CustomizableObjectToClipWith = NewClipperCO;
+		Node->ArrayMaterialNodeToClipWithID.Empty();
+	}
+
+	if (DetailBuilderPtr)
+	{
+		DetailBuilderPtr->ForceRefreshDetails();
+	}
+}
+
+
+TSharedPtr<FString> FCustomizableObjectNodeModifierClipWithMeshDetails::SetInitialClippingMethod()
+{
+	TSharedPtr<FString> InitialMethod = SelectAMethod;
+
+	if (!Node->bUseMaterials && !Node->bUseTags)
+	{
+		if (Node->ArrayMaterialNodeToClipWithID.Num() > 0)
+		{
+			InitialMethod = ClippingMethods[0];
+			Node->bUseMaterials = true;
+		}
+
+		if (Node->Tags.Num() > 0)
+		{
+			InitialMethod = ClippingMethods[1];
+			Node->bUseTags = true;
+		}
+
+		if (Node->bUseMaterials && Node->bUseTags)
+		{
+			InitialMethod = ClippingMethods[2];
+		}
+	}
+	else
+	{
+		if (Node->bUseMaterials && !Node->bUseTags)
+		{
+			InitialMethod = ClippingMethods[0];
+		}
+		else if (!Node->bUseMaterials && Node->bUseTags)
+		{
+			InitialMethod = ClippingMethods[1];
+		}
+		else
+		{
+			InitialMethod = ClippingMethods[2];
+		}
+	}
+
+	return InitialMethod;
+}
+
+
+void FCustomizableObjectNodeModifierClipWithMeshDetails::OnClippingMethodComboBoxSelectionChanged(TSharedPtr<FString> Selection, ESelectInfo::Type SelectInfo)
+{
+	if (Selection.IsValid())
+	{
+		if (*Selection == "Tags")
+		{
+			Node->bUseMaterials = false;
+			Node->bUseTags = true;
+			Node->ArrayMaterialNodeToClipWithID.Empty();
+		}
+		else if(*Selection == "Material")
+		{
+			Node->bUseMaterials = true;
+			Node->bUseTags = false;
+			Node->Tags.Empty();
+		}
+		else
+		{
+			Node->bUseMaterials = true;
+			Node->bUseTags = true;
+		}
+
+		DetailBuilderPtr->ForceRefreshDetails();
+	}	
+}
+
+
+void FCustomizableObjectNodeModifierClipWithMeshDetails::OnMeshClipWithMeshNodeComboBoxSelectionChanged(TSharedPtr<FString> Selection, ESelectInfo::Type SelectInfo)
+{
+	if (Selection.IsValid() && (SelectedCO != nullptr))
+	{
+		TArray<UCustomizableObjectNodeMaterial*> ArrayMaterialNode;
+		SelectedCO->GetPrivate()->GetSource()->GetNodesOfClass<UCustomizableObjectNodeMaterial>(ArrayMaterialNode);
+
+		const FScopedTransaction Transaction(LOCTEXT("ChangedClipMeshWithMeshMaterialTransaction", "Changed Clip Mesh With Mesh Material"));
+		Node->Modify();
+		Node->ArrayMaterialNodeToClipWithID.Empty();
+		for (const UCustomizableObjectNodeMaterial* MaterialNode : ArrayMaterialNode)
+		{
+			if (MaterialNode && MaterialNode->GetMaterial() && *Selection == MaterialNode->GetMaterial()->GetName())
+			{
+				Node->ArrayMaterialNodeToClipWithID.Add(MaterialNode->NodeGuid);
+			}
+		}
+	}
+}
+
+
+FReply FCustomizableObjectNodeModifierClipWithMeshDetails::OnAddTagPressed()
+{
+	if (Node)
+	{
+		Node->Tags.Add("");
+		DetailBuilderPtr->ForceRefreshDetails();
+
+		return FReply::Handled();
+	}
+
+	return FReply::Unhandled();
+}
+
+
+#undef LOCTEXT_NAMESPACE
diff --git a/Engine/Plugins/Experimental/Mutable/Source/CustomizableObjectEditor/Private/MuCOE/Nodes/CustomizableObjectNodeModifierClipWithMeshDetails.h b/Engine/Plugins/Experimental/Mutable/Source/CustomizableObjectEditor/Private/MuCOE/Nodes/CustomizableObjectNodeModifierClipWithMeshDetails.h
new file mode 100644
--- /dev/null
+++ b/Engine/Plugins/Experimental/Mutable/Source/CustomizableObjectEditor/Private/MuCOE/Nodes/CustomizableObjectNodeModifierClipWithMeshDetails.h
@@ -0,0 +1,65 @@
+// Copyright Epic Games, Inc. All Rights Reserved.
+
+#pragma once
+
+#include "CustomizableObjectNodeDetails.h"
+#include "IDetailCustomization.h"
+#include "Widgets/SCompoundWidget.h"
+
+namespace ESelectInfo { enum Type : int; }
+namespace ETextCommit { enum Type : int; }
+
+class FText;
+class IDetailLayoutBuilder;
+class SHorizontalBox;
+class SVerticalBox;
+class UCustomizableObjectNodeModifierClipWithMesh;
+struct FAssetData;
+
+
+class FCustomizableObjectNodeModifierClipWithMeshDetails : public FCustomizableObjectNodeDetails
+{
+public:
+	/** Makes a new instance of this detail layout class for a specific detail view requesting it */
+	static TSharedRef<IDetailCustomization> MakeInstance();
+
+	/** ILayoutDetails interface */
+	virtual void CustomizeDetails( IDetailLayoutBuilder& DetailBuilder ) override;
+
+	/** Callback when new CO is selected */
+	void ClipperCustomizableObjectSelectionChanged(const FAssetData& AssetData);
+
+private:
+	/** Callback when the clipping method is selected */
+	void OnClippingMethodComboBoxSelectionChanged(TSharedPtr<FString> Selection, ESelectInfo::Type SelectInfo);
+
+	/** Callback when new material name is selected */
+	void OnMeshClipWithMeshNodeComboBoxSelectionChanged(TSharedPtr<FString> Selection, ESelectInfo::Type SelectInfo);
+
+	/** Initializes the text of the method combobox and the view to display */
+	TSharedPtr<FString> SetInitialClippingMethod();
+
+	/** The node for which details are being customized */
+	class UCustomizableObjectNodeModifierClipWithMesh* Node;
+
+	/** Pointer to the builder passed by CustomizeDetails method */
+	IDetailLayoutBuilder* DetailBuilderPtr = nullptr;
+
+	/** References to the names of all possible methods to select the material */
+	TArray<TSharedPtr<FString>> ClippingMethods;
+
+	/** References to the names of all possible selectable material nodes in SelectedCO */
+	TArray<TSharedPtr<FString>> ArrayMaterialNodeOption;
+
+	/** To cover all LODs from a single selection, only unique names of materials are considered */
+	TArray<FString> ArrayMaterialNodeName;
+
+	/** Selected CO in the details tab of the Node node, for which to select a material that will be clipped by Node node */
+	class UCustomizableObject* SelectedCO;
+
+	/** Pointer to a FString to indicate that a method must be selected */
+	TSharedPtr<FString> SelectAMethod;
+
+	/** Callback to add a tag to the arry of tags */
+	FReply OnAddTagPressed();
+};
diff --git a/Engine/Plugins/Experimental/Mutable/Source/CustomizableObjectEditor/Private/MuCOE/Nodes/CustomizableObjectNodeModifierClipWithUVMask.cpp b/Engine/Plugins/Experimental/Mutable/Source/CustomizableObjectEditor/Private/MuCOE/Nodes/CustomizableObjectNodeModifierClipWithUVMask.cpp
--- a/Engine/Plugins/Experimental/Mutable/Source/CustomizableObjectEditor/Private/MuCOE/Nodes/CustomizableObjectNodeModifierClipWithUVMask.cpp
+++ b/Engine/Plugins/Experimental/Mutable/Source/CustomizableObjectEditor/Private/MuCOE/Nodes/CustomizableObjectNodeModifierClipWithUVMask.cpp
@@ -42,12 +42,6 @@ void UCustomizableObjectNodeModifierClipWithUVMask::PinConnectionListChanged(UEd
 }
 
 
-UEdGraphPin* UCustomizableObjectNodeModifierClipWithUVMask::OutputPin() const
-{
-	return FindPin(TEXT("Modifier"));
-}
-
-
 UEdGraphPin* UCustomizableObjectNodeModifierClipWithUVMask::ClipMaskPin() const
 {
 	return FindPin(TEXT("Clip Mask"));
diff --git a/Engine/Plugins/Experimental/Mutable/Source/CustomizableObjectEditor/Private/MuCOE/Nodes/CustomizableObjectNodeModifierEditMeshSection.cpp b/Engine/Plugins/Experimental/Mutable/Source/CustomizableObjectEditor/Private/MuCOE/Nodes/CustomizableObjectNodeModifierEditMeshSection.cpp
new file mode 100644
--- /dev/null
+++ b/Engine/Plugins/Experimental/Mutable/Source/CustomizableObjectEditor/Private/MuCOE/Nodes/CustomizableObjectNodeModifierEditMeshSection.cpp
@@ -0,0 +1,312 @@
+// Copyright Epic Games, Inc. All Rights Reserved.
+
+#include "MuCOE/Nodes/CustomizableObjectNodeModifierEditMeshSection.h"
+
+#include "MuCOE/Nodes/CustomizableObjectNodeModifierExtendMeshSection.h"
+#include "MuCOE/CustomizableObjectEditor_Deprecated.h"
+#include "MuCOE/CustomizableObjectLayout.h"
+#include "MuCOE/EdGraphSchema_CustomizableObject.h"
+#include "MuCOE/CustomizableObjectPin.h"
+#include "MuCO/CustomizableObjectCustomVersion.h"
+#include "MaterialTypes.h"
+
+class FCustomizableObjectNodeParentedMaterial;
+class UCustomizableObjectNodeRemapPins;
+
+#define LOCTEXT_NAMESPACE "CustomizableObjectEditor"
+
+
+void UCustomizableObjectNodeModifierEditMeshSection::AllocateDefaultPins(UCustomizableObjectNodeRemapPins* RemapPins)
+{
+	const UEdGraphSchema_CustomizableObject* Schema = GetDefault<UEdGraphSchema_CustomizableObject>();
+	
+	const int32 NumImages = GetNumParameters(EMaterialParameterType::Texture);
+	for (int32 ImageIndex = 0; ImageIndex < NumImages; ++ImageIndex)
+	{
+		{
+			UCustomizableObjectNodeEditMaterialPinEditImageData* PinEditImageData = NewObject<UCustomizableObjectNodeEditMaterialPinEditImageData>(this);
+			PinEditImageData->ImageParamId = GetParameterId(EMaterialParameterType::Texture, ImageIndex);
+			
+			const FName ImageName = GetParameterName(EMaterialParameterType::Texture, ImageIndex);
+			UEdGraphPin* PinImage = CustomCreatePin(EGPD_Input, Schema->PC_Image, ImageName, PinEditImageData);
+			PinImage->bHidden = true;
+			PinImage->bDefaultValueIsIgnored = true;
+
+			PinsParameterMap.Add(PinEditImageData->ImageParamId, FEdGraphPinReference(PinImage));
+
+			FString PinMaskName = ImageName.ToString() + FString(TEXT(" Mask"));
+			UEdGraphPin* PinMask = CustomCreatePin(EGPD_Input, Schema->PC_Image, *PinMaskName);
+			PinMask->bHidden = true;
+			PinMask->bDefaultValueIsIgnored = true;
+			
+			PinEditImageData->PinMask = FEdGraphPinReference(PinMask);
+		}
+	}
+
+	CustomCreatePin(EGPD_Output, Schema->PC_Modifier, TEXT("Modifier"));
+}
+
+
+const UEdGraphPin* UCustomizableObjectNodeModifierEditMeshSection::GetUsedImageMaskPin(const FNodeMaterialParameterId& ImageId) const
+{
+	if (const UEdGraphPin* Pin = GetUsedImagePin(ImageId))
+	{
+		UCustomizableObjectNodeEditMaterialPinEditImageData& PinData = GetPinData<UCustomizableObjectNodeEditMaterialPinEditImageData>(*Pin);
+		return PinData.PinMask.Get();
+	}
+
+	return nullptr;
+}
+
+
+bool UCustomizableObjectNodeModifierEditMeshSection::IsSingleOutputNode() const
+{
+	return true;
+}
+
+
+bool UCustomizableObjectNodeModifierEditMeshSection::CustomRemovePin(UEdGraphPin& Pin)
+{
+	for (TMap<FNodeMaterialParameterId, FEdGraphPinReference>::TIterator It = PinsParameterMap.CreateIterator(); It; ++It)
+	{
+		if (It.Value().Get() == &Pin) // We could improve performance if FEdGraphPinReference exposed the pin id.
+		{
+			It.RemoveCurrent();
+			break;
+		}
+	}
+
+	return Super::CustomRemovePin(Pin);
+}
+
+
+bool UCustomizableObjectNodeModifierEditMeshSection::HasPinViewer() const
+{
+	return true;
+}
+
+
+void UCustomizableObjectNodeModifierEditMeshSection::BackwardsCompatibleFixup(int32 CustomizableObjectCustomVersion)
+{
+	Super::BackwardsCompatibleFixup(CustomizableObjectCustomVersion);
+
+	// Convert deprecated node index list to the node id list.
+	if (CustomizableObjectCustomVersion == FCustomizableObjectCustomVersion::PostLoadToCustomVersion
+		&& BlockIds_DEPRECATED.Num() < Blocks_DEPRECATED.Num())
+	{
+		UCustomizableObjectNodeMaterialBase* ParentMaterialNode = GetCustomizableObjectExternalNode<UCustomizableObjectNodeMaterialBase>(ParentMaterialObject_DEPRECATED.Get(), ParentMaterialNodeId_DEPRECATED);
+
+		TArray<UCustomizableObjectLayout*> Layouts = ParentMaterialNode->GetLayouts();
+
+		if (!Layouts.IsValidIndex(ParentLayoutIndex))
+		{
+			UE_LOG(LogMutable, Warning, TEXT("[%s] UCustomizableObjectNodeModifierEditMeshSection refers to an invalid texture layout index %d. Parent node has %d layouts."), 
+				*GetOutermost()->GetName(), ParentLayoutIndex, Layouts.Num());
+		}
+		else
+		{
+			UCustomizableObjectLayout* ParentLayout = Layouts[ParentLayoutIndex];
+
+			if (Cast<UCustomizableObjectNodeMaterial>(ParentMaterialNode))
+			{
+				for (int32 IndexIndex = BlockIds_DEPRECATED.Num(); IndexIndex < Blocks_DEPRECATED.Num(); ++IndexIndex)
+				{
+					int32 BlockIndex = Blocks_DEPRECATED[IndexIndex];
+					if (ParentLayout->Blocks.IsValidIndex(BlockIndex) )
+					{
+						const FGuid Id = ParentLayout->Blocks[BlockIndex].Id;
+						if (Id.IsValid())
+						{
+							BlockIds_DEPRECATED.Add(Id);
+						}
+						else
+						{
+							UE_LOG(LogMutable, Warning, TEXT("[%s] UCustomizableObjectNodeModifierEditMeshSection refers to an valid layout block %d but that block doesn't have an id."),
+								*GetOutermost()->GetName(), BlockIndex );
+						}
+					}
+					else
+					{
+						UE_LOG(LogMutable, Warning, TEXT("[%s] UCustomizableObjectNodeModifierEditMeshSection refers to an invalid layout block index %d. Parent node has %d blocks."),
+							*GetOutermost()->GetName(), BlockIndex, ParentLayout->Blocks.Num());
+					}
+				}
+			}
+		}
+	}
+	
+	// Convert deprecated node id list to absolute rect list.
+	if (CustomizableObjectCustomVersion == FCustomizableObjectCustomVersion::UseUVRects)
+	{
+		// If we are here, it means this node was loaded from a version that didn't have it's own layout.
+		check(Layout->Blocks.IsEmpty());
+
+		UCustomizableObjectNodeMaterialBase* ParentMaterialNode = GetCustomizableObjectExternalNode<UCustomizableObjectNodeMaterialBase>(ParentMaterialObject_DEPRECATED.Get(), ParentMaterialNodeId_DEPRECATED);
+
+		TArray<UCustomizableObjectLayout*> ParentLayouts = ParentMaterialNode->GetLayouts();
+
+		if (!ParentLayouts.IsValidIndex(ParentLayoutIndex))
+		{
+			UE_LOG(LogMutable, Warning, TEXT("[%s] UCustomizableObjectNodeModifierEditMeshSection refers to an invalid texture layout index %d. Parent node has %d layouts."),
+				*GetOutermost()->GetName(), ParentLayoutIndex, ParentLayouts.Num());
+		}
+		else
+		{
+			UCustomizableObjectLayout* ParentLayout = ParentLayouts[ParentLayoutIndex];
+			FIntPoint GridSize = ParentLayout->GetGridSize();
+
+			Layout->SetGridSize(GridSize);
+
+			if (Cast<UCustomizableObjectNodeMaterial>(ParentMaterialNode))
+			{
+				for (const FGuid& BlockId : BlockIds_DEPRECATED)
+				{
+					bool bSkipBlock = false;
+					for (const FCustomizableObjectLayoutBlock& ExistingBlock : Layout->Blocks)
+					{
+						if (ExistingBlock.Id == BlockId)
+						{
+							bSkipBlock = true;
+							UE_LOG(LogMutable, Log, TEXT("[%s] UCustomizableObjectNodeModifierEditMeshSection has a duplicated layout block id. One has been ignored during version upgrade."), *GetOutermost()->GetName());
+							break;
+						}
+					}
+
+					if (bSkipBlock)
+					{
+						continue;
+					}
+
+					bool bFoundInParent = false;
+					for (const FCustomizableObjectLayoutBlock& ParentBlock : ParentLayout->Blocks)
+					{
+						if (ParentBlock.Id == BlockId)
+						{
+							bFoundInParent = true;
+
+							FCustomizableObjectLayoutBlock NewBlock;
+							NewBlock = ParentBlock;
+
+							// Clear some unnecessary data.
+							NewBlock.bReduceBothAxes = false;
+							NewBlock.bReduceByTwo = false;
+							NewBlock.Priority = 0;
+
+							Layout->Blocks.Add(NewBlock);
+							break;
+						}
+					}
+
+					if (!bFoundInParent)
+					{
+						UE_LOG(LogMutable, Warning, TEXT("[%s] UCustomizableObjectNodeModifierEditMeshSection refers to and invalid layout block. It has been ignored during version upgrade."), *GetOutermost()->GetName());
+					}
+				}
+			}
+		}
+	}
+
+	if (CustomizableObjectCustomVersion == FCustomizableObjectCustomVersion::AutomaticNodeMaterial)
+	{
+		UCustomizableObjectNodeMaterialBase* ParentMaterial = GetCustomizableObjectExternalNode<UCustomizableObjectNodeMaterialBase>(ParentMaterialObject_DEPRECATED.Get(), ParentMaterialNodeId_DEPRECATED);
+
+		if (ParentMaterial)
+		{
+			for (const FCustomizableObjectNodeEditMaterialImage& Image : Images_DEPRECATED)
+			{
+				const UEdGraphPin* ImagePin = FindPin(Image.Name);
+				const UEdGraphPin* PinMask = FindPin(Image.Name + " Mask");
+				if (!ImagePin || !PinMask)
+				{
+					continue;
+				}
+				
+				UCustomizableObjectNodeEditMaterialPinEditImageData*  PinEditImageData = NewObject<UCustomizableObjectNodeEditMaterialPinEditImageData>(this);
+				PinEditImageData->ImageId_DEPRECATED = FGuid::NewGuid();
+				PinEditImageData->PinMask = PinMask;
+				
+				// Search for the Image Id the Edit pin was referring to.
+				const int32 NumImages = ParentMaterial->GetNumParameters(EMaterialParameterType::Texture);
+				for (int32 ImageIndex = 0; ImageIndex < NumImages; ++ImageIndex)
+				{
+					if (ParentMaterial->GetParameterName(EMaterialParameterType::Texture, ImageIndex).ToString() == Image.Name)
+					{
+						PinEditImageData->ImageId_DEPRECATED = ParentMaterial->GetParameterId(EMaterialParameterType::Texture, ImageIndex).ParameterId;
+						break;
+					}
+				}
+				
+				AddPinData(*ImagePin, *PinEditImageData);
+			}
+		}
+		
+		Images_DEPRECATED.Empty();
+	}
+
+	// Fill PinsParameter.
+	if (CustomizableObjectCustomVersion == FCustomizableObjectCustomVersion::AutomaticNodeMaterialPerformanceBug) // || CustomizableObjectCustomVersion < FCustomizableObjectCustomVersion::AutomaticNodeMaterialPerformance
+	{
+		for (const UEdGraphPin* Pin : GetAllNonOrphanPins())
+		{
+			if (const UCustomizableObjectNodeEditMaterialPinEditImageData* PinData = Cast<UCustomizableObjectNodeEditMaterialPinEditImageData>(GetPinData(*Pin)))
+			{
+				PinsParameter_DEPRECATED.Add(PinData->ImageId_DEPRECATED, FEdGraphPinReference(Pin));
+			}
+		}
+	}
+
+	if (CustomizableObjectCustomVersion == FCustomizableObjectCustomVersion::ConvertEditAndExtendToModifiers)
+	{
+		// Look for the parent material and set it as the modifier reference material 
+
+		PinsParameterMap = PinsParameterMap_DEPRECATED;
+		PinsParameterMap_DEPRECATED.Empty();
+
+		UCustomizableObjectNode* ParentNode = GetCustomizableObjectExternalNode<UCustomizableObjectNode>(ParentMaterialObject_DEPRECATED.Get(), ParentMaterialNodeId_DEPRECATED);
+
+		if (UCustomizableObjectNodeMaterial* MaterialParentNode = Cast<UCustomizableObjectNodeMaterial>(ParentNode))
+		{
+			ReferenceMaterial = MaterialParentNode->GetMaterial();
+		}
+		else if (UCustomizableObjectNodeModifierExtendMeshSection* ExtendParentNode = Cast<UCustomizableObjectNodeModifierExtendMeshSection>(ParentNode))
+		{
+			ReferenceMaterial = ExtendParentNode->ReferenceMaterial;
+		}
+		else
+		{
+			// Conversion failed?
+			ensure(false);
+			UE_LOG(LogMutable, Warning, TEXT("[%s] UCustomizableObjectNodeModifierExtendMeshSection version upgrade failed."), *GetOutermost()->GetName());
+		}
+
+		ReconstructNode();
+	}
+}
+
+
+FText UCustomizableObjectNodeModifierEditMeshSection::GetNodeTitle(ENodeTitleType::Type TitleType) const
+{
+	return LOCTEXT("Edit_MeshSection", "Edit Mesh Section");
+}
+
+
+FString UCustomizableObjectNodeModifierEditMeshSection::GetRefreshMessage() const
+{
+	return TEXT("Source material has changed, texture channels might have been added, removed or renamed. Please refresh the parent material node to reflect those changes.");
+}
+
+
+FText UCustomizableObjectNodeModifierEditMeshSection::GetTooltipText() const
+{
+	return LOCTEXT("Edit_Material_Tooltip", "Modify the texture parameters of an ancestor's material partially or completely.");
+}
+
+
+void UCustomizableObjectNodeModifierEditMeshSection::SetLayoutIndex(const int32 LayoutIndex)
+{
+	ParentLayoutIndex = LayoutIndex;
+}
+
+
+#undef LOCTEXT_NAMESPACE
+
diff --git a/Engine/Plugins/Experimental/Mutable/Source/CustomizableObjectEditor/Private/MuCOE/Nodes/CustomizableObjectNodeModifierEditMeshSectionDetails.cpp b/Engine/Plugins/Experimental/Mutable/Source/CustomizableObjectEditor/Private/MuCOE/Nodes/CustomizableObjectNodeModifierEditMeshSectionDetails.cpp
new file mode 100644
--- /dev/null
+++ b/Engine/Plugins/Experimental/Mutable/Source/CustomizableObjectEditor/Private/MuCOE/Nodes/CustomizableObjectNodeModifierEditMeshSectionDetails.cpp
@@ -0,0 +1,186 @@
+// Copyright Epic Games, Inc. All Rights Reserved.
+
+#include "MuCOE/Nodes/CustomizableObjectNodeModifierEditMeshSectionDetails.h"
+
+#include "DetailLayoutBuilder.h"
+#include "DetailWidgetRow.h"
+#include "IDetailGroup.h"
+#include "IDetailsView.h"
+#include "MuCOE/Nodes/CustomizableObjectNodeModifierEditMeshSection.h"
+#include "MuCOE/Nodes/CustomizableObjectNodeModifierExtendMeshSection.h"
+#include "MuCOE/Nodes/CustomizableObjectNodeMaterialBase.h"
+#include "MuCOE/SCustomizableObjectNodeLayoutBlocksEditor.h"
+#include "MuCOE/PinViewer/SPinViewer.h"
+#include "MuR/MutableTrace.h"
+#include "Widgets/Layout/SBox.h"
+#include "Widgets/Text/STextBlock.h"
+#include "Widgets/Input/STextComboBox.h"
+#include "UObject/Package.h"
+
+class FString;
+
+#define LOCTEXT_NAMESPACE "CustomizableObjectDetails"
+
+
+TSharedRef<IDetailCustomization> FCustomizableObjectNodeModifierEditMeshSectionDetails::MakeInstance()
+{
+	return MakeShareable( new FCustomizableObjectNodeModifierEditMeshSectionDetails);
+}
+
+
+void FCustomizableObjectNodeModifierEditMeshSectionDetails::CustomizeDetails( IDetailLayoutBuilder& DetailBuilder )
+{
+	FCustomizableObjectNodeDetails::CustomizeDetails(DetailBuilder);
+
+	Node = nullptr;
+	const IDetailsView* DetailsView = DetailBuilder.GetDetailsView();
+	if (DetailsView->GetSelectedObjects().Num())
+	{
+		Node = Cast<UCustomizableObjectNodeModifierEditMeshSection>(DetailsView->GetSelectedObjects()[0].Get());
+	}
+
+	IDetailCategoryBuilder& LayoutCategory = DetailBuilder.EditCategory("LayoutOptions");
+
+	if (!Node)
+	{
+		LayoutCategory.AddCustomRow(LOCTEXT("BlocksDetails_NodeNotFound", "NodeNotFound"))
+			[
+				SNew(STextBlock)
+					.Text(LOCTEXT("Node not found", "Node not found"))
+			];
+		return;
+	}
+
+	// UV Channel combo (for now hardcoded to a maximum of 4)
+	{
+		int32 MaxGridSize = 128;
+		UVChannelOptions.Empty();
+
+		TSharedPtr<FString> CurrentUVChannel;
+		for (int32 Index = 0; Index < 4; ++Index)
+		{
+			UVChannelOptions.Add(MakeShareable(new FString(FString::Printf(TEXT("%d"), Index))));
+
+			if (Node->ParentLayoutIndex == Index)
+			{
+				CurrentUVChannel = UVChannelOptions.Last();
+			}
+		}
+
+		IDetailGroup& LayoutOptionsGroup = LayoutCategory.AddGroup(TEXT("LayoutOptionsGroup"), LOCTEXT("LayoutGroup", "Layout Group"), false, true);
+		LayoutOptionsGroup.HeaderRow()
+			.NameContent()
+			[
+				SNew(STextBlock)
+					.Text(LOCTEXT("UVChannel", "UV Channel"))
+					.Font(DetailBuilder.GetDetailFont())
+			]
+			.ValueContent()
+			[
+				SNew(STextComboBox)
+					.InitiallySelectedItem(CurrentUVChannel)
+					.OptionsSource(&UVChannelOptions)
+					.OnSelectionChanged(this, &FCustomizableObjectNodeModifierEditMeshSectionDetails::OnUVChannelChanged)
+					.Font(DetailBuilder.GetDetailFont())
+			];
+	}
+
+	// Grid size combo
+	{
+		int32 MaxGridSize = 128;
+		LayoutGridSizes.Empty();
+
+		TSharedPtr<FString> CurrentSize;
+		for (int32 Size = 1; Size <= MaxGridSize; Size *= 2)
+		{
+			LayoutGridSizes.Add(MakeShareable(new FString(FString::Printf(TEXT("%d x %d"), Size, Size))));
+
+			if (Node->Layout->GetGridSize() == FIntPoint(Size))
+			{
+				CurrentSize = LayoutGridSizes.Last();
+			}
+		}
+
+		IDetailGroup& LayoutOptionsGroup = LayoutCategory.AddGroup(TEXT("LayoutOptionsGroup"), LOCTEXT("LayoutGroup", "Layout Group"), false, true);
+		LayoutOptionsGroup.HeaderRow()
+			.NameContent()
+			[
+				SNew(STextBlock)
+					.Text(LOCTEXT("LayoutGridSizeText", "Grid Size"))
+					.Font(DetailBuilder.GetDetailFont())
+			]
+			.ValueContent()
+			[
+				SNew(STextComboBox)
+					.InitiallySelectedItem(CurrentSize)
+					.OptionsSource(&LayoutGridSizes)
+					.OnSelectionChanged(this, &FCustomizableObjectNodeModifierEditMeshSectionDetails::OnGridSizeChanged)
+					.Font(DetailBuilder.GetDetailFont())
+			];
+	}
+
+	// Block editor
+	LayoutBlocksEditor = SNew(SCustomizableObjectNodeLayoutBlocksEditor);
+
+	LayoutCategory.AddCustomRow(LOCTEXT("BlocksDetails_BlockInstructions", "BlockInstructions"))
+	[
+		SNew(SBox)
+		.HeightOverride(700.0f)
+		.WidthOverride(700.0f)
+		[
+			LayoutBlocksEditor.ToSharedRef()
+		]
+	];
+
+	UpdateLayout();
+}
+
+
+void FCustomizableObjectNodeModifierEditMeshSectionDetails::UpdateLayout()
+{
+	MUTABLE_CPUPROFILER_SCOPE(FCustomizableObjectNodeModifierEditMeshSectionDetails_UpdateLayout);
+
+	// Try to find the parent layout, because we want to show its UVs in the widget
+	UCustomizableObjectLayout* ParentLayout = Node->GetPossibleParentLayout();
+
+	LayoutBlocksEditor->SetCurrentLayout(Node->Layout, ParentLayout);
+}
+
+
+void FCustomizableObjectNodeModifierEditMeshSectionDetails::OnGridSizeChanged(TSharedPtr<FString> NewSelection, ESelectInfo::Type SelectInfo)
+{
+	if (!Node->Layout)
+	{
+		return;
+	}
+
+	int32 Size = 1 << LayoutGridSizes.Find(NewSelection);
+
+	if (Node->Layout->GetGridSize().X != Size || Node->Layout->GetGridSize().Y != Size)
+	{
+		Node->Layout->SetGridSize(FIntPoint(Size));
+
+		Node->MarkPackageDirty();
+		UpdateLayout();
+	}
+}
+
+
+void FCustomizableObjectNodeModifierEditMeshSectionDetails::OnUVChannelChanged(TSharedPtr<FString> NewSelection, ESelectInfo::Type SelectInfo)
+{
+	if (!Node->Layout)
+	{
+		return;
+	}
+
+	int32 Index = UVChannelOptions.Find(NewSelection);
+
+	if (Node->ParentLayoutIndex != Index)
+	{
+		Node->ParentLayoutIndex = Index;
+		Node->Modify();
+		UpdateLayout();
+	}
+}
+
+#undef LOCTEXT_NAMESPACE
diff --git a/Engine/Plugins/Experimental/Mutable/Source/CustomizableObjectEditor/Private/MuCOE/Nodes/CustomizableObjectNodeModifierEditMeshSectionDetails.h b/Engine/Plugins/Experimental/Mutable/Source/CustomizableObjectEditor/Private/MuCOE/Nodes/CustomizableObjectNodeModifierEditMeshSectionDetails.h
new file mode 100644
--- /dev/null
+++ b/Engine/Plugins/Experimental/Mutable/Source/CustomizableObjectEditor/Private/MuCOE/Nodes/CustomizableObjectNodeModifierEditMeshSectionDetails.h
@@ -0,0 +1,41 @@
+// Copyright Epic Games, Inc. All Rights Reserved.
+
+#pragma once
+
+
+#include "CustomizableObjectNodeDetails.h"
+#include "IDetailCustomization.h"
+#include "Types/SlateEnums.h"
+
+class FString;
+class IDetailLayoutBuilder;
+class UCustomizableObjectNodeModifierEditMeshSection;
+class SCustomizableObjectNodeLayoutBlocksEditor;
+
+
+class FCustomizableObjectNodeModifierEditMeshSectionDetails : public FCustomizableObjectNodeDetails
+{
+public:
+	/** Makes a new instance of this detail layout class for a specific detail view requesting it */
+	static TSharedRef<IDetailCustomization> MakeInstance();
+
+	/** ILayoutDetails interface */
+	virtual void CustomizeDetails(IDetailLayoutBuilder& DetailBuilder) override;
+
+private:
+	UCustomizableObjectNodeModifierEditMeshSection* Node;
+
+	// Layout block editor widget
+	TSharedPtr<SCustomizableObjectNodeLayoutBlocksEditor> LayoutBlocksEditor;
+
+	/** List of available layout grid sizes. */
+	TArray< TSharedPtr< FString > > LayoutGridSizes;
+	TArray< TSharedPtr< FString > > UVChannelOptions;
+
+	/** Layout Options Callbacks */
+	void OnGridSizeChanged(TSharedPtr<FString> NewSelection, ESelectInfo::Type SelectInfo);
+	void OnUVChannelChanged(TSharedPtr<FString> NewSelection, ESelectInfo::Type SelectInfo);
+
+	/** Reset the layout in the widget to force a refresh. */
+	void UpdateLayout();
+};
diff --git a/Engine/Plugins/Experimental/Mutable/Source/CustomizableObjectEditor/Private/MuCOE/Nodes/CustomizableObjectNodeModifierExtendMeshSection.cpp b/Engine/Plugins/Experimental/Mutable/Source/CustomizableObjectEditor/Private/MuCOE/Nodes/CustomizableObjectNodeModifierExtendMeshSection.cpp
new file mode 100644
--- /dev/null
+++ b/Engine/Plugins/Experimental/Mutable/Source/CustomizableObjectEditor/Private/MuCOE/Nodes/CustomizableObjectNodeModifierExtendMeshSection.cpp
@@ -0,0 +1,216 @@
+// Copyright Epic Games, Inc. All Rights Reserved.
+
+#include "MuCOE/Nodes/CustomizableObjectNodeModifierExtendMeshSection.h"
+
+#include "MaterialTypes.h"
+#include "MuCO/CustomizableObjectCustomVersion.h"
+#include "MuCOE/CustomizableObjectEditor_Deprecated.h"
+#include "MuCOE/CustomizableObjectGraph.h"
+#include "MuCOE/EdGraphSchema_CustomizableObject.h"
+#include "MuCOE/GraphTraversal.h"
+#include "MuCOE/Nodes/CustomizableObjectNodeCopyMaterial.h"
+#include "MuCOE/Nodes/CustomizableObjectNodeSkeletalMesh.h"
+#include "MuCOE/Nodes/CustomizableObjectNodeTable.h"
+#include "MuCOE/CustomizableObjectPin.h"
+
+class UCustomizableObjectNode;
+class UCustomizableObjectNodeRemapPins;
+
+#define LOCTEXT_NAMESPACE "CustomizableObjectEditor"
+
+
+void UCustomizableObjectNodeModifierExtendMeshSection::BackwardsCompatibleFixup(int32 CustomizableObjectCustomVersion)
+{
+	Super::BackwardsCompatibleFixup(CustomizableObjectCustomVersion);
+	
+	if (CustomizableObjectCustomVersion == FCustomizableObjectCustomVersion::ExtendMaterialRemoveImages)
+	{
+		const UCustomizableObjectNodeMaterialBase* ParentMaterial = GetCustomizableObjectExternalNode<UCustomizableObjectNodeMaterialBase>(ParentMaterialObject_DEPRECATED.Get(), ParentMaterialNodeId_DEPRECATED);
+		if (ParentMaterial)
+		{
+			for (const FCustomizableObjectNodeExtendMaterialImage& Image : Images_DEPRECATED)
+			{
+				const UEdGraphPin* ImagePin = FindPin(Image.Name);
+				if (!ImagePin)
+				{
+					continue;
+				}
+
+				FGuid ImageId = FGuid::NewGuid();
+
+				// Search for the Image Id the Extend pin was referring to.
+				const int32 NumImages = ParentMaterial->GetNumParameters(EMaterialParameterType::Texture);
+				for (int32 ImageIndex = 0; ImageIndex < NumImages; ++ImageIndex)
+				{
+					if (ParentMaterial->GetParameterName(EMaterialParameterType::Texture, ImageIndex).ToString() == Image.Name)
+					{
+						ImageId = ParentMaterial->GetParameterId(EMaterialParameterType::Texture, ImageIndex).ParameterId;
+						break;
+					}
+				}
+
+				PinsParameter_DEPRECATED.Add(ImageId, FEdGraphPinReference(ImagePin));                
+			}
+		}
+
+		Images_DEPRECATED.Empty();
+	}
+
+	if (CustomizableObjectCustomVersion == FCustomizableObjectCustomVersion::ConvertEditAndExtendToModifiers)
+	{
+		// Look for the parent material and set it as the modifier reference material 
+		PinsParameterMap = PinsParameterMap_DEPRECATED;
+		PinsParameterMap_DEPRECATED.Empty();
+
+		UCustomizableObjectNode* ParentNode = GetCustomizableObjectExternalNode<UCustomizableObjectNode>(ParentMaterialObject_DEPRECATED.Get(), ParentMaterialNodeId_DEPRECATED);
+		if (ParentNode)
+		{
+			// Add an autogenerated tag to the legacy parent, so that it can be referenced in this modifier.
+			FString NewLegacyTag = MakeNodeAutoTag(ParentNode);
+			RequiredTags.Add(NewLegacyTag);
+
+			if (UCustomizableObjectNodeMaterial* MaterialParentNode = Cast<UCustomizableObjectNodeMaterial>(ParentNode))
+			{
+				ReferenceMaterial = MaterialParentNode->GetMaterial();
+				MaterialParentNode->Tags.AddUnique(NewLegacyTag);
+			}
+			else if (UCustomizableObjectNodeModifierExtendMeshSection* ExtendParentNode = Cast<UCustomizableObjectNodeModifierExtendMeshSection>(ParentNode))
+			{
+				ReferenceMaterial = ExtendParentNode->ReferenceMaterial;
+				ExtendParentNode->Tags.AddUnique(NewLegacyTag);
+			}
+			else
+			{
+				// Conversion failed?
+				ensure(false);
+				UE_LOG(LogMutable, Warning, TEXT("[%s] UCustomizableObjectNodeModifierExtendMeshSection version upgrade failed."), *GetOutermost()->GetName());
+			}
+
+			// If the tag was added to another CO, keep track of the addition to be able to repeat it in case
+			// that CO is not re-saved but this one is.
+			UCustomizableObject* ThisNodeObject = GetRootObject(*this);
+			if (ParentMaterialObject_DEPRECATED != ThisNodeObject)
+			{
+				FLegacyTag LegacyTag;
+				LegacyTag.ParentObject = ParentMaterialObject_DEPRECATED;
+				LegacyTag.ParentNode = ParentMaterialNodeId_DEPRECATED;
+				LegacyTag.Tag = NewLegacyTag;
+				LegacyBackportsRequiredTags.AddUnique(LegacyTag);
+			}
+		}
+		else
+		{
+			// Conversion failed?
+			ensure(false);
+			UE_LOG(LogMutable, Warning, TEXT("[%s] UCustomizableObjectNodeModifierEditMeshSectionBase version upgrade failed."), *GetOutermost()->GetName());
+		}
+
+		ReconstructNode();
+	}
+}
+
+
+void UCustomizableObjectNodeModifierExtendMeshSection::AllocateDefaultPins(UCustomizableObjectNodeRemapPins* RemapPins)
+{
+	const UEdGraphSchema_CustomizableObject* Schema = GetDefault<UEdGraphSchema_CustomizableObject>();
+
+	CustomCreatePin(EGPD_Input, Schema->PC_Mesh, FName("Add Mesh") );
+
+	// Begin texture pins
+	{
+		const int32 NumImages = GetNumParameters(EMaterialParameterType::Texture);
+		for (int32 ImageIndex = 0; ImageIndex < NumImages; ++ImageIndex)
+		{
+			const FName ImageName = GetParameterName(EMaterialParameterType::Texture, ImageIndex);
+			UEdGraphPin* PinImage = CustomCreatePin(EGPD_Input, Schema->PC_Image, ImageName);
+			PinImage->bDefaultValueIsIgnored = true;
+
+			const FNodeMaterialParameterId ImageId = GetParameterId(EMaterialParameterType::Texture, ImageIndex);
+				
+			PinsParameterMap.Add(ImageId, FEdGraphPinReference(PinImage));
+		}
+	}
+	// End texture pins
+
+	CustomCreatePin(EGPD_Output, Schema->PC_Modifier, TEXT("Modifier"));
+}
+
+
+FText UCustomizableObjectNodeModifierExtendMeshSection::GetNodeTitle(ENodeTitleType::Type TitleType) const
+{
+	return LOCTEXT("Extend_MeshSection", "Extend Mesh Section");
+}
+
+
+FString UCustomizableObjectNodeModifierExtendMeshSection::GetRefreshMessage() const
+{
+	return TEXT("The Reference Material has changed, texture channels might have been added, removed or renamed. Please refresh to reflect those changes.");
+}
+
+
+UEdGraphPin* UCustomizableObjectNodeModifierExtendMeshSection::AddMeshPin() const
+{
+	return FindPin(TEXT("Add Mesh"));
+}
+
+
+TArray<UCustomizableObjectLayout*> UCustomizableObjectNodeModifierExtendMeshSection::GetLayouts() const
+{
+	TArray<UCustomizableObjectLayout*> Result;
+
+	if (UEdGraphPin* MeshPin = AddMeshPin())
+	{
+		if (const UEdGraphPin* ConnectedPin = FollowInputPin(*MeshPin))
+		{
+			if (const UEdGraphPin* SourceMeshPin = FindMeshBaseSource(*ConnectedPin, false))
+			{
+				if (const UCustomizableObjectNodeSkeletalMesh* MeshNode = Cast<UCustomizableObjectNodeSkeletalMesh>(SourceMeshPin->GetOwningNode()))
+				{
+					Result = MeshNode->GetLayouts(*SourceMeshPin);
+				}
+				else if (const UCustomizableObjectNodeTable* TableNode = Cast<UCustomizableObjectNodeTable>(SourceMeshPin->GetOwningNode()))
+				{
+					Result = TableNode->GetLayouts(SourceMeshPin);
+				}
+			}
+		}
+	}
+
+	return Result;
+}
+
+
+FText UCustomizableObjectNodeModifierExtendMeshSection::GetTooltipText() const
+{
+	return LOCTEXT("Extend_Material_Tooltip", "Extend a mesh section: add a new mesh fragment, and add its corresponding texture to the modified material texture parameters.");
+}
+
+
+bool UCustomizableObjectNodeModifierExtendMeshSection::IsSingleOutputNode() const
+{
+	return true;
+}
+
+
+bool UCustomizableObjectNodeModifierExtendMeshSection::CustomRemovePin(UEdGraphPin& Pin)
+{
+	for (TMap<FNodeMaterialParameterId, FEdGraphPinReference>::TIterator It = PinsParameterMap.CreateIterator(); It; ++It)
+	{
+		if (It.Value().Get() == &Pin) // We could improve performance if FEdGraphPinReference exposed the pin id.
+		{
+			It.RemoveCurrent();
+			break;
+		}
+	}
+
+	return Super::CustomRemovePin(Pin);
+}
+
+
+TArray<FString>* UCustomizableObjectNodeModifierExtendMeshSection::GetEnableTags() 
+{ 
+	return &Tags; 
+}
+
+
+#undef LOCTEXT_NAMESPACE
diff --git a/Engine/Plugins/Experimental/Mutable/Source/CustomizableObjectEditor/Private/MuCOE/Nodes/CustomizableObjectNodeModifierMorphMeshSectionDetails.cpp b/Engine/Plugins/Experimental/Mutable/Source/CustomizableObjectEditor/Private/MuCOE/Nodes/CustomizableObjectNodeModifierMorphMeshSectionDetails.cpp
new file mode 100644
--- /dev/null
+++ b/Engine/Plugins/Experimental/Mutable/Source/CustomizableObjectEditor/Private/MuCOE/Nodes/CustomizableObjectNodeModifierMorphMeshSectionDetails.cpp
@@ -0,0 +1,291 @@
+// Copyright Epic Games, Inc. All Rights Reserved.
+
+#include "MuCOE/Nodes/CustomizableObjectNodeModifierMorphMeshSectionDetails.h"
+
+#include "MuCOE/CustomizableObjectEditorUtilities.h"
+#include "MuCOE/GraphTraversal.h"
+#include "MuCOE/Nodes/CustomizableObjectNodeMaterial.h"
+#include "MuCOE/Nodes/CustomizableObjectNodeModifierMorphMeshSection.h"
+#include "MuCOE/Nodes/CustomizableObjectNodeModifierExtendMeshSection.h"
+#include "MuCOE/Nodes/CustomizableObjectNodeSkeletalMesh.h"
+#include "MuCOE/UnrealEditorPortabilityHelpers.h"
+#include "PropertyCustomizationHelpers.h"
+#include "Widgets/Input/STextComboBox.h"
+#include "Widgets/Input/SSearchBox.h"
+#include "DetailLayoutBuilder.h"
+#include "IDetailsView.h"
+#include "Engine/SkeletalMesh.h"
+#include "Misc/Attribute.h"
+
+
+#define LOCTEXT_NAMESPACE "CustomizableObjectDetails"
+
+
+TSharedRef<IDetailCustomization> FCustomizableObjectNodeModifierMorphMeshSectionDetails::MakeInstance()
+{
+	return MakeShareable( new FCustomizableObjectNodeModifierMorphMeshSectionDetails );
+}
+
+
+void FCustomizableObjectNodeModifierMorphMeshSectionDetails::CustomizeDetails( IDetailLayoutBuilder& DetailBuilder )
+{
+	FCustomizableObjectNodeDetails::CustomizeDetails(DetailBuilder);
+
+	const IDetailsView* DetailsView = DetailBuilder.GetDetailsView();
+	if ( DetailsView->GetSelectedObjects().Num() )
+	{
+		Node = Cast<UCustomizableObjectNodeModifierMorphMeshSection>( DetailsView->GetSelectedObjects()[0].Get() );
+	}
+
+
+	// This property is not relevant for this node
+	DetailBuilder.HideProperty(GET_MEMBER_NAME_CHECKED(UCustomizableObjectNodeModifierWithMaterial, ReferenceMaterial), UCustomizableObjectNodeModifierWithMaterial::StaticClass());
+	
+	IDetailCategoryBuilder& ModifierCategory = DetailBuilder.EditCategory("Modifier");
+	ModifierCategory.SetSortOrder(-10000);
+
+	IDetailCategoryBuilder& BlocksCategory = DetailBuilder.EditCategory( "Customizable Object" );
+
+	if (Node)
+	{
+		// Scan for hint morph names
+		{
+			TArray<UCustomizableObjectNode*> CandidateNodes;
+			Node->GetPossiblyModifiedNodes(CandidateNodes);
+
+			// For now just show the first one
+			for (UCustomizableObjectNode* Candidate : CandidateNodes)
+			{
+				USkeletalMesh* SkeletalMesh = nullptr;
+
+				if (UCustomizableObjectNodeMaterialBase* MaterialNode = Cast<UCustomizableObjectNodeMaterialBase>(Candidate))
+				{
+					const UEdGraphPin* SourceMeshPin = FindMeshBaseSource(*MaterialNode->OutputPin(), false);
+					if (SourceMeshPin)
+					{
+						UCustomizableObjectNodeSkeletalMesh* SkeletalNode = Cast<UCustomizableObjectNodeSkeletalMesh>(SourceMeshPin->GetOwningNode());
+						if (SkeletalNode)
+						{
+							SkeletalMesh = SkeletalNode->SkeletalMesh;
+						}
+					}
+				}
+
+				else if (UCustomizableObjectNodeModifierExtendMeshSection* ExtendNode = Cast<UCustomizableObjectNodeModifierExtendMeshSection>(Candidate))
+				{
+					const UEdGraphPin* SourceMeshPin = FindMeshBaseSource(*ExtendNode->OutputPin(), false);
+					if (SourceMeshPin)
+					{
+						UCustomizableObjectNodeSkeletalMesh* SkeletalNode = Cast<UCustomizableObjectNodeSkeletalMesh>(SourceMeshPin->GetOwningNode());
+						if (SkeletalNode)
+						{
+							SkeletalMesh = SkeletalNode->SkeletalMesh;
+						}
+					}
+				}
+
+				if (SkeletalMesh)
+				{
+					const TArray<TObjectPtr<UMorphTarget>>& Morphs = SkeletalMesh->GetMorphTargets();
+					for (TObjectPtr<UMorphTarget> Morph : Morphs)
+					{
+						FString MorphTargetName = Morph->GetName();
+						OptionsSource.AddUnique(MakeShared<FString>(MorphTargetName));
+					}
+				}
+			}
+
+			// TODO: Add all morphs if no candidate is found? Add both in different menu sections?
+		}
+		FilteredOptionsSource = OptionsSource;
+
+		// Morph target selection
+		TSharedRef<IPropertyHandle> MorphTargetNameProperty = DetailBuilder.GetProperty("MorphTargetName");
+		BlocksCategory.AddCustomRow(LOCTEXT("MorphMaterialDetails_Target", "Target"))
+			[
+				SNew(SProperty, MorphTargetNameProperty)
+					.ShouldDisplayName(false)
+					.CustomWidget()
+					[
+						SNew(SBorder)
+							.BorderImage(UE_MUTABLE_GET_BRUSH("NoBorder"))
+							.Padding(FMargin(0.0f, 0.0f, 10.0f, 0.0f))
+							[
+								SNew(SHorizontalBox)
+
+								+ SHorizontalBox::Slot()
+								.FillWidth(10.0f)
+								.VAlign(VAlign_Center)
+								[
+									SNew(STextBlock)
+										.Text(LOCTEXT("MorphMaterialDetails_MorphTarget", "Morph Target"))
+										.Font(IDetailLayoutBuilder::GetDetailFont())
+								]
+
+								+ SHorizontalBox::Slot()
+								.FillWidth(10.0f)
+								.HAlign(HAlign_Fill)
+								.VAlign(VAlign_Center)
+								[
+									SAssignNew(this->MorphCombo , SComboButton)
+									.ButtonStyle(FAppStyle::Get(), "NoBorder")
+									.ButtonContent()
+									[
+										SNew(SEditableTextBox)
+										.Font(IDetailLayoutBuilder::GetDetailFont())
+										.Text_Lambda([&]() { return Node ? FText::FromString(Node->MorphTargetName) : FText();})
+										.OnTextChanged(this, &FCustomizableObjectNodeModifierMorphMeshSectionDetails::OnMorphTargetComboBoxSelectionChanged)
+									]
+									.MenuContent()
+									[
+										SNew(SVerticalBox)
+
+											+ SVerticalBox::Slot()
+											.AutoHeight()
+											[
+												SAssignNew(this->SearchField, SEditableTextBox)
+													.HintText(LOCTEXT("Search", "Search"))
+													.OnTextChanged(this, &FCustomizableObjectNodeModifierMorphMeshSectionDetails::OnSearchTextChanged)
+													.OnTextCommitted(this, &FCustomizableObjectNodeModifierMorphMeshSectionDetails::OnSearchTextCommitted)
+											]
+
+											+ SVerticalBox::Slot()
+											[
+												SAssignNew(this->ComboListView, SListView< TSharedPtr<FString> >)
+													.ListItemsSource(&FilteredOptionsSource)
+													.OnGenerateRow(this, &FCustomizableObjectNodeModifierMorphMeshSectionDetails::GenerateMenuItemRow)
+													.OnSelectionChanged(this, &FCustomizableObjectNodeModifierMorphMeshSectionDetails::OnSelectionChanged)
+													.OnKeyDownHandler(this, &FCustomizableObjectNodeModifierMorphMeshSectionDetails::OnKeyDownHandler)
+													.SelectionMode(ESelectionMode::Single)
+											]
+									]
+								]
+							]
+					]
+			];
+	}
+	else
+	{
+		BlocksCategory.AddCustomRow( LOCTEXT("MorphMaterialDetails_Node", "Node") )
+		[
+			SNew( STextBlock )
+			.Text( LOCTEXT( "MorphMaterialDetails_NodeNotFound", "Node not found" ) )
+		];
+	}
+}
+
+
+void FCustomizableObjectNodeModifierMorphMeshSectionDetails::OnSearchTextChanged(const FText& ChangedText)
+{
+	SearchText = ChangedText;
+
+	RefreshOptions();
+}
+
+void FCustomizableObjectNodeModifierMorphMeshSectionDetails::OnSearchTextCommitted(const FText& InText, ETextCommit::Type InCommitType)
+{
+	if ((InCommitType == ETextCommit::Type::OnEnter) && FilteredOptionsSource.Num() > 0)
+	{
+		ComboListView->SetSelection(FilteredOptionsSource[0], ESelectInfo::OnKeyPress);
+	}
+}
+
+
+void FCustomizableObjectNodeModifierMorphMeshSectionDetails::OnMorphTargetComboBoxSelectionChanged(const FText& NewText)
+{	
+	if (Node && Node->MorphTargetName!=NewText.ToString())
+	{
+		Node->MorphTargetName = NewText.ToString();
+		Node->Modify();
+	}
+}
+
+void FCustomizableObjectNodeModifierMorphMeshSectionDetails::RefreshOptions()
+{
+	// Need to refresh filtered list whenever options change
+	FilteredOptionsSource.Reset();
+
+	if (SearchText.IsEmpty())
+	{
+		FilteredOptionsSource.Append(OptionsSource);
+	}
+	else
+	{
+		TArray<FString> SearchTokens;
+		SearchText.ToString().ParseIntoArrayWS(SearchTokens);
+
+		for (const TSharedPtr<FString>& Option : OptionsSource)
+		{
+			bool bAllTokensMatch = true;
+			for (const FString& SearchToken : SearchTokens)
+			{
+				if (Option->Find(SearchToken, ESearchCase::Type::IgnoreCase) == INDEX_NONE)
+				{
+					bAllTokensMatch = false;
+					break;
+				}
+			}
+
+			if (bAllTokensMatch)
+			{
+				FilteredOptionsSource.Add(Option);
+			}
+		}
+	}
+
+	ComboListView->RequestListRefresh();
+}
+
+
+TSharedRef<ITableRow> FCustomizableObjectNodeModifierMorphMeshSectionDetails::GenerateMenuItemRow(TSharedPtr<FString> InItem, const TSharedRef<STableViewBase>& OwnerTable)
+{
+	return SNew(SComboRow<TSharedPtr<FString>>, OwnerTable)
+		[
+			SNew(STextBlock).Text(FText::FromString(*InItem))
+		];
+}
+
+
+void FCustomizableObjectNodeModifierMorphMeshSectionDetails::OnSelectionChanged(TSharedPtr<FString> ProposedSelection, ESelectInfo::Type SelectInfo)
+{
+	if (!ProposedSelection || !Node)
+	{
+		return;
+	}
+
+	// Ensure that the proposed selection is different from selected
+	if (*ProposedSelection != Node->MorphTargetName)
+	{
+		Node->MorphTargetName = *ProposedSelection;
+	}
+
+	// close combo as long as the selection wasn't from navigation
+	if (SelectInfo != ESelectInfo::OnNavigation)
+	{
+		MorphCombo->SetIsOpen(false);
+	}
+	else
+	{
+		ComboListView->RequestScrollIntoView(ProposedSelection, 0);
+	}
+}
+
+
+FReply FCustomizableObjectNodeModifierMorphMeshSectionDetails::OnKeyDownHandler(const FGeometry& MyGeometry, const FKeyEvent& InKeyEvent)
+{
+	if (InKeyEvent.GetKey() == EKeys::Enter)
+	{
+		// Select the first selected item on hitting enter
+		TArray<TSharedPtr<FString>> SelectedItems = ComboListView->GetSelectedItems();
+		if (SelectedItems.Num() > 0)
+		{
+			OnSelectionChanged(SelectedItems[0], ESelectInfo::OnKeyPress);
+			return FReply::Handled();
+		}
+	}
+
+	return FReply::Unhandled();
+}
+
+
+#undef LOCTEXT_NAMESPACE
diff --git a/Engine/Plugins/Experimental/Mutable/Source/CustomizableObjectEditor/Private/MuCOE/Nodes/CustomizableObjectNodeModifierMorphMeshSectionDetails.h b/Engine/Plugins/Experimental/Mutable/Source/CustomizableObjectEditor/Private/MuCOE/Nodes/CustomizableObjectNodeModifierMorphMeshSectionDetails.h
new file mode 100644
--- /dev/null
+++ b/Engine/Plugins/Experimental/Mutable/Source/CustomizableObjectEditor/Private/MuCOE/Nodes/CustomizableObjectNodeModifierMorphMeshSectionDetails.h
@@ -0,0 +1,53 @@
+// Copyright Epic Games, Inc. All Rights Reserved.
+
+#pragma once
+
+#include "CustomizableObjectNodeDetails.h"
+#include "IDetailCustomization.h"
+#include "Types/SlateEnums.h"
+#include "Widgets/Input/SComboButton.h"
+#include "Widgets/Input/SEditableTextBox.h"
+#include "Widgets/Views/SListView.h"
+#include "Widgets/Views/ITableRow.h"
+
+namespace ESelectInfo { enum Type : int; }
+
+class FString;
+class IDetailLayoutBuilder;
+class IPropertyHandle;
+class UCustomizableObjectNodeObject;
+
+class FCustomizableObjectNodeModifierMorphMeshSectionDetails : public FCustomizableObjectNodeDetails
+{
+public:
+	// Makes a new instance of this detail layout class for a specific detail view requesting it 
+	static TSharedRef<IDetailCustomization> MakeInstance();
+
+	// ILayoutDetails interface
+	void CustomizeDetails( IDetailLayoutBuilder& DetailBuilder ) override;
+
+private:
+	class UCustomizableObjectNodeModifierMorphMeshSection* Node = nullptr;
+
+	/** */
+	TSharedPtr<SComboButton> MorphCombo;
+	/** The search field used for the combox box's contents */
+	TSharedPtr<SEditableTextBox> SearchField;
+	/** The ListView that we pop up; visualized the available options. */
+	TSharedPtr< SListView< TSharedPtr<FString> > > ComboListView;
+	/** */
+	TArray< TSharedPtr<FString> > OptionsSource;
+	/** Updated whenever search text is changed */
+	FText SearchText;
+	/** Filtered list that is actually displayed */
+	TArray< TSharedPtr<FString> > FilteredOptionsSource;
+
+	void OnMorphTargetComboBoxSelectionChanged(const FText& NewText);
+	void OnSearchTextChanged(const FText& ChangedText);
+	void OnSearchTextCommitted(const FText& InText, ETextCommit::Type InCommitType);
+	TSharedRef<ITableRow> GenerateMenuItemRow(TSharedPtr<FString> InItem, const TSharedRef<STableViewBase>& OwnerTable);
+	void RefreshOptions();
+	void OnSelectionChanged(TSharedPtr<FString> ProposedSelection, ESelectInfo::Type SelectInfo);
+	FReply OnKeyDownHandler(const FGeometry& MyGeometry, const FKeyEvent& InKeyEvent);
+
+};
diff --git a/Engine/Plugins/Experimental/Mutable/Source/CustomizableObjectEditor/Private/MuCOE/Nodes/CustomizableObjectNodeModifierRemoveMesh.cpp b/Engine/Plugins/Experimental/Mutable/Source/CustomizableObjectEditor/Private/MuCOE/Nodes/CustomizableObjectNodeModifierRemoveMesh.cpp
new file mode 100644
--- /dev/null
+++ b/Engine/Plugins/Experimental/Mutable/Source/CustomizableObjectEditor/Private/MuCOE/Nodes/CustomizableObjectNodeModifierRemoveMesh.cpp
@@ -0,0 +1,58 @@
+// Copyright Epic Games, Inc. All Rights Reserved.
+
+#include "MuCOE/Nodes/CustomizableObjectNodeModifierRemoveMesh.h"
+
+#include "MuCOE/EdGraphSchema_CustomizableObject.h"
+#include "MuCOE/ICustomizableObjectEditor.h"
+#include "MuCO/CustomizableObjectCustomVersion.h"
+
+class UCustomizableObjectNodeRemapPins;
+
+#define LOCTEXT_NAMESPACE "CustomizableObjectEditor"
+
+
+void UCustomizableObjectNodeModifierRemoveMesh::AllocateDefaultPins(UCustomizableObjectNodeRemapPins* RemapPins)
+{
+	const UEdGraphSchema_CustomizableObject* Schema = GetDefault<UEdGraphSchema_CustomizableObject>();
+
+	UEdGraphPin* RemoveMeshPin = CustomCreatePin(EGPD_Input, Schema->PC_Mesh, FName("Remove Mesh") );
+	RemoveMeshPin->bDefaultValueIsIgnored = true;
+
+	UEdGraphPin* OutputPin = CustomCreatePin(EGPD_Output, Schema->PC_Modifier, FName("Modifier"));
+}
+
+
+FText UCustomizableObjectNodeModifierRemoveMesh::GetNodeTitle(ENodeTitleType::Type TitleType) const
+{
+	return LOCTEXT("Remove_Mesh", "Remove Mesh");
+}
+
+
+void UCustomizableObjectNodeModifierRemoveMesh::PinConnectionListChanged(UEdGraphPin* Pin)
+{
+	Super::PinConnectionListChanged(Pin);
+
+	if (Pin == OutputPin())
+	{
+		TSharedPtr<ICustomizableObjectEditor> Editor = GetGraphEditor();
+
+		if (Editor.IsValid())
+		{
+			Editor->UpdateGraphNodeProperties();
+		}
+	}
+}
+
+FText UCustomizableObjectNodeModifierRemoveMesh::GetTooltipText() const
+{
+	return LOCTEXT("Remove_Mesh_Tooltip",
+	"Removes the faces of a material that are defined only by vertexes shared by the material and the input mesh.It also removes any vertex\nand edge that only define deleted faces, they are not left dangling. If the mesh removed covers at least all the faces included in one or\nmore layout blocs, those blocs are removed, freeing final texture layout space.");
+}
+
+bool UCustomizableObjectNodeModifierRemoveMesh::IsSingleOutputNode() const
+{
+	return true;
+}
+
+
+#undef LOCTEXT_NAMESPACE
