[mutable] Convert the following nodes to modifiers: EditMaterial, ExtendMaterial, RemoveMesh, RemoveMeshBlocks, MorphMaterial.

diff --git a/Engine/Plugins/Experimental/Mutable/Source/CustomizableObjectEditor/Private/MuCOE/Nodes/CustomizableObjectNodeModifierRemoveMeshBlocks.cpp b/Engine/Plugins/Experimental/Mutable/Source/CustomizableObjectEditor/Private/MuCOE/Nodes/CustomizableObjectNodeModifierRemoveMeshBlocks.cpp
new file mode 100644
--- /dev/null
+++ b/Engine/Plugins/Experimental/Mutable/Source/CustomizableObjectEditor/Private/MuCOE/Nodes/CustomizableObjectNodeModifierRemoveMeshBlocks.cpp
@@ -0,0 +1,181 @@
+// Copyright Epic Games, Inc. All Rights Reserved.
+
+#include "MuCOE/Nodes/CustomizableObjectNodeModifierRemoveMeshBlocks.h"
+
+#include "MuCOE/Nodes/CustomizableObjectNodeModifierExtendMeshSection.h"
+#include "MuCO/CustomizableObjectCustomVersion.h"
+#include "MuCOE/CustomizableObjectLayout.h"
+#include "MuCOE/EdGraphSchema_CustomizableObject.h"
+#include "MuCOE/ICustomizableObjectEditor.h"
+#include "MuCOE/Nodes/CustomizableObjectNodeMaterial.h"
+
+class UCustomizableObjectNodeRemapPins;
+class UEdGraphPin;
+
+#define LOCTEXT_NAMESPACE "CustomizableObjectEditor"
+
+
+void UCustomizableObjectNodeModifierRemoveMeshBlocks::BackwardsCompatibleFixup(int32 CustomizableObjectCustomVersion)
+{
+	Super::BackwardsCompatibleFixup(CustomizableObjectCustomVersion);
+	
+	// Convert deprecated node index list to the node id list.
+	if (CustomizableObjectCustomVersion == FCustomizableObjectCustomVersion::PostLoadToCustomVersion
+		&& BlockIds_DEPRECATED.Num() < Blocks_DEPRECATED.Num())
+	{
+		UCustomizableObjectNodeMaterialBase* ParentMaterialNode = GetCustomizableObjectExternalNode<UCustomizableObjectNodeMaterialBase>(ParentMaterialObject_DEPRECATED.Get(), ParentMaterialNodeId_DEPRECATED);
+		if (ParentMaterialNode)
+		{
+			TArray<UCustomizableObjectLayout*> Layouts = ParentMaterialNode->GetLayouts();
+
+			if (!Layouts.IsValidIndex(ParentLayoutIndex))
+			{
+				UE_LOG(LogMutable, Warning, TEXT("[%s] UCustomizableObjectNodeModifierRemoveMeshBlocks refers to an invalid texture layout index %d. Parent node has %d layouts."),
+					*GetOutermost()->GetName(), ParentLayoutIndex, Layouts.Num());
+			}
+			else if (UCustomizableObjectNodeMaterial* ParentMaterial = Cast<UCustomizableObjectNodeMaterial>(ParentMaterialNode))
+			{
+				UCustomizableObjectLayout* ParentLayout = Layouts[ParentLayoutIndex];
+
+				for (int IndexIndex = BlockIds_DEPRECATED.Num(); IndexIndex < Blocks_DEPRECATED.Num(); ++IndexIndex)
+				{
+					const int BlockIndex = Blocks_DEPRECATED[IndexIndex];
+					if (!ParentLayout->Blocks.IsValidIndex(BlockIndex))
+					{
+						UE_LOG(LogMutable, Warning, TEXT("[%s] UCustomizableObjectNodeModifierRemoveMeshBlocks refers to an invalid layout block index %d. Parent node has %d blocks."),
+							*GetOutermost()->GetName(), BlockIndex, ParentLayout->Blocks.Num());
+
+						continue;
+					}
+					
+					const FGuid Id = ParentLayout->Blocks[BlockIndex].Id;
+					if (!Id.IsValid())
+					{
+						UE_LOG(LogMutable, Warning, TEXT("[%s] UCustomizableObjectNodeModifierRemoveMeshBlocks refers to an valid layout block %d but that block doesn't have an id."),
+							*GetOutermost()->GetName(), BlockIndex);
+
+						continue;
+					}
+
+					BlockIds_DEPRECATED.Add(Id);
+				}
+			}
+		}
+		else
+		{
+			// Conversion failed?
+			ensure(false);
+			UE_LOG(LogMutable, Warning, TEXT("[%s] UCustomizableObjectNodeModifierRemoveMeshBlocks version upgrade failed."), *GetOutermost()->GetName());
+		}
+	}
+
+	// Convert deprecated node id list to absolute rect list.
+	if (CustomizableObjectCustomVersion == FCustomizableObjectCustomVersion::UseUVRects)
+	{
+		// If we are here, it means this node was loaded from a version that didn't have it's own layout.
+		check(Layout->Blocks.IsEmpty());
+
+		UCustomizableObjectNodeMaterialBase* ParentMaterialNode = GetCustomizableObjectExternalNode<UCustomizableObjectNodeMaterialBase>(ParentMaterialObject_DEPRECATED.Get(), ParentMaterialNodeId_DEPRECATED);
+
+		TArray<UCustomizableObjectLayout*> ParentLayouts;
+
+		if (!ParentMaterialNode)
+		{
+			// Conversion failed?
+			ensure(false);
+			UE_LOG(LogMutable, Warning, TEXT("[%s] UCustomizableObjectNodeModifierRemoveMeshBlocks version upgrade failed."), *GetOutermost()->GetName());
+		}
+		else
+		{
+			ParentLayouts = ParentMaterialNode->GetLayouts();
+		}
+
+		if (!ParentLayouts.IsValidIndex(ParentLayoutIndex))
+		{
+			UE_LOG(LogMutable, Warning, TEXT("[%s] UCustomizableObjectNodeModifierRemoveMeshBlocks refers to an invalid texture layout index %d. Parent node has %d layouts."),
+				*GetOutermost()->GetName(), ParentLayoutIndex, ParentLayouts.Num());
+		}
+		else
+		{
+			UCustomizableObjectLayout* ParentLayout = ParentLayouts[ParentLayoutIndex];
+			FIntPoint GridSize = ParentLayout->GetGridSize();
+
+			Layout->SetGridSize(GridSize);
+
+			if (Cast<UCustomizableObjectNodeMaterial>(ParentMaterialNode))
+			{
+				for (const FGuid& BlockId : BlockIds_DEPRECATED)
+				{
+					bool bSkipBlock = false;
+					for (const FCustomizableObjectLayoutBlock& ExistingBlock : Layout->Blocks)
+					{
+						if (ExistingBlock.Id == BlockId)
+						{
+							bSkipBlock = true;
+							UE_LOG(LogMutable, Log, TEXT("[%s] UCustomizableObjectNodeModifierRemoveMeshBlocks has a duplicated layout block id. One has been ignored during version upgrade."), *GetOutermost()->GetName());
+							break;
+						}
+					}
+
+					if (bSkipBlock)
+					{
+						continue;
+					}
+
+					bool bFoundInParent = false;
+					for (const FCustomizableObjectLayoutBlock& ParentBlock : ParentLayout->Blocks)
+					{
+						if (ParentBlock.Id == BlockId)
+						{
+							bFoundInParent = true;
+
+							FCustomizableObjectLayoutBlock NewBlock;
+							NewBlock = ParentBlock;
+
+							// Clear some unnecessary data.
+							NewBlock.bReduceBothAxes = false;
+							NewBlock.bReduceByTwo = false;
+							NewBlock.Priority = 0;
+
+							Layout->Blocks.Add(NewBlock);
+							break;
+						}
+					}
+
+					if (!bFoundInParent)
+					{
+						UE_LOG(LogMutable, Warning, TEXT("[%s] UCustomizableObjectNodeModifierRemoveMeshBlocks refers to and invalid layout block. It has been ignored during version upgrade."), *GetOutermost()->GetName());
+					}
+				}
+			}
+		}
+	}
+}
+
+
+void UCustomizableObjectNodeModifierRemoveMeshBlocks::AllocateDefaultPins(UCustomizableObjectNodeRemapPins* RemapPins)
+{
+	const UEdGraphSchema_CustomizableObject* Schema = GetDefault<UEdGraphSchema_CustomizableObject>();
+
+	UEdGraphPin* OutputPin = CustomCreatePin(EGPD_Output, Schema->PC_Modifier, FName("Modifier"));
+}
+
+
+FText UCustomizableObjectNodeModifierRemoveMeshBlocks::GetNodeTitle(ENodeTitleType::Type TitleType) const
+{
+	return LOCTEXT("Remove_Mesh_Blocks", "Remove Mesh Blocks");
+}
+
+
+FText UCustomizableObjectNodeModifierRemoveMeshBlocks::GetTooltipText() const
+{
+	return LOCTEXT("Remove_Mesh_Blocks_Tooltip", "Remove all the geometry in the chosen layout blocks from a material.");
+}
+
+
+bool UCustomizableObjectNodeModifierRemoveMeshBlocks::IsSingleOutputNode() const
+{
+	return true;
+}
+
+#undef LOCTEXT_NAMESPACE
diff --git a/Engine/Plugins/Experimental/Mutable/Source/CustomizableObjectEditor/Private/MuCOE/Nodes/CustomizableObjectNodeModifierRemoveMeshBlocksDetails.cpp b/Engine/Plugins/Experimental/Mutable/Source/CustomizableObjectEditor/Private/MuCOE/Nodes/CustomizableObjectNodeModifierRemoveMeshBlocksDetails.cpp
new file mode 100644
--- /dev/null
+++ b/Engine/Plugins/Experimental/Mutable/Source/CustomizableObjectEditor/Private/MuCOE/Nodes/CustomizableObjectNodeModifierRemoveMeshBlocksDetails.cpp
@@ -0,0 +1,190 @@
+// Copyright Epic Games, Inc. All Rights Reserved.
+
+#include "MuCOE/Nodes/CustomizableObjectNodeModifierRemoveMeshBlocksDetails.h"
+
+#include "DetailLayoutBuilder.h"
+#include "DetailWidgetRow.h"
+#include "IDetailGroup.h"
+#include "IDetailsView.h"
+#include "MuCOE/Nodes/CustomizableObjectNodeModifierRemoveMeshBlocks.h"
+#include "MuCOE/Nodes/CustomizableObjectNodeMaterialBase.h"
+#include "MuCOE/SCustomizableObjectNodeLayoutBlocksEditor.h"
+#include "Widgets/Layout/SBox.h"
+#include "Widgets/Text/STextBlock.h"
+#include "Widgets/Input/STextComboBox.h"
+#include "UObject/Package.h"
+
+class FString;
+
+
+#define LOCTEXT_NAMESPACE "CustomizableObjectDetails"
+
+
+TSharedRef<IDetailCustomization> FCustomizableObjectNodeModifierRemoveMeshBlocksDetails::MakeInstance()
+{
+	return MakeShareable( new FCustomizableObjectNodeModifierRemoveMeshBlocksDetails );
+}
+
+
+void FCustomizableObjectNodeModifierRemoveMeshBlocksDetails::CustomizeDetails(IDetailLayoutBuilder& DetailBuilder)
+{
+	FCustomizableObjectNodeDetails::CustomizeDetails(DetailBuilder);
+
+	Node = nullptr;
+	const IDetailsView* DetailsView = DetailBuilder.GetDetailsView();
+	if (DetailsView->GetSelectedObjects().Num())
+	{
+		Node = Cast<UCustomizableObjectNodeModifierRemoveMeshBlocks>(DetailsView->GetSelectedObjects()[0].Get());
+	}
+
+	// This property is not relevant for this node
+	DetailBuilder.HideProperty(GET_MEMBER_NAME_CHECKED(UCustomizableObjectNodeModifierWithMaterial, ReferenceMaterial), UCustomizableObjectNodeModifierWithMaterial::StaticClass());
+
+	IDetailCategoryBuilder& ModifierCategory = DetailBuilder.EditCategory("Modifier");
+	ModifierCategory.SetSortOrder(-10000);
+
+	IDetailCategoryBuilder& LayoutCategory = DetailBuilder.EditCategory("LayoutOptions");
+	LayoutCategory.SetSortOrder(10000);
+
+	if (!Node)
+	{
+		LayoutCategory.AddCustomRow(LOCTEXT("BlocksDetails_NodeNotFound", "NodeNotFound"))
+			[
+				SNew(STextBlock)
+					.Text(LOCTEXT("Node not found", "Node not found"))
+			];
+		return;
+	}
+
+	// UV Channel combo (for now hardcoded to a maximum of 4)
+	{
+		int32 MaxGridSize = 128;
+		UVChannelOptions.Empty();
+
+		TSharedPtr<FString> CurrentUVChannel;
+		for (int32 Index = 0; Index<4; ++Index)
+		{
+			UVChannelOptions.Add(MakeShareable(new FString(FString::Printf(TEXT("%d"), Index))));
+
+			if (Node->ParentLayoutIndex == Index)
+			{
+				CurrentUVChannel = UVChannelOptions.Last();
+			}
+		}
+
+		IDetailGroup& LayoutOptionsGroup = LayoutCategory.AddGroup(TEXT("LayoutOptionsGroup"), LOCTEXT("LayoutGroup", "Layout Group"), false, true);
+		LayoutOptionsGroup.HeaderRow()
+			.NameContent()
+			[
+				SNew(STextBlock)
+					.Text(LOCTEXT("UVChannel", "UV Channel"))
+					.Font(DetailBuilder.GetDetailFont())
+			]
+			.ValueContent()
+			[
+				SNew(STextComboBox)
+					.InitiallySelectedItem(CurrentUVChannel)
+					.OptionsSource(&UVChannelOptions)
+					.OnSelectionChanged(this, &FCustomizableObjectNodeModifierRemoveMeshBlocksDetails::OnUVChannelChanged)
+					.Font(DetailBuilder.GetDetailFont())
+			];
+	}
+
+	// Grid size combo
+	{
+		int32 MaxGridSize = 128;
+		LayoutGridSizes.Empty();
+
+		TSharedPtr<FString> CurrentSize;
+		for (int32 Size = 1; Size <= MaxGridSize; Size *= 2)
+		{
+			LayoutGridSizes.Add(MakeShareable(new FString(FString::Printf(TEXT("%d x %d"), Size, Size))));
+
+			if (Node->Layout->GetGridSize() == FIntPoint(Size))
+			{
+				CurrentSize = LayoutGridSizes.Last();
+			}
+		}
+
+		IDetailGroup& LayoutOptionsGroup = LayoutCategory.AddGroup(TEXT("LayoutOptionsGroup"), LOCTEXT("LayoutGroup", "Layout Group"), false, true);
+		LayoutOptionsGroup.HeaderRow()
+			.NameContent()
+			[
+				SNew(STextBlock)
+					.Text(LOCTEXT("LayoutGridSizeText", "Grid Size"))
+					.Font(DetailBuilder.GetDetailFont())
+			]
+			.ValueContent()
+			[
+				SNew(STextComboBox)
+					.InitiallySelectedItem(CurrentSize)
+					.OptionsSource(&LayoutGridSizes)
+					.OnSelectionChanged(this, &FCustomizableObjectNodeModifierRemoveMeshBlocksDetails::OnGridSizeChanged)
+					.Font(DetailBuilder.GetDetailFont())
+			];
+	}
+
+	// Block editor
+	LayoutBlocksEditor = SNew(SCustomizableObjectNodeLayoutBlocksEditor);
+
+	LayoutCategory.AddCustomRow(LOCTEXT("BlocksDetails_BlockInstructions", "BlockInstructions"))
+		[
+			SNew(SBox)
+				.HeightOverride(700.0f)
+				.WidthOverride(700.0f)
+				[
+					LayoutBlocksEditor.ToSharedRef()
+				]
+		];
+
+	UpdateLayout();
+}
+
+
+void FCustomizableObjectNodeModifierRemoveMeshBlocksDetails::UpdateLayout()
+{
+	// Try to find the parent layout, because we want to show its UVs in the widget
+	UCustomizableObjectLayout* ParentLayout = Node->GetPossibleParentLayout();
+
+	LayoutBlocksEditor->SetCurrentLayout(Node->Layout, ParentLayout);
+
+}
+
+
+void FCustomizableObjectNodeModifierRemoveMeshBlocksDetails::OnGridSizeChanged(TSharedPtr<FString> NewSelection, ESelectInfo::Type SelectInfo)
+{
+	if (!Node->Layout)
+	{
+		return;
+	}
+
+	int32 Size = 1 << LayoutGridSizes.Find(NewSelection);
+
+	if (Node->Layout->GetGridSize().X != Size || Node->Layout->GetGridSize().Y != Size)
+	{
+		Node->Layout->SetGridSize(FIntPoint(Size));
+		Node->MarkPackageDirty();
+		UpdateLayout();
+	}
+}
+
+
+void FCustomizableObjectNodeModifierRemoveMeshBlocksDetails::OnUVChannelChanged(TSharedPtr<FString> NewSelection, ESelectInfo::Type SelectInfo)
+{
+	if (!Node->Layout)
+	{
+		return;
+	}
+
+	int32 Index = UVChannelOptions.Find(NewSelection);
+
+	if (Node->ParentLayoutIndex != Index)
+	{
+		Node->ParentLayoutIndex = Index;
+		Node->MarkPackageDirty();
+		UpdateLayout();
+	}
+}
+
+
+#undef LOCTEXT_NAMESPACE
diff --git a/Engine/Plugins/Experimental/Mutable/Source/CustomizableObjectEditor/Private/MuCOE/Nodes/CustomizableObjectNodeModifierRemoveMeshBlocksDetails.h b/Engine/Plugins/Experimental/Mutable/Source/CustomizableObjectEditor/Private/MuCOE/Nodes/CustomizableObjectNodeModifierRemoveMeshBlocksDetails.h
new file mode 100644
--- /dev/null
+++ b/Engine/Plugins/Experimental/Mutable/Source/CustomizableObjectEditor/Private/MuCOE/Nodes/CustomizableObjectNodeModifierRemoveMeshBlocksDetails.h
@@ -0,0 +1,42 @@
+// Copyright Epic Games, Inc. All Rights Reserved.
+
+#pragma once
+
+
+#include "CustomizableObjectNodeDetails.h"
+#include "IDetailCustomization.h"
+#include "Types/SlateEnums.h"
+
+class FString;
+class IDetailLayoutBuilder;
+class UCustomizableObjectNodeModifierRemoveMeshBlocks;
+class SCustomizableObjectNodeLayoutBlocksEditor;
+
+
+class FCustomizableObjectNodeModifierRemoveMeshBlocksDetails : public FCustomizableObjectNodeDetails
+{
+public:
+	/** Makes a new instance of this detail layout class for a specific detail view requesting it */
+	static TSharedRef<IDetailCustomization> MakeInstance();
+
+	/** ILayoutDetails interface */
+	virtual void CustomizeDetails(IDetailLayoutBuilder& DetailBuilder) override;
+
+private:
+	UCustomizableObjectNodeModifierRemoveMeshBlocks* Node;
+
+	// Layout block editor widget
+	TSharedPtr<SCustomizableObjectNodeLayoutBlocksEditor> LayoutBlocksEditor;
+
+	/** List of available layout grid sizes. */
+	TArray< TSharedPtr< FString > > LayoutGridSizes;
+	TArray< TSharedPtr< FString > > UVChannelOptions;
+
+	/** Layout Options Callbacks */
+	void OnGridSizeChanged(TSharedPtr<FString> NewSelection, ESelectInfo::Type SelectInfo);
+	void OnUVChannelChanged(TSharedPtr<FString> NewSelection, ESelectInfo::Type SelectInfo);
+
+	/** Reset the layout in the widget to force a refresh. */
+	void UpdateLayout();
+
+};
diff --git a/Engine/Plugins/Experimental/Mutable/Source/CustomizableObjectEditor/Private/MuCOE/Nodes/CustomizableObjectNodeModifierRemoveMeshDetails.h b/Engine/Plugins/Experimental/Mutable/Source/CustomizableObjectEditor/Private/MuCOE/Nodes/CustomizableObjectNodeModifierRemoveMeshDetails.h
new file mode 100644
--- /dev/null
+++ b/Engine/Plugins/Experimental/Mutable/Source/CustomizableObjectEditor/Private/MuCOE/Nodes/CustomizableObjectNodeModifierRemoveMeshDetails.h
@@ -0,0 +1,21 @@
+// Copyright Epic Games, Inc. All Rights Reserved.
+
+#pragma once
+
+#include "CustomizableObjectNodeDetails.h"
+#include "IDetailCustomization.h"
+#include "Types/SlateEnums.h"
+#include "Misc/AssertionMacros.h"
+
+class IDetailLayoutBuilder;
+
+class FCustomizableObjectNodeModifierRemoveMeshDetails : public FCustomizableObjectNodeDetails
+{
+public:
+	/** Makes a new instance of this detail layout class for a specific detail view requesting it */
+	static TSharedRef<IDetailCustomization> MakeInstance();
+
+	/** ILayoutDetails interface */
+	virtual void CustomizeDetails(IDetailLayoutBuilder& DetailBuilder) override;
+
+};
diff --git a/Engine/Plugins/Experimental/Mutable/Source/CustomizableObjectEditor/Private/MuCOE/Nodes/CustomizableObjectNodeModifierWithMaterial.cpp b/Engine/Plugins/Experimental/Mutable/Source/CustomizableObjectEditor/Private/MuCOE/Nodes/CustomizableObjectNodeModifierWithMaterial.cpp
new file mode 100644
--- /dev/null
+++ b/Engine/Plugins/Experimental/Mutable/Source/CustomizableObjectEditor/Private/MuCOE/Nodes/CustomizableObjectNodeModifierWithMaterial.cpp
@@ -0,0 +1,190 @@
+// Copyright Epic Games, Inc. All Rights Reserved.
+
+#include "MuCOE/Nodes/CustomizableObjectNodeModifierWithMaterial.h"
+
+#include "MuCOE/Nodes/CustomizableObjectNodeMaterial.h"
+#include "MuCOE/GraphTraversal.h"
+#include "MuCOE/CustomizableObjectPin.h"
+#include "MuCO/CustomizableObjectCustomVersion.h"
+#include "MaterialTypes.h"
+#include "MaterialCachedData.h"
+#include "Materials/MaterialExpressionTextureCoordinate.h"
+#include "Materials/MaterialExpressionTextureSample.h"
+#include "Materials/MaterialInstance.h"
+
+
+#define LOCTEXT_NAMESPACE "CustomizableObjectEditor"
+
+
+
+void UCustomizableObjectNodeModifierWithMaterial::PostEditChangeProperty(FPropertyChangedEvent& PropertyChangedEvent)
+{
+	Super::PostEditChangeProperty(PropertyChangedEvent);
+
+	if (const FProperty* PropertyThatChanged = PropertyChangedEvent.Property)
+	{
+		if (PropertyThatChanged->GetFName() == GET_MEMBER_NAME_CHECKED(UCustomizableObjectNodeModifierWithMaterial, ReferenceMaterial))
+		{
+			UCustomizableObjectNode::ReconstructNode();
+		}
+	}
+}
+
+
+bool UCustomizableObjectNodeModifierWithMaterial::UsesImage(const FNodeMaterialParameterId& ImageId) const
+{
+	if (const UEdGraphPin* Pin = GetUsedImagePin(ImageId))
+	{
+		return FollowInputPin(*Pin) != nullptr;
+	}
+	else
+	{
+		return false;
+	}
+}
+
+
+const UEdGraphPin* UCustomizableObjectNodeModifierWithMaterial::GetUsedImagePin(const FNodeMaterialParameterId& ImageId) const
+{
+	if (const FEdGraphPinReference* PinReference = PinsParameterMap.Find(ImageId))
+	{
+		if (const UEdGraphPin& Pin = *PinReference->Get(); !IsPinOrphan(Pin)) // We always have a valid pin reference. If it is nullptr, it means that something went wrong.
+		{
+			return &Pin;
+		}
+	}
+
+	return nullptr;
+}
+
+
+bool UCustomizableObjectNodeModifierWithMaterial::IsNodeOutDatedAndNeedsRefresh()
+{
+	const bool bOutdated = [&, this]()
+		{
+			for (const TTuple<FNodeMaterialParameterId, FEdGraphPinReference> Pair : PinsParameterMap)
+			{
+				const UEdGraphPin& Pin = *Pair.Value.Get();
+				if (!IsPinOrphan(Pin))
+				{
+					bool bHas = UCustomizableObjectNodeMaterial::HasParameter(ReferenceMaterial,Pair.Key);
+					if (UsesImage(Pair.Key) && !bHas)
+					{
+						return true;
+					}
+				}
+			}
+
+			return false;
+		}();
+
+	// Remove previous compilation warnings
+	if (!bOutdated && bHasCompilerMessage)
+	{
+		RemoveWarnings();
+		GetGraph()->NotifyGraphChanged();
+	}
+
+	return bOutdated;
+}
+
+
+int32 UCustomizableObjectNodeModifierWithMaterial::GetNumParameters(const EMaterialParameterType Type) const
+{
+	if (ReferenceMaterial)
+	{
+		return ReferenceMaterial->GetCachedExpressionData().GetParameterTypeEntry(Type).ParameterInfoSet.Num();
+	}
+	else
+	{
+		return 0;
+	}
+}
+
+
+FNodeMaterialParameterId UCustomizableObjectNodeModifierWithMaterial::GetParameterId(const EMaterialParameterType Type, const int32 ParameterIndex) const
+{
+	const FMaterialCachedExpressionData& Data = ReferenceMaterial->GetCachedExpressionData();
+
+	if (Data.EditorOnlyData)
+	{
+		if (Data.EditorOnlyData->EditorEntries[(int32)Type].EditorInfo.Num() != 0)
+		{
+			const FGuid ParameterId = Data.EditorOnlyData->EditorEntries[(int32)Type].EditorInfo[ParameterIndex].ExpressionGuid;
+			const int32 LayerIndex = GetParameterLayerIndex(Type, ParameterIndex);
+
+			return { ParameterId, LayerIndex };
+		}
+	}
+
+	return FNodeMaterialParameterId();
+}
+
+
+FName UCustomizableObjectNodeModifierWithMaterial::GetParameterName(const EMaterialParameterType Type, const int32 ParameterIndex) const
+{
+	check(ReferenceMaterial);
+
+	const FMaterialCachedParameterEntry& Entry = ReferenceMaterial->GetCachedExpressionData().GetParameterTypeEntry(Type);
+
+	for (TSet<FMaterialParameterInfo>::TConstIterator It(Entry.ParameterInfoSet); It; ++It)
+	{
+		const int32 IteratorIndex = It.GetId().AsInteger();
+
+		if (IteratorIndex == ParameterIndex)
+		{
+			return (*It).Name;
+		}
+	}
+
+	// The parameter should exist
+	check(false);
+
+	return FName();
+}
+
+
+int32 UCustomizableObjectNodeModifierWithMaterial::GetParameterLayerIndex(const UMaterialInterface* InMaterial, const EMaterialParameterType Type, const int32 ParameterIndex)
+{
+	check(InMaterial);
+
+	const FMaterialCachedParameterEntry& Entry = InMaterial->GetCachedExpressionData().GetParameterTypeEntry(Type);
+
+	for (TSet<FMaterialParameterInfo>::TConstIterator It(Entry.ParameterInfoSet); It; ++It)
+	{
+		const int32 IteratorIndex = It.GetId().AsInteger();
+
+		if (IteratorIndex == ParameterIndex)
+		{
+			return (*It).Index;
+		}
+	}
+
+	// The parameter should exist
+	check(false);
+
+	return -1;
+}
+
+
+int32 UCustomizableObjectNodeModifierWithMaterial::GetParameterLayerIndex(const EMaterialParameterType Type, const int32 ParameterIndex) const
+{
+	return GetParameterLayerIndex(ReferenceMaterial.Get(), Type, ParameterIndex);
+}
+
+
+FText UCustomizableObjectNodeModifierWithMaterial::GetParameterLayerName(const EMaterialParameterType Type, const int32 ParameterIndex) const
+{
+	check(ReferenceMaterial)
+
+		int32 LayerIndex = GetParameterLayerIndex(Type, ParameterIndex);
+
+	FMaterialLayersFunctions LayersValue;
+	ReferenceMaterial->GetMaterialLayers(LayersValue);
+
+
+	return LayersValue.EditorOnly.LayerNames.IsValidIndex(LayerIndex) ? LayersValue.EditorOnly.LayerNames[LayerIndex] : FText();
+}
+
+
+#undef LOCTEXT_NAMESPACE
diff --git a/Engine/Plugins/Experimental/Mutable/Source/CustomizableObjectEditor/Private/MuCOE/SCustomizableObjectEditorTagExplorer.cpp b/Engine/Plugins/Experimental/Mutable/Source/CustomizableObjectEditor/Private/MuCOE/SCustomizableObjectEditorTagExplorer.cpp
--- a/Engine/Plugins/Experimental/Mutable/Source/CustomizableObjectEditor/Private/MuCOE/SCustomizableObjectEditorTagExplorer.cpp
+++ b/Engine/Plugins/Experimental/Mutable/Source/CustomizableObjectEditor/Private/MuCOE/SCustomizableObjectEditorTagExplorer.cpp
@@ -8,12 +8,13 @@
 #include "MuCOE/CustomizableObjectEditor.h"
 #include "MuCOE/CustomizableObjectGraph.h"
 #include "MuCOE/GraphTraversal.h"
-#include "MuCOE/Nodes/CustomizableObjectNodeExtendMaterial.h"
+#include "MuCOE/Nodes/CustomizableObjectNodeModifierExtendMeshSection.h"
+#include "MuCOE/Nodes/CustomizableObjectNodeModifierEditMeshSection.h"
 #include "MuCOE/Nodes/CustomizableObjectNodeMaterial.h"
 #include "MuCOE/Nodes/CustomizableObjectNodeMaterialVariation.h"
-#include "MuCOE/Nodes/CustomizableObjectNodeMeshClipDeform.h"
-#include "MuCOE/Nodes/CustomizableObjectNodeMeshClipMorph.h"
-#include "MuCOE/Nodes/CustomizableObjectNodeMeshClipWithMesh.h"
+#include "MuCOE/Nodes/CustomizableObjectNodeModifierClipDeform.h"
+#include "MuCOE/Nodes/CustomizableObjectNodeModifierClipMorph.h"
+#include "MuCOE/Nodes/CustomizableObjectNodeModifierClipWithMesh.h"
 #include "MuCOE/Nodes/CustomizableObjectNodeModifierClipWithUVMask.h"
 #include "Widgets/Input/SButton.h"
 #include "Widgets/Input/SComboButton.h"
@@ -155,95 +156,43 @@ void SCustomizableObjectEditorTagExplorer::FillTagInformation(const UCustomizabl
 	{
 		for (const TObjectPtr<UEdGraphNode>& Node : Object.GetPrivate()->GetSource()->Nodes)
 		{
-			if (UCustomizableObjectNodeMaterialBase* TypedNodeMat = Cast<UCustomizableObjectNodeMaterialBase>(Node))
+			// Gather tags from nodes that activate them
+			if (UCustomizableObjectNode* TypedNodeMat = Cast<UCustomizableObjectNode>(Node))
 			{
-				const TArray<FString>& MaterialTags = TypedNodeMat->GetTags();
+				const TArray<FString>* MaterialTags = TypedNodeMat->GetEnableTags();
 				
-				for (int32 i = 0; i < MaterialTags.Num(); ++i)
+				if (MaterialTags)
 				{
-					NodeTags.Add(MaterialTags[i], TypedNodeMat);
-					if (Tags.Find(MaterialTags[i]) == INDEX_NONE)
+					for (const FString& CurrentTag : *MaterialTags)
 					{
-						Tags.Add(MaterialTags[i]);
+						NodeTags.Add(CurrentTag, TypedNodeMat);
+						Tags.AddUnique(CurrentTag);
 					}
 				}
 			}
 
+			// Gather tags from nodes that require them
+
+			// TODO: All other variations?
 			if (UCustomizableObjectNodeMaterialVariation* TypedNodeVariations = Cast<UCustomizableObjectNodeMaterialVariation>(Node))
 			{
 				for (int32 i = 0; i < TypedNodeVariations->GetNumVariations(); ++i)
 				{
-					const FString& VairationTag = TypedNodeVariations->GetVariation(i).Tag; 
-					NodeTags.Add(VairationTag, TypedNodeVariations);
-					
-					if (Tags.Find(VairationTag) == INDEX_NONE)
-					{
-						Tags.Add(VairationTag);
-					}
-				}
-			}
-
-			if (UCustomizableObjectNodeMeshClipMorph* TypedNodeClipMorph = Cast<UCustomizableObjectNodeMeshClipMorph>(Node))
-			{
-				for (int32 i = 0; i < TypedNodeClipMorph->Tags.Num(); ++i)
-				{
-					NodeTags.Add(TypedNodeClipMorph->Tags[i], TypedNodeClipMorph);
-					
-					if (Tags.Find(TypedNodeClipMorph->Tags[i]) == INDEX_NONE)
-					{
-						Tags.Add(TypedNodeClipMorph->Tags[i]);
-					}
-				}
-			}
-
-			if (UCustomizableObjectNodeMeshClipWithMesh* TypedNodeClipMesh = Cast<UCustomizableObjectNodeMeshClipWithMesh>(Node))
-			{
-				for (int32 i = 0; i < TypedNodeClipMesh->Tags.Num(); ++i)
-				{
-					NodeTags.Add(TypedNodeClipMesh->Tags[i], TypedNodeClipMesh);
-					
-					if (Tags.Find(TypedNodeClipMesh->Tags[i]) == INDEX_NONE)
-					{
-						Tags.Add(TypedNodeClipMesh->Tags[i]);
-					}
-				}
-			}
-
-			if (UCustomizableObjectNodeModifierClipWithUVMask* TypedNodeClipMask = Cast<UCustomizableObjectNodeModifierClipWithUVMask>(Node))
-			{
-				for (int32 i = 0; i < TypedNodeClipMask->Tags.Num(); ++i)
-				{
-					NodeTags.Add(TypedNodeClipMask->Tags[i], TypedNodeClipMask);
-					
-					if (Tags.Find(TypedNodeClipMask->Tags[i]) == INDEX_NONE)
-					{
-						Tags.Add(TypedNodeClipMask->Tags[i]);
-					}
-				}
-			}
-
-			if (UCustomizableObjectNodeExtendMaterial* TypedNodeExtend = Cast<UCustomizableObjectNodeExtendMaterial>(Node))
-			{
-				for (int32 i = 0; i < TypedNodeExtend->Tags.Num(); ++i)
-				{
-					NodeTags.Add(TypedNodeExtend->Tags[i], TypedNodeExtend);
-					
-					if (Tags.Find(TypedNodeExtend->Tags[i]) == INDEX_NONE)
-					{
-						Tags.Add(TypedNodeExtend->Tags[i]);
-					}
+					const FString& VariationTag = TypedNodeVariations->GetVariation(i).Tag; 
+					NodeTags.Add(VariationTag, TypedNodeVariations);					
+					Tags.AddUnique(VariationTag);
 				}
 			}
 
-			if (UCustomizableObjectNodeMeshClipDeform* TypedNodeClipDeform = Cast<UCustomizableObjectNodeMeshClipDeform>(Node))
+			if (UCustomizableObjectNodeModifierBase* TypedNodeModifier = Cast<UCustomizableObjectNodeModifierBase>(Node))
 			{
-				for (int32 i = 0; i < TypedNodeClipDeform->Tags.Num(); ++i)
+				const TArray<FString>* RequiredTags = TypedNodeModifier->GetRequiredTags();
+				if (RequiredTags)
 				{
-					NodeTags.Add(TypedNodeClipDeform->Tags[i], TypedNodeClipDeform);
-					
-					if (Tags.Find(TypedNodeClipDeform->Tags[i]) == INDEX_NONE)
+					for (const FString& CurrentTag : *RequiredTags)
 					{
-						Tags.Add(TypedNodeClipDeform->Tags[i]);
+						NodeTags.Add(CurrentTag, TypedNodeModifier);
+						Tags.AddUnique(CurrentTag);
 					}
 				}
 			}
diff --git a/Engine/Plugins/Experimental/Mutable/Source/CustomizableObjectEditor/Private/MuCOE/SCustomizableObjectEditorViewport.cpp b/Engine/Plugins/Experimental/Mutable/Source/CustomizableObjectEditor/Private/MuCOE/SCustomizableObjectEditorViewport.cpp
--- a/Engine/Plugins/Experimental/Mutable/Source/CustomizableObjectEditor/Private/MuCOE/SCustomizableObjectEditorViewport.cpp
+++ b/Engine/Plugins/Experimental/Mutable/Source/CustomizableObjectEditor/Private/MuCOE/SCustomizableObjectEditorViewport.cpp
@@ -339,7 +339,7 @@ void SCustomizableObjectEditorViewportTabBody::Tick(const FGeometry& AllottedGeo
 }
 
 
-void SCustomizableObjectEditorViewportTabBody::ShowGizmoClipMorph(UCustomizableObjectNodeMeshClipMorph& ClipPlainNode) const
+void SCustomizableObjectEditorViewportTabBody::ShowGizmoClipMorph(UCustomizableObjectNodeModifierClipMorph& ClipPlainNode) const
 {
 	if (const TSharedPtr<SWidget> Toolbar = ViewportToolbarTransformWidget.Pin())
 	{
@@ -361,7 +361,7 @@ void SCustomizableObjectEditorViewportTabBody::HideGizmoClipMorph() const
 }
 
 
-void SCustomizableObjectEditorViewportTabBody::ShowGizmoClipMesh(UCustomizableObjectNodeMeshClipWithMesh& ClipMeshNode, UObject& ClipMesh, int32 LODIndex, int32 SectionIndex, int32 MaterialSlotIndex) const
+void SCustomizableObjectEditorViewportTabBody::ShowGizmoClipMesh(UCustomizableObjectNodeModifierClipWithMesh& ClipMeshNode, UObject& ClipMesh, int32 LODIndex, int32 SectionIndex, int32 MaterialSlotIndex) const
 {
 	if (const TSharedPtr<SWidget> Toolbar = ViewportToolbarTransformWidget.Pin())
 	{
diff --git a/Engine/Plugins/Experimental/Mutable/Source/CustomizableObjectEditor/Private/MuCOE/SCustomizableObjectEditorViewport.h b/Engine/Plugins/Experimental/Mutable/Source/CustomizableObjectEditor/Private/MuCOE/SCustomizableObjectEditorViewport.h
--- a/Engine/Plugins/Experimental/Mutable/Source/CustomizableObjectEditor/Private/MuCOE/SCustomizableObjectEditorViewport.h
+++ b/Engine/Plugins/Experimental/Mutable/Source/CustomizableObjectEditor/Private/MuCOE/SCustomizableObjectEditorViewport.h
@@ -111,13 +111,13 @@ public:
 	const TSharedPtr<FUICommandList>& GetCommandList() const { return UICommandList; }
 
 	/** Do not call directly. Use ICustomizableObjectEditor functions instead. */
-	void ShowGizmoClipMorph(UCustomizableObjectNodeMeshClipMorph& ClipPlainNode) const;
+	void ShowGizmoClipMorph(UCustomizableObjectNodeModifierClipMorph& ClipPlainNode) const;
 
 	/** Do not call directly. Use ICustomizableObjectEditor functions instead. */
 	void HideGizmoClipMorph() const;
 
 	/** Do not call directly. Use ICustomizableObjectEditor functions instead. */
-	void ShowGizmoClipMesh(UCustomizableObjectNodeMeshClipWithMesh& ClipMeshNode, UObject& ClipMesh, int32 LODIndex, int32 SectionIndex, int32 MaterialSlotIndex) const;
+	void ShowGizmoClipMesh(UCustomizableObjectNodeModifierClipWithMesh& ClipMeshNode, UObject& ClipMesh, int32 LODIndex, int32 SectionIndex, int32 MaterialSlotIndex) const;
 
 	/** Do not call directly. Use ICustomizableObjectEditor functions instead. */
 	void HideGizmoClipMesh() const;
diff --git a/Engine/Plugins/Experimental/Mutable/Source/CustomizableObjectEditor/Private/MuCOE/SMutableGraphViewer.cpp b/Engine/Plugins/Experimental/Mutable/Source/CustomizableObjectEditor/Private/MuCOE/SMutableGraphViewer.cpp
--- a/Engine/Plugins/Experimental/Mutable/Source/CustomizableObjectEditor/Private/MuCOE/SMutableGraphViewer.cpp
+++ b/Engine/Plugins/Experimental/Mutable/Source/CustomizableObjectEditor/Private/MuCOE/SMutableGraphViewer.cpp
@@ -39,7 +39,6 @@
 #include "MuT/NodeObjectGroup.h"
 #include "MuT/NodeObjectNew.h"
 #include "MuT/NodeSurfaceNew.h"
-#include "MuT/NodeSurfaceEdit.h"
 #include "MuT/NodeSurfaceSwitch.h"
 #include "MuT/NodeSurfaceVariation.h"
 #include "MuT/NodeLOD.h"
@@ -53,6 +52,7 @@
 #include "MuT/NodeModifierMeshClipMorphPlane.h"
 #include "MuT/NodeModifierMeshClipWithUVMask.h"
 #include "MuT/NodeModifierMeshClipMorphPlane.h"
+#include "MuT/NodeModifierSurfaceEdit.h"
 #include "MuT/NodeScalarConstant.h"
 #include "MuT/NodeScalarCurve.h"
 #include "MuT/NodeScalarSwitch.h"
@@ -360,13 +360,12 @@ void SMutableGraphViewer::GetChildrenForInfo(TSharedPtr<FMutableGraphTreeElement
 		}
 	}
 
-	else if (ParentNode->GetType() == mu::NodeSurfaceEdit::GetStaticType())
+	else if (ParentNode->GetType() == mu::NodeModifierSurfaceEdit::GetStaticType())
 	{
-		mu::NodeSurfaceEdit* SurfaceEdit = StaticCast<mu::NodeSurfaceEdit*>(ParentNode);
+		mu::NodeModifierSurfaceEdit* SurfaceEdit = StaticCast<mu::NodeModifierSurfaceEdit*>(ParentNode);
 		AddChildFunc(SurfaceEdit->MeshAdd.get(), TEXT("MESH_ADD"));
 		AddChildFunc(SurfaceEdit->MeshRemove.get(), TEXT("MESH_REMOVE"));
-		AddChildFunc(SurfaceEdit->MeshMorph.get(), TEXT("MORPH"));
-		AddChildFunc(SurfaceEdit->MorphFactor.get(), TEXT("MORPH_FACTOR"));
+		AddChildFunc(SurfaceEdit->MorphFactor.get(), FString::Printf(TEXT("MORPH_FACTOR [%s]"), *SurfaceEdit->MeshMorph) );
 
 		for (int32 l = 0; l < SurfaceEdit->Textures.Num(); ++l)
 		{
@@ -490,6 +489,7 @@ void SMutableGraphViewer::GetChildrenForInfo(TSharedPtr<FMutableGraphTreeElement
 	{
 		mu::NodeModifierMeshClipWithUVMask* ModifierMeshClipWithUVMaskVar = StaticCast<mu::NodeModifierMeshClipWithUVMask*>(ParentNode);
 		AddChildFunc(ModifierMeshClipWithUVMaskVar->ClipMask.get(), FString::Printf(TEXT("CLIP MASK")));
+		AddChildFunc(ModifierMeshClipWithUVMaskVar->ClipLayout.get(), FString::Printf(TEXT("CLIP LAYOUT")));
 	}
 
 	else if (ParentNode->GetType() == mu::NodeImageSwitch::GetStaticType())
diff --git a/Engine/Plugins/Experimental/Mutable/Source/CustomizableObjectEditor/Public/MuCOE/CustomizableObjectEditorViewportClient.h b/Engine/Plugins/Experimental/Mutable/Source/CustomizableObjectEditor/Public/MuCOE/CustomizableObjectEditorViewportClient.h
--- a/Engine/Plugins/Experimental/Mutable/Source/CustomizableObjectEditor/Public/MuCOE/CustomizableObjectEditorViewportClient.h
+++ b/Engine/Plugins/Experimental/Mutable/Source/CustomizableObjectEditor/Public/MuCOE/CustomizableObjectEditorViewportClient.h
@@ -18,8 +18,8 @@ class FViewport;
 class ICustomizableObjectInstanceEditor;
 class UAnimationAsset;
 class UCustomizableObject;
-class UCustomizableObjectNodeMeshClipMorph;
-class UCustomizableObjectNodeMeshClipWithMesh;
+class UCustomizableObjectNodeModifierClipMorph;
+class UCustomizableObjectNodeModifierClipWithMesh;
 class UCustomizableObjectNodeProjectorConstant;
 class UCustomizableObjectNodeProjectorParameter;
 class UDebugSkelMeshComponent;
@@ -182,13 +182,13 @@ public:
 	void SetShowBounds();
 
 	/** Do not call directly. Use ICustomizableObjectEditor functions instead. */
-	void ShowGizmoClipMorph(UCustomizableObjectNodeMeshClipMorph& ClipPlainNode);
+	void ShowGizmoClipMorph(UCustomizableObjectNodeModifierClipMorph& ClipPlainNode);
 
 	/** Do not call directly. Use ICustomizableObjectEditor functions instead. */
 	void HideGizmoClipMorph();
 	
 	/** Do not call directly. Use ICustomizableObjectEditor functions instead. */
-	void ShowGizmoClipMesh(UCustomizableObjectNodeMeshClipWithMesh& ClipMeshNode, UObject& ClipMesh, int32 LODIndex, int32 SectionIndex, int32 MaterialSlotIndex);
+	void ShowGizmoClipMesh(UCustomizableObjectNodeModifierClipWithMesh& ClipMeshNode, UObject& ClipMesh, int32 LODIndex, int32 SectionIndex, int32 MaterialSlotIndex);
 
 	/** Do not call directly. Use ICustomizableObjectEditor functions instead. */
 	void HideGizmoClipMesh();
@@ -344,7 +344,7 @@ private:
 	int32 UVDrawLODIndex = 0;
 	int32 UVDrawUVIndex = 0;
 	
-	UCustomizableObjectNodeMeshClipMorph* ClipMorphNode;
+	UCustomizableObjectNodeModifierClipMorph* ClipMorphNode;
 	TObjectPtr<UMaterial> ClipMorphMaterial;
 	bool bClipMorphLocalStartOffset;
 	FVector ClipMorphOrigin;
@@ -355,7 +355,7 @@ private:
 	FVector ClipMorphYAxis;
 	float MorphLength;
 	FBoxSphereBounds MorphBounds;
-	UCustomizableObjectNodeMeshClipWithMesh* ClipMeshNode;
+	UCustomizableObjectNodeModifierClipWithMesh* ClipMeshNode;
 	TObjectPtr<UMaterial> ClipMeshMaterial;
 	UStaticMeshComponent* ClipMeshStaticMeshComp;
 	USkeletalMeshComponent* ClipMeshSkeletalMeshComp;
diff --git a/Engine/Plugins/Experimental/Mutable/Source/CustomizableObjectEditor/Public/MuCOE/CustomizableObjectEditor_Deprecated.h b/Engine/Plugins/Experimental/Mutable/Source/CustomizableObjectEditor/Public/MuCOE/CustomizableObjectEditor_Deprecated.h
--- a/Engine/Plugins/Experimental/Mutable/Source/CustomizableObjectEditor/Public/MuCOE/CustomizableObjectEditor_Deprecated.h
+++ b/Engine/Plugins/Experimental/Mutable/Source/CustomizableObjectEditor/Public/MuCOE/CustomizableObjectEditor_Deprecated.h
@@ -70,7 +70,7 @@ struct CUSTOMIZABLEOBJECTEDITOR_API FCustomizableObjectNodeMaterialScalar
 // Deprecated, do not use!
 constexpr int32 UV_LAYOUT_DEFAULT = -2;
 
-// UCustomizableObjectNodeEditMaterial
+// UCustomizableObjectNodeModifierEditMeshSection
 // Deprecated, do not use!
 USTRUCT()
 struct CUSTOMIZABLEOBJECTEDITOR_API FCustomizableObjectNodeEditMaterialImage
@@ -81,7 +81,7 @@ struct CUSTOMIZABLEOBJECTEDITOR_API FCustomizableObjectNodeEditMaterialImage
 	FString Name;
 };
 
-// UCustomizableObjectNodeExtendMaterial
+// UCustomizableObjectNodeModifierExtendMeshSection
 // Deprecated, do not use!
 USTRUCT()
 struct CUSTOMIZABLEOBJECTEDITOR_API FCustomizableObjectNodeExtendMaterialImage
diff --git a/Engine/Plugins/Experimental/Mutable/Source/CustomizableObjectEditor/Public/MuCOE/ICustomizableObjectInstanceEditor.h b/Engine/Plugins/Experimental/Mutable/Source/CustomizableObjectEditor/Public/MuCOE/ICustomizableObjectInstanceEditor.h
--- a/Engine/Plugins/Experimental/Mutable/Source/CustomizableObjectEditor/Public/MuCOE/ICustomizableObjectInstanceEditor.h
+++ b/Engine/Plugins/Experimental/Mutable/Source/CustomizableObjectEditor/Public/MuCOE/ICustomizableObjectInstanceEditor.h
@@ -5,7 +5,6 @@
 #include "Toolkits/AssetEditorToolkit.h"
 
 class UPointLightComponent;
-class UCustomizableObjectNodeMeshClipWithMesh;
 class UCustomizableObjectNodeProjectorParameter;
 class UCustomizableObjectNodeProjectorConstant;
 class SCustomizableObjectEditorViewportTabBody;
@@ -13,8 +12,9 @@ class UCustomizableObject;
 class UCustomizableObjectInstance;
 class UProjectorParameter;
 class UCustomSettings;
-class UCustomizableObjectNodeMeshClipMorph;
-class UCustomizableObjectNodeMeshClipWithMesh;
+class UCustomizableObjectNodeModifierClipWithMesh;
+class UCustomizableObjectNodeModifierClipMorph;
+class UCustomizableObjectNodeModifierClipWithMesh;
 class ULightComponent;
 class UPoseAsset;
 class UCustomizableObjectEditorProperties;
@@ -70,7 +70,7 @@ public:
 	
 	/** Show the clip morph plane gizmo of the NodeClipMorph.
 	 * Synchronizes all editor widgets so that the selection and the widgets is always consistent. */
-	virtual void ShowGizmoClipMorph(UCustomizableObjectNodeMeshClipMorph& ClipPlainNode) {}
+	virtual void ShowGizmoClipMorph(UCustomizableObjectNodeModifierClipMorph& ClipPlainNode) {}
 
 	/** Hide the clip morph plane gizmo of the NodeClipMorph.
 	 * Synchronizes all editor widgets so that the selection and the widgets is always consistent. */
@@ -78,7 +78,7 @@ public:
 	
 	/** Show the clipping mesh gizmo from the NodeMeshClipWithMesh.
 	 * Synchronizes all editor widgets so that the selection and the widgets is always consistent. */
-	virtual void ShowGizmoClipMesh(UCustomizableObjectNodeMeshClipWithMesh& ClipMeshNode) {}
+	virtual void ShowGizmoClipMesh(UCustomizableObjectNodeModifierClipWithMesh& ClipMeshNode) {}
 
 	/** Hide the clipping mesh gizmo from the NodeMeshClipWithMesh.
 	 * Synchronizes all editor widgets so that the selection and the widgets is always consistent. */
diff --git a/Engine/Plugins/Experimental/Mutable/Source/CustomizableObjectEditor/Public/MuCOE/Nodes/CustomizableObjectNode.h b/Engine/Plugins/Experimental/Mutable/Source/CustomizableObjectEditor/Public/MuCOE/Nodes/CustomizableObjectNode.h
--- a/Engine/Plugins/Experimental/Mutable/Source/CustomizableObjectEditor/Public/MuCOE/Nodes/CustomizableObjectNode.h
+++ b/Engine/Plugins/Experimental/Mutable/Source/CustomizableObjectEditor/Public/MuCOE/Nodes/CustomizableObjectNode.h
@@ -164,7 +164,11 @@ public:
 	virtual bool CanConnect( const UEdGraphPin* InOwnedInputPin, const UEdGraphPin* InOutputPin, bool& bOutIsOtherNodeBlocklisted, bool& bOutArePinsCompatible) const;
 
 	// Used during compilation process to cache the node for all LODs, or generate it specifically for each of them.
-	virtual bool IsAffectedByLOD() const { return true; }
+	virtual bool IsAffectedByLOD() const;
+
+	/** Return an array of tags that this node will enable and apply to its data. Null if none, or of it doesn't apply. */
+	virtual TArray<FString>* GetEnableTags();
+
 
 	// Get the CustomizableObject graph that owns this node
 	class UCustomizableObjectGraph* GetCustomizableObjectGraph() const;
@@ -218,12 +222,6 @@ public:
 
 	/** Return the LOD which this node belongs to. Return -1 if it is not connected (directly or indirectly) to a LOD pin. */
 	int32 GetLOD() const;
-
-	/** Get all parent Object nodes of this node which are parents of this node.
-	 * 
-	 * If a Object node is connected through a LOD pin, it has to match the given LOD.
-	 */
-	TArray<UCustomizableObjectNodeObject*> GetParentObjectNodes(int LOD) const;
 	
 	/** Set a pin to be hidden or not. */
 	void SetPinHidden(UEdGraphPin& Pin, bool bHidden);
@@ -324,7 +322,7 @@ private:
 	UPROPERTY()
 	TMap<FGuid, TObjectPtr<UCustomizableObjectNodePinData>> PinsDataId;
 
-	// Deprectated properties
+	// Deprecated properties in favour of PinsDataId
 	UPROPERTY()
 	TMap<FEdGraphPinReference, TObjectPtr<UCustomizableObjectNodePinData>> PinsData_DEPRECATED;
 };
diff --git a/Engine/Plugins/Experimental/Mutable/Source/CustomizableObjectEditor/Public/MuCOE/Nodes/CustomizableObjectNodeCopyMaterial.h b/Engine/Plugins/Experimental/Mutable/Source/CustomizableObjectEditor/Public/MuCOE/Nodes/CustomizableObjectNodeCopyMaterial.h
--- a/Engine/Plugins/Experimental/Mutable/Source/CustomizableObjectEditor/Public/MuCOE/Nodes/CustomizableObjectNodeCopyMaterial.h
+++ b/Engine/Plugins/Experimental/Mutable/Source/CustomizableObjectEditor/Public/MuCOE/Nodes/CustomizableObjectNodeCopyMaterial.h
@@ -36,7 +36,7 @@ public:
 	virtual UMaterialInterface* GetMaterial() const override;
 	virtual bool IsReuseMaterialBetweenLODs() const override;
 	virtual FName GetMeshComponentName() const override;
-	virtual TArray<FString> GetTags() const override;
+	virtual TArray<FString>* GetEnableTags() override;
 	virtual UEdGraphPin* GetMaterialAssetPin() const override;
 	virtual int32 GetNumParameters(EMaterialParameterType Type) const override;
 	virtual FNodeMaterialParameterId GetParameterId(EMaterialParameterType Type, int32 ParameterIndex) const override;
diff --git a/Engine/Plugins/Experimental/Mutable/Source/CustomizableObjectEditor/Public/MuCOE/Nodes/CustomizableObjectNodeMaterial.h b/Engine/Plugins/Experimental/Mutable/Source/CustomizableObjectEditor/Public/MuCOE/Nodes/CustomizableObjectNodeMaterial.h
--- a/Engine/Plugins/Experimental/Mutable/Source/CustomizableObjectEditor/Public/MuCOE/Nodes/CustomizableObjectNodeMaterial.h
+++ b/Engine/Plugins/Experimental/Mutable/Source/CustomizableObjectEditor/Public/MuCOE/Nodes/CustomizableObjectNodeMaterial.h
@@ -133,6 +133,7 @@ public:
 	virtual FString GetRefreshMessage() const override;
 	virtual TSharedPtr<IDetailsView> CustomizePinDetails(const UEdGraphPin& Pin) const override;
 	void ReconstructNode(UCustomizableObjectNodeRemapPins* RemapPinsMode);
+	virtual TArray<FString>* GetEnableTags() override;
 
 	
 	// UCustomizableObjectNodeMaterialBase interface
@@ -141,7 +142,6 @@ public:
 	virtual UMaterialInterface* GetMaterial() const override;
 	virtual bool IsReuseMaterialBetweenLODs() const override;
 	virtual FName GetMeshComponentName() const override;
-	virtual TArray<FString> GetTags() const override;
 	virtual int32 GetNumParameters(EMaterialParameterType Type) const override;
 	virtual FNodeMaterialParameterId GetParameterId(EMaterialParameterType Type, int32 ParameterIndex) const override;
 	virtual FName GetParameterName(EMaterialParameterType Type, int32 ParameterIndex) const override;
@@ -166,6 +166,12 @@ public:
 	void SetMaterial(UMaterialInterface* InMaterial);
 	void SetComponentName(const FName& Name);
 
+	static bool HasParameter(const UMaterialInterface* InMaterial, const FNodeMaterialParameterId& ParameterId);
+	static int32 GetParameterLayerIndex(const UMaterialInterface* InMaterial, EMaterialParameterType Type, int32 ParameterIndex);
+
+	UPROPERTY(EditAnywhere, Category = CustomizableObject)
+	TArray<FString> Tags;
+
 private:
 	/** Set the Pin Mode of a Texture Parameter Pin. */
 	void SetImagePinMode(UEdGraphPin& Pin, EPinMode PinMode) const;
@@ -179,9 +185,6 @@ private:
 	UPROPERTY(EditAnywhere, Category=CustomizableObject, DisplayName = "Default Texture Parameter Mode", Meta = (ToolTip = "All Mateiral Texture Parameters set to \"Node Defined\" will use this mode."))
 	ENodePinMode TextureParametersMode = ENodePinMode::Passthrough;
 
-	UPROPERTY(EditAnywhere, Category = CustomizableObject)
-	TArray<FString> Tags;
-
 	UPROPERTY()
 	int32 MeshComponentIndex_DEPRECATED = 0;
 
@@ -196,6 +199,7 @@ private:
 	/** Last static or skeletal mesh connected. Used to remove the callback once disconnected. */
 	TWeakObjectPtr<UCustomizableObjectNode> LastMeshNodeConnected;
 
+	/** List of material parameter types that are actually relevant to mutable. */
 	static const TArray<EMaterialParameterType> ParameterTypes;
 
 	/** Relates a Parameter id (key) (and layer if is a layered material) to a Pin (value). Only used to improve performance.
diff --git a/Engine/Plugins/Experimental/Mutable/Source/CustomizableObjectEditor/Public/MuCOE/Nodes/CustomizableObjectNodeMaterialBase.h b/Engine/Plugins/Experimental/Mutable/Source/CustomizableObjectEditor/Public/MuCOE/Nodes/CustomizableObjectNodeMaterialBase.h
--- a/Engine/Plugins/Experimental/Mutable/Source/CustomizableObjectEditor/Public/MuCOE/Nodes/CustomizableObjectNodeMaterialBase.h
+++ b/Engine/Plugins/Experimental/Mutable/Source/CustomizableObjectEditor/Public/MuCOE/Nodes/CustomizableObjectNodeMaterialBase.h
@@ -69,9 +69,7 @@ public:
 	virtual bool IsReuseMaterialBetweenLODs() const PURE_VIRTUAL(UCustomizableObjectNodeMaterial::IsReuseMaterialBetweenLODs, return {}; );
 
 	virtual FName GetMeshComponentName() const PURE_VIRTUAL(UCustomizableObjectNodeMaterial::GetComponentName, return {}; );
-	
-	virtual TArray<FString> GetTags() const PURE_VIRTUAL(UCustomizableObjectNodeMaterial::GetTags, return {}; );
-	
+		
 	virtual UEdGraphPin* GetMeshPin() const PURE_VIRTUAL(UCustomizableObjectNodeMaterial::GetMeshPin, return {}; );
 
 	virtual UEdGraphPin* GetMaterialAssetPin() const PURE_VIRTUAL(UCustomizableObjectNodeMaterial::GetMaterialAssetPin, return {}; );
diff --git a/Engine/Plugins/Experimental/Mutable/Source/CustomizableObjectEditor/Public/MuCOE/Nodes/CustomizableObjectNodeModifierBase.h b/Engine/Plugins/Experimental/Mutable/Source/CustomizableObjectEditor/Public/MuCOE/Nodes/CustomizableObjectNodeModifierBase.h
--- a/Engine/Plugins/Experimental/Mutable/Source/CustomizableObjectEditor/Public/MuCOE/Nodes/CustomizableObjectNodeModifierBase.h
+++ b/Engine/Plugins/Experimental/Mutable/Source/CustomizableObjectEditor/Public/MuCOE/Nodes/CustomizableObjectNodeModifierBase.h
@@ -2,24 +2,70 @@
 
 #pragma once
 
-#include "MuCOE/Nodes/CustomizableObjectNodeMaterialBase.h"
+#include "MuCOE/Nodes/CustomizableObjectNode.h"
 
 #include "CustomizableObjectNodeModifierBase.generated.h"
 
-class UObject;
+enum class EMutableMultipleTagPolicy:uint8;
+
+
+USTRUCT()
+struct FLegacyTag
+{
+	GENERATED_BODY()
+
+	UPROPERTY()
+	TSoftObjectPtr<UCustomizableObject> ParentObject;
+
+	UPROPERTY()
+	FGuid ParentNode;
+
+	UPROPERTY()
+	FString Tag;
+
+	bool operator==(const FLegacyTag&) const = default;
+};
 
 
 UCLASS(abstract)
 class CUSTOMIZABLEOBJECTEDITOR_API UCustomizableObjectNodeModifierBase : public UCustomizableObjectNode
 {
-public:
+
 	GENERATED_BODY()
 
+public:
+
 	// EdGraphNode interface
 	virtual FLinearColor GetNodeTitleColor() const override;
 
 	// UCustomizableObjectNode interface
 	virtual void BackwardsCompatibleFixup(int32 CustomizableObjectCustomVersion) override;
 	virtual void PostBackwardsCompatibleFixup() override;
+
+	// Own interface
+	virtual UEdGraphPin* OutputPin() const;
+	virtual const TArray<FString>* GetRequiredTags() const PURE_VIRTUAL(UCustomizableObjectNodeModifierBase::GetRequiredTags, return nullptr; );
+
+	virtual EMutableMultipleTagPolicy GetMultipleTagsPolicy() const PURE_VIRTUAL(UCustomizableObjectNodeModifierBase::GetMultipleTagsPolicy, return {}; );
+
+	/** Return true if the tags and policy in this node would make this node modify the given parameter node. */
+	bool IsApplicableTo(UCustomizableObjectNode* Candidate) const;
+
+	/** Get the list of all nodes of this object that may be affected by this modifier. This is an approximate
+	* query: it includes all nodes that may be affected by this modifier but it may include nodes that will not 
+	* be modified by this modifier because of other conditions.
+	*/
+	void GetPossiblyModifiedNodes(TArray<UCustomizableObjectNode*>&) const;
+
+	/** Generate a tag that is unique and stable for the given node. */
+	static FString MakeNodeAutoTag(UEdGraphNode*);
+
+protected:
+
+	/** This is a list of tags autogenerated due to backwards compatibility upgrade. 
+	* It must be removed when we remove the related backwards compatibility code.
+	*/
+	UPROPERTY()
+	TArray<FLegacyTag> LegacyBackportsRequiredTags;
 };
 
diff --git a/Engine/Plugins/Experimental/Mutable/Source/CustomizableObjectEditor/Public/MuCOE/Nodes/CustomizableObjectNodeModifierClipDeform.h b/Engine/Plugins/Experimental/Mutable/Source/CustomizableObjectEditor/Public/MuCOE/Nodes/CustomizableObjectNodeModifierClipDeform.h
new file mode 100644
--- /dev/null
+++ b/Engine/Plugins/Experimental/Mutable/Source/CustomizableObjectEditor/Public/MuCOE/Nodes/CustomizableObjectNodeModifierClipDeform.h
@@ -0,0 +1,58 @@
+// Copyright Epic Games, Inc. All Rights Reserved.
+
+#pragma once
+
+#include "MuCOE/Nodes/CustomizableObjectNodeModifierBase.h"
+#include "MuT/NodeModifier.h"
+
+#include "CustomizableObjectNodeModifierClipDeform.generated.h"
+
+class UCustomizableObjectNodeRemapPins;
+class UEdGraphPin;
+class UObject;
+
+UENUM()
+enum class EShapeBindingMethod : uint32
+{
+	ClosestProject = 0,
+	ClosestToSurface = 1,
+	NormalProject = 2
+};
+
+UCLASS()
+class CUSTOMIZABLEOBJECTEDITOR_API UCustomizableObjectNodeModifierClipDeform : public UCustomizableObjectNodeModifierBase
+{
+
+	GENERATED_BODY()
+
+public:
+
+	UPROPERTY(EditAnywhere, BlueprintReadWrite, Category = MeshClipDeform)
+	TArray<FString> Tags;
+	
+	/** Policy to use tags in case more than one is added. */
+	UPROPERTY(EditAnywhere, Category = MeshClipDeform)
+	EMutableMultipleTagPolicy MultipleTagPolicy = EMutableMultipleTagPolicy::OnlyOneRequired;
+
+	UPROPERTY(EditAnywhere, Category = MeshClipDeform)
+	EShapeBindingMethod BindingMethod;
+	
+public:
+
+	// EdGraphNode interface
+	virtual FText GetNodeTitle(ENodeTitleType::Type TitleType) const override;
+	virtual FText GetTooltipText() const override;
+
+	// UCustomizableObjectNode interface
+	virtual void AllocateDefaultPins(UCustomizableObjectNodeRemapPins* RemapPins) override;
+
+	// UCustomizableObjectNodeModifierBase interface
+	virtual const TArray<FString>* GetRequiredTags() const override { return &Tags; }
+	virtual EMutableMultipleTagPolicy GetMultipleTagsPolicy() const override { return MultipleTagPolicy; }
+
+	// Own interface
+
+	/** Input pins. */
+	UEdGraphPin* ClipShapePin() const;
+};
+
diff --git a/Engine/Plugins/Experimental/Mutable/Source/CustomizableObjectEditor/Public/MuCOE/Nodes/CustomizableObjectNodeModifierClipMorph.h b/Engine/Plugins/Experimental/Mutable/Source/CustomizableObjectEditor/Public/MuCOE/Nodes/CustomizableObjectNodeModifierClipMorph.h
new file mode 100644
--- /dev/null
+++ b/Engine/Plugins/Experimental/Mutable/Source/CustomizableObjectEditor/Public/MuCOE/Nodes/CustomizableObjectNodeModifierClipMorph.h
@@ -0,0 +1,104 @@
+// Copyright Epic Games, Inc. All Rights Reserved.
+
+#pragma once
+
+#include "MuCOE/Nodes/CustomizableObjectNodeModifierBase.h"
+#include "MuT/NodeModifier.h"
+
+#include "CustomizableObjectNodeModifierClipMorph.generated.h"
+
+namespace ENodeTitleType { enum Type : int; }
+
+class FArchive;
+class UCustomizableObjectNodeRemapPins;
+class UEdGraphPin;
+class UObject;
+struct FPropertyChangedEvent;
+
+UCLASS()
+class CUSTOMIZABLEOBJECTEDITOR_API UCustomizableObjectNodeModifierClipMorph : public UCustomizableObjectNodeModifierBase
+{
+public:
+	GENERATED_BODY()
+
+	UPROPERTY(EditAnywhere, Category = MeshToClipAndMorph)
+	FName BoneName;
+
+	UPROPERTY(EditAnywhere, BlueprintReadWrite, Category = MeshToClipAndMorph)
+	TArray<FString> Tags;
+
+	/** Policy to use tags in case more than one is added. */
+	UPROPERTY(EditAnywhere, Category = MeshToClipAndMorph)
+	EMutableMultipleTagPolicy MultipleTagPolicy = EMutableMultipleTagPolicy::OnlyOneRequired;
+
+	UPROPERTY(EditAnywhere, Category = MeshToClipAndMorph)
+	uint32 ReferenceSkeletonIndex = 0;
+
+	UPROPERTY(EditAnywhere, Category = MeshClipParameters, meta = (DisplayName="Morph Start Offset", ToolTip="Offset from the origin of the selected bone to the actual start of the morph."))
+	FVector StartOffset;
+
+	UPROPERTY(EditAnywhere, Category = MeshClipParameters, meta = (DisplayName = "Local Start Offset", ToolTip = "Toggles between a local or global start offset."))
+	bool bLocalStartOffset;
+
+	UPROPERTY(EditAnywhere, Category = MeshClipParameters, meta = (DisplayName="Morph length", ToolTip="The length from the morph start to the clip plane."))
+	float B;
+
+	UPROPERTY(EditAnywhere, Category = MeshClipParameters, meta = (DisplayName="Ellipse Radius 1", ToolTip="First radius of the ellipse that the mesh is morphed into."))
+	float Radius;
+
+	UPROPERTY(EditAnywhere, Category = MeshClipParameters, meta = (DisplayName = "Ellipse Radius 2", ToolTip = "Second radius of the ellipse that the mesh is morphed into."))
+	float Radius2;
+
+	UPROPERTY(EditAnywhere, Category = MeshClipParameters, meta = (DisplayName = "Ellipse Rotation", ToolTip = "Ellipse Rotation in degrees around the bone axis."))
+	float RotationAngle;
+
+	UPROPERTY(EditAnywhere, Category = MeshClipParameters, meta = (DisplayName = "Morph Curve Control", ToolTip = "Controls the morph curve shape. A value of 1 is linear, less than 1 is concave and greater than 1 convex."))
+	float Exponent;
+
+	UPROPERTY(EditAnywhere, Category = MeshClipParameters, meta = (DisplayName = "Invert normal direction", ToolTip = "Flag to invert the normal direction"))
+	bool bInvertNormal;
+
+	UPROPERTY()
+	FVector Origin;
+
+	UPROPERTY()
+	FVector Normal;
+
+private:
+	UPROPERTY()
+	bool bOldOffset_DEPRECATED;
+
+public:
+	UPROPERTY(EditAnywhere, Category = MeshClipParameters, meta = (DisplayName = "Max Effect Radius", ToolTip = "The maximum distance from the origin of the widget where vertices will be affected. If negative, there will be no limit."))
+	float MaxEffectRadius;
+
+	bool bUpdateViewportWidget;
+
+	UCustomizableObjectNodeModifierClipMorph();
+
+	// UObject interface.
+	virtual void PostEditChangeProperty(FPropertyChangedEvent& PropertyChangedEvent) override;
+	virtual void Serialize(FArchive& Ar) override;
+
+	// EdGraphNode interface
+	virtual FText GetNodeTitle(ENodeTitleType::Type TitleType) const override;
+	virtual FText GetTooltipText() const override;
+
+	// UCustomizableObjectNode interface
+	virtual void AllocateDefaultPins(UCustomizableObjectNodeRemapPins* RemapPins) override;
+	virtual void PinConnectionListChanged(UEdGraphPin* Pin) override;
+
+	// UCustomizableObjectNodeModifierBase interface
+	virtual const TArray<FString>* GetRequiredTags() const override { return &Tags; }
+	virtual EMutableMultipleTagPolicy GetMultipleTagsPolicy() const override { return MultipleTagPolicy; }
+
+	// Own interface
+	UEdGraphPin* OutputPin() const;
+
+	FVector GetOriginWithOffset() const;
+	void FindLocalAxes(FVector& Right, FVector& Up, FVector& Forward) const;
+
+	// Change StartOffset from World to Local of the other way around
+	void ChangeStartOffsetTransform();
+};
+
diff --git a/Engine/Plugins/Experimental/Mutable/Source/CustomizableObjectEditor/Public/MuCOE/Nodes/CustomizableObjectNodeModifierClipWithMesh.h b/Engine/Plugins/Experimental/Mutable/Source/CustomizableObjectEditor/Public/MuCOE/Nodes/CustomizableObjectNodeModifierClipWithMesh.h
new file mode 100644
--- /dev/null
+++ b/Engine/Plugins/Experimental/Mutable/Source/CustomizableObjectEditor/Public/MuCOE/Nodes/CustomizableObjectNodeModifierClipWithMesh.h
@@ -0,0 +1,72 @@
+// Copyright Epic Games, Inc. All Rights Reserved.
+
+#pragma once
+
+#include "MuCOE/Nodes/CustomizableObjectNodeModifierBase.h"
+#include "MuT/NodeModifier.h"
+
+#include "CustomizableObjectNodeModifierClipWithMesh.generated.h"
+
+namespace ENodeTitleType { enum Type : int; }
+
+class UCustomizableObject;
+class UCustomizableObjectNodeRemapPins;
+class UEdGraphPin;
+class UObject;
+struct FGuid;
+
+UCLASS()
+class CUSTOMIZABLEOBJECTEDITOR_API UCustomizableObjectNodeModifierClipWithMesh : public UCustomizableObjectNodeModifierBase
+{
+public:
+	GENERATED_BODY()
+
+	UCustomizableObjectNodeModifierClipWithMesh();
+
+	UPROPERTY(EditAnywhere, BlueprintReadWrite, Category = MeshToClip)
+	TArray<FString> Tags;
+
+	UPROPERTY(EditAnywhere, Category = CustomizableObjectToClip)
+	EMutableMultipleTagPolicy MultipleTagPolicy = EMutableMultipleTagPolicy::OnlyOneRequired;
+
+	//!< If assigned, then a material inside this CO will be clipped by this node.
+    //!< If several materials with the same name, all are considered (to cover all LOD levels)
+    UPROPERTY(EditAnywhere, Category = CustomizableObjectToClip)
+	TObjectPtr<UCustomizableObject> CustomizableObjectToClipWith;
+
+    //!< Array with the Guids of the nodes with the same material inside the CustomizableObjectToClipWith CO (if any is assigned)
+    UPROPERTY(EditAnywhere, Category = CustomizableObjectToClip)
+    TArray<FGuid> ArrayMaterialNodeToClipWithID;
+
+	/** Mesh Transform*/
+	UPROPERTY(EditAnywhere, Category = CustomizableObject)
+	FTransform Transform;
+
+	// Details view variables
+	// The clipping node uses tags
+	bool bUseTags;
+
+	// The clipping node uses a material name
+	bool bUseMaterials;
+
+	// EdGraphNode interface
+	virtual FText GetNodeTitle(ENodeTitleType::Type TitleType) const override;
+	virtual FText GetTooltipText() const override;
+
+	// UCustomizableObjectNode interface
+	virtual void AllocateDefaultPins(UCustomizableObjectNodeRemapPins* RemapPins) override;
+	virtual void ReconstructNode(UCustomizableObjectNodeRemapPins* RemapPinsMode) override;
+	virtual void PinConnectionListChanged(UEdGraphPin* Pin) override;
+	virtual void UpdateReferencedNodeId(const FGuid& NewGuid) override;
+	virtual void BeginPostDuplicate(bool bDuplicateForPIE) override;
+
+	// UCustomizableObjectNodeModifierBase interface
+	virtual const TArray<FString>* GetRequiredTags() const override { return &Tags; }
+	virtual EMutableMultipleTagPolicy GetMultipleTagsPolicy() const override { return MultipleTagPolicy; }
+
+	// Own interface
+	UEdGraphPin* OutputPin() const;
+
+	UEdGraphPin* ClipMeshPin() const;
+};
+
diff --git a/Engine/Plugins/Experimental/Mutable/Source/CustomizableObjectEditor/Public/MuCOE/Nodes/CustomizableObjectNodeModifierClipWithUVMask.h b/Engine/Plugins/Experimental/Mutable/Source/CustomizableObjectEditor/Public/MuCOE/Nodes/CustomizableObjectNodeModifierClipWithUVMask.h
--- a/Engine/Plugins/Experimental/Mutable/Source/CustomizableObjectEditor/Public/MuCOE/Nodes/CustomizableObjectNodeModifierClipWithUVMask.h
+++ b/Engine/Plugins/Experimental/Mutable/Source/CustomizableObjectEditor/Public/MuCOE/Nodes/CustomizableObjectNodeModifierClipWithUVMask.h
@@ -16,9 +16,10 @@ class UEdGraphPin;
 UCLASS()
 class CUSTOMIZABLEOBJECTEDITOR_API UCustomizableObjectNodeModifierClipWithUVMask : public UCustomizableObjectNodeModifierBase
 {
-public:
 	GENERATED_BODY()
 
+public:
+
 	/** Materials in all other objects that activate this tags will be clipped with this UV mask. */
 	UPROPERTY(EditAnywhere, BlueprintReadWrite, Category = MeshToClip)
 	TArray<FString> Tags;
@@ -31,6 +32,8 @@ public:
 	UPROPERTY(EditAnywhere, BlueprintReadWrite, Category = MeshToClip)
 	int32 UVChannelForMask = 0;
 
+public:
+
 	// EdGraphNode interface
 	virtual FText GetNodeTitle(ENodeTitleType::Type TitleType) const override;
 	virtual FText GetTooltipText() const override;
@@ -39,9 +42,13 @@ public:
 	virtual void AllocateDefaultPins(UCustomizableObjectNodeRemapPins* RemapPins) override;
 	virtual void PinConnectionListChanged(UEdGraphPin* Pin) override;
 
+	// UCustomizableObjectNodeModifierBase interface
+	virtual const TArray<FString>* GetRequiredTags() const override { return &Tags; }
+	virtual EMutableMultipleTagPolicy GetMultipleTagsPolicy() const override { return MultipleTagPolicy; }
+
 	// Own interface
-	UEdGraphPin* OutputPin() const;
 
+	/** Access to input pins. */
 	UEdGraphPin* ClipMaskPin() const;
 };
 
diff --git a/Engine/Plugins/Experimental/Mutable/Source/CustomizableObjectEditor/Public/MuCOE/Nodes/CustomizableObjectNodeModifierEditMeshSection.h b/Engine/Plugins/Experimental/Mutable/Source/CustomizableObjectEditor/Public/MuCOE/Nodes/CustomizableObjectNodeModifierEditMeshSection.h
new file mode 100644
--- /dev/null
+++ b/Engine/Plugins/Experimental/Mutable/Source/CustomizableObjectEditor/Public/MuCOE/Nodes/CustomizableObjectNodeModifierEditMeshSection.h
@@ -0,0 +1,86 @@
+// Copyright Epic Games, Inc. All Rights Reserved.
+
+#pragma once
+
+#include "MuCOE/Nodes/CustomizableObjectNodeModifierEditLayoutBlocks.h"
+#include "MuCOE/Nodes/CustomizableObjectNodeMaterial.h"
+
+#include "CustomizableObjectNodeModifierEditMeshSection.generated.h"
+
+namespace ENodeTitleType { enum Type : int; }
+struct FCustomizableObjectNodeEditMaterialImage;
+struct FNodeMaterialParameterId;
+
+class FArchive;
+class FCustomizableObjectNodeParentedMaterial;
+class UCustomizableObject;
+class UCustomizableObjectNodeRemapPins;
+class UObject;
+
+
+/** Additional data for the Parent Texture Parameter to edit pin. */
+UCLASS()
+class CUSTOMIZABLEOBJECTEDITOR_API UCustomizableObjectNodeEditMaterialPinEditImageData : public UCustomizableObjectNodePinData
+{
+	GENERATED_BODY()
+
+public:
+
+	/** Parameter FGuid + Layer Index */
+	UPROPERTY()
+	FNodeMaterialParameterId ImageParamId;
+
+	UPROPERTY();
+	FEdGraphPinReference PinMask;
+
+	UPROPERTY()
+	FGuid ImageId_DEPRECATED;
+};
+
+
+UCLASS()
+class CUSTOMIZABLEOBJECTEDITOR_API UCustomizableObjectNodeModifierEditMeshSection : public UCustomizableObjectNodeModifierEditLayoutBlocks
+{
+
+	GENERATED_BODY()
+
+public:
+	
+	// EdGraphNode interface
+	virtual FText GetNodeTitle(ENodeTitleType::Type TitleType) const override;
+	virtual FText GetTooltipText() const override;
+
+	// UCustomizableObjectNode interface
+	virtual void BackwardsCompatibleFixup(int32 CustomizableObjectCustomVersion) override;
+	virtual void AllocateDefaultPins(UCustomizableObjectNodeRemapPins* RemapPins) override;
+	virtual FString GetRefreshMessage() const override;
+	virtual bool IsSingleOutputNode() const override;
+	virtual bool CustomRemovePin(UEdGraphPin& Pin) override;
+	virtual bool HasPinViewer() const override;
+	
+	// Own interface
+	
+	/** */
+	virtual void SetLayoutIndex(const int32 LayoutIndex);
+
+	/** Returns the Image mask pin of the given Image that will be edited.
+	 *
+	 * @returns Always returns a valid pin if EditsImage(const FGuid&) returns true. */
+	const UEdGraphPin* GetUsedImageMaskPin(const FNodeMaterialParameterId& ImageId) const;
+
+private:
+
+	// Deprecated properties
+	UPROPERTY()
+	TMap<FNodeMaterialParameterId, FEdGraphPinReference> PinsParameterMap_DEPRECATED;
+
+	UPROPERTY()
+	TMap<FGuid, FEdGraphPinReference> PinsParameter_DEPRECATED;
+	
+	UPROPERTY()
+	TArray<FCustomizableObjectNodeEditMaterialImage> Images_DEPRECATED;
+	
+	// Old layout blocks to patch. Now in parent class.
+	UPROPERTY()
+	TArray<int32> Blocks_DEPRECATED;
+};
diff --git a/Engine/Plugins/Experimental/Mutable/Source/CustomizableObjectEditor/Public/MuCOE/Nodes/CustomizableObjectNodeModifierExtendMeshSection.h b/Engine/Plugins/Experimental/Mutable/Source/CustomizableObjectEditor/Public/MuCOE/Nodes/CustomizableObjectNodeModifierExtendMeshSection.h
new file mode 100644
--- /dev/null
+++ b/Engine/Plugins/Experimental/Mutable/Source/CustomizableObjectEditor/Public/MuCOE/Nodes/CustomizableObjectNodeModifierExtendMeshSection.h
@@ -0,0 +1,72 @@
+// Copyright Epic Games, Inc. All Rights Reserved.
+
+#pragma once
+
+#include "MuCOE/Nodes/CustomizableObjectNodeModifierWithMaterial.h"
+#include "MuCOE/CustomizableObjectLayout.h"
+
+#include "CustomizableObjectNodeModifierExtendMeshSection.generated.h"
+
+namespace ENodeTitleType { enum Type : int; }
+
+class FArchive;
+class UCustomizableObject;
+class UCustomizableObjectNode;
+class UCustomizableObjectNodeRemapPins;
+class UEdGraphPin;
+class UObject;
+struct FCustomizableObjectNodeExtendMaterialImage;
+struct FEdGraphPinReference;
+
+
+UCLASS()
+class CUSTOMIZABLEOBJECTEDITOR_API UCustomizableObjectNodeModifierExtendMeshSection : public UCustomizableObjectNodeModifierWithMaterial
+{
+	GENERATED_BODY()
+
+public:
+
+	/** Tags enabled when the modifier is applied. Also used to decide what modifiers get applied to the data added by this node.  */
+	UPROPERTY(EditAnywhere, BlueprintReadWrite, Category = AddedData)
+	TArray<FString> Tags;
+
+public:
+	
+	// Begin EdGraphNode interface
+	virtual FText GetNodeTitle(ENodeTitleType::Type TitleType) const override;
+	virtual FText GetTooltipText() const override;
+
+	// UCustomizableObjectNode interface
+	virtual void BackwardsCompatibleFixup(int32 CustomizableObjectCustomVersion) override;
+	virtual void AllocateDefaultPins(UCustomizableObjectNodeRemapPins* RemapPins) override;
+	virtual FString GetRefreshMessage() const override;
+	virtual bool IsSingleOutputNode() const override;
+	virtual bool CustomRemovePin(UEdGraphPin& Pin) override;
+	virtual TArray<FString>* GetEnableTags() override;
+
+	// Own interface
+	UEdGraphPin* AddMeshPin() const;
+
+	TArray<UCustomizableObjectLayout*> GetLayouts() const;
+
+private:
+
+	// Deprecated properties
+
+	UPROPERTY()
+	TObjectPtr<UCustomizableObject> ParentMaterialObject_DEPRECATED = nullptr;
+
+	UPROPERTY()
+	FGuid ParentMaterialNodeId_DEPRECATED;
+	
+
+	UPROPERTY()
+	TMap<FNodeMaterialParameterId, FEdGraphPinReference> PinsParameterMap_DEPRECATED;
+
+	UPROPERTY()
+	TMap<FGuid, FEdGraphPinReference> PinsParameter_DEPRECATED;
+
+	UPROPERTY()
+	TArray<FCustomizableObjectNodeExtendMaterialImage> Images_DEPRECATED;
+};
+
diff --git a/Engine/Plugins/Experimental/Mutable/Source/CustomizableObjectEditor/Public/MuCOE/Nodes/CustomizableObjectNodeModifierMorphMeshSection.h b/Engine/Plugins/Experimental/Mutable/Source/CustomizableObjectEditor/Public/MuCOE/Nodes/CustomizableObjectNodeModifierMorphMeshSection.h
new file mode 100644
--- /dev/null
+++ b/Engine/Plugins/Experimental/Mutable/Source/CustomizableObjectEditor/Public/MuCOE/Nodes/CustomizableObjectNodeModifierMorphMeshSection.h
@@ -0,0 +1,43 @@
+// Copyright Epic Games, Inc. All Rights Reserved.
+
+#pragma once
+
+#include "MuCOE/Nodes/CustomizableObjectNodeModifierEditMeshSectionBase.h"
+
+#include "CustomizableObjectNodeModifierMorphMeshSection.generated.h"
+
+namespace ENodeTitleType { enum Type : int; }
+
+class UCustomizableObjectNodeRemapPins;
+class UEdGraphPin;
+
+
+UCLASS()
+class CUSTOMIZABLEOBJECTEDITOR_API UCustomizableObjectNodeModifierMorphMeshSection: public UCustomizableObjectNodeModifierEditMeshSectionBase
+{
+
+	GENERATED_BODY()
+	
+public:
+
+	UPROPERTY(EditAnywhere, Category = CustomizableObject)
+	FString MorphTargetName;
+
+public:
+
+	// Begin EdGraphNode interface
+	FText GetNodeTitle(ENodeTitleType::Type TitleType) const override;
+	FText GetTooltipText() const override;
+
+	// UCustomizableObjectNode interface
+	virtual void AllocateDefaultPins(UCustomizableObjectNodeRemapPins* RemapPins) override;
+	virtual FString GetRefreshMessage() const override;
+	virtual bool IsSingleOutputNode() const override;
+
+	UEdGraphPin* FactorPin() const
+	{
+		return FindPin(TEXT("Factor"));
+	}
+
+};
+
diff --git a/Engine/Plugins/Experimental/Mutable/Source/CustomizableObjectEditor/Public/MuCOE/Nodes/CustomizableObjectNodeModifierRemoveMesh.h b/Engine/Plugins/Experimental/Mutable/Source/CustomizableObjectEditor/Public/MuCOE/Nodes/CustomizableObjectNodeModifierRemoveMesh.h
new file mode 100644
--- /dev/null
+++ b/Engine/Plugins/Experimental/Mutable/Source/CustomizableObjectEditor/Public/MuCOE/Nodes/CustomizableObjectNodeModifierRemoveMesh.h
@@ -0,0 +1,36 @@
+// Copyright Epic Games, Inc. All Rights Reserved.
+
+#pragma once
+
+#include "MuCOE/Nodes/CustomizableObjectNodeModifierEditMeshSectionBase.h"
+
+#include "CustomizableObjectNodeModifierRemoveMesh.generated.h"
+
+namespace ENodeTitleType { enum Type : int; }
+
+class UCustomizableObjectNodeRemapPins;
+class UEdGraphPin;
+class UObject;
+
+UCLASS()
+class CUSTOMIZABLEOBJECTEDITOR_API UCustomizableObjectNodeModifierRemoveMesh : public UCustomizableObjectNodeModifierEditMeshSectionBase
+{
+public:
+	GENERATED_BODY()
+
+	// Begin EdGraphNode interface
+	FText GetNodeTitle(ENodeTitleType::Type TitleType) const override;
+	FText GetTooltipText() const override;
+
+	// UCustomizableObjectNode interface
+	virtual void AllocateDefaultPins(UCustomizableObjectNodeRemapPins* RemapPins) override;
+	virtual void PinConnectionListChanged(UEdGraphPin * Pin);
+
+	UEdGraphPin* RemoveMeshPin() const
+	{
+		return FindPin(TEXT("Remove Mesh"));
+	}
+
+	bool IsSingleOutputNode() const override;
+};
+
diff --git a/Engine/Plugins/Experimental/Mutable/Source/CustomizableObjectEditor/Public/MuCOE/Nodes/CustomizableObjectNodeModifierRemoveMeshBlocks.h b/Engine/Plugins/Experimental/Mutable/Source/CustomizableObjectEditor/Public/MuCOE/Nodes/CustomizableObjectNodeModifierRemoveMeshBlocks.h
new file mode 100644
--- /dev/null
+++ b/Engine/Plugins/Experimental/Mutable/Source/CustomizableObjectEditor/Public/MuCOE/Nodes/CustomizableObjectNodeModifierRemoveMeshBlocks.h
@@ -0,0 +1,35 @@
+// Copyright Epic Games, Inc. All Rights Reserved.
+
+#pragma once
+
+#include "MuCOE/Nodes/CustomizableObjectNodeModifierEditLayoutBlocks.h"
+
+#include "CustomizableObjectNodeModifierRemoveMeshBlocks.generated.h"
+
+namespace ENodeTitleType { enum Type : int; }
+
+class FArchive;
+class UCustomizableObjectNodeRemapPins;
+class UEdGraphPin;
+class UObject;
+struct FPropertyChangedEvent;
+
+UCLASS()
+class CUSTOMIZABLEOBJECTEDITOR_API UCustomizableObjectNodeModifierRemoveMeshBlocks : public UCustomizableObjectNodeModifierEditLayoutBlocks
+{
+public:
+	GENERATED_BODY()
+
+	UPROPERTY()
+	TArray<int32> Blocks_DEPRECATED;
+
+	// Begin EdGraphNode interface
+	virtual FText GetNodeTitle(ENodeTitleType::Type TitleType) const override;
+	virtual FText GetTooltipText() const override;
+
+	// UCustomizableObjectNode interface
+	virtual void BackwardsCompatibleFixup(int32 CustomizableObjectCustomVersion) override;
+	virtual void AllocateDefaultPins(UCustomizableObjectNodeRemapPins* RemapPins) override;
+	virtual bool IsSingleOutputNode() const override;
+};
+
diff --git a/Engine/Plugins/Experimental/Mutable/Source/CustomizableObjectEditor/Public/MuCOE/Nodes/CustomizableObjectNodeModifierWithMaterial.h b/Engine/Plugins/Experimental/Mutable/Source/CustomizableObjectEditor/Public/MuCOE/Nodes/CustomizableObjectNodeModifierWithMaterial.h
new file mode 100644
--- /dev/null
+++ b/Engine/Plugins/Experimental/Mutable/Source/CustomizableObjectEditor/Public/MuCOE/Nodes/CustomizableObjectNodeModifierWithMaterial.h
@@ -0,0 +1,69 @@
+// Copyright Epic Games, Inc. All Rights Reserved.
+
+#pragma once
+
+#include "MuCOE/Nodes/CustomizableObjectNodeModifierBase.h"
+#include "MuCOE/Nodes/CustomizableObjectNodeMaterialBase.h"
+#include "MuT/NodeModifier.h"
+
+#include "CustomizableObjectNodeModifierWithMaterial.generated.h"
+
+class UObject;
+
+
+UCLASS(abstract)
+class CUSTOMIZABLEOBJECTEDITOR_API UCustomizableObjectNodeModifierWithMaterial : public UCustomizableObjectNodeModifierBase
+{
+
+	GENERATED_BODY()
+
+public:
+
+	/** Materials in all other objects that activate this tags will be affected by this modifier. */
+	UPROPERTY(EditAnywhere, BlueprintReadWrite, Category = Modifier)
+	TArray<FString> RequiredTags;
+
+	/** Policy to use required tags in case more than one is added. */
+	UPROPERTY(EditAnywhere, Category = Modifier)
+	EMutableMultipleTagPolicy MultipleTagPolicy = EMutableMultipleTagPolicy::OnlyOneRequired;
+
+	/** Reference material that defines the structure of the data to be extended.
+	* The Sections modified with this modifier are supposed to have the same texture parameters, but don't need to have exactly
+	* the ReferenceMaterial set.
+	*/
+	UPROPERTY(EditAnywhere, Category = Parameters)
+	TObjectPtr<UMaterialInterface> ReferenceMaterial = nullptr;
+
+	/** Relates a Parameter id (and layer if is a layered material) to a Pin. Only used to improve performance. */
+	UPROPERTY()
+	TMap<FNodeMaterialParameterId, FEdGraphPinReference> PinsParameterMap;
+
+public:
+
+	//
+	virtual void PostEditChangeProperty(FPropertyChangedEvent& PropertyChangedEvent) override;
+
+	// UCustomizableObjectNode interface
+	virtual bool IsNodeOutDatedAndNeedsRefresh() override;
+
+	// UCustomizableObjectNodeModifierBase interface
+	virtual const TArray<FString>* GetRequiredTags() const override { return &RequiredTags;  }
+	virtual EMutableMultipleTagPolicy GetMultipleTagsPolicy() const override { return MultipleTagPolicy;  }
+
+	// Own interface
+
+	/** */
+	bool UsesImage(const FNodeMaterialParameterId& ImageId) const;
+
+	/** */
+	const UEdGraphPin* GetUsedImagePin(const FNodeMaterialParameterId& ImageId) const;
+
+	/** */
+	int32 GetNumParameters(const EMaterialParameterType Type) const;
+	FNodeMaterialParameterId GetParameterId(const EMaterialParameterType Type, const int32 ParameterIndex) const;
+	FName GetParameterName(const EMaterialParameterType Type, const int32 ParameterIndex) const;
+	static int32 GetParameterLayerIndex(const UMaterialInterface* InMaterial, const EMaterialParameterType Type, const int32 ParameterIndex);
+	int32 GetParameterLayerIndex(const EMaterialParameterType Type, const int32 ParameterIndex) const;
+	FText GetParameterLayerName(const EMaterialParameterType Type, const int32 ParameterIndex) const;
+};
+
diff --git a/Engine/Plugins/Experimental/Mutable/Source/MutableRuntime/Internal/MuR/Operations.h b/Engine/Plugins/Experimental/Mutable/Source/MutableRuntime/Internal/MuR/Operations.h
--- a/Engine/Plugins/Experimental/Mutable/Source/MutableRuntime/Internal/MuR/Operations.h
+++ b/Engine/Plugins/Experimental/Mutable/Source/MutableRuntime/Internal/MuR/Operations.h
@@ -864,7 +864,9 @@ namespace mu
 		struct MeshMaskClipUVMaskArgs
 		{
 			ADDRESS Source = 0;
-			ADDRESS Mask = 0;
+			ADDRESS UVSource = 0;
+			ADDRESS MaskImage = 0;
+			ADDRESS MaskLayout = 0;
 			uint8 LayoutIndex = 0;
 		};
 
diff --git a/Engine/Plugins/Experimental/Mutable/Source/MutableRuntime/Private/MuR/CodeRunner.cpp b/Engine/Plugins/Experimental/Mutable/Source/MutableRuntime/Private/MuR/CodeRunner.cpp
--- a/Engine/Plugins/Experimental/Mutable/Source/MutableRuntime/Private/MuR/CodeRunner.cpp
+++ b/Engine/Plugins/Experimental/Mutable/Source/MutableRuntime/Private/MuR/CodeRunner.cpp
@@ -1674,26 +1674,51 @@ namespace mu
 			{
 				AddOp(FScheduledOp(item.At, item, 1),
 					FScheduledOp(args.Source, item),
-					FScheduledOp(args.Mask, item));
+					FScheduledOp(args.UVSource, item),
+					FScheduledOp(args.MaskImage, item),
+					FScheduledOp(args.MaskLayout, item));
 				break;
 			}
 			case 1:
 			{
-				MUTABLE_CPUPROFILER_SCOPE(ME_MASKCLIPUVMASK_1)
+				MUTABLE_CPUPROFILER_SCOPE(ME_MASKCLIPUVMASK_1);
 
 				Ptr<const Mesh> Source = LoadMesh(FCacheAddress(args.Source, item));
-				Ptr<const Image> Mask = LoadImage(FCacheAddress(args.Mask, item));
+				Ptr<const Mesh> UVSource = LoadMesh(FCacheAddress(args.UVSource, item));
+				Ptr<const Image> MaskImage = LoadImage(FCacheAddress(args.MaskImage, item));
+				Ptr<const Layout> MaskLayout = LoadLayout(FCacheAddress(args.MaskLayout, item));
 
 				// Only if both are valid.
-				if (Source.get() && Mask.get())
+				if (Source.get() && MaskImage.get())
 				{
 					Ptr<Mesh> Result = CreateMesh();
 
 					bool bOutSuccess = false;
-					MeshMaskClipUVMask(Result.get(), Source.get(), Mask.get(), args.LayoutIndex, bOutSuccess);
+					MakeMeshMaskFromUVMask(Result.get(), Source.get(), UVSource.get(), MaskImage.get(), args.LayoutIndex, bOutSuccess);
 
 					Release(Source);
-					Release(Mask);
+					Release(UVSource);
+					Release(MaskImage);
+					if (!bOutSuccess)
+					{
+						Release(Result);
+						StoreMesh(item, nullptr);
+					}
+					else
+					{
+						StoreMesh(item, Result);
+					}
+				}
+				else if (Source.get() && MaskLayout.get())
+				{					
+					Ptr<Mesh> Result = CreateMesh();
+
+					bool bOutSuccess = false;
+					MakeMeshMaskFromLayout(Result.get(), Source.get(), UVSource.get(), MaskLayout.get(), args.LayoutIndex, bOutSuccess);
+
+					Release(Source);
+					Release(UVSource);
+					Release(MaskImage);
 					if (!bOutSuccess)
 					{
 						Release(Result);
@@ -1707,7 +1732,8 @@ namespace mu
 				else
 				{
 					Release(Source);
-					Release(Mask);
+					Release(UVSource);
+					Release(MaskImage);
 					StoreMesh(item, nullptr);
 				}
 
diff --git a/Engine/Plugins/Experimental/Mutable/Source/MutableRuntime/Private/MuR/OpMeshClipWithMesh.cpp b/Engine/Plugins/Experimental/Mutable/Source/MutableRuntime/Private/MuR/OpMeshClipWithMesh.cpp
--- a/Engine/Plugins/Experimental/Mutable/Source/MutableRuntime/Private/MuR/OpMeshClipWithMesh.cpp
+++ b/Engine/Plugins/Experimental/Mutable/Source/MutableRuntime/Private/MuR/OpMeshClipWithMesh.cpp
@@ -4,6 +4,7 @@
 
 #include "MuR/MeshPrivate.h"
 #include "MuR/ImagePrivate.h"
+#include "MuR/Layout.h"
 #include "MuR/ConvertData.h"
 #include "MuR/Platform.h"
 #include "MuR/MutableTrace.h"
@@ -558,6 +559,69 @@ namespace mu { namespace
 	}
 
 
+	void MeshLayoutMaskClassifyVertices(TArray<uint8>& VertexClipped, const Mesh* Base, const Layout* Mask, const uint8 LayoutIndex)
+	{
+		using namespace UE::Geometry;
+
+		MUTABLE_CPUPROFILER_SCOPE(MeshUVMaskClassifyVertices);
+
+		uint32 VertexCount = Base->GetVertexCount();
+
+		// Stores whether each vertex in the original mesh in the clip mesh volume
+		VertexClipped.SetNumUninitialized(VertexCount, EAllowShrinking::No);
+		FMemory::Memzero(VertexClipped.GetData(), VertexClipped.Num());
+
+		// Now go through all vertices in the mesh and record whether they are inside or outside of the ClipMesh
+		const FMeshBufferSet& MBSPriv = Base->GetVertexBuffers();
+		for (int32 b = 0; b < MBSPriv.Buffers.Num(); ++b)
+		{
+			const TArray<mu::FMeshBufferChannel>& Channels = MBSPriv.Buffers[b].Channels;
+
+			for (int32 c = 0; c < Channels.Num(); ++c)
+			{
+				EMeshBufferSemantic Sem = Channels[c].Semantic;
+				if (Sem != MBS_TEXCOORDS)
+				{
+					continue;
+				}
+
+				int32 SemIndex = Channels[c].SemanticIndex;
+				if (SemIndex != LayoutIndex)
+				{
+					continue;
+				}
+
+				UntypedMeshBufferIteratorConst It(Base->GetVertexBuffers(), Sem, SemIndex);
+				for (uint32 V = 0; V < VertexCount; ++V)
+				{
+					// \TODO: This could be optimized.
+					FVector2f UV = It.GetAsVec2f();
+
+					FVector2f Cell = UV * FVector2f(Mask->Size.X, Mask->Size.Y);
+
+					// \TODO: This could also be optimized
+					for (const FLayoutBlock& Block : Mask->Blocks)
+					{
+						if (
+							float(Block.Min.X) <= Cell.X && float(Block.Min.Y) <= Cell.Y
+							&&
+							float(Block.Min.X + Block.Size.X) >= Cell.X && float(Block.Min.Y + Block.Size.Y) >= Cell.Y
+							)
+						{
+							VertexClipped[V] = true; 
+							break;
+						}
+					}
+
+					++It;
+				}
+
+				break;
+			}
+		}
+	}
+
+
 	/** Make a mask with the indices of the vertices with 0 in the IncludedVertices array. */
 	void CreateMask(Mesh* Result, const Mesh* Base, const TArray<uint8>& IncludedVertices)
 	{
@@ -727,18 +791,64 @@ namespace mu
     }
 
 
-	void MeshMaskClipUVMask(Mesh* Result, const Mesh* Base, const Image* Mask, uint8 LayoutIndex, bool& bOutSuccess)
+	void MakeMeshMaskFromUVMask(Mesh* Result, const Mesh* Base, const Mesh* BaseForUVs, const Image* Mask, uint8 LayoutIndex, bool& bOutSuccess)
 	{
-		MUTABLE_CPUPROFILER_SCOPE(MeshMaskClipMesh);
+		MUTABLE_CPUPROFILER_SCOPE(MeshMaskUVMask);
 
-		check(Result && Base && Mask);
+		check(Result && Base && Mask && BaseForUVs);
+		check(Base->VertexBuffers.GetElementCount() == BaseForUVs->VertexBuffers.GetElementCount());
 
 		bOutSuccess = true;
 
 		// Stores whether each vertex in the original mesh in in the clip mesh volume
 		TArray<uint8> VertexClipped; 
 
-		MeshUVMaskClassifyVertices(VertexClipped, Base, Mask, LayoutIndex);
+		MeshUVMaskClassifyVertices(VertexClipped, BaseForUVs, Mask, LayoutIndex);
+
+		// We only remove vertices if all their faces are clipped
+		TArray<uint8> VertexWithFaceNotClipped;
+		VertexWithFaceNotClipped.AddZeroed(VertexClipped.Num());
+
+		UntypedMeshBufferIteratorConst Ito(Base->GetIndexBuffers(), MBS_VERTEXINDEX);
+		int32 FaceCount = Base->GetFaceCount();
+		for (int32 F = 0; F < FaceCount; ++F)
+		{
+			uint32 OV[3] = { 0, 0, 0 };
+
+			OV[0] = Ito.GetAsUINT32(); ++Ito;
+			OV[1] = Ito.GetAsUINT32(); ++Ito;
+			OV[2] = Ito.GetAsUINT32(); ++Ito;
+
+			bool bFaceClipped =
+				VertexClipped[OV[0]] &&
+				VertexClipped[OV[1]] &&
+				VertexClipped[OV[2]];
+
+			if (!bFaceClipped)
+			{
+				VertexWithFaceNotClipped[OV[0]] = true;
+				VertexWithFaceNotClipped[OV[1]] = true;
+				VertexWithFaceNotClipped[OV[2]] = true;
+			}
+		}
+
+		CreateMask(Result, Base, VertexWithFaceNotClipped);
+	}
+
+
+	void MakeMeshMaskFromLayout(Mesh* Result, const Mesh* Base, const Mesh* BaseForUVs, const Layout* Mask, uint8 LayoutIndex, bool& bOutSuccess)
+	{
+		MUTABLE_CPUPROFILER_SCOPE(MakeMeshMaskFromLayout);
+
+		check(Result && Base && Mask && BaseForUVs);
+		check(Base->VertexBuffers.GetElementCount()== BaseForUVs->VertexBuffers.GetElementCount());
+
+		bOutSuccess = true;
+
+		// Stores whether each vertex in the original mesh in in the clip mesh volume
+		TArray<uint8> VertexClipped;
+
+		MeshLayoutMaskClassifyVertices(VertexClipped, BaseForUVs, Mask, LayoutIndex);
 
 		// We only remove vertices if all their faces are clipped
 		TArray<uint8> VertexWithFaceNotClipped;
diff --git a/Engine/Plugins/Experimental/Mutable/Source/MutableRuntime/Private/MuR/OpMeshClipWithMesh.h b/Engine/Plugins/Experimental/Mutable/Source/MutableRuntime/Private/MuR/OpMeshClipWithMesh.h
--- a/Engine/Plugins/Experimental/Mutable/Source/MutableRuntime/Private/MuR/OpMeshClipWithMesh.h
+++ b/Engine/Plugins/Experimental/Mutable/Source/MutableRuntime/Private/MuR/OpMeshClipWithMesh.h
@@ -8,6 +8,7 @@ namespace mu
 {
 	class Mesh;
 	class Image;
+	class Layout;
 
     /**  */
 	extern void MeshClipWithMesh(Mesh* Result, const Mesh* pBase, const Mesh* pClipMesh, bool& bOutSuccess);
@@ -15,10 +16,13 @@ namespace mu
     /** Generate a mask mesh with the faces of the base mesh inside the clip mesh. */
 	extern void MeshMaskClipMesh(Mesh* Result, const Mesh* pBase, const Mesh* pClipMesh, bool& bOutSuccess);
 
-	/** Generate a mask mesh with the faces of the base mesh that have all 3 vertices marked in the fiven mask. */
-	extern void MeshMaskClipUVMask(Mesh* Result, const Mesh* Base, const Image* Mask, uint8 LayoutIndex, bool& bOutSuccess);
+	/** Generate a mask mesh with the faces of the base mesh that have all 3 vertices marked in the given mask. */
+	extern void MakeMeshMaskFromUVMask(Mesh* Result, const Mesh* Base, const Mesh* BaseForUVs, const Image* Mask, uint8 LayoutIndex, bool& bOutSuccess);
 
     /** Generate a mask mesh with the faces of the base mesh matching the fragment. */
 	extern void MeshMaskDiff(Mesh* Result, const Mesh* pBase, const Mesh* pFragment, bool& bOutSuccess);
 
+	/** Generate a mask mesh with the faces of the base mesh that have all 3 vertices inside any block of the given layout. */
+	extern void MakeMeshMaskFromLayout(Mesh* Result, const Mesh* Base, const Mesh* BaseForUVs, const Layout* Mask, uint8 LayoutIndex, bool& bOutSuccess);
+
 }
diff --git a/Engine/Plugins/Experimental/Mutable/Source/MutableRuntime/Private/MuR/Operations.cpp b/Engine/Plugins/Experimental/Mutable/Source/MutableRuntime/Private/MuR/Operations.cpp
--- a/Engine/Plugins/Experimental/Mutable/Source/MutableRuntime/Private/MuR/Operations.cpp
+++ b/Engine/Plugins/Experimental/Mutable/Source/MutableRuntime/Private/MuR/Operations.cpp
@@ -849,7 +849,8 @@ namespace mu
 		{
 			OP::MeshMaskClipUVMaskArgs args = program.GetOpArgs<OP::MeshMaskClipUVMaskArgs>(at);
 			f(args.Source);
-			f(args.Mask);
+			f(args.MaskImage);
+			f(args.MaskLayout);
 			break;
 		}
 
diff --git a/Engine/Plugins/Experimental/Mutable/Source/MutableRuntime/Private/MuR/SparseIndexMap.h b/Engine/Plugins/Experimental/Mutable/Source/MutableRuntime/Private/MuR/SparseIndexMap.h
--- a/Engine/Plugins/Experimental/Mutable/Source/MutableRuntime/Private/MuR/SparseIndexMap.h
+++ b/Engine/Plugins/Experimental/Mutable/Source/MutableRuntime/Private/MuR/SparseIndexMap.h
@@ -143,7 +143,7 @@ namespace mu
 				}
 			}
 
-			check(false);
+			ensure(false);
 			return false;
 		}
 
diff --git a/Engine/Plugins/Experimental/Mutable/Source/MutableTools/Internal/MuT/AST.h b/Engine/Plugins/Experimental/Mutable/Source/MutableTools/Internal/MuT/AST.h
--- a/Engine/Plugins/Experimental/Mutable/Source/MutableTools/Internal/MuT/AST.h
+++ b/Engine/Plugins/Experimental/Mutable/Source/MutableTools/Internal/MuT/AST.h
@@ -372,9 +372,7 @@ namespace mu
 	};
 
 
-	//---------------------------------------------------------------------------------------------
-	//!
-	//---------------------------------------------------------------------------------------------
+	/** */
 	class Sink_MeshFormatAST
 	{
 	public:
@@ -392,6 +390,24 @@ namespace mu
 	};
 
 
+	/** */
+	class Sink_MeshExtractLayoutBlocksAST
+	{
+	public:
+
+		// \TODO This is recursive and may cause stack overflows in big models.
+		Ptr<ASTOp> Apply(const class ASTOpMeshExtractLayoutBlocks* Root);
+
+	protected:
+
+		const class ASTOpMeshExtractLayoutBlocks* m_root = nullptr;
+		Ptr<ASTOp> m_initialSource;
+		TMap<FSinkerOldToNewKey, Ptr<ASTOp>> OldToNew;
+
+		Ptr<ASTOp> Visit(const Ptr<ASTOp>& at, const class ASTOpMeshExtractLayoutBlocks* currentFormatOp);
+	};
+
+
 	//---------------------------------------------------------------------------------------------
 	//!
 	//---------------------------------------------------------------------------------------------
@@ -420,6 +436,7 @@ namespace mu
 		Sink_ImageSwizzleAST ImageSwizzleSinker;
 		Sink_ImageMipmapAST ImageMipmapSinker;
 		Sink_MeshFormatAST MeshFormatSinker;
+		Sink_MeshExtractLayoutBlocksAST MeshExtractLayoutBlocksSinker;
 	};
 
 
diff --git a/Engine/Plugins/Experimental/Mutable/Source/MutableTools/Internal/MuT/NodeMeshConstantPrivate.h b/Engine/Plugins/Experimental/Mutable/Source/MutableTools/Internal/MuT/NodeMeshConstantPrivate.h
--- a/Engine/Plugins/Experimental/Mutable/Source/MutableTools/Internal/MuT/NodeMeshConstantPrivate.h
+++ b/Engine/Plugins/Experimental/Mutable/Source/MutableTools/Internal/MuT/NodeMeshConstantPrivate.h
@@ -21,6 +21,14 @@ namespace mu
 
 		TArray<Ptr<NodeLayout>> Layouts;
 
+		struct FMorph
+		{
+			FString Name;
+			Ptr<Mesh> MorphedMesh;
+		};
+
+		TArray<FMorph> Morphs;
+
 	};
 
 }
diff --git a/Engine/Plugins/Experimental/Mutable/Source/MutableTools/Private/MuT/ASTOpMeshExtractLayoutBlocks.cpp b/Engine/Plugins/Experimental/Mutable/Source/MutableTools/Private/MuT/ASTOpMeshExtractLayoutBlocks.cpp
--- a/Engine/Plugins/Experimental/Mutable/Source/MutableTools/Private/MuT/ASTOpMeshExtractLayoutBlocks.cpp
+++ b/Engine/Plugins/Experimental/Mutable/Source/MutableTools/Private/MuT/ASTOpMeshExtractLayoutBlocks.cpp
@@ -8,6 +8,7 @@
 #include "MuT/ASTOpMeshMorph.h"
 #include "MuT/ASTOpMeshRemoveMask.h"
 #include "MuT/ASTOpMeshClipMorphPlane.h"
+#include "MuT/ASTOpMeshApplyPose.h"
 #include "MuR/ModelPrivate.h"
 #include "MuR/RefCounted.h"
 #include "MuR/Types.h"
@@ -94,148 +95,168 @@ namespace mu
 		}
 	}
 
+	 
+	mu::Ptr<ASTOp> ASTOpMeshExtractLayoutBlocks::OptimiseSink(const FModelOptimizationOptions&, FOptimizeSinkContext& Context) const
+	{
+		mu::Ptr<ASTOp> NewOp = Context.MeshExtractLayoutBlocksSinker.Apply(this);
+		return NewOp;
+	}
 
-	mu::Ptr<ASTOp> ASTOpMeshExtractLayoutBlocks::OptimiseSink(const FModelOptimizationOptions&, FOptimizeSinkContext&) const
+
+	//---------------------------------------------------------------------------------------------
+	mu::Ptr<ASTOp> Sink_MeshExtractLayoutBlocksAST::Apply(const ASTOpMeshExtractLayoutBlocks* root)
 	{
-		Ptr<ASTOp> NewOp;
+		m_root = root;
 
-		if (!Source.child())
+		OldToNew.Reset();
+
+		m_initialSource = m_root->Source.child();
+		mu::Ptr<ASTOp> newSource = Visit(m_initialSource, m_root);
+
+		// If there is any change, it is the new root.
+		if (newSource != m_initialSource)
 		{
-			return nullptr;
+			return newSource;
 		}
 
-		OP_TYPE SourceType = Source.child()->GetOpType();
+		return nullptr;
+	}
 
-		// Optimize only the mesh parameter
-		switch (SourceType)
-		{
+	mu::Ptr<ASTOp> Sink_MeshExtractLayoutBlocksAST::Visit(const mu::Ptr<ASTOp>& at, const ASTOpMeshExtractLayoutBlocks* currentSinkOp)
+	{
+		if (!at) return nullptr;
 
-		case OP_TYPE::ME_SWITCH:
+		// Already visited?
+		const Ptr<ASTOp>* Cached = OldToNew.Find({ at, currentSinkOp });
+		if (Cached)
 		{
-			// Move the operation down all the paths
-			Ptr<ASTOpSwitch> NewSwitch = mu::Clone<ASTOpSwitch>(Source.child());
-
-			if (NewSwitch->def)
-			{
-				Ptr<ASTOpMeshExtractLayoutBlocks> NewBind = mu::Clone<ASTOpMeshExtractLayoutBlocks>(this);
-				NewBind->Source = NewSwitch->def.child();
-				NewSwitch->def = NewBind;
-			}
+			return *Cached;
+		}
 
-			for (int32 v = 0; v < NewSwitch->cases.Num(); ++v)
-			{
-				if (NewSwitch->cases[v].branch)
-				{
-					Ptr<ASTOpMeshExtractLayoutBlocks> NewBind = mu::Clone<ASTOpMeshExtractLayoutBlocks>(this);
-					NewBind->Source = NewSwitch->cases[v].branch.child();
-					NewSwitch->cases[v].branch = NewBind;
-				}
-			}
+		mu::Ptr<ASTOp> newAt = at;
+		switch (at->GetOpType())
+		{
 
-			NewOp = NewSwitch;
+		case OP_TYPE::ME_APPLYLAYOUT:
+		{
+			Ptr<ASTOpFixed> newOp = mu::Clone<ASTOpFixed>(at);
+			newOp->SetChild(newOp->op.args.MeshApplyLayout.mesh, Visit(newOp->children[newOp->op.args.MeshApplyLayout.mesh].child(), currentSinkOp));
+			newAt = newOp;
 			break;
 		}
 
-		case OP_TYPE::ME_CONDITIONAL:
+		case OP_TYPE::ME_SETSKELETON:
 		{
-			// Move the operation down all the paths
-			Ptr<ASTOpConditional> NewConditional = mu::Clone<ASTOpConditional>(Source.child());
-
-			if (NewConditional->yes)
-			{
-				Ptr<ASTOpMeshExtractLayoutBlocks> NewBind = mu::Clone<ASTOpMeshExtractLayoutBlocks>(this);
-				NewBind->Source = NewConditional->yes.child();
-				NewConditional->yes = NewBind;
-			}
-
-			if (NewConditional->no)
-			{
-				Ptr<ASTOpMeshExtractLayoutBlocks> NewBind = mu::Clone<ASTOpMeshExtractLayoutBlocks>(this);
-				NewBind->Source = NewConditional->no.child();
-				NewConditional->no = NewBind;
-			}
-
-			NewOp = NewConditional;
+			Ptr<ASTOpFixed> newOp = mu::Clone<ASTOpFixed>(at);
+			newOp->SetChild(newOp->op.args.MeshSetSkeleton.source, Visit(newOp->children[newOp->op.args.MeshSetSkeleton.source].child(), currentSinkOp));
+			newAt = newOp;
 			break;
 		}
 
 		case OP_TYPE::ME_ADDTAGS:
 		{
-			Ptr<ASTOpMeshAddTags> NewAddTags = mu::Clone<ASTOpMeshAddTags>(Source.child());
-
-			if (NewAddTags->Source)
-			{
-				Ptr<ASTOpMeshExtractLayoutBlocks> New = mu::Clone<ASTOpMeshExtractLayoutBlocks>(this);
-				New->Source = NewAddTags->Source.child();
-				NewAddTags->Source = New;
-			}
+			Ptr<ASTOpMeshAddTags> newOp = mu::Clone<ASTOpMeshAddTags>(at);
+			newOp->Source = Visit(newOp->Source.child(), currentSinkOp);
+			newAt = newOp;
+			break;
+		}
 
-			NewOp = NewAddTags;
+		case OP_TYPE::ME_CLIPMORPHPLANE:
+		{
+			Ptr<ASTOpMeshClipMorphPlane> newOp = mu::Clone<ASTOpMeshClipMorphPlane>(at);
+			newOp->source = Visit(newOp->source.child(), currentSinkOp);
+			newAt = newOp;
 			break;
 		}
 
 		case OP_TYPE::ME_MORPH:
 		{
-			// Move the operation down the base and the target
-			Ptr<ASTOpMeshMorph> NewMorph = mu::Clone<ASTOpMeshMorph>(Source.child());
-
-			if (NewMorph->Base)
-			{
-				Ptr<ASTOpMeshExtractLayoutBlocks> New = mu::Clone<ASTOpMeshExtractLayoutBlocks>(this);
-				New->Source = NewMorph->Base.child();
-				NewMorph->Base = New;
-			}
+			Ptr<ASTOpMeshMorph> NewOp = mu::Clone<ASTOpMeshMorph>(at);
+			NewOp->Base = Visit(NewOp->Base.child(), currentSinkOp);
+			NewOp->Target = Visit(NewOp->Target.child(), currentSinkOp);
+			newAt = NewOp;
+			break;
+		}
 
-			if (NewMorph->Target)
-			{
-				Ptr<ASTOpMeshExtractLayoutBlocks> New = mu::Clone<ASTOpMeshExtractLayoutBlocks>(this);
-				New->Source = NewMorph->Target.child();
-				NewMorph->Target = New;
-			}
+		case OP_TYPE::ME_MERGE:
+		{
+			Ptr<ASTOpFixed> newOp = mu::Clone<ASTOpFixed>(at);
+			newOp->SetChild(newOp->op.args.MeshMerge.base, Visit(newOp->children[newOp->op.args.MeshMerge.base].child(), currentSinkOp));
+			newOp->SetChild(newOp->op.args.MeshMerge.added, Visit(newOp->children[newOp->op.args.MeshMerge.added].child(), currentSinkOp));
+			newAt = newOp;
+			break;
+		}
 
-			NewOp = NewMorph;
+		case OP_TYPE::ME_APPLYPOSE:
+		{
+			Ptr<ASTOpMeshApplyPose> NewOp = mu::Clone<ASTOpMeshApplyPose>(at);
+			NewOp->base = Visit(NewOp->base.child(), currentSinkOp);
+			newAt = NewOp;
 			break;
 		}
 
-		case OP_TYPE::ME_REMOVEMASK:
+		case OP_TYPE::ME_INTERPOLATE:
 		{
-			// Move the operation down the base
-			// \TODO: mask too to try to make them smaller?
-			Ptr<ASTOpMeshRemoveMask> NewRemove = mu::Clone<ASTOpMeshRemoveMask>(Source.child());
+			Ptr<ASTOpFixed> newOp = mu::Clone<ASTOpFixed>(at);
+			newOp->SetChild(newOp->op.args.MeshInterpolate.base, Visit(newOp->children[newOp->op.args.MeshInterpolate.base].child(), currentSinkOp));
 
-			if (NewRemove->source)
+			for (int32 t = 0; t < MUTABLE_OP_MAX_INTERPOLATE_COUNT - 1; ++t)
 			{
-				Ptr<ASTOpMeshExtractLayoutBlocks> New = mu::Clone<ASTOpMeshExtractLayoutBlocks>(this);
-				New->Source = NewRemove->source.child();
-				NewRemove->source = New;
+				if (newOp->children[newOp->op.args.MeshInterpolate.targets[t]])
+				{
+					newOp->SetChild(newOp->op.args.MeshInterpolate.targets[t], Visit(newOp->children[newOp->op.args.MeshInterpolate.targets[t]].child(), currentSinkOp));
+				}
 			}
 
-			NewOp = NewRemove;
+			newAt = newOp;
 			break;
 		}
 
-		case OP_TYPE::ME_CLIPMORPHPLANE:
+		case OP_TYPE::ME_REMOVEMASK:
+		{
+			// TODO: Make mask smaller?
+			Ptr<ASTOpMeshRemoveMask> newOp = mu::Clone<ASTOpMeshRemoveMask>(at);
+			newOp->source = Visit(newOp->source.child(), currentSinkOp);
+			newAt = newOp;
+			break;
+		}
+
+		case OP_TYPE::ME_CONDITIONAL:
 		{
-			// Move the operation down the source
-			Ptr<ASTOpMeshClipMorphPlane> NewMorph = mu::Clone<ASTOpMeshClipMorphPlane>(Source.child());
+			Ptr<ASTOpConditional> newOp = mu::Clone<ASTOpConditional>(at);
+			newOp->yes = Visit(newOp->yes.child(), currentSinkOp);
+			newOp->no = Visit(newOp->no.child(), currentSinkOp);
+			newAt = newOp;
+			break;
+		}
 
-			if (NewMorph->source)
+		case OP_TYPE::ME_SWITCH:
+		{
+			Ptr<ASTOpSwitch> newOp = mu::Clone<ASTOpSwitch>(at);
+			newOp->def = Visit(newOp->def.child(), currentSinkOp);
+			for (ASTOpSwitch::FCase& c : newOp->cases)
 			{
-				Ptr<ASTOpMeshExtractLayoutBlocks> New = mu::Clone<ASTOpMeshExtractLayoutBlocks>(this);
-				New->Source = NewMorph->source.child();
-				NewMorph->source = New;
+				c.branch = Visit(c.branch.child(), currentSinkOp);
 			}
-
-			NewOp = NewMorph;
+			newAt = newOp;
 			break;
 		}
 
+		// If we reach here it means the operation type has not bee optimized.
 		default:
+			if (at != m_initialSource)
+			{
+				mu::Ptr<ASTOpMeshExtractLayoutBlocks> newOp = mu::Clone<ASTOpMeshExtractLayoutBlocks>(currentSinkOp);
+				newOp->Source = at;
+				newAt = newOp;
+			}
 			break;
 
 		}
 
-		return NewOp;
+		OldToNew.Add({ at, currentSinkOp }, newAt);
+
+		return newAt;
 	}
 
 
diff --git a/Engine/Plugins/Experimental/Mutable/Source/MutableTools/Private/MuT/ASTOpMeshFormat.cpp b/Engine/Plugins/Experimental/Mutable/Source/MutableTools/Private/MuT/ASTOpMeshFormat.cpp
--- a/Engine/Plugins/Experimental/Mutable/Source/MutableTools/Private/MuT/ASTOpMeshFormat.cpp
+++ b/Engine/Plugins/Experimental/Mutable/Source/MutableTools/Private/MuT/ASTOpMeshFormat.cpp
@@ -365,24 +365,16 @@ mu::Ptr<ASTOp> Sink_MeshFormatAST::Visit(const mu::Ptr<ASTOp>& at, const ASTOpMe
 		break;
 	}
 
-	// This cannot be sunk since the result is different. Since the clipping is now correctly
-	// generated at the end of the chain when really necessary, this wrong optimisation is no 
-	// longer needed.
-	//case OP_TYPE::ME_CLIPMORPHPLANE:
-//         {
-//             // We move the mask creation down the source
-//             auto typedAt = dynamic_cast<const ASTOpMeshClipMorphPlane*>(at.get());
-//             newAt = Visit(typedAt->source.child());
-//             break;
-//         }
-
 	case OP_TYPE::ME_FORMAT:
 		// TODO: The child format can be removed. 
 		// Unless channels are removed and re-added, which would change their content?
 		break;
 
+
+	// This operation should not be optimized.
 	case OP_TYPE::ME_DIFFERENCE:
 
+	// If we reach here it means the operation type has not bee optimized.
 	default:
 		if (at != m_initialSource)
 		{
