[mutable] Convert the following nodes to modifiers: EditMaterial, ExtendMaterial, RemoveMesh, RemoveMeshBlocks, MorphMaterial.

diff --git a/Engine/Plugins/Experimental/Mutable/Source/MutableTools/Private/MuT/ASTOpMeshMaskClipUVMask.cpp b/Engine/Plugins/Experimental/Mutable/Source/MutableTools/Private/MuT/ASTOpMeshMaskClipUVMask.cpp
--- a/Engine/Plugins/Experimental/Mutable/Source/MutableTools/Private/MuT/ASTOpMeshMaskClipUVMask.cpp
+++ b/Engine/Plugins/Experimental/Mutable/Source/MutableTools/Private/MuT/ASTOpMeshMaskClipUVMask.cpp
@@ -21,7 +21,9 @@ namespace mu
 	//-------------------------------------------------------------------------------------------------
 	ASTOpMeshMaskClipUVMask::ASTOpMeshMaskClipUVMask()
 		: Source(this)
-		, Mask(this)
+		, UVSource(this)
+		, MaskImage(this)
+		, MaskLayout(this)
 	{
 	}
 
@@ -40,7 +42,11 @@ namespace mu
 		if (otherUntyped.GetOpType() == GetOpType())
 		{
 			const ASTOpMeshMaskClipUVMask* other = static_cast<const ASTOpMeshMaskClipUVMask*>(&otherUntyped);
-			return Source == other->Source && Mask == other->Mask && LayoutIndex==other->LayoutIndex;
+			return Source == other->Source
+				&& UVSource == other->UVSource
+				&& MaskImage == other->MaskImage
+				&& MaskLayout == other->MaskLayout 
+				&& LayoutIndex==other->LayoutIndex;
 		}
 		return false;
 	}
@@ -50,7 +56,9 @@ namespace mu
 	uint64 ASTOpMeshMaskClipUVMask::Hash() const
 	{
 		uint64 res = std::hash<void*>()(Source.child().get());
-		hash_combine(res, Mask.child());
+		hash_combine(res, UVSource.child());
+		hash_combine(res, MaskImage.child());
+		hash_combine(res, MaskLayout.child());
 		return res;
 	}
 
@@ -60,7 +68,9 @@ namespace mu
 	{
 		Ptr<ASTOpMeshMaskClipUVMask> n = new ASTOpMeshMaskClipUVMask();
 		n->Source = mapChild(Source.child());
-		n->Mask = mapChild(Mask.child());
+		n->UVSource = mapChild(UVSource.child());
+		n->MaskImage = mapChild(MaskImage.child());
+		n->MaskLayout = mapChild(MaskLayout.child());
 		n->LayoutIndex = LayoutIndex;
 		return n;
 	}
@@ -70,7 +80,9 @@ namespace mu
 	void ASTOpMeshMaskClipUVMask::ForEachChild(const TFunctionRef<void(ASTChild&)> f)
 	{
 		f(Source);
-		f(Mask);
+		f(UVSource);
+		f(MaskImage);
+		f(MaskLayout);
 	}
 
 
@@ -84,11 +96,13 @@ namespace mu
 			FMemory::Memzero(args);
 
 			if (Source) args.Source = Source->linkedAddress;
-			if (Mask) args.Mask = Mask->linkedAddress;
+			if (UVSource) args.UVSource = UVSource->linkedAddress;
+			if (MaskImage) args.MaskImage = MaskImage->linkedAddress;
+			if (MaskLayout) args.MaskLayout = MaskLayout->linkedAddress;
 			args.LayoutIndex = LayoutIndex;
 
 			linkedAddress = (OP::ADDRESS)program.m_opAddress.Num();
-			program.m_opAddress.Add((uint32_t)program.m_byteCode.Num());
+			program.m_opAddress.Add((uint32)program.m_byteCode.Num());
 			AppendCode(program.m_byteCode, OP_TYPE::ME_MASKCLIPUVMASK);
 			AppendCode(program.m_byteCode, args);
 		}
@@ -248,7 +262,7 @@ namespace mu
 				m_root = root;
 				m_oldToNew.Empty();
 
-				m_initialClip = root->Mask.child();
+				m_initialClip = root->MaskImage.child();
 				mu::Ptr<ASTOp> NewClip = Visit(m_initialClip);
 
 				// If there is any change, it is the new root.
@@ -321,7 +335,7 @@ namespace mu
 					if (at != m_initialClip)
 					{
 						Ptr<ASTOpMeshMaskClipUVMask> newOp = mu::Clone<ASTOpMeshMaskClipUVMask>(m_root);
-						newOp->Mask = at;
+						newOp->MaskImage = at;
 						newAt = newOp;
 					}
 					break;
diff --git a/Engine/Plugins/Experimental/Mutable/Source/MutableTools/Private/MuT/ASTOpMeshMaskClipUVMask.h b/Engine/Plugins/Experimental/Mutable/Source/MutableTools/Private/MuT/ASTOpMeshMaskClipUVMask.h
--- a/Engine/Plugins/Experimental/Mutable/Source/MutableTools/Private/MuT/ASTOpMeshMaskClipUVMask.h
+++ b/Engine/Plugins/Experimental/Mutable/Source/MutableTools/Private/MuT/ASTOpMeshMaskClipUVMask.h
@@ -10,18 +10,17 @@
 
 namespace mu
 {
-struct FProgram;
+	struct FProgram;
 
-
-	//---------------------------------------------------------------------------------------------
-	//!
-	//---------------------------------------------------------------------------------------------
+	/** Generate a mesh mask from an image mask or a layout with blocks. */
 	class ASTOpMeshMaskClipUVMask final : public ASTOp
 	{
 	public:
 
 		ASTChild Source;
-		ASTChild Mask;
+		ASTChild UVSource;
+		ASTChild MaskImage;
+		ASTChild MaskLayout;
 		uint8 LayoutIndex = 0;
 
 	public:
diff --git a/Engine/Plugins/Experimental/Mutable/Source/MutableTools/Private/MuT/CodeGenerator.cpp b/Engine/Plugins/Experimental/Mutable/Source/MutableTools/Private/MuT/CodeGenerator.cpp
--- a/Engine/Plugins/Experimental/Mutable/Source/MutableTools/Private/MuT/CodeGenerator.cpp
+++ b/Engine/Plugins/Experimental/Mutable/Source/MutableTools/Private/MuT/CodeGenerator.cpp
@@ -36,6 +36,7 @@
 #include "MuT/ASTOpMeshDifference.h"
 #include "MuT/ASTOpMeshMorph.h"
 #include "MuT/ASTOpMeshOptimizeSkinning.h"
+#include "MuT/ASTOpMeshExtractLayoutBlocks.h"
 #include "MuT/ASTOpParameter.h"
 #include "MuT/ASTOpLayoutRemoveBlocks.h"
 #include "MuT/ASTOpLayoutFromMesh.h"
@@ -59,7 +60,7 @@
 #include "MuT/NodeMesh.h"
 #include "MuT/NodeMeshClipMorphPlane.h"
 #include "MuT/NodeMeshClipWithMesh.h"
-#include "MuT/NodeMeshConstant.h"
+#include "MuT/NodeMeshConstantPrivate.h"
 #include "MuT/NodeMeshFormat.h"
 #include "MuT/NodeMeshFragment.h"
 #include "MuT/NodeMeshGeometryOperation.h"
@@ -71,6 +72,7 @@
 #include "MuT/NodeModifierMeshClipMorphPlane.h"
 #include "MuT/NodeModifierMeshClipWithMesh.h"
 #include "MuT/NodeModifierMeshClipWithUVMask.h"
+#include "MuT/NodeModifierSurfaceEdit.h"
 #include "MuT/NodeObject.h"
 #include "MuT/NodeObjectGroupPrivate.h"
 #include "MuT/NodeObjectNew.h"
@@ -80,7 +82,6 @@
 #include "MuT/NodeScalar.h"
 #include "MuT/NodeScalarConstant.h"
 #include "MuT/NodeSurface.h"
-#include "MuT/NodeSurfaceEdit.h"
 #include "MuT/NodeSurfaceNew.h"
 #include "MuT/NodeSurfaceVariation.h"
 #include "MuT/NodeSurfaceSwitch.h"
@@ -182,9 +183,8 @@ namespace mu
 
 			// This happens only if we generate a node graph that has a NodeSurfaceNew at the root.
 			FSurfaceGenerationResult surfResult;
-			const TArray<FirstPassGenerator::FSurface::FEdit> edits;
 			FSurfaceGenerationOptions SurfaceOptions(Options);
-			GenerateSurface(surfResult, SurfaceOptions, surfNode, edits);
+			GenerateSurface(surfResult, SurfaceOptions, surfNode);
 			return surfResult.surfaceOp;
 		}
 
@@ -200,12 +200,6 @@ namespace mu
 			return nullptr;
 		}
 
-		else if (pNode->GetType()->IsA(NodeSurfaceEdit::GetStaticType()))
-		{
-			// This happens only if we generate a node graph that has a NodeSurfaceEdit at the root.
-			return nullptr;
-		}
-
 		else if (pNode->GetType()->IsA(NodeModifier::GetStaticType()))
 		{
 			// This happens only if we generate a node graph that has a modifier at the root.
@@ -426,7 +420,7 @@ namespace mu
 
 
 	Ptr<ASTOp> CodeGenerator::GenerateImageBlockPatch(Ptr<ASTOp> InBlockOp,
-		const NodeSurfaceEdit::FTexture& Patch,
+		const NodeModifierSurfaceEdit::FTexture& Patch,
 		Ptr<Image> PatchMask,
 		Ptr<ASTOp> conditionAd,
 		const FImageGenerationOptions& ImageOptions )
@@ -791,7 +785,7 @@ namespace mu
 	}
 
 
-	Ptr<Image> CodeGenerator::GenerateImageBlockPatchMask(const NodeSurfaceEdit::FTexture& Patch, FIntPoint GridSize, int32 BlockPixelsX, int32 BlockPixelsY, box<FIntVector2> RectInCells )
+	Ptr<Image> CodeGenerator::GenerateImageBlockPatchMask(const NodeModifierSurfaceEdit::FTexture& Patch, FIntPoint GridSize, int32 BlockPixelsX, int32 BlockPixelsY, box<FIntVector2> RectInCells )
 	{
 		// Create a patching mask for the block
 		Ptr<Image> PatchMask;
@@ -838,8 +832,7 @@ namespace mu
     //---------------------------------------------------------------------------------------------
     void CodeGenerator::GenerateSurface( FSurfaceGenerationResult& result, 
 										 const FSurfaceGenerationOptions& Options,
-                                         Ptr<const NodeSurfaceNew> surfaceNode,
-                                         const TArray<FirstPassGenerator::FSurface::FEdit>& edits )
+                                         Ptr<const NodeSurfaceNew> surfaceNode )
     {
         MUTABLE_CPUPROFILER_SCOPE(GenerateSurface);
 
@@ -850,7 +843,7 @@ namespace mu
 
         // Generate the mesh
         //------------------------------------------------------------------------
-        FMeshGenerationResult meshResults;
+        FMeshGenerationResult MeshResults;
 
 		// We don't add the mesh here, since it will be added directly at the top of the
 		// component expression in the NodeComponentNew generator with the right merges
@@ -925,7 +918,8 @@ namespace mu
 			}
 
 			// Normalize UVs if we're going to work with images and layouts.
-			const bool bNormalizeUVs = node.Images.Num() && edits.Num();
+			// TODO: This should come from per-layout settings!
+			const bool bNormalizeUVs = false; // !node.Images.IsEmpty();
 			MeshOptions.bNormalizeUVs = bNormalizeUVs;
 
 			// Ensure UV islands remain within their main layout block on lower LODs to avoid unexpected reordering 
@@ -933,226 +927,22 @@ namespace mu
 			// that may cause them to fall on a different block.
 			MeshOptions.bClampUVIslands = bShareSurface && bNormalizeUVs;
 
-            GenerateMesh(MeshOptions, meshResults, node.Mesh);
-            lastMeshOp = meshResults.MeshOp;
-            meshResults.ExtraMeshLayouts.SetNum( edits.Num() );
+            GenerateMesh(MeshOptions, MeshResults, node.Mesh);
 
 			ActiveTags.Pop();
 
-            // Let's remember the base mesh of each added mesh, so that we can use them for the
-            // "remove mesh" operations that apply to them instead of the base.
-            std::map<const Node*,Ptr<ASTOp>> baseMeshesForEachAddedMesh;
-
-
-            // Apply mesh merges from child objects "edit surface" nodes
-            for ( int32 editIndex=0; editIndex<edits.Num(); ++editIndex )
-            {
-                const FirstPassGenerator::FSurface::FEdit& e = edits[editIndex];
-
-                if ( Ptr<NodeMesh> pAdd = e.Node->MeshAdd)
-                {
-					// Store the data necessary to apply modifiers for the pre-normal operations stage.
-					ActiveTags.Add(e.Node->EnableTags);
-						
-					FMeshGenerationOptions MergedMeshOptions;
-					MergedMeshOptions.bLayouts = true;
-					MergedMeshOptions.bClampUVIslands = bShareSurface && bNormalizeUVs;
-					MergedMeshOptions.bNormalizeUVs = bNormalizeUVs;
-					MergedMeshOptions.State = Options.State;
-					MergedMeshOptions.ActiveTags = e.Node->EnableTags;
-
-					if (SharedMeshResults)
-					{
-						check(SharedMeshResults->ExtraMeshLayouts.Num()>editIndex);
-						MergedMeshOptions.OverrideLayouts = SharedMeshResults->ExtraMeshLayouts[editIndex].GeneratedLayouts;
-					}
-
-					FMeshGenerationResult addResults;
-                    GenerateMesh(MergedMeshOptions, addResults, pAdd);
-
-					ActiveTags.Pop();
-
-                    baseMeshesForEachAddedMesh[e.Node] = addResults.BaseMeshOp;
-
-					// Apply the modifier for the post-normal operations stage to the added mesh
-					bool bModifiersForBeforeOperations = false;
-					FGenericGenerationOptions ModifierOptions(Options);
-					ModifierOptions.ActiveTags = e.Node->EnableTags;
-					lastMeshOp = ApplyMeshModifiers(Options, lastMeshOp, bModifiersForBeforeOperations, surfaceNode->GetMessageContext());
-
-                    FMeshGenerationResult::FExtraLayouts data;
-					data.GeneratedLayouts = addResults.GeneratedLayouts;
-					data.Condition = e.Condition;
-                    data.MeshFragment = addResults.MeshOp;
-                    meshResults.ExtraMeshLayouts[editIndex] = data;
-
-                    Ptr<ASTOpFixed> mop = new ASTOpFixed();
-                    mop->op.type = OP_TYPE::ME_MERGE;
-                    mop->SetChild(mop->op.args.MeshMerge.base, lastMeshOp );
-                    mop->SetChild(mop->op.args.MeshMerge.added, addResults.MeshOp );
-                    // will merge the meshes under the same surface
-                    mop->op.args.MeshMerge.newSurfaceID = 0;
-
-                    // Condition to apply
-                    if (e.Condition)
-                    {
-                        Ptr<ASTOpConditional> conditionalAd = new ASTOpConditional();
-                        conditionalAd->type = OP_TYPE::ME_CONDITIONAL;
-                        conditionalAd->no = lastMeshOp;
-                        conditionalAd->yes = mop;
-                        conditionalAd->condition = e.Condition;
-                        lastMeshOp = conditionalAd;
-                    }
-                    else
-                    {
-                        lastMeshOp = mop;
-                    }
-                }
-            }
-
-
-            // Apply mesh removes from child objects "edit surface" nodes.
-            // "Removes" need to come after "Adds" because some removes may refer to added meshes,
-            // and not the base.
-			// \TODO: Apply base removes first, and then "added meshes" removes here. It may have lower memory footprint during generation.
-            Ptr<ASTOpMeshRemoveMask> rop;
-            for ( const FirstPassGenerator::FSurface::FEdit& e: edits )
-            {
-                if ( Ptr<NodeMesh> pRemove = e.Node->MeshRemove )
-                {
-                    FMeshGenerationResult removeResults;
-					FMeshGenerationOptions RemoveMeshOptions;
-					RemoveMeshOptions.bLayouts = false;
-					RemoveMeshOptions.State = Options.State;
-					RemoveMeshOptions.ActiveTags = e.Node->EnableTags;
-
-                    GenerateMesh(RemoveMeshOptions, removeResults, pRemove );
-
-                    Ptr<ASTOpFixed> maskOp = new ASTOpFixed();
-                    maskOp->op.type = OP_TYPE::ME_MASKDIFF;
-
-                    // By default, remove from the base
-                    Ptr<ASTOp> removeFrom = meshResults.BaseMeshOp;
-                    // See if we want to remove from an added mesh instead.
-                    if ( e.Node->Parent )
-                    {
-                        auto addedBaseMeshIt = baseMeshesForEachAddedMesh.find(e.Node->Parent.get());
-                        if (addedBaseMeshIt!=baseMeshesForEachAddedMesh.end())
-                        {
-                            removeFrom = addedBaseMeshIt->second;
-                        }
-                    }
-
-                    maskOp->SetChild(maskOp->op.args.MeshMaskDiff.source, removeFrom );
-                    maskOp->SetChild(maskOp->op.args.MeshMaskDiff.fragment, removeResults.MeshOp );
-
-                    if (!rop)
-                    {
-                        rop = new ASTOpMeshRemoveMask();
-                        rop->source = lastMeshOp;
-                    }
-
-                    rop->AddRemove( e.Condition, maskOp );
-                }
-            }
-
-            if (rop)
-            {
-                lastMeshOp = rop;
-            }
-
-            // Apply mesh morphs from child objects "edit surface" nodes
-            for ( const FirstPassGenerator::FSurface::FEdit& e: edits )
-            {
-                if ( NodeMeshPtr pMorph = e.Node->MeshMorph )
-                {
-					// Not needed because it has been generated before already.
-					// Base mesh
-					//FMeshGenerationResult baseMesh;
-					//GenerateMesh(baseMesh, pMesh);
-
-                    // Target mesh
-					FMeshGenerationOptions MorphTargetMeshOptions;
-					MorphTargetMeshOptions.bLayouts = false;
-					MorphTargetMeshOptions.State = Options.State;
-
-                    FMeshGenerationResult morphResult;
-                    GenerateMesh(MorphTargetMeshOptions, morphResult, pMorph );
-
-					// Morph generation through mesh diff
-					Ptr<ASTOp> targetAd = morphResult.MeshOp;
-                    Ptr<ASTOpMeshDifference> diffAd;
-                    {
-                        Ptr<ASTOpMeshDifference> op = new ASTOpMeshDifference();
-                        op->Base = meshResults.BaseMeshOp;
-                        op->Target = targetAd;
-					
-                        // Morphing tex coords here is not supported:
-                        // Generating the homogoneous UVs is difficult since we don't have the base
-                        // layout yet.                       
-                        op->bIgnoreTextureCoords = true;
-                        diffAd = op;
-                    }
-
-                    // Morph operation
-                    Ptr<ASTOp> morphAd;
-                    {
-                        Ptr<ASTOpMeshMorph> op = new ASTOpMeshMorph();
-
-						// Factor
-						if (e.Node->MorphFactor)
-						{
-							FScalarGenerationResult ChildResult;
-							GenerateScalar(ChildResult, Options, e.Node->MorphFactor);
-							op->Factor = ChildResult.op;
-						}
-						else
-						{
-							NodeScalarConstantPtr auxNode = new NodeScalarConstant();
-							auxNode->SetValue(1.0f);
-
-							FScalarGenerationResult ChildResult;
-							GenerateScalar(ChildResult, Options, auxNode);
-							op->Factor = ChildResult.op;
-						}
-
-						// Base		
-						op->Base = lastMeshOp;
-
-						// Targets
-						op->Target = diffAd;
-                        morphAd = op;
-                    }
-
-                    // Condition to apply the morph
-                    if (e.Condition)
-                    {
-                        Ptr<ASTOpConditional> conditionalAd = new ASTOpConditional();
-                        conditionalAd->type = OP_TYPE::ME_CONDITIONAL;
-                        conditionalAd->no = lastMeshOp;
-                        conditionalAd->yes = morphAd;
-                        conditionalAd->condition = e.Condition;
-                        lastMeshOp = conditionalAd;
-                    }
-                    else
-                    {
-                        lastMeshOp = morphAd;
-                    }
-                }
-            }
-
 			// Apply the modifier for the post-normal operations stage.
 			bool bModifiersForBeforeOperations = false;
-			FGenericGenerationOptions ModifierOptions(Options);
+			FMeshGenerationOptions ModifierOptions(MeshOptions);
 			ModifierOptions.ActiveTags = node.Tags;
-			lastMeshOp = ApplyMeshModifiers(Options, lastMeshOp, bModifiersForBeforeOperations, surfaceNode->GetMessageContext());
+			lastMeshOp = ApplyMeshModifiers(MeshOptions, MeshResults, SharedMeshResults, bModifiersForBeforeOperations, surfaceNode->GetMessageContext(), nullptr);
 
             // Layouts
-            for ( int32 LayoutIndex=0; LayoutIndex <meshResults.GeneratedLayouts.Num(); ++LayoutIndex)
+            for ( int32 LayoutIndex=0; LayoutIndex < MeshResults.GeneratedLayouts.Num(); ++LayoutIndex)
             {
                 Ptr<ASTOp> layoutOp;
 
-				Ptr<const Layout> pLayout = meshResults.GeneratedLayouts[LayoutIndex].Layout;
+				Ptr<const Layout> pLayout = MeshResults.GeneratedLayouts[LayoutIndex].Layout;
                 if ( pLayout )
                 {
 					if (SharedMeshResults)
@@ -1174,16 +964,16 @@ namespace mu
 						}
 
 						// Add children merged meshes layouts
-						for (const FMeshGenerationResult::FExtraLayouts& data : meshResults.ExtraMeshLayouts)
+						for (const FMeshGenerationResult::FExtraLayouts& Data : MeshResults.ExtraMeshLayouts)
 						{
-							if (!data.MeshFragment)
+							if (!Data.MeshFragment)
 							{
 								// No mesh to add, we assume there are no layouts to add either.
-								check(data.GeneratedLayouts.IsEmpty());
+								check(Data.GeneratedLayouts.IsEmpty());
 								continue;
 							}
 
-							if (data.GeneratedLayouts.Num() != meshResults.GeneratedLayouts.Num())
+							if (Data.GeneratedLayouts.Num() != MeshResults.GeneratedLayouts.Num())
 							{
 								ErrorLog->GetPrivate()->Add(TEXT("Merged layout has been ignored because the number of layouts is different."), ELMT_ERROR, surfaceNode->GetMessageContext());
 							}
@@ -1195,7 +985,7 @@ namespace mu
 									Ptr<ASTOpConstantResource> op = new ASTOpConstantResource();
 									op->Type = OP_TYPE::LA_CONSTANT;
 
-									Ptr<const Layout> pCloned = data.GeneratedLayouts[LayoutIndex].Layout;
+									Ptr<const Layout> pCloned = Data.GeneratedLayouts[LayoutIndex].Layout;
 									op->SetValue(pCloned, CompilerOptions->OptimisationOptions.DiskCacheContext);
 
 									layoutFragmentAd = op;
@@ -1207,13 +997,13 @@ namespace mu
 								mergeAd->Added = layoutFragmentAd;
 
 								// Condition to apply
-								if (data.Condition)
+								if (Data.Condition)
 								{
 									Ptr<ASTOpConditional> conditionalAd = new ASTOpConditional();
 									conditionalAd->type = OP_TYPE::LA_CONDITIONAL;
 									conditionalAd->no = layoutOp;
 									conditionalAd->yes = mergeAd;
-									conditionalAd->condition = data.Condition;
+									conditionalAd->condition = Data.Condition;
 									layoutOp = conditionalAd;
 								}
 								else
@@ -1261,7 +1051,7 @@ namespace mu
 
                 }
 
-                meshResults.LayoutOps.Add( layoutOp );
+                MeshResults.LayoutOps.Add( layoutOp );
             }
 
             // Store in the surface for later use.
@@ -1273,7 +1063,7 @@ namespace mu
         //------------------------------------------------------------------------
 		if (!bShareSurface)
 		{
-			for (int32 t = 0; t < node.Images.Num(); ++t)
+			for (int32 ImageIndex = 0; ImageIndex < node.Images.Num(); ++ImageIndex)
 			{
 				MUTABLE_CPUPROFILER_SCOPE(SurfaceTexture);
 
@@ -1283,7 +1073,7 @@ namespace mu
 				Ptr<NodeImageSwizzle> swizzleNode;
 
 				bool bFound = false;
-				NodeImagePtr pImageNode = node.Images[t].Image;
+				Ptr<NodeImage> pImageNode = node.Images[ImageIndex].Image;
 
 				while (!bFound && pImageNode)
 				{
@@ -1322,7 +1112,7 @@ namespace mu
 				if (bFound)
 				{
 
-					const int32 LayoutIndex = node.Images[t].LayoutIndex;
+					const int32 LayoutIndex = node.Images[ImageIndex].LayoutIndex;
 
 					// If the layout index has been set to negative, it means we should ignore the layout for this image.
 					CompilerOptions::TextureLayoutStrategy ImageLayoutStrategy = (LayoutIndex < 0)
@@ -1344,40 +1134,19 @@ namespace mu
 						GenerateImage(ImageOptions, Result, pImageNode);
 						Ptr<ASTOp> imageAd = Result.op;
 
-						// Look for patches to this block
-						for (int32 editIndex = 0; editIndex < edits.Num(); ++editIndex)
-						{
-							const FirstPassGenerator::FSurface::FEdit& e = edits[editIndex];
-							if (t >= e.Node->Textures.Num())
-							{
-								continue;
-							}
-
-							const NodeSurfaceEdit::FTexture& Patch = e.Node->Textures[t];
-							if (Patch.PatchImage.get())
-							{
-								// Does the current block need to be patched?
-
-								// Ideally this should be the actual image size
-								constexpr int32 FakeLayoutSize = 256;
-
-								FIntPoint GridSize(FakeLayoutSize, FakeLayoutSize);
-								int32 BlockPixelsX = 1;
-								int32 BlockPixelsY = 1;
-								box< FIntVector2 > RectInCells;
-								RectInCells.min = { 0,0 };
-								RectInCells.size = { FakeLayoutSize ,FakeLayoutSize };
-
-								Ptr<Image> PatchMask = GenerateImageBlockPatchMask(Patch, GridSize, BlockPixelsX, BlockPixelsY, RectInCells);
-
-								if (PatchMask)
-								{
-									imageAd = GenerateImageBlockPatch(imageAd, Patch, PatchMask, e.Condition, ImageOptions);
-								}
-							}
-
-						}
-
+						// Placeholder block. Ideally this should be the actual image size
+						constexpr int32 FakeLayoutSize = 256;
+						FIntPoint GridSize(FakeLayoutSize, FakeLayoutSize);
+						FLayoutBlockDesc LayoutBlockDesc;
+						LayoutBlockDesc.BlockPixelsX = 1;
+						LayoutBlockDesc.BlockPixelsY = 1;
+						box< FIntVector2 > RectInCells;
+						RectInCells.min = { 0,0 };
+						RectInCells.size = { FakeLayoutSize ,FakeLayoutSize };
+						
+						bool bModifiersForBeforeOperations = false;
+						imageAd = ApplyImageBlockModifiers(ImageOptions, imageAd, ImageIndex, GridSize, LayoutBlockDesc, RectInCells, bModifiersForBeforeOperations, surfaceNode->GetMessageContext());
+						
 						check(imageAd);
 
 						if (swizzleNode)
@@ -1422,21 +1191,21 @@ namespace mu
 						op->type = OP_TYPE::IN_ADDIMAGE;
 						op->instance = lastSurfOp;
 						op->value = imageAd;
-						op->name = node.Images[t].Name;
+						op->name = node.Images[ImageIndex].Name;
 
 						lastSurfOp = op;
 					}
 
 					else if (ImageLayoutStrategy == CompilerOptions::TextureLayoutStrategy::Pack) //-V547
 					{						
-						if (LayoutIndex >= meshResults.GeneratedLayouts.Num() ||
-							LayoutIndex >= meshResults.LayoutOps.Num())
+						if (LayoutIndex >= MeshResults.GeneratedLayouts.Num() ||
+							LayoutIndex >= MeshResults.LayoutOps.Num())
 						{
 							ErrorLog->GetPrivate()->Add("Missing layout in object, or its parent.", ELMT_ERROR, surfaceNode->GetMessageContext());
 						}
 						else
 						{
-							const Layout* pLayout = meshResults.GeneratedLayouts[LayoutIndex].Layout.get();
+							const Layout* pLayout = MeshResults.GeneratedLayouts[LayoutIndex].Layout.get();
 							check(pLayout);
 
 							Ptr<ASTOpInstanceAdd> op = new ASTOpInstanceAdd();
@@ -1448,11 +1217,17 @@ namespace mu
 
 							// Size of a layout block in pixels
 							FIntPoint GridSize = pLayout->GetGridSize();
-
-							EImageFormat FinalFormat = EImageFormat::IF_NONE;
-							int32 BlockPixelsX = 0;
-							int32 BlockPixelsY = 0;
-							bool bBlocksHaveMips = false;
+	
+							// Try to guess the layout block description from the first valid block that is generated.
+							FLayoutBlockDesc LayoutBlockDesc;
+							if (formatNode)
+							{
+								LayoutBlockDesc.FinalFormat = formatNode->GetPrivate()->m_formatIfAlpha;
+								if (LayoutBlockDesc.FinalFormat == EImageFormat::IF_NONE)
+								{
+									LayoutBlockDesc.FinalFormat = formatNode->GetPrivate()->m_format;
+								}
+							}
 
 							bool bImageSizeWarning = false;
 
@@ -1462,57 +1237,12 @@ namespace mu
 							{
 								BlankImageOp = new ASTOpFixed();
 								BlankImageOp->op.type = OP_TYPE::IM_BLANKLAYOUT;
-								BlankImageOp->SetChild(BlankImageOp->op.args.ImageBlankLayout.layout, meshResults.LayoutOps[LayoutIndex]);
+								BlankImageOp->SetChild(BlankImageOp->op.args.ImageBlankLayout.layout, MeshResults.LayoutOps[LayoutIndex]);
 								// The rest ok the op will be completed below
 								BlankImageOp->op.args.ImageBlankLayout.mipmapCount = 0;
 								imageAd = BlankImageOp;
 							}
 
-							auto UpdateBlockSize = [&BlockPixelsX, &BlockPixelsY, &FinalFormat, &bBlocksHaveMips, &bImageSizeWarning, &formatNode, &BlankImageOp, &node, surfaceNode, &t, &Options, this]( FImageDesc BlockDesc, FIntVector2 LayoutCellSize )
-							{
-								if (BlockPixelsX == 0 && LayoutCellSize.X>0 && LayoutCellSize.Y>0)
-								{
-									if (!bImageSizeWarning)
-									{
-										// If the block pixels is not a multiple of the block layout cells
-										if ((BlockDesc.m_size[0] % LayoutCellSize[0] != 0) || (BlockDesc.m_size[1] % LayoutCellSize[1] != 0))
-										{
-											bImageSizeWarning = true;
-
-											check(CurrentParents.Last().Lod== Options.LODIndex);
-											int32 currentLOD = Options.LODIndex;
-											FString Msg = FString::Printf(TEXT("A texture [%s] for material [%s] parameter [%s] in LOD [%d] has been resized because it didn't fit the layout. "),
-												*node.Images[t].Name,
-												*node.Images[t].MaterialName,
-												*node.Images[t].MaterialParameterName,
-												currentLOD);
-											ErrorLog->GetPrivate()->Add(Msg, ELMT_INFO, surfaceNode->GetMessageContext());
-										}
-									}
-
-									BlockPixelsX = FMath::Max(1, BlockDesc.m_size[0] / LayoutCellSize[0]);
-									BlockPixelsY = FMath::Max(1, BlockDesc.m_size[1] / LayoutCellSize[1]);
-									bBlocksHaveMips = BlockDesc.m_lods > 1;
-
-									FinalFormat = BlockDesc.m_format;
-									if (formatNode)
-									{
-										FinalFormat = formatNode->GetPrivate()->m_formatIfAlpha;
-										if (FinalFormat == EImageFormat::IF_NONE)
-										{
-											FinalFormat = formatNode->GetPrivate()->m_format;
-										}
-									}
-
-									// Complete the base op
-									BlankImageOp->op.args.ImageBlankLayout.blockSize[0] = uint16(BlockPixelsX);
-									BlankImageOp->op.args.ImageBlankLayout.blockSize[1] = uint16(BlockPixelsY);
-									BlankImageOp->op.args.ImageBlankLayout.format = GetUncompressedFormat(FinalFormat);
-									BlankImageOp->op.args.ImageBlankLayout.generateMipmaps = bBlocksHaveMips;
-									BlankImageOp->op.args.ImageBlankLayout.mipmapCount = 0;
-								}
-							};
-
 							for (int32 BlockIndex = 0; BlockIndex < pLayout->GetBlockCount(); ++BlockIndex)
 							{
 								// Generate the image
@@ -1543,33 +1273,14 @@ namespace mu
 								RectInCells.min = pLayout->Blocks[BlockIndex].Min;
 								RectInCells.size = pLayout->Blocks[BlockIndex].Size;
 
-								// If we don't know the size of a layout block in pixels, calculate it
-								UpdateBlockSize(BlockDesc, RectInCells.size);
+								// Try to update the layout block desc if we don't know it yet.
+								UpdateLayoutBlockDesc(LayoutBlockDesc, BlockDesc, RectInCells.size);
 
 								// Even if we force the size afterwards, we need some size hint in some cases, like image projections.
 								ImageOptions.RectSize = UE::Math::TIntVector2<int32>(BlockDesc.m_size);
 
-								// Look for patches to this block
-								for (const FirstPassGenerator::FSurface::FEdit& Edit : edits)
-								{
-									if (t >= Edit.Node->Textures.Num())
-									{
-										continue;
-									}
-
-									const NodeSurfaceEdit::FTexture& Patch = Edit.Node->Textures[t];
-									if (Patch.PatchImage)
-									{
-										// Is the current block to be patched?
-										Ptr<Image> PatchMask = GenerateImageBlockPatchMask( Patch, GridSize, BlockPixelsX, BlockPixelsY, RectInCells );
-
-										// If there was any edit affecting this block
-										if (PatchMask)
-										{
-											blockAd = GenerateImageBlockPatch(blockAd, Patch, PatchMask, Edit.Condition, ImageOptions);
-										}
-									}
-								}
+								bool bModifiersForBeforeOperations = false;
+								blockAd = ApplyImageBlockModifiers(ImageOptions, blockAd, ImageIndex, GridSize, LayoutBlockDesc, RectInCells, bModifiersForBeforeOperations, surfaceNode->GetMessageContext());
 
 								// Enforce block size and optimizations
 								blockAd = GenerateImageSize(blockAd, FIntVector2(BlockDesc.m_size));
@@ -1579,11 +1290,11 @@ namespace mu
 								//Ptr<ASTOp> blockAd = GenerateImageFormat(blockAd, baseFormat);
 
 								// Apply tiling to avoid generating chunks of image that are too big.
-								blockAd = ApplyTiling(blockAd, ImageOptions.RectSize, FinalFormat);
+								blockAd = ApplyTiling(blockAd, ImageOptions.RectSize, LayoutBlockDesc.FinalFormat);
 
 								// Compose layout operation
 								Ptr<ASTOpImageCompose> composeOp = new ASTOpImageCompose();
-								composeOp->Layout = meshResults.LayoutOps[LayoutIndex];
+								composeOp->Layout = MeshResults.LayoutOps[LayoutIndex];
 								composeOp->Base = imageAd;
 								composeOp->BlockImage = blockAd;
 
@@ -1595,104 +1306,17 @@ namespace mu
 							}
 							check(imageAd);
 
-							// Apply composition of blocks coming from child objects
-							for (int32 editIndex = 0; editIndex < edits.Num(); ++editIndex)
-							{
-								const FirstPassGenerator::FSurface::FEdit& e = edits[editIndex];
-								if (t < e.Node->Textures.Num())
-								{
-									Ptr<NodeImage> pExtend = e.Node->Textures[t].Extend;
-									if (pExtend)
-									{
-										if (LayoutIndex >= meshResults.ExtraMeshLayouts[editIndex].GeneratedLayouts.Num()
-											||
-											!meshResults.ExtraMeshLayouts[editIndex].GeneratedLayouts[LayoutIndex].Layout)
-										{
-											FString Msg = FString::Printf(TEXT("Trying to extend a layout that doesn't exist in object [%s]."),
-												*CurrentParents.Last().ObjectNode->Name
-											);
-
-											ErrorLog->GetPrivate()->Add(Msg, ELMT_ERROR, surfaceNode->GetMessageContext());
-										}
-										else
-										{
-											Ptr<const Layout> pExtendLayout = meshResults.ExtraMeshLayouts[editIndex].GeneratedLayouts[LayoutIndex].Layout;
-
-											// Size of a layout block in pixels
-											FIntPoint extlayout = pExtendLayout->GetGridSize();
-
-											Ptr<ASTOp> lastBase = imageAd;
-
-											for (int32 b = 0; b < pExtendLayout->GetBlockCount(); ++b)
-											{
-												// Generate the image block
-												FImageGenerationOptions ImageOptions;
-												ImageOptions.State = Options.State;
-												ImageOptions.ImageLayoutStrategy = ImageLayoutStrategy;
-												ImageOptions.ActiveTags = node.Tags;
-												ImageOptions.RectSize = { 0,0 };
-												ImageOptions.LayoutToApply = pExtendLayout;
-												ImageOptions.LayoutBlockId = pExtendLayout->Blocks[b].Id;
-												FImageGenerationResult ExtendResult;
-												GenerateImage(ImageOptions, ExtendResult, pExtend);
-												Ptr<ASTOp> fragmentAd = ExtendResult.op;
-
-												// Block in layout grid units
-												box< FIntVector2 > rectInCells;
-												rectInCells.min = pExtendLayout->Blocks[b].Min;
-												rectInCells.size = pExtendLayout->Blocks[b].Size;
-
-												FImageDesc ExtendDesc = fragmentAd->GetImageDesc();
-
-												// If we don't know the size of a layout block in pixels, calculate it
-												UpdateBlockSize(ExtendDesc, rectInCells.size);
-
-												// Adjust the format and size of the block to be added
-												// Actually don't do it, it will be propagated from the top format operation.
-												//fragmentAd = GenerateImageFormat(fragmentAd, FinalFormat);
-
-												UE::Math::TIntVector2<int32> expectedSize;
-												expectedSize[0] = BlockPixelsX * rectInCells.size[0];
-												expectedSize[1] = BlockPixelsY * rectInCells.size[1];
-												fragmentAd = GenerateImageSize(fragmentAd, expectedSize);
-
-												// Apply tiling to avoid generating chunks of image that are too big.
-												fragmentAd = ApplyTiling(fragmentAd, expectedSize, FinalFormat);
-
-												// Compose operation
-												Ptr<ASTOpImageCompose> composeOp = new ASTOpImageCompose();
-												composeOp->Layout = meshResults.LayoutOps[LayoutIndex];
-												composeOp->Base = lastBase;
-												composeOp->BlockImage = fragmentAd;
-
-												// Set the absolute block index.
-												check(pExtendLayout->Blocks[b].Id != FLayoutBlock::InvalidBlockId);
-												composeOp->BlockId = pExtendLayout->Blocks[b].Id;
-
-												lastBase = composeOp;
-											}
-
-											// Condition to enable this image extension
-											if (e.Condition)
-											{
-												Ptr<ASTOp> conditionalAd;
-												Ptr<ASTOpConditional> cop = new ASTOpConditional();
-												cop->type = OP_TYPE::IM_CONDITIONAL;
-												cop->no = imageAd;
-												cop->yes = lastBase;
-												cop->condition = e.Condition;
-												conditionalAd = cop;
-												imageAd = conditionalAd;
-											}
-											else
-											{
-												imageAd = lastBase;
-											}
-
-										}
-									}
-								}
-							}
+							FMeshGenerationOptions ModifierOptions(Options);
+							ModifierOptions.ActiveTags = node.Tags;
+							bool bModifiersForBeforeOperations = false;
+							imageAd = ApplyImageExtendModifiers(ModifierOptions, MeshResults, imageAd, ImageLayoutStrategy, LayoutIndex, ImageIndex, GridSize, LayoutBlockDesc, bModifiersForBeforeOperations, surfaceNode->GetMessageContext());
+
+							// Complete the base op
+							BlankImageOp->op.args.ImageBlankLayout.blockSize[0] = uint16(LayoutBlockDesc.BlockPixelsX);
+							BlankImageOp->op.args.ImageBlankLayout.blockSize[1] = uint16(LayoutBlockDesc.BlockPixelsY);
+							BlankImageOp->op.args.ImageBlankLayout.format = GetUncompressedFormat(LayoutBlockDesc.FinalFormat);
+							BlankImageOp->op.args.ImageBlankLayout.generateMipmaps = LayoutBlockDesc.bBlocksHaveMips;
+							BlankImageOp->op.args.ImageBlankLayout.mipmapCount = 0;
 
 							if (swizzleNode)
 							{
@@ -1711,7 +1335,12 @@ namespace mu
 							}
 
 							// Apply mipmap and format if necessary, skip if format is IF_NONE (possibly because a block was skipped above)
-							if (mipmapNode && FinalFormat != EImageFormat::IF_NONE)
+							bool bNeedsMips =
+								(mipmapNode && LayoutBlockDesc.FinalFormat != EImageFormat::IF_NONE)
+								||
+								LayoutBlockDesc.bBlocksHaveMips;
+
+							if (bNeedsMips)
 							{
 								Ptr<ASTOpImageMipmap> mop = new ASTOpImageMipmap();
 
@@ -1723,55 +1352,28 @@ namespace mu
 
 								// We have to avoid mips smaller than the image format block size, so
 								// we will devide the layout block by the format block
-								const FImageFormatData& finfo = GetImageFormatData(FinalFormat);
+								const FImageFormatData& PixelFormatInfo = GetImageFormatData(LayoutBlockDesc.FinalFormat);
 
-								int32 mipsX = FMath::CeilLogTwo(BlockPixelsX / finfo.PixelsPerBlockX);
-								int32 mipsY = FMath::CeilLogTwo(BlockPixelsY / finfo.PixelsPerBlockY);
+								int32 mipsX = FMath::CeilLogTwo(LayoutBlockDesc.BlockPixelsX / PixelFormatInfo.PixelsPerBlockX);
+								int32 mipsY = FMath::CeilLogTwo(LayoutBlockDesc.BlockPixelsY / PixelFormatInfo.PixelsPerBlockY);
 								mop->BlockLevels = (uint8)FMath::Max(mipsX, mipsY);
 								
-								if (BlockPixelsX < finfo.PixelsPerBlockX || BlockPixelsY < finfo.PixelsPerBlockY)
+								if (LayoutBlockDesc.BlockPixelsX < PixelFormatInfo.PixelsPerBlockX || LayoutBlockDesc.BlockPixelsY < PixelFormatInfo.PixelsPerBlockY)
 								{
 									// In this case, the mipmap will never be useful for blocks, so we indicate that
 									// it should make the mips at the root of the expression.
 									mop->bOnlyTail = true;
 								}
 
-								mop->AddressMode = mipmapNode->GetPrivate()->m_settings.AddressMode;
-								mop->FilterType = mipmapNode->GetPrivate()->m_settings.FilterType;
-
-								imageAd = mop;
-							}
-
-							else if (bBlocksHaveMips)
-							{
-								// If the blocks had mipmaps, we still need to generate them after the compose, to have the full chain
-								Ptr<ASTOpImageMipmap> mop = new ASTOpImageMipmap();
-
-								// At the end of the day, we want all the mipmaps. Maybe the code
-								// optimiser will split the process later.
-								mop->Levels = 0;
-								mop->bOnlyTail = false;
-								mop->Source = imageAd;
-
-								// We have to avoid mips smaller than the image format block size, so
-								// we will devide the layout block by the format block
-								const FImageFormatData& finfo = GetImageFormatData(FinalFormat);
-
-								int32 mipsX = (int32)ceilf(logf((float)BlockPixelsX / finfo.PixelsPerBlockX) / logf(2.0f));
-								int32 mipsY = (int32)ceilf(logf((float)BlockPixelsY / finfo.PixelsPerBlockY) / logf(2.0f));
-								mop->BlockLevels = (uint8_t)FMath::Max(mipsX, mipsY);
+								mop->AddressMode = EAddressMode::ClampToEdge;
+								mop->FilterType = EMipmapFilterType::SimpleAverage;
 
-								if (BlockPixelsX < finfo.PixelsPerBlockX || BlockPixelsY < finfo.PixelsPerBlockY)
+								if (mipmapNode)
 								{
-									// In this case, the mipmap will never be useful for blocks, so we indicate that
-									// it should make the mips at the root of the expression.
-									mop->bOnlyTail = true;
+									mop->AddressMode = mipmapNode->GetPrivate()->m_settings.AddressMode;
+									mop->FilterType = mipmapNode->GetPrivate()->m_settings.FilterType;
 								}
 
-								// Not important for the end of the mip tail?
-								mop->AddressMode = EAddressMode::ClampToEdge;
-								mop->FilterType = EMipmapFilterType::SimpleAverage;
-
 								imageAd = mop;
 							}
 
@@ -1788,7 +1390,7 @@ namespace mu
 							op->value = imageAd;
 
 							// Name
-							op->name = node.Images[t].Name;
+							op->name = node.Images[ImageIndex].Name;
 
 							lastSurfOp = op;
 						}
@@ -1879,7 +1481,7 @@ namespace mu
 		if (bIsBaseForSharedSurface)
 		{
 			check(!SharedMeshOptionsMap.Contains(node.SharedSurfaceId));
-			SharedMeshOptionsMap.Add( node.SharedSurfaceId, meshResults );
+			SharedMeshOptionsMap.Add( node.SharedSurfaceId, MeshResults );
 		}
     }
 
@@ -1932,7 +1534,7 @@ namespace mu
 
 				FSurfaceGenerationOptions SurfaceOptions(Options);
 				FSurfaceGenerationResult surfaceGenerationResult;
-                GenerateSurface( surfaceGenerationResult, SurfaceOptions, its.Node, its.Edits );
+                GenerateSurface( surfaceGenerationResult, SurfaceOptions, its.Node );
                 sop->value = surfaceGenerationResult.surfaceOp;
 
                 sop->id = surfaceID;
@@ -1940,11 +1542,7 @@ namespace mu
                 sop->SharedSurfaceId = its.Node->SharedSurfaceId;
                 Ptr<ASTOp> surfaceAt = sop;
 
-				// TODO: This could be done earlier?
-                Ptr<ASTOpFixed> SurfaceConditionOp = new ASTOpFixed();
-				SurfaceConditionOp->op.type = OP_TYPE::BO_AND;
-				SurfaceConditionOp->SetChild(SurfaceConditionOp->op.args.BoolBinary.a, its.ObjectCondition );
-				SurfaceConditionOp->SetChild(SurfaceConditionOp->op.args.BoolBinary.b, its.SurfaceCondition );
+                Ptr<ASTOp> SurfaceConditionOp = its.FinalCondition;
 
                 {
                     Ptr<ASTOpConditional> op = new ASTOpConditional();
@@ -2329,33 +1927,237 @@ namespace mu
 	}
 
 	//---------------------------------------------------------------------------------------------
-	Ptr<ASTOp> CodeGenerator::ApplyMeshModifiers(
-		const FGenericGenerationOptions& Options,
-		const Ptr<ASTOp>& sourceOp,
+	Ptr<ASTOp> CodeGenerator::ApplyMeshModifiers(const FMeshGenerationOptions& Options, 
+		FMeshGenerationResult& BaseMeshResult,
+		const FMeshGenerationResult* SharedMeshResults,
 		bool bModifiersForBeforeOperations,
-		const void* errorContext )
+		const void* ErrorContext,
+		const NodeMeshConstant* OriginalMeshNode )
 	{
-		Ptr<ASTOp> lastMeshOp = sourceOp;
+		Ptr<ASTOp> LastMeshOp = BaseMeshResult.MeshOp;
 
 		// Apply mesh modifiers
-		TArray<FirstPassGenerator::FModifier> modifiers;
+		TArray<FirstPassGenerator::FModifier> Modifiers;
 
 		int32 currentLOD = CurrentParents.Last().Lod;
-		GetModifiersFor(Options.ActiveTags, currentLOD, bModifiersForBeforeOperations, modifiers);
+		GetModifiersFor(Options.ActiveTags, currentLOD, bModifiersForBeforeOperations, Modifiers);
+
+		Ptr<ASTOp> PreModifiersMesh = LastMeshOp;
+
+		ActiveTags.Push({});
+
+		// Process mesh extend modifiers (from edit modifiers)
+		int32 EditIndex = 0;
+		for (const FirstPassGenerator::FModifier& m : Modifiers)
+		{
+			if (m.Node->GetType() == NodeModifierSurfaceEdit::GetStaticType())
+			{
+				const NodeModifierSurfaceEdit* Edit = static_cast<const NodeModifierSurfaceEdit*>(m.Node);
+
+				BaseMeshResult.ExtraMeshLayouts.Emplace();
+
+				if (Ptr<NodeMesh> pAdd = Edit->MeshAdd)
+				{
+					// Store the data necessary to apply modifiers for the pre-normal operations stage.
+					ActiveTags.Push(Edit->EnableTags);
+
+					FMeshGenerationOptions MergedMeshOptions(Options);
+					MergedMeshOptions.ActiveTags = Edit->EnableTags; // TODO: Append to current?
+
+					if (SharedMeshResults)
+					{
+						check(SharedMeshResults->ExtraMeshLayouts.IsValidIndex(EditIndex));
+						MergedMeshOptions.OverrideLayouts = SharedMeshResults->ExtraMeshLayouts[EditIndex].GeneratedLayouts;
+					}
+
+					FMeshGenerationResult AddResults;
+					GenerateMesh(MergedMeshOptions, AddResults, pAdd);
+
+					ActiveTags.Pop();
+
+					// Apply the modifier for the post-normal operations stage to the added mesh
+					FMeshGenerationOptions ModifierOptions(Options);
+					ModifierOptions.ActiveTags = Edit->EnableTags;
+					Ptr<ASTOp> AddedMeshOp = ApplyMeshModifiers(ModifierOptions, AddResults, SharedMeshResults, bModifiersForBeforeOperations, ErrorContext, nullptr);
+
+					FMeshGenerationResult::FExtraLayouts data;
+					data.GeneratedLayouts = AddResults.GeneratedLayouts;
+					data.Condition = m.FinalCondition;
+					data.MeshFragment = AddedMeshOp;
+					BaseMeshResult.ExtraMeshLayouts[EditIndex] = data;
+
+					Ptr<ASTOpFixed> mop = new ASTOpFixed();
+					mop->op.type = OP_TYPE::ME_MERGE;
+					mop->SetChild(mop->op.args.MeshMerge.base, LastMeshOp);
+					mop->SetChild(mop->op.args.MeshMerge.added, AddedMeshOp);
+					// will merge the meshes under the same surface
+					mop->op.args.MeshMerge.newSurfaceID = 0;
+
+					// Condition to apply
+					if (m.FinalCondition)
+					{
+						Ptr<ASTOpConditional> conditionalAd = new ASTOpConditional();
+						conditionalAd->type = OP_TYPE::ME_CONDITIONAL;
+						conditionalAd->no = LastMeshOp;
+						conditionalAd->yes = mop;
+						conditionalAd->condition = m.FinalCondition;
+						LastMeshOp = conditionalAd;
+					}
+					else
+					{
+						LastMeshOp = mop;
+					}
+				}
+
+				++EditIndex;
+			}
+
+		}
+
+		// "remove" operation to group all the removes
+		Ptr<ASTOpMeshRemoveMask> RemoveOp;
+
+		// Process mesh remove modifiers (from edit modifiers)
+		for (const FirstPassGenerator::FModifier& m : Modifiers)
+		{
+			if (m.Node->GetType() == NodeModifierSurfaceEdit::GetStaticType())
+			{
+				const NodeModifierSurfaceEdit* Edit = static_cast<const NodeModifierSurfaceEdit*>(m.Node);
+
+				// Apply mesh removes from child objects "edit surface" nodes.
+				// "Removes" need to come after "Adds" because some removes may refer to added meshes,
+				// and not the base.
+				// \TODO: Apply base removes first, and then "added meshes" removes here. It may have lower memory footprint during generation.
+				if (Ptr<NodeMesh> pRemove = Edit->MeshRemove)
+				{
+					FMeshGenerationResult removeResults;
+					FMeshGenerationOptions RemoveMeshOptions;
+					RemoveMeshOptions.bLayouts = false;
+					RemoveMeshOptions.State = Options.State;
+					RemoveMeshOptions.ActiveTags = Edit->EnableTags;
+
+					GenerateMesh(RemoveMeshOptions, removeResults, pRemove);
+
+					Ptr<ASTOpFixed> maskOp = new ASTOpFixed();
+					maskOp->op.type = OP_TYPE::ME_MASKDIFF;
+
+					// By default, remove from the base
+					Ptr<ASTOp> removeFrom = BaseMeshResult.BaseMeshOp;
+					maskOp->SetChild(maskOp->op.args.MeshMaskDiff.source, removeFrom);
+					maskOp->SetChild(maskOp->op.args.MeshMaskDiff.fragment, removeResults.MeshOp);
+
+					if (!RemoveOp)
+					{
+						RemoveOp = new ASTOpMeshRemoveMask();
+						RemoveOp->source = LastMeshOp;
+						LastMeshOp = RemoveOp;
+					}
+
+					RemoveOp->AddRemove(m.FinalCondition, maskOp);
+				}
+			}
+		}
+
+
+		// Process mesh morph modifiers (from edit modifiers)
+		for (const FirstPassGenerator::FModifier& m : Modifiers)
+		{
+			if (m.Node->GetType() == NodeModifierSurfaceEdit::GetStaticType())
+			{
+				const NodeModifierSurfaceEdit* Edit = static_cast<const NodeModifierSurfaceEdit*>(m.Node);
+
+				if (Edit->MeshMorph.IsEmpty())
+				{
+					continue;
+				}
+
+				check(OriginalMeshNode);
+
+				Ptr<Mesh> TargetMesh = OriginalMeshNode->FindMorph(Edit->MeshMorph);
+				if (!TargetMesh)
+				{
+					continue;
+				}
+				
+				{
+					// Target mesh
+					Ptr<ASTOpConstantResource> TargetMeshOp = new ASTOpConstantResource;
+					TargetMeshOp->Type = OP_TYPE::ME_CONSTANT;
+					TargetMeshOp->SetValue(TargetMesh->Clone(), CompilerOptions->OptimisationOptions.DiskCacheContext);
+
+					// Morph generation through mesh diff
+					Ptr<ASTOpMeshDifference> diffAd;
+					{
+						Ptr<ASTOpMeshDifference> op = new ASTOpMeshDifference();
+						op->Base = BaseMeshResult.BaseMeshOp;
+						op->Target = TargetMeshOp;
+
+						// Morphing tex coords here is not supported:
+						// Generating the homogoneous UVs is difficult since we don't have the base
+						// layout yet.                       
+						op->bIgnoreTextureCoords = true;
+						diffAd = op;
+					}
+
+					// Morph operation
+					Ptr<ASTOp> morphAd;
+					{
+						Ptr<ASTOpMeshMorph> op = new ASTOpMeshMorph();
+
+						// Factor
+						if (Edit->MorphFactor)
+						{
+							FScalarGenerationResult ChildResult;
+							GenerateScalar(ChildResult, Options, Edit->MorphFactor);
+							op->Factor = ChildResult.op;
+						}
+						else
+						{
+							Ptr<NodeScalarConstant> auxNode = new NodeScalarConstant();
+							auxNode->SetValue(1.0f);
+
+							FScalarGenerationResult ChildResult;
+							GenerateScalar(ChildResult, Options, auxNode);
+							op->Factor = ChildResult.op;
+						}
+
+						// Base		
+						op->Base = LastMeshOp;
+
+						// Targets
+						op->Target = diffAd;
+						morphAd = op;
+					}
+
+					// Condition to apply the morph
+					if (m.FinalCondition)
+					{
+						Ptr<ASTOpConditional> conditionalAd = new ASTOpConditional();
+						conditionalAd->type = OP_TYPE::ME_CONDITIONAL;
+						conditionalAd->no = LastMeshOp;
+						conditionalAd->yes = morphAd;
+						conditionalAd->condition = m.FinalCondition;
+						LastMeshOp = conditionalAd;
+					}
+					else
+					{
+						LastMeshOp = morphAd;
+					}
+				}
+			}
+		}
 
-		Ptr<ASTOp> preModifiersMesh = lastMeshOp;
 
-		ActiveTags.Add({});
 
 		// Process clip-with-mesh modifiers
-		Ptr<ASTOpMeshRemoveMask> removeOp;
-		for (const FirstPassGenerator::FModifier& m : modifiers)
+		RemoveOp = nullptr;
+		for (const FirstPassGenerator::FModifier& m : Modifiers)
 		{
 			if (m.Node->GetType()== NodeModifierMeshClipWithMesh::GetStaticType())
 			{
 				const NodeModifierMeshClipWithMesh* TypedClipNode = static_cast<const NodeModifierMeshClipWithMesh*>(m.Node);
 				Ptr<ASTOpMeshMaskClipMesh> op = new ASTOpMeshMaskClipMesh();
-				op->source = preModifiersMesh;
+				op->source = PreModifiersMesh;
 
 				// Parameters
 				FMeshGenerationOptions ClipOptions;
@@ -2368,78 +2170,103 @@ namespace mu
 
 				if (!op->clip)
 				{
-					ErrorLog->GetPrivate()->Add("Clip mesh has not been generated", ELMT_ERROR, errorContext);
+					ErrorLog->GetPrivate()->Add("Clip mesh has not been generated", ELMT_ERROR, ErrorContext);
 					continue;
 				}
 
 				Ptr<ASTOp> maskAt = op;
 
-				if (!removeOp)
+				if (!RemoveOp)
 				{
-					removeOp = new ASTOpMeshRemoveMask();
-					removeOp->source = lastMeshOp;
-					lastMeshOp = removeOp;
+					RemoveOp = new ASTOpMeshRemoveMask();
+					RemoveOp->source = LastMeshOp;
+					LastMeshOp = RemoveOp;
 				}
 
-				Ptr<ASTOpFixed> surfCondOp = new ASTOpFixed();
-				surfCondOp->op.type = OP_TYPE::BO_AND;
-				surfCondOp->SetChild(surfCondOp->op.args.BoolBinary.a, m.ObjectCondition);
-				surfCondOp->SetChild(surfCondOp->op.args.BoolBinary.b, m.SurfaceCondition);
-				Ptr<ASTOp> fullCondition = surfCondOp;
+				Ptr<ASTOp> fullCondition = m.FinalCondition;
 
-				removeOp->AddRemove(fullCondition, maskAt);
+				RemoveOp->AddRemove(fullCondition, maskAt);
 			}
 		}
 
 		// Process clip-with-mask modifiers
-		for (const FirstPassGenerator::FModifier& m : modifiers)
+		for (const FirstPassGenerator::FModifier& m : Modifiers)
 		{
 			if (m.Node->GetType() == NodeModifierMeshClipWithUVMask::GetStaticType())
 			{
+				// Create a constant mesh with the original UVs required by this modifier.
+				// TODO: Optimize, by caching.
+				// TODO: Optimize by formatting and keeping only UVs
+				check(OriginalMeshNode);
+				const Mesh* OriginalMesh = OriginalMeshNode->GetPrivate()->Value.get();
+				Ptr<ASTOpConstantResource> UVMeshOp = new ASTOpConstantResource();
+				UVMeshOp->Type = OP_TYPE::ME_CONSTANT;
+				UVMeshOp->SetValue(OriginalMesh->Clone(), CompilerOptions->OptimisationOptions.DiskCacheContext);
+
 				const NodeModifierMeshClipWithUVMask* TypedClipNode = static_cast<const NodeModifierMeshClipWithUVMask*>(m.Node);
+
+				Ptr<ASTOp> MeshMaskAt;
+
 				Ptr<ASTOpMeshMaskClipUVMask> op = new ASTOpMeshMaskClipUVMask();
-				op->Source = preModifiersMesh;
+				MeshMaskAt = op;
+				op->Source = BaseMeshResult.BaseMeshOp; 
+				op->UVSource = UVMeshOp; 
 				op->LayoutIndex = TypedClipNode->LayoutIndex;
 
-				// Parameters
-				FImageGenerationOptions ClipOptions;
-				ClipOptions.ImageLayoutStrategy = CompilerOptions::TextureLayoutStrategy::None;
-				ClipOptions.LayoutBlockId = FLayoutBlock::InvalidBlockId;
-				ClipOptions.State = Options.State;
+				if (TypedClipNode->ClipMask)
+				{
+					// Parameters to generate the mask image
+					FImageGenerationOptions ClipOptions;
+					ClipOptions.ImageLayoutStrategy = CompilerOptions::TextureLayoutStrategy::None;
+					ClipOptions.LayoutBlockId = FLayoutBlock::InvalidBlockId;
+					ClipOptions.State = Options.State;
 
-				FImageGenerationResult ClipMaskResult;
-				GenerateImage(ClipOptions, ClipMaskResult, TypedClipNode->ClipMask);
+					FImageGenerationResult ClipMaskResult;
+					GenerateImage(ClipOptions, ClipMaskResult, TypedClipNode->ClipMask);
 
-				// It could be IF_L_UBIT, but since this should be optimized out at compile time, leave the most cpu efficient.
-				op->Mask = GenerateImageFormat(ClipMaskResult.op, mu::EImageFormat::IF_L_UBYTE);
+					// It could be IF_L_UBIT, but since this should be optimized out at compile time, leave the most cpu efficient.
+					op->MaskImage = GenerateImageFormat(ClipMaskResult.op, mu::EImageFormat::IF_L_UBYTE);
 
-				if (!op->Mask)
-				{
-					ErrorLog->GetPrivate()->Add("Clip UV mask has not been generated", ELMT_ERROR, errorContext);
-					continue;
+					if (!op->MaskImage)
+					{
+						ErrorLog->GetPrivate()->Add("Clip UV mask has not been generated", ELMT_ERROR, ErrorContext);
+						continue;
+					}
 				}
 
-				Ptr<ASTOp> maskAt = op;
+				else if (TypedClipNode->ClipLayout)
+				{
+					// Generate the layout with blocks to extract
+					Ptr<const Layout> Layout = GenerateLayout(TypedClipNode->ClipLayout, 0);
+
+					Ptr<ASTOpConstantResource> LayoutOp = new ASTOpConstantResource();
+					LayoutOp->Type = OP_TYPE::LA_CONSTANT;
+					LayoutOp->SetValue(Layout, CompilerOptions->OptimisationOptions.DiskCacheContext);
+					op->MaskLayout = LayoutOp;
+				}
 
-				if (!removeOp)
+				else
 				{
-					removeOp = new ASTOpMeshRemoveMask();
-					removeOp->source = lastMeshOp;
-					lastMeshOp = removeOp;
+					// No mask or layout specified to clip. Don't clip anything.
 				}
 
-				Ptr<ASTOpFixed> surfCondOp = new ASTOpFixed();
-				surfCondOp->op.type = OP_TYPE::BO_AND;
-				surfCondOp->SetChild(surfCondOp->op.args.BoolBinary.a, m.ObjectCondition);
-				surfCondOp->SetChild(surfCondOp->op.args.BoolBinary.b, m.SurfaceCondition);
-				Ptr<ASTOp> fullCondition = surfCondOp;
+				if (MeshMaskAt)
+				{
+					if (!RemoveOp)
+					{
+						RemoveOp = new ASTOpMeshRemoveMask();
+						RemoveOp->source = LastMeshOp;
+						LastMeshOp = RemoveOp;
+					}
 
-				removeOp->AddRemove(fullCondition, maskAt);
+					Ptr<ASTOp> fullCondition = m.FinalCondition;
+					RemoveOp->AddRemove(fullCondition, MeshMaskAt);
+				}
 			}
 		}
 
 		// Process clip-morph-plane modifiers
-		for (const FirstPassGenerator::FModifier& m : modifiers)
+		for (const FirstPassGenerator::FModifier& m : Modifiers)
 		{
 			Ptr<ASTOp> modifiedMeshOp;
 
@@ -2447,7 +2274,7 @@ namespace mu
 			{
 				const NodeModifierMeshClipMorphPlane* TypedNode = static_cast<const NodeModifierMeshClipMorphPlane*>(m.Node);
 				Ptr<ASTOpMeshClipMorphPlane> op = new ASTOpMeshClipMorphPlane();
-				op->source = lastMeshOp;
+				op->source = LastMeshOp;
 
 				// Morph to an ellipse
 				{
@@ -2500,27 +2327,20 @@ namespace mu
 				op->factor = TypedNode->Parameters.LinearityFactor;
 
 				modifiedMeshOp = op;
-			}
 
-			if (modifiedMeshOp)
-			{
-				Ptr<ASTOpFixed> surfCondOp = new ASTOpFixed();
-				surfCondOp->op.type = OP_TYPE::BO_AND;
-				surfCondOp->SetChild(surfCondOp->op.args.BoolBinary.a, m.ObjectCondition);
-				surfCondOp->SetChild(surfCondOp->op.args.BoolBinary.b, m.SurfaceCondition);
-				Ptr<ASTOp> fullCondition = surfCondOp;
+				Ptr<ASTOp> fullCondition = m.FinalCondition;
 
-				Ptr<ASTOpConditional> op = new ASTOpConditional();
-				op->type = OP_TYPE::ME_CONDITIONAL;
-				op->no = lastMeshOp;
-				op->yes = modifiedMeshOp;
-				op->condition = fullCondition;
-				lastMeshOp = op;
+				Ptr<ASTOpConditional> ConditionalOp = new ASTOpConditional();
+				ConditionalOp->type = OP_TYPE::ME_CONDITIONAL;
+				ConditionalOp->no = LastMeshOp;
+				ConditionalOp->yes = modifiedMeshOp;
+				ConditionalOp->condition = fullCondition;
+				LastMeshOp = ConditionalOp;
 			}
 		}
 
     	// Process clip deform modifiers.
-		for (const FirstPassGenerator::FModifier& M : modifiers)
+		for (const FirstPassGenerator::FModifier& M : Modifiers)
 		{
 			Ptr<ASTOp> ModifiedMeshOp;
 
@@ -2538,7 +2358,7 @@ namespace mu
 				GenerateMesh(ClipOptions, ClipShapeResult, TypedClipNode->ClipMesh);
 				ClipOp->ClipShape = ClipShapeResult.MeshOp;
 				
-				BindOp->Mesh = lastMeshOp;
+				BindOp->Mesh = LastMeshOp;
 				BindOp->Shape = ClipShapeResult.MeshOp; 
 				BindOp->BindingMethod = static_cast<uint32>(TypedClipNode->BindingMethod);
 	
@@ -2547,7 +2367,7 @@ namespace mu
 				if (!ClipOp->ClipShape)
 				{
 					ErrorLog->GetPrivate()->Add
-					("Clip shape mesh has not been generated", ELMT_ERROR, errorContext);
+					("Clip shape mesh has not been generated", ELMT_ERROR, ErrorContext);
 				}
 				else
 				{
@@ -2557,24 +2377,217 @@ namespace mu
 			
 			if (ModifiedMeshOp)
 			{
-				Ptr<ASTOpFixed> SurfCondOp = new ASTOpFixed();
-				SurfCondOp->op.type = OP_TYPE::BO_AND;
-				SurfCondOp->SetChild(SurfCondOp->op.args.BoolBinary.a, M.ObjectCondition);
-				SurfCondOp->SetChild(SurfCondOp->op.args.BoolBinary.b, M.SurfaceCondition);
-				Ptr<ASTOp> FullCondition = SurfCondOp;
+				Ptr<ASTOp> FullCondition = M.FinalCondition;
 
 				Ptr<ASTOpConditional> Op = new ASTOpConditional();
 				Op->type = OP_TYPE::ME_CONDITIONAL;
-				Op->no = lastMeshOp;
+				Op->no = LastMeshOp;
 				Op->yes = ModifiedMeshOp;
 				Op->condition = FullCondition;
-				lastMeshOp = Op;
+				LastMeshOp = Op;
 			}
 		}
 			
 		ActiveTags.Pop();
 
-		return lastMeshOp;
+		return LastMeshOp;
+	}
+
+
+	Ptr<ASTOp> CodeGenerator::ApplyImageBlockModifiers(const FImageGenerationOptions& Options, Ptr<ASTOp> BaseImageOp, int32 ImageIndex, 
+		FIntPoint GridSize,
+		const FLayoutBlockDesc& LayoutBlockDesc,
+		box< FIntVector2 > RectInCells,
+		bool bModifiersForBeforeOperations, const void* ErrorContext)
+	{
+		Ptr<ASTOp> LastImageOp = BaseImageOp;
+
+		// Apply mesh modifiers
+		TArray<FirstPassGenerator::FModifier> Modifiers;
+
+		int32 currentLOD = CurrentParents.Last().Lod;
+		GetModifiersFor(Options.ActiveTags, currentLOD, bModifiersForBeforeOperations, Modifiers);
+
+		ActiveTags.Push({});
+
+		// Process patch image modifiers (from edit modifiers)
+		for (const FirstPassGenerator::FModifier& m : Modifiers)
+		{
+			if (m.Node->GetType() == NodeModifierSurfaceEdit::GetStaticType())
+			{
+				const NodeModifierSurfaceEdit* Edit = static_cast<const NodeModifierSurfaceEdit*>(m.Node);
+
+
+				if (ImageIndex >= Edit->Textures.Num())
+				{
+					continue;
+				}
+
+				const NodeModifierSurfaceEdit::FTexture& Patch = Edit->Textures[ImageIndex];
+				if (Patch.PatchImage.get())
+				{
+					// Does the current block need to be patched? Find out by building a mask.
+					Ptr<Image> PatchMask = GenerateImageBlockPatchMask(Patch, GridSize, LayoutBlockDesc.BlockPixelsX, LayoutBlockDesc.BlockPixelsY, RectInCells);
+
+					if (PatchMask)
+					{
+						LastImageOp = GenerateImageBlockPatch(LastImageOp, Patch, PatchMask, m.FinalCondition, Options);
+					}
+				}
+			}
+
+		}
+
+		ActiveTags.Pop();
+
+		return LastImageOp;
+	}
+
+
+	void CodeGenerator::UpdateLayoutBlockDesc(CodeGenerator::FLayoutBlockDesc& Out, FImageDesc BlockDesc, FIntVector2 LayoutCellSize)
+	{
+		if (Out.BlockPixelsX == 0 && LayoutCellSize.X > 0 && LayoutCellSize.Y > 0)
+		{
+			Out.BlockPixelsX = FMath::Max(1, BlockDesc.m_size[0] / LayoutCellSize[0]);
+			Out.BlockPixelsY = FMath::Max(1, BlockDesc.m_size[1] / LayoutCellSize[1]);
+			Out.bBlocksHaveMips = BlockDesc.m_lods > 1;
+
+			if (Out.FinalFormat==EImageFormat::IF_NONE)
+			{
+				Out.FinalFormat = BlockDesc.m_format;
+			}
+		}
+	};
+
+
+	Ptr<ASTOp> CodeGenerator::ApplyImageExtendModifiers(
+		const FGenericGenerationOptions& Options, 
+		const FMeshGenerationResult& BaseMeshResults,
+		Ptr<ASTOp> BaseImageOp, 
+		CompilerOptions::TextureLayoutStrategy ImageLayoutStrategy,
+		int32 LayoutIndex, int32 ImageIndex,
+		FIntPoint GridSize, 
+		CodeGenerator::FLayoutBlockDesc& InOutLayoutBlockDesc,
+		bool bModifiersForBeforeOperations, const void* ErrorContext)
+	{
+		Ptr<ASTOp> LastImageOp = BaseImageOp;
+
+		// Apply mesh modifiers
+		TArray<FirstPassGenerator::FModifier> Modifiers;
+
+		int32 currentLOD = CurrentParents.Last().Lod;
+		GetModifiersFor(Options.ActiveTags, currentLOD, bModifiersForBeforeOperations, Modifiers);
+
+		ActiveTags.Push({});
+
+		// Process mesh extend modifiers (from edit modifiers)
+		int32 EditIndex = 0;
+		for (const FirstPassGenerator::FModifier& m : Modifiers)
+		{
+			if (m.Node->GetType() == NodeModifierSurfaceEdit::GetStaticType())
+			{
+				const NodeModifierSurfaceEdit* Edit = static_cast<const NodeModifierSurfaceEdit*>(m.Node);
+
+				int32 ThisEditIndex = EditIndex;
+				++EditIndex;
+
+				if (ImageIndex >= Edit->Textures.Num())
+				{
+					continue;
+				}
+
+				Ptr<NodeImage> pExtend = Edit->Textures[ImageIndex].Extend;
+				if (!pExtend)
+				{
+					continue;
+				}
+
+				const TArray<FGeneratedLayout>& ExtraLayouts = BaseMeshResults.ExtraMeshLayouts[ThisEditIndex].GeneratedLayouts;
+
+				if (LayoutIndex >= ExtraLayouts.Num() || !ExtraLayouts[LayoutIndex].Layout)
+				{
+					FString Msg = FString::Printf(TEXT("Trying to extend a layout that doesn't exist in object [%s]."), *CurrentParents.Last().ObjectNode->Name );
+					ErrorLog->GetPrivate()->Add(Msg, ELMT_ERROR, ErrorContext);
+				}
+				else
+				{
+					Ptr<const Layout> pExtendLayout = ExtraLayouts[LayoutIndex].Layout;
+
+					Ptr<ASTOp> lastBase = LastImageOp;
+
+					for (int32 b = 0; b < pExtendLayout->GetBlockCount(); ++b)
+					{
+						// Generate the image block
+						FImageGenerationOptions ImageOptions;
+						ImageOptions.State = Options.State;
+						ImageOptions.ImageLayoutStrategy = ImageLayoutStrategy;
+						ImageOptions.ActiveTags = Edit->EnableTags; // TODO: Merge with current tags?
+						ImageOptions.RectSize = { 0,0 };
+						ImageOptions.LayoutToApply = pExtendLayout;
+						ImageOptions.LayoutBlockId = pExtendLayout->Blocks[b].Id;
+						FImageGenerationResult ExtendResult;
+						GenerateImage(ImageOptions, ExtendResult, pExtend);
+						Ptr<ASTOp> fragmentAd = ExtendResult.op;
+
+						// Block in layout grid units
+						box< FIntVector2 > rectInCells;
+						rectInCells.min = pExtendLayout->Blocks[b].Min;
+						rectInCells.size = pExtendLayout->Blocks[b].Size;
+
+						FImageDesc ExtendDesc = fragmentAd->GetImageDesc();
+
+						// If we don't know the size of a layout block in pixels, calculate it
+						UpdateLayoutBlockDesc(InOutLayoutBlockDesc, ExtendDesc, rectInCells.size);
+
+						// Adjust the format and size of the block to be added
+						// Actually don't do it, it will be propagated from the top format operation.
+						//fragmentAd = GenerateImageFormat(fragmentAd, FinalImageFormat);
+
+						UE::Math::TIntVector2<int32> expectedSize;
+						expectedSize[0] = InOutLayoutBlockDesc.BlockPixelsX * rectInCells.size[0];
+						expectedSize[1] = InOutLayoutBlockDesc.BlockPixelsY * rectInCells.size[1];
+						fragmentAd = GenerateImageSize(fragmentAd, expectedSize);
+
+						// Apply tiling to avoid generating chunks of image that are too big.
+						fragmentAd = ApplyTiling(fragmentAd, expectedSize, InOutLayoutBlockDesc.FinalFormat);
+
+						// Compose operation
+						Ptr<ASTOpImageCompose> composeOp = new ASTOpImageCompose();
+						composeOp->Layout = BaseMeshResults.LayoutOps[LayoutIndex];
+						composeOp->Base = lastBase;
+						composeOp->BlockImage = fragmentAd;
+
+						// Set the absolute block index.
+						check(pExtendLayout->Blocks[b].Id != FLayoutBlock::InvalidBlockId);
+						composeOp->BlockId = pExtendLayout->Blocks[b].Id;
+
+						lastBase = composeOp;
+					}
+
+					// Condition to enable this image extension
+					if (m.FinalCondition)
+					{
+						Ptr<ASTOp> conditionalAd;
+						Ptr<ASTOpConditional> cop = new ASTOpConditional();
+						cop->type = OP_TYPE::IM_CONDITIONAL;
+						cop->no = LastImageOp;
+						cop->yes = lastBase;
+						cop->condition = m.FinalCondition;
+						conditionalAd = cop;
+						LastImageOp = conditionalAd;
+					}
+					else
+					{
+						LastImageOp = lastBase;
+					}
+
+				}
+			}
+		}
+
+		ActiveTags.Pop();
+
+		return LastImageOp;
 	}
 
 
diff --git a/Engine/Plugins/Experimental/Mutable/Source/MutableTools/Private/MuT/CodeGenerator.h b/Engine/Plugins/Experimental/Mutable/Source/MutableTools/Private/MuT/CodeGenerator.h
--- a/Engine/Plugins/Experimental/Mutable/Source/MutableTools/Private/MuT/CodeGenerator.h
+++ b/Engine/Plugins/Experimental/Mutable/Source/MutableTools/Private/MuT/CodeGenerator.h
@@ -36,6 +36,7 @@
 #include "MuT/NodeProjector.h"
 #include "MuT/NodeString.h"
 #include "MuT/NodeSurfaceNew.h"
+#include "MuT/NodeModifierSurfaceEdit.h"
 #include "MuT/Table.h"
 #include "MuT/TablePrivate.h"
 #include "Templates/TypeHash.h"
@@ -361,16 +362,6 @@ namespace mu
 		void GenerateComponent_Switch(const FComponentGenerationOptions&, FGenericGenerationResult&, const NodeComponentSwitch*);
 		void GenerateComponent_Variation(const FComponentGenerationOptions&, FGenericGenerationResult&, const NodeComponentVariation*);
 
-		//-----------------------------------------------------------------------------------------
-
-		// Get the modifiers that have to be applied to elements with a specific tag.
-		void GetModifiersFor(const TArray<FString>& SurfaceTags, int32 LOD,
-			bool bModifiersForBeforeOperations, TArray<FirstPassGenerator::FModifier>& OutModifiers);
-
-		// Apply the required mesh modifiers to the given operation.
-		Ptr<ASTOp> ApplyMeshModifiers(const FGenericGenerationOptions&, const Ptr<ASTOp>& SourceOp,
-			bool bModifiersForBeforeOperations, const void* errorContext);
-
         //-----------------------------------------------------------------------------------------
         //!
         Ptr<ASTOp> GenerateTableVariable(Ptr<const Node>, const FTableCacheKey&, bool bAddNoneOption, const FString& DefaultRowName);
@@ -504,10 +495,10 @@ namespace mu
 		Ptr<ASTOp> ApplyTiling(Ptr<ASTOp> Source, UE::Math::TIntVector2<int32> Size, EImageFormat Format);
 
 		/** Generate a layout block-sized image with a mask including all pixels in the blocks defined in the patch node. */
-		Ptr<Image> GenerateImageBlockPatchMask(const NodeSurfaceEdit::FTexture&, FIntPoint GridSize, int32 BlockPixelsX, int32 BlockPixelsY, box<FIntVector2> RectInCells);
+		Ptr<Image> GenerateImageBlockPatchMask(const NodeModifierSurfaceEdit::FTexture&, FIntPoint GridSize, int32 BlockPixelsX, int32 BlockPixelsY, box<FIntVector2> RectInCells);
 
 		/** Generate all the operations to apply the block patching on top of the BlockOp, and masking with PatchMask. */
-		Ptr<ASTOp> GenerateImageBlockPatch(Ptr<ASTOp> BlockOp, const NodeSurfaceEdit::FTexture&, Ptr<Image> PatchMask, Ptr<ASTOp> ConditionOp, const FImageGenerationOptions&);
+		Ptr<ASTOp> GenerateImageBlockPatch(Ptr<ASTOp> BlockOp, const NodeModifierSurfaceEdit::FTexture&, Ptr<Image> PatchMask, Ptr<ASTOp> ConditionOp, const FImageGenerationOptions&);
 
         //-----------------------------------------------------------------------------------------
         // Meshes
@@ -801,13 +792,50 @@ namespace mu
             Ptr<ASTOp> surfaceOp;
         };
 
-        void GenerateSurface( FSurfaceGenerationResult&, const FSurfaceGenerationOptions&,
-                              Ptr<const NodeSurfaceNew>,
-                              const TArray<FirstPassGenerator::FSurface::FEdit>& Edits );
+        void GenerateSurface( FSurfaceGenerationResult&, const FSurfaceGenerationOptions&, Ptr<const NodeSurfaceNew> );
 
 		//-----------------------------------------------------------------------------------------
-		//Default Table Parameters
+		// Default Table Parameters
 		Ptr<ASTOp> GenerateDefaultTableValue(ETableColumnType NodeType);
+
+
+		//-----------------------------------------------------------------------------------------
+
+		struct FLayoutBlockDesc
+		{
+			EImageFormat FinalFormat = EImageFormat::IF_NONE;
+			int32 BlockPixelsX = 0;
+			int32 BlockPixelsY = 0;
+			bool bBlocksHaveMips = false;
+		};
+
+		void UpdateLayoutBlockDesc(FLayoutBlockDesc& Out, FImageDesc BlockDesc, FIntVector2 LayoutCellSize);
+
+		// Get the modifiers that have to be applied to elements with a specific tag.
+		void GetModifiersFor(const TArray<FString>& SurfaceTags, int32 LOD, bool bModifiersForBeforeOperations, TArray<FirstPassGenerator::FModifier>& OutModifiers);
+
+		// Apply the required mesh modifiers to the given operation.
+		Ptr<ASTOp> ApplyMeshModifiers(const FMeshGenerationOptions&, FMeshGenerationResult& BaseResults, 
+			const FMeshGenerationResult* SharedSurfaceResults, 
+			bool bModifiersForBeforeOperations, const void* ErrorContext,
+			const NodeMeshConstant* OriginalMeshNode);
+
+		Ptr<ASTOp> ApplyImageBlockModifiers(const FImageGenerationOptions&, Ptr<ASTOp> BaseImageOp, int32 ImageIndex, 
+			FIntPoint GridSize,
+			const FLayoutBlockDesc& LayoutBlockDesc,
+			box< FIntVector2 > RectInCells,
+			bool bModifiersForBeforeOperations, const void* ErrorContext);
+
+		Ptr<ASTOp> ApplyImageExtendModifiers(
+			const FGenericGenerationOptions& Options, 
+			const FMeshGenerationResult& BaseMeshResults,
+			Ptr<ASTOp> ImageAd, 
+			CompilerOptions::TextureLayoutStrategy ImageLayoutStrategy, 
+			int32 LayoutIndex, int32 ImageIndex,
+			FIntPoint GridSize,
+			CodeGenerator::FLayoutBlockDesc& InOutLayoutBlockDesc,
+			bool bModifiersForBeforeOperations, const void* ErrorContext);
+
     };
 
 	
diff --git a/Engine/Plugins/Experimental/Mutable/Source/MutableTools/Private/MuT/CodeGenerator_FirstPass.cpp b/Engine/Plugins/Experimental/Mutable/Source/MutableTools/Private/MuT/CodeGenerator_FirstPass.cpp
--- a/Engine/Plugins/Experimental/Mutable/Source/MutableTools/Private/MuT/CodeGenerator_FirstPass.cpp
+++ b/Engine/Plugins/Experimental/Mutable/Source/MutableTools/Private/MuT/CodeGenerator_FirstPass.cpp
@@ -21,12 +21,12 @@
 #include "MuT/NodeModifierMeshClipMorphPlane.h"
 #include "MuT/NodeModifierMeshClipWithMesh.h"
 #include "MuT/NodeModifierMeshClipWithUVMask.h"
+#include "MuT/NodeModifierSurfaceEdit.h"
 #include "MuT/NodeObject.h"
 #include "MuT/NodeObjectGroupPrivate.h"
 #include "MuT/NodeObjectNew.h"
 #include "MuT/NodePrivate.h"
 #include "MuT/NodeSurface.h"
-#include "MuT/NodeSurfaceEdit.h"
 #include "MuT/NodeSurfaceNew.h"
 #include "MuT/NodeSurfaceVariation.h"
 #include "MuT/NodeSurfaceSwitch.h"
@@ -62,7 +62,7 @@ namespace mu
  			Generate_Generic(Root);
  		}
 
-		// Step 2: Collect all tags and a list of the surfaces that activate them
+		// Step 2: Collect all tags and a list of the surfaces and modifiers that activate them
 		for (int32 s=0; s<Surfaces.Num(); ++s)
 		{
 			// Collect the tags in new surfaces
@@ -92,43 +92,35 @@ namespace mu
                     Tags[tag].Surfaces.Add(s);
 				}
 			}
+		}
 
-            // Collect the tags in edit surfaces
-            for (int32 e=0; e< Surfaces[s].Edits.Num(); ++e)
-            {
-                const FSurface::FEdit& edit = Surfaces[s].Edits[e];
-                for (int32 t=0; t<edit.Node->EnableTags.Num(); ++t)
-                {
-                    int32 tag = -1;
-					const FString& tagStr = edit.Node->EnableTags[t];
-
-                    for (int32 i = 0; i<Tags.Num() && tag<0; ++i)
-                    {
-                        if (Tags[i].Tag == tagStr)
-                        {
-                            tag = i;
-                        }
-                    }
-
-                    // New tag?
-                    if (tag < 0)
-                    {
-                        tag = Tags.Num();
-						FTag newTag;
-                        newTag.Tag = tagStr;
-                        Tags.Add(newTag);
-                    }
+		// TODO: Modifier's enabling tags?
+		for (int32 ModifierIndex = 0; ModifierIndex < Modifiers.Num(); ++ModifierIndex)
+		{
+			// Collect the tags in the modifiers
+			for (const FString& ModifierTag: Modifiers[ModifierIndex].Node->EnableTags)
+			{
+				int32 TagIndex = Tags.IndexOfByPredicate([&](const FTag& Candidate)
+					{ 
+						return Candidate.Tag == ModifierTag;
+					});
 
-                    if (Tags[tag].Edits.Find({s,e}) == INDEX_NONE)
-                    {
-                        Tags[tag].Edits.Add({s,e});
-                    }
-                }
-            }
+				// New tag?
+				if (TagIndex < 0)
+				{
+					FTag newTag;
+					newTag.Tag = ModifierTag;
+					TagIndex = Tags.Add(newTag);
+				}
 
+				if (Tags[TagIndex].Modifiers.Find(ModifierIndex) == INDEX_NONE)
+				{
+					Tags[TagIndex].Modifiers.Add(ModifierIndex);
+				}
+			}
 		}
 
-        // Step 3: Create default state if necessary
+		// Step 3: Create default state if necessary
         if ( bIgnoreStates )
         {
             States.Empty();
@@ -154,10 +146,6 @@ namespace mu
 		{
 			Generate_SurfaceNew(static_cast<const NodeSurfaceNew*>(Root));
 		}
-		else if (Root->GetType() == NodeSurfaceEdit::GetStaticType())
-		{
-			Generate_SurfaceEdit(static_cast<const NodeSurfaceEdit*>(Root));
-		}
 		else if (Root->GetType() == NodeSurfaceVariation::GetStaticType())
 		{
 			Generate_SurfaceVariation(static_cast<const NodeSurfaceVariation*>(Root));
@@ -237,46 +225,6 @@ namespace mu
 	}
 
 
-	//---------------------------------------------------------------------------------------------
-    void FirstPassGenerator::Generate_SurfaceEdit(const NodeSurfaceEdit* InNode)
-	{
-		// Store a reference to this node in the surface data for the surface that this node is
-		// editing.
-		FSurface* Surface = Surfaces.FindByPredicate([&InNode](const FSurface& s)
-        {
-            // Are we editing the main surface node of this surface?
-            if (s.Node.get() == InNode->Parent.get()) return true;
-
-            // Are we editing an edit node modifying this surface?
-            for (const FSurface::FEdit& e: s.Edits)
-            {
-                if (InNode->Parent && e.Node== InNode->Parent)
-                {
-                    return true;
-                }
-            }
-
-            return false;
-        });
-		
-		// The surface could be missing if the parent is not in the hierarchy. This could happen
-		// with wrong input or in case of partial models for preview.
-		if (Surface)
-		{
-			FSurface::FEdit edit;
-			edit.PositiveTags = CurrentPositiveTags;
-			edit.NegativeTags = CurrentNegativeTags;
-            edit.Node = InNode;
-            edit.Condition = CurrentCondition.Last().ObjectCondition;
-			Surface->Edits.Add(edit);
-		}
-		else
-		{
-			ErrorLog->GetPrivate()->Add("Missing parent object for edit node.", ELMT_WARNING, InNode->GetMessageContext());
-		}
-	}
-
-
 	//---------------------------------------------------------------------------------------------
 	void FirstPassGenerator::Generate_SurfaceVariation(const NodeSurfaceVariation* InNode)
 	{
diff --git a/Engine/Plugins/Experimental/Mutable/Source/MutableTools/Private/MuT/CodeGenerator_FirstPass.h b/Engine/Plugins/Experimental/Mutable/Source/MutableTools/Private/MuT/CodeGenerator_FirstPass.h
--- a/Engine/Plugins/Experimental/Mutable/Source/MutableTools/Private/MuT/CodeGenerator_FirstPass.h
+++ b/Engine/Plugins/Experimental/Mutable/Source/MutableTools/Private/MuT/CodeGenerator_FirstPass.h
@@ -18,13 +18,14 @@
 #include "MuT/NodeModifierMeshClipMorphPlane.h"
 #include "MuT/NodeModifierMeshClipWithMesh.h"
 #include "MuT/NodeModifierMeshClipWithUVMask.h"
+#include "MuT/NodeModifierSurfaceEdit.h"
 #include "MuT/NodeObjectGroup.h"
 #include "MuT/NodeObjectNew.h"
-#include "MuT/NodeSurfaceEdit.h"
 #include "MuT/NodeSurfaceNew.h"
 #include "MuT/NodeSurfaceVariation.h"
 #include "MuT/NodeSurfaceSwitch.h"
 #include "MuT/ASTOpParameter.h"
+#include "MuT/ErrorLog.h"
 
 namespace mu
 {
@@ -47,7 +48,6 @@ namespace mu
 		void Generate_Generic(const Node*);
 		void Generate_Modifier(const NodeModifier*);
 		void Generate_SurfaceNew(const NodeSurfaceNew*);
-		void Generate_SurfaceEdit(const NodeSurfaceEdit*);
 		void Generate_SurfaceSwitch(const NodeSurfaceSwitch*);
 		void Generate_SurfaceVariation(const NodeSurfaceVariation*);
 		void Generate_ComponentNew(const NodeComponentNew*);
@@ -129,26 +129,9 @@ namespace mu
             // This is filled in the first pass.
             StateCondition StateCondition;
 
-            // Condition for this surface to be enabled when all the object conditions are met.
+            // Combined condition for the surface and the object conditions.
             // This is filled in CodeGenerator_SecondPass.
-            Ptr<ASTOp> SurfaceCondition;
-
-			// All surface editing nodes that edit this surface
-            struct FEdit
-            {
-				//! Reference to the edit node, used during compilation.
-				const NodeSurfaceEdit* Node = nullptr;
-				
-				// List of tags that are required for the presence of this surface
-            	TArray<FString> PositiveTags;
-
-            	// List of tags that block the presence of this surface
-            	TArray<FString> NegativeTags;
-            	
-                //! Condition that enables the effects of this edit node on the surface
-                Ptr<ASTOp> Condition;
-            };
-			TArray<FEdit> Edits;
+            Ptr<ASTOp> FinalCondition;
 
             // This is filled in the final code generation pass
             Ptr<ASTOp> ResultSurfaceOp;
@@ -171,9 +154,9 @@ namespace mu
             // be the parent object where this surface will be added.
             Ptr<ASTOp> ObjectCondition;
 
-            // This conditions is the condition for this modifier to be enabled when all the object conditions are met.
-            // This is filled in CodeGenerator_SecondPass.
-            Ptr<ASTOp> SurfaceCondition;
+			// Combined condition for the this modifier and the object conditions.
+			// This is filled in CodeGenerator_SecondPass.
+			Ptr<ASTOp> FinalCondition;
 
             // This is filled in CodeGenerator_SecondPass.
             StateCondition StateCondition;
@@ -183,22 +166,20 @@ namespace mu
         };
 		TArray<FModifier> Modifiers;
 
-		//! Info about all found tags.
+		/** Info about all found tags. */
 		struct FTag
 		{
 			FString Tag;
 
-            // Surfaces that activate the tag. These are indices to the FirstPassGenerator::surfaces
-            // vector.
+            /** Surfaces that activate the tag.These are indices to the FirstPassGenerator::Surfaces array.*/
 			TArray<int32> Surfaces;
 
-            // Edit Surfaces that activate the tag. These first element of the pair are indices to
-            // the FirstPassGenerator::surfaces vector. The second element are indices to the
-            // "edits" in the specific surface.
-			TArray<TPair<int32,int32>> Edits;
+            /** Modifiers that activate the tag. The index refers to the FirstPassGenerator::Modifiers. */
+			TArray<int32> Modifiers;
 
-            // This conditions is the condition for this tag to be enabled considering no other
-            // condition. This is filled in CodeGenerator_SecondPass.
+            /** This conditions is the condition for this tag to be enabled considering no other condition. 
+			* This is filled in CodeGenerator_SecondPass. 
+			*/
             Ptr<ASTOp> GenericCondition;
         };
         TArray<FTag> Tags;
diff --git a/Engine/Plugins/Experimental/Mutable/Source/MutableTools/Private/MuT/CodeGenerator_Image.cpp b/Engine/Plugins/Experimental/Mutable/Source/MutableTools/Private/MuT/CodeGenerator_Image.cpp
--- a/Engine/Plugins/Experimental/Mutable/Source/MutableTools/Private/MuT/CodeGenerator_Image.cpp
+++ b/Engine/Plugins/Experimental/Mutable/Source/MutableTools/Private/MuT/CodeGenerator_Image.cpp
@@ -1438,8 +1438,12 @@ namespace mu
 				ErrorLog->GetPrivate()->Add("Layout or block index error.", ELMT_ERROR, InNode->GetMessageContext());
 			}
 
-
-			Ptr<ASTOp> CurrentMeshToProjectOp = MeshResult.MeshOp;
+			// TODO: 
+			// MeshResult.MeshOp has some modifiers applied already: the ones applied before other operations directly in the mesh constant generation. 
+			// This is not what was happening before the refactor  so use MeshResult.BaseMeshOp. This is another case of ambiguity of order of modifiers 
+			// that whould be fixed with the general ordering design.
+			//Ptr<ASTOp> CurrentMeshToProjectOp = MeshResult.MeshOp;
+			Ptr<ASTOp> CurrentMeshToProjectOp = MeshResult.BaseMeshOp;
 
             if (projectorResult.type == PROJECTOR_TYPE::WRAPPING)
             {
diff --git a/Engine/Plugins/Experimental/Mutable/Source/MutableTools/Private/MuT/CodeGenerator_Mesh.cpp b/Engine/Plugins/Experimental/Mutable/Source/MutableTools/Private/MuT/CodeGenerator_Mesh.cpp
--- a/Engine/Plugins/Experimental/Mutable/Source/MutableTools/Private/MuT/CodeGenerator_Mesh.cpp
+++ b/Engine/Plugins/Experimental/Mutable/Source/MutableTools/Private/MuT/CodeGenerator_Mesh.cpp
@@ -1301,7 +1301,6 @@ namespace mu
 		}
 
 		OutResult.BaseMeshOp = LastMeshOp;
-		OutResult.MeshOp = LastMeshOp;
 
 		// Add the tags operation
 		if (Tags.Num())
@@ -1312,9 +1311,11 @@ namespace mu
 			LastMeshOp = AddTagsOp;
 		}
 
+		OutResult.MeshOp = LastMeshOp;
+
 		// Apply the modifier for the pre-normal operations stage.
 		bool bModifiersForBeforeOperations = true;
-		OutResult.MeshOp = ApplyMeshModifiers(InOptions, LastMeshOp, bModifiersForBeforeOperations, InNode->GetMessageContext());
+		OutResult.MeshOp = ApplyMeshModifiers(InOptions, OutResult, nullptr, bModifiersForBeforeOperations, InNode->GetMessageContext(), InNode);
     }
 
 
diff --git a/Engine/Plugins/Experimental/Mutable/Source/MutableTools/Private/MuT/CodeGenerator_SecondPass.cpp b/Engine/Plugins/Experimental/Mutable/Source/MutableTools/Private/MuT/CodeGenerator_SecondPass.cpp
--- a/Engine/Plugins/Experimental/Mutable/Source/MutableTools/Private/MuT/CodeGenerator_SecondPass.cpp
+++ b/Engine/Plugins/Experimental/Mutable/Source/MutableTools/Private/MuT/CodeGenerator_SecondPass.cpp
@@ -141,85 +141,6 @@ namespace mu
 			//PartialOptimise( c, CompilerOptions->OptimisationOptions );
 		}
 
-
-		// Condition expression for all the edit-surfaces that activate the tag
-		for (const TPair<int32,int32>& EditKey : t.Edits)
-		{
-			if (negSurf.find(EditKey.Key) != negSurf.end())
-			{
-				// The surface in the edit is a negative requirement higher up in the condition so
-				// this branch never be true.
-				continue;
-			}
-
-			Ptr<ASTOp> surfCondition;
-			if (posSurf.find(EditKey.Key) != posSurf.end())
-			{
-				// This surface in the edit is already a positive requirement higher up in the condition
-				// so we don't need that part of the condition
-			}
-			else
-			{
-				auto PositiveTags = posTag;
-				PositiveTags.insert(tagIndex);
-
-				surfCondition = GenerateDataCodition(EditKey.Key,
-					FirstPass->Surfaces[EditKey.Key].PositiveTags,
-					FirstPass->Surfaces[EditKey.Key].NegativeTags,
-					posSurf,
-					negSurf,
-					PositiveTags,
-					negTag);
-			}
-
-			const FirstPassGenerator::FSurface& surface = FirstPass->Surfaces[EditKey.Key];
-			const FirstPassGenerator::FSurface::FEdit& edit = surface.Edits[EditKey.Value];
-
-			// Combine object and surface conditions
-			Ptr<ASTOp> fullCondition;
-			if (surfCondition)
-			{
-				Ptr<ASTOpFixed> f = new ASTOpFixed();
-				f->op.type = OP_TYPE::BO_AND;
-				f->SetChild(f->op.args.BoolBinary.a, surface.ObjectCondition);
-				f->SetChild(f->op.args.BoolBinary.b, surfCondition);
-				fullCondition = m_opPool.Add(f);
-			}
-			else
-			{
-				fullCondition = m_opPool.Add(surface.ObjectCondition);
-			}
-
-			// Combine with edit condition
-			if (fullCondition)
-			{
-				Ptr<ASTOpFixed> f = new ASTOpFixed();
-				f->op.type = OP_TYPE::BO_AND;
-				f->SetChild(f->op.args.BoolBinary.a, edit.Condition);
-				f->SetChild(f->op.args.BoolBinary.b, fullCondition);
-				fullCondition = m_opPool.Add(f);
-			}
-			else
-			{
-				fullCondition = m_opPool.Add(edit.Condition);
-			}
-
-
-			if (!c)
-			{
-				c = fullCondition;
-			}
-			else
-			{
-				Ptr<ASTOpFixed> o = new ASTOpFixed();
-				o->op.type = OP_TYPE::BO_OR;
-				o->SetChild(o->op.args.BoolBinary.a, fullCondition);
-				o->SetChild(o->op.args.BoolBinary.b, c);
-				c = m_opPool.Add(o);
-			}
-		}
-
-
 		TagConditionGenerationCache.insert(std::make_pair<>(key, c));
 
 		return c;
@@ -490,31 +411,37 @@ namespace mu
 			{
 				Ptr<ASTOp> c = GenerateDataCodition(
 					SurfaceIndex, 
-					FirstPass->Surfaces[SurfaceIndex].PositiveTags,
-					FirstPass->Surfaces[SurfaceIndex].NegativeTags,
+					Surface.PositiveTags,
+					Surface.NegativeTags,
 					Empty, Empty, Empty, Empty);
 
-				FirstPass->Surfaces[SurfaceIndex].SurfaceCondition = c;
-			}
-			
-			for (int32 EditIndex = 0; EditIndex < Surface.Edits.Num(); ++EditIndex)
-			{
-				FirstPassGenerator::FSurface::FEdit& Edit = Surface.Edits[EditIndex];
-				
-				Ptr<ASTOp> c = GenerateDataCodition(
-					EditIndex, 
-					Surface.Edits[EditIndex].PositiveTags,
-					Surface.Edits[EditIndex].NegativeTags,
-					Empty, Empty, Empty, Empty);
+				Ptr<ASTOpFixed> ConditionOp = new ASTOpFixed();
+				ConditionOp->op.type = OP_TYPE::BO_AND;
+				ConditionOp->SetChild(ConditionOp->op.args.BoolBinary.a, Surface.ObjectCondition);
+				ConditionOp->SetChild(ConditionOp->op.args.BoolBinary.b, c);
 
-				Ptr<ASTOpFixed> OpAnd = new ASTOpFixed;
-				OpAnd->op.type = OP_TYPE::BO_AND;
-				OpAnd->SetChild(OpAnd->op.args.BoolBinary.a, Edit.Condition);
-				OpAnd->SetChild(OpAnd->op.args.BoolBinary.b, c);
-				c = m_opPool.Add(OpAnd);
-				
-				Edit.Condition = OpAnd;
+				Surface.FinalCondition = ConditionOp;
 			}
+			
+			// TODO: Convert to modifiers that enable tags?
+			//for (int32 EditIndex = 0; EditIndex < Surface.Edits.Num(); ++EditIndex)
+			//{
+			//	FirstPassGenerator::FSurface::FEdit& Edit = Surface.Edits[EditIndex];
+			//	
+			//	Ptr<ASTOp> c = GenerateDataCodition(
+			//		EditIndex, 
+			//		Surface.Edits[EditIndex].PositiveTags,
+			//		Surface.Edits[EditIndex].NegativeTags,
+			//		Empty, Empty, Empty, Empty);
+
+			//	Ptr<ASTOpFixed> OpAnd = new ASTOpFixed;
+			//	OpAnd->op.type = OP_TYPE::BO_AND;
+			//	OpAnd->SetChild(OpAnd->op.args.BoolBinary.a, Edit.Condition);
+			//	OpAnd->SetChild(OpAnd->op.args.BoolBinary.b, c);
+			//	c = m_opPool.Add(OpAnd);
+			//	
+			//	Edit.Condition = OpAnd;
+			//}
 		}
 
 		for (int32 ModifierIndex = 0; ModifierIndex < FirstPass->Modifiers.Num(); ++ModifierIndex)
@@ -525,7 +452,12 @@ namespace mu
 				FirstPass->Modifiers[ModifierIndex].NegativeTags,
 				Empty, Empty, Empty, Empty);
 			
-			FirstPass->Modifiers[ModifierIndex].SurfaceCondition = c;
+			Ptr<ASTOpFixed> ConditionOp = new ASTOpFixed();
+			ConditionOp->op.type = OP_TYPE::BO_AND;
+			ConditionOp->SetChild(ConditionOp->op.args.BoolBinary.a, FirstPass->Modifiers[ModifierIndex].ObjectCondition);
+			ConditionOp->SetChild(ConditionOp->op.args.BoolBinary.b, c);
+
+			FirstPass->Modifiers[ModifierIndex].FinalCondition = ConditionOp;
 		}
 
 		for (int32 ComponentIndex = 0; ComponentIndex < FirstPass->Components.Num(); ++ComponentIndex)
diff --git a/Engine/Plugins/Experimental/Mutable/Source/MutableTools/Private/MuT/Node.cpp b/Engine/Plugins/Experimental/Mutable/Source/MutableTools/Private/MuT/Node.cpp
--- a/Engine/Plugins/Experimental/Mutable/Source/MutableTools/Private/MuT/Node.cpp
+++ b/Engine/Plugins/Experimental/Mutable/Source/MutableTools/Private/MuT/Node.cpp
@@ -12,7 +12,6 @@
 #include "MuT/NodeLOD.h"
 #include "MuT/NodeSurface.h"
 #include "MuT/NodeSurfaceNew.h"
-#include "MuT/NodeSurfaceEdit.h"
 #include "MuT/NodeSurfaceSwitch.h"
 #include "MuT/NodeSurfaceVariation.h"
 #include "MuT/NodeExtensionData.h"
@@ -37,6 +36,7 @@
 #include "MuT/NodeModifierMeshClipMorphPlane.h"
 #include "MuT/NodeModifierMeshClipWithMesh.h"
 #include "MuT/NodeModifierMeshClipWithUVMask.h"
+#include "MuT/NodeModifierSurfaceEdit.h"
 
 namespace mu
 {
@@ -53,7 +53,6 @@ namespace mu
 
 	FNodeType NodeSurface::StaticType = FNodeType(Node::EType::Surface, Node::GetStaticType());
 	FNodeType NodeSurfaceNew::StaticType = FNodeType(Node::EType::SurfaceNew, NodeSurface::GetStaticType());
-	FNodeType NodeSurfaceEdit::StaticType = FNodeType(Node::EType::SurfaceEdit, NodeSurface::GetStaticType());
 	FNodeType NodeSurfaceSwitch::StaticType = FNodeType(Node::EType::SurfaceSwitch, NodeSurface::GetStaticType());
 	FNodeType NodeSurfaceVariation::StaticType = FNodeType(Node::EType::SurfaceVariation, NodeSurface::GetStaticType());
 
@@ -82,6 +81,7 @@ namespace mu
 	FNodeType NodeModifierMeshClipMorphPlane::StaticType = FNodeType(Node::EType::ModifierMeshClipMorphPlane, NodeModifier::GetStaticType());
 	FNodeType NodeModifierMeshClipWithMesh::StaticType = FNodeType(Node::EType::ModifierMeshClipWithMesh, NodeModifier::GetStaticType());
 	FNodeType NodeModifierMeshClipWithUVMask::StaticType = FNodeType(Node::EType::ModifierMeshClipWithUVMask, NodeModifier::GetStaticType());
+	FNodeType NodeModifierSurfaceEdit::StaticType = FNodeType(Node::EType::ModifierSurfaceEdit, NodeSurface::GetStaticType());
 
 
 	FNodeType::FNodeType()
diff --git a/Engine/Plugins/Experimental/Mutable/Source/MutableTools/Private/MuT/NodeMeshConstant.cpp b/Engine/Plugins/Experimental/Mutable/Source/MutableTools/Private/MuT/NodeMeshConstant.cpp
--- a/Engine/Plugins/Experimental/Mutable/Source/MutableTools/Private/MuT/NodeMeshConstant.cpp
+++ b/Engine/Plugins/Experimental/Mutable/Source/MutableTools/Private/MuT/NodeMeshConstant.cpp
@@ -36,6 +36,25 @@ namespace mu
     }
 
 
+	void NodeMeshConstant::AddMorph(const FString& Name, Ptr<Mesh> Morphed)
+	{
+		m_pD->Morphs.Add({Name,Morphed});
+	}
+
+
+	Ptr<Mesh> NodeMeshConstant::FindMorph(const FString& Name) const
+	{
+		for (Private::FMorph& Morph : m_pD->Morphs )
+		{
+			if (Morph.Name == Name)
+			{
+				return Morph.MorphedMesh;
+			}
+		}
+		return nullptr;
+	}
+
+
 	int32 NodeMeshConstant::GetLayoutCount() const
 	{
 		return m_pD->Layouts.Num();
diff --git a/Engine/Plugins/Experimental/Mutable/Source/MutableTools/Public/MuT/Node.h b/Engine/Plugins/Experimental/Mutable/Source/MutableTools/Public/MuT/Node.h
--- a/Engine/Plugins/Experimental/Mutable/Source/MutableTools/Public/MuT/Node.h
+++ b/Engine/Plugins/Experimental/Mutable/Source/MutableTools/Public/MuT/Node.h
@@ -120,7 +120,6 @@ namespace mu
 
 			Surface,
 			SurfaceNew,
-			SurfaceEdit,
 			SurfaceSwitch,
 			SurfaceVariation,
 
@@ -141,6 +140,7 @@ namespace mu
 			ModifierMeshClipWithMesh,
 			ModifierMeshClipDeform,
 			ModifierMeshClipWithUVMask,
+			ModifierSurfaceEdit,
 
 			ExtensionData,
 			ExtensionDataConstant,
diff --git a/Engine/Plugins/Experimental/Mutable/Source/MutableTools/Public/MuT/NodeMeshConstant.h b/Engine/Plugins/Experimental/Mutable/Source/MutableTools/Public/MuT/NodeMeshConstant.h
--- a/Engine/Plugins/Experimental/Mutable/Source/MutableTools/Public/MuT/NodeMeshConstant.h
+++ b/Engine/Plugins/Experimental/Mutable/Source/MutableTools/Public/MuT/NodeMeshConstant.h
@@ -42,7 +42,13 @@ namespace mu
 		Ptr<Mesh> GetValue() const;
 
 		//! Set the constant mesh that will be returned.
-		void SetValue( Ptr<Mesh> );
+		void SetValue(Ptr<Mesh>);
+
+		/** */
+		void AddMorph(const FString& Name, Ptr<Mesh>);
+
+		/** */
+		Ptr<Mesh> FindMorph(const FString& Name) const;
 
 		//! Get the number of layouts defined in this mesh.
 		int32 GetLayoutCount() const;
diff --git a/Engine/Plugins/Experimental/Mutable/Source/MutableTools/Public/MuT/NodeModifier.h b/Engine/Plugins/Experimental/Mutable/Source/MutableTools/Public/MuT/NodeModifier.h
--- a/Engine/Plugins/Experimental/Mutable/Source/MutableTools/Public/MuT/NodeModifier.h
+++ b/Engine/Plugins/Experimental/Mutable/Source/MutableTools/Public/MuT/NodeModifier.h
@@ -42,9 +42,13 @@ namespace mu
 		/** In case of multiple tags in RequiredTags: are they all required, or one is enough? */
 		EMutableMultipleTagPolicy MultipleTagsPolicy = EMutableMultipleTagPolicy::OnlyOneRequired;
 
-		// Wether the modifier has to be applied after the normal node operations or before
+		/** Wether the modifier has to be applied before the normal node operations or after. */
 		bool bApplyBeforeNormalOperations = true;
 
+		/** Tags enabled by this modifier. Other modifiers activated by these tags will be applied to this modifier's "child data" like meshes added by this modifier.
+		* Not to be confused with the RequiredTags.
+		*/
+		TArray<FString> EnableTags;
 
 	public:
 
diff --git a/Engine/Plugins/Experimental/Mutable/Source/MutableTools/Public/MuT/NodeModifierMeshClipWithUVMask.h b/Engine/Plugins/Experimental/Mutable/Source/MutableTools/Public/MuT/NodeModifierMeshClipWithUVMask.h
--- a/Engine/Plugins/Experimental/Mutable/Source/MutableTools/Public/MuT/NodeModifierMeshClipWithUVMask.h
+++ b/Engine/Plugins/Experimental/Mutable/Source/MutableTools/Public/MuT/NodeModifierMeshClipWithUVMask.h
@@ -7,12 +7,13 @@
 #include "MuT/Node.h"
 #include "MuT/NodeModifier.h"
 #include "MuT/NodeImage.h"
+#include "MuT/NodeLayout.h"
 
 
 namespace mu
 {
 
-	/** */
+	/** This node can clip part of a mesh using UV space-data (an image or a list of layout blocks). */
 	class MUTABLETOOLS_API NodeModifierMeshClipWithUVMask : public NodeModifier
 	{
 	public:
@@ -20,6 +21,9 @@ namespace mu
 		/** Image with the regions to remove. It will be interpreted as a bitmap. */
 		Ptr<NodeImage> ClipMask;
 
+		/** If the image above is null, clipping may hapopen with layout blocks. */
+		Ptr<NodeLayout> ClipLayout;
+
 		/** Layout index of the UVs to use inthe source mesh to ben clipped with the mask. */
 		uint8 LayoutIndex = 0;
 
diff --git a/Engine/Plugins/Experimental/Mutable/Source/MutableTools/Public/MuT/NodeModifierSurfaceEdit.h b/Engine/Plugins/Experimental/Mutable/Source/MutableTools/Public/MuT/NodeModifierSurfaceEdit.h
new file mode 100644
--- /dev/null
+++ b/Engine/Plugins/Experimental/Mutable/Source/MutableTools/Public/MuT/NodeModifierSurfaceEdit.h
@@ -0,0 +1,89 @@
+// Copyright Epic Games, Inc. All Rights Reserved.
+
+#pragma once
+
+#include "MuT/NodeModifier.h"
+#include "MuT/NodeMesh.h"
+#include "MuT/NodeImage.h"
+#include "MuT/NodeScalar.h"
+#include "MuR/ImageTypes.h"
+#include "MuR/RefCounted.h"
+#include "MuR/Ptr.h"
+
+
+namespace mu
+{
+
+	/** This node modifies a surface node.
+    * It allows to extend, cut and morph the parent Surface's meshes.
+    * It also allows to patch the parent Surface's textures.
+	*/
+    class MUTABLETOOLS_API NodeModifierSurfaceEdit : public NodeModifier
+	{
+	public:
+
+		/** Mesh to remove from the modified surface. */
+		Ptr<NodeMesh> MeshRemove;
+
+		/** Mesh to add to the modified surface. */
+		Ptr<NodeMesh> MeshAdd;
+
+		/** Name of the morph to apply to the modified surface if it has it. */
+		FString MeshMorph;
+
+		/** Factor of the morph to apply. */
+		Ptr<NodeScalar> MorphFactor;
+
+		/** Data for every modified texture. */
+		struct FTexture
+		{
+			/** Image to add if extgending. */
+			Ptr<NodeImage> Extend;
+
+			/** Image to blend if patching. */
+			Ptr<NodeImage> PatchImage;
+
+			/** Optional mask controlling the blending area. */
+			Ptr<NodeImage> PatchMask;
+
+			/** Rects in the parent layout homogeneous UV space to patch. */
+			TArray<FBox2f> PatchBlocks;
+
+			/** */
+			EBlendType PatchBlendType = EBlendType::BT_BLEND;
+
+			/** Patch alpha channel as well? */
+			bool bPatchApplyToAlpha = false;
+
+		};
+
+		/** Textures to modify. */
+		TArray<FTexture> Textures;
+
+	public:
+
+		NodeModifierSurfaceEdit() 
+		{
+			// This modifier needs to be applied at the end of the operations.
+			bApplyBeforeNormalOperations = false;
+		}
+
+		// Node interface
+		virtual const FNodeType* GetType() const override { return &StaticType; }
+		static const FNodeType* GetStaticType() { return &StaticType; }
+
+	protected:
+
+		//! Forbidden. Manage with the Ptr<> template.
+		~NodeModifierSurfaceEdit() {}
+
+	private:
+
+		static FNodeType StaticType;
+
+	};
+
+
+
+}
+
