[mutable] Fixes case problem breaking compile in case-sensitive platforms.
#rnx

diff --git a/Engine/Plugins/Experimental/Mutable/Source/CustomizableObjectEditor/Private/MuCOE/GenerateMutableSource/GenerateMutableSourceSurface.h b/Engine/Plugins/Experimental/Mutable/Source/CustomizableObjectEditor/Private/MuCOE/GenerateMutableSource/GenerateMutableSourceSurface.h
--- a/Engine/Plugins/Experimental/Mutable/Source/CustomizableObjectEditor/Private/MuCOE/GenerateMutableSource/GenerateMutableSourceSurface.h
+++ b/Engine/Plugins/Experimental/Mutable/Source/CustomizableObjectEditor/Private/MuCOE/GenerateMutableSource/GenerateMutableSourceSurface.h
@@ -3,7 +3,7 @@
 #pragma once
 
 #include "MuR/Ptr.h"
-#include "Hal/Platform.h"
+#include "HAL/Platform.h"
 
 class UEdGraphPin;
 struct FMutableGraphGenerationContext;
