[mutable] Fix mac compilation issues.
#rnx

diff --git a/Engine/Plugins/Experimental/Mutable/Source/MutableTools/Private/MuT/CodeGenerator.cpp b/Engine/Plugins/Experimental/Mutable/Source/MutableTools/Private/MuT/CodeGenerator.cpp
--- a/Engine/Plugins/Experimental/Mutable/Source/MutableTools/Private/MuT/CodeGenerator.cpp
+++ b/Engine/Plugins/Experimental/Mutable/Source/MutableTools/Private/MuT/CodeGenerator.cpp
@@ -1306,7 +1306,8 @@ namespace mu
 							}
 							check(imageAd);
 
-							FMeshGenerationOptions ModifierOptions(Options);
+							FMeshGenerationOptions ModifierOptions;
+							ModifierOptions.State = Options.State;
 							ModifierOptions.ActiveTags = node.Tags;
 							bool bModifiersForBeforeOperations = false;
 							imageAd = ApplyImageExtendModifiers(ModifierOptions, MeshResults, imageAd, ImageLayoutStrategy, LayoutIndex, ImageIndex, GridSize, LayoutBlockDesc, bModifiersForBeforeOperations, surfaceNode->GetMessageContext());
