[Mutable] New component system.

- Renamed Mesh Component node to Passthrough Mesh Component.
- Added new Mesh Component node.
    - Moved old Root Object nodes to this new node.
- Added new Add To Mesh Component node.
    - Moved old Child Object nodes to this new node.
- Added new FixupReconstructPins (ReconstructNode that can be called inside BackwardsCompatibleFixup).

Mesh Component node:
- Properties moved from Object node to this node:
    - LOD pins
    - LODReductionSettings (FComponentSettings no longer needed)
    - AutoLODStrategy

#rnx

diff --git a/Engine/Plugins/Experimental/Mutable/Source/CustomizableObject/Internal/MuCO/CustomizableObjectCustomVersion.h b/Engine/Plugins/Experimental/Mutable/Source/CustomizableObject/Internal/MuCO/CustomizableObjectCustomVersion.h
--- a/Engine/Plugins/Experimental/Mutable/Source/CustomizableObject/Internal/MuCO/CustomizableObjectCustomVersion.h
+++ b/Engine/Plugins/Experimental/Mutable/Source/CustomizableObject/Internal/MuCO/CustomizableObjectCustomVersion.h
@@ -140,6 +140,8 @@ struct CUSTOMIZABLEOBJECT_API FCustomizableObjectCustomVersion
 		
 		ConvertEditAndExtendToModifiers,
 
+		NodeComponentMesh, // Moved LOD pins from NodeObject to NodeComponentMesh.
+		
 		// -----<new versions can be added above this line>-------------------------------------------------
 		VersionPlusOne,
 		LatestVersion = VersionPlusOne - 1
diff --git a/Engine/Plugins/Experimental/Mutable/Source/CustomizableObject/Internal/MuCO/CustomizableObjectPrivate.h b/Engine/Plugins/Experimental/Mutable/Source/CustomizableObject/Internal/MuCO/CustomizableObjectPrivate.h
--- a/Engine/Plugins/Experimental/Mutable/Source/CustomizableObject/Internal/MuCO/CustomizableObjectPrivate.h
+++ b/Engine/Plugins/Experimental/Mutable/Source/CustomizableObject/Internal/MuCO/CustomizableObjectPrivate.h
@@ -1458,6 +1458,8 @@ public:
 		CorrectlySerializeTableToParamNames,
 		
 		AddMaterialSlotNameIndexToSurfaceMetadata,
+
+		NodeComponentMesh,
 		
 		MoveEditNodesToModifiers,
 
diff --git a/Engine/Plugins/Experimental/Mutable/Source/CustomizableObjectEditor/Private/MuCOE/CustomizableObjectCompiler.cpp b/Engine/Plugins/Experimental/Mutable/Source/CustomizableObjectEditor/Private/MuCOE/CustomizableObjectCompiler.cpp
--- a/Engine/Plugins/Experimental/Mutable/Source/CustomizableObjectEditor/Private/MuCOE/CustomizableObjectCompiler.cpp
+++ b/Engine/Plugins/Experimental/Mutable/Source/CustomizableObjectEditor/Private/MuCOE/CustomizableObjectCompiler.cpp
@@ -657,14 +657,7 @@ mu::NodeObjectPtr FCustomizableObjectCompiler::GenerateMutableRoot(
 			}
 		}
 	}
-
-	// Ensure that the CO has a valid AutoLODStrategy on the ActualRoot.
-	if (ActualRoot->AutoLODStrategy == ECustomizableObjectAutomaticLODStrategy::Inherited)
-	{
-		CompilerLog(LOCTEXT("RootInheritsFromParent", "Error! Base CustomizableObject's LOD Strategy can't be set to 'Inherit from parent object'"), ActualRoot, EMessageSeverity::Error);
-		return nullptr;
-	}
-
+	
 	if (ActualRootObject->GetPrivate()->MutableMeshComponents.IsEmpty())
 	{
 		CompilerLog(LOCTEXT("NoComponentsError", "Error! There are no components defined in the Object Properties Tab."), ActualRoot, EMessageSeverity::Error);
@@ -719,6 +712,62 @@ mu::NodeObjectPtr FCustomizableObjectCompiler::GenerateMutableRoot(
 		GenerationContext.AddParticipatingObject(*RefSkeletalMesh);
 	}
 
+	// Find all component nodes
+	TSet<UCustomizableObject*> ParticipatingObjects; // All COs participating in the hierarchy
+	{
+		auto GetParents = [&ParticipatingObjects](UCustomizableObject* Object)
+		{
+			TArray<UCustomizableObjectNodeObject*> ArrayNodeObject;
+			TArray<UCustomizableObject*> ArrayCustomizableObject;
+
+			GetParentsUntilRoot(Object, ArrayNodeObject, ArrayCustomizableObject);
+
+			ParticipatingObjects.Append(ArrayCustomizableObject);
+		};
+
+		// Parents
+		GetParents(Object);
+
+		// Working set and their parents
+		for (TSoftObjectPtr<UCustomizableObject> SoftObject :  Object->WorkingSet)
+		{
+			if (UCustomizableObject* WorkingSetObject = SoftObject.LoadSynchronous())
+			{
+				GetParents(WorkingSetObject);
+			}
+		}
+
+		// Children
+		for (FAssetData& AssetData : ArrayAssetData)
+		{
+			ParticipatingObjects.Add(Cast<UCustomizableObject>(AssetData.GetAsset()));
+		}
+	}
+	
+	for (UCustomizableObject* ParticipatingObject : ParticipatingObjects)
+	{
+		for (UEdGraphNode* Node : ParticipatingObject->GetPrivate()->GetSource()->Nodes)
+		{
+			if (UCustomizableObjectNodeComponentMesh* NodeComponentMesh = Cast<UCustomizableObjectNodeComponentMesh>(Node))
+			{
+				if (FMutableComponentInfo* ComponentInfo = GenerationContext.ComponentInfos.FindByPredicate([&](const FMutableComponentInfo& Element)
+				{
+					return Element.ComponentName == NodeComponentMesh->ComponentName;
+				}))
+				{
+					if (ComponentInfo->NodeComponentMesh)
+					{
+						FText Msg = FText::Format(LOCTEXT("NoReferenceMeshObjectTab", "Already exist a Component node with the same name in Customizable Object [{0}]"), FText::FromString(GetRootObject(*ComponentInfo->NodeComponentMesh)->GetName()));
+						CompilerLog(Msg, Node, EMessageSeverity::Error);
+						return nullptr;
+					}
+
+					ComponentInfo->NodeComponentMesh = NodeComponentMesh;
+				}
+			}
+		}
+	}
+
 	// Copy component data to the object being compiled
 	Object->GetPrivate()->MutableMeshComponents = ActualRootObject->GetPrivate()->MutableMeshComponents;
 
@@ -730,9 +779,11 @@ mu::NodeObjectPtr FCustomizableObjectCompiler::GenerateMutableRoot(
 		GenerationContext.TableToParamNames = Object->GetPrivate()->GetModelResources().TableToParamNames;
 	}
 
+	GenerationContext.bPartialCompilation = !bOutIsRootObject;
+
 	// Generate the object expression
 	UE_LOG(LogMutable, Verbose, TEXT("PROFILE: [ %16.8f ] GenerateMutableSource start."), FPlatformTime::Seconds());
-	mu::NodeObjectPtr MutableRoot = GenerateMutableSource(ActualRoot->OutputPin(), GenerationContext, !bOutIsRootObject);
+	mu::NodeObjectPtr MutableRoot = GenerateMutableSource(ActualRoot->OutputPin(), GenerationContext);
 	UE_LOG(LogMutable, Verbose, TEXT("PROFILE: [ %16.8f ] GenerateMutableSource end."), FPlatformTime::Seconds());
 
     ActualRoot->RealTimeMorphSelectionOverrides = GenerationContext.RealTimeMorphTargetsOverrides;
@@ -1143,7 +1194,7 @@ void FCustomizableObjectCompiler::CompileInternal(bool bAsync)
 		ModelResources.CustomizableObjectPathMap = MoveTemp(GenerationContext.CustomizableObjectPathMap);
 #endif
 
-		ModelResources.NumComponents = GenerationContext.NumMeshComponentsInRoot + GenerationContext.NumExplicitMeshComponents;
+		ModelResources.NumComponents = GenerationContext.NumMeshComponentsInRoot + GenerationContext.NumPassthroughMeshComponents;
 		ModelResources.NumLODs = GenerationContext.NumLODsInRoot;
 		ModelResources.NumLODsToStream = GenerationContext.bEnableLODStreaming ? GenerationContext.NumMaxLODsToStream : 0;
 		ModelResources.FirstLODAvailable = GenerationContext.FirstLODAvailable;
diff --git a/Engine/Plugins/Experimental/Mutable/Source/CustomizableObjectEditor/Private/MuCOE/CustomizableObjectEditorModule.cpp b/Engine/Plugins/Experimental/Mutable/Source/CustomizableObjectEditor/Private/MuCOE/CustomizableObjectEditorModule.cpp
--- a/Engine/Plugins/Experimental/Mutable/Source/CustomizableObjectEditor/Private/MuCOE/CustomizableObjectEditorModule.cpp
+++ b/Engine/Plugins/Experimental/Mutable/Source/CustomizableObjectEditor/Private/MuCOE/CustomizableObjectEditorModule.cpp
@@ -34,7 +34,6 @@
 #include "MuCOE/Nodes/CustomizableObjectNodeExternalPinDetails.h"
 #include "MuCOE/Nodes/CustomizableObjectNodeLayoutBlocks.h"
 #include "MuCOE/Nodes/CustomizableObjectNodeLayoutBlocksDetails.h"
-#include "MuCOE/Nodes/CustomizableObjectNodeMaterialDetails.h"
 #include "MuCOE/Nodes/CustomizableObjectNodeModifierEditMeshSection.h"
 #include "MuCOE/Nodes/CustomizableObjectNodeModifierEditMeshSectionDetails.h"
 #include "MuCOE/Nodes/CustomizableObjectNodeModifierExtendMeshSection.h"
@@ -62,6 +61,7 @@
 #include "MuCOE/Nodes/CustomizableObjectNodeStaticMesh.h"
 #include "MuCOE/Nodes/CustomizableObjectNodeTable.h"
 #include "MuCOE/Nodes/CustomizableObjectNodeTableDetails.h"
+#include "MuCOE/Nodes/CustomizableObjectNodeComponentMeshDetails.h"
 #include "MuCOE/Widgets/CustomizableObjectLODReductionSettings.h"
 #include "PropertyEditorModule.h"
 #include "AssetRegistry/AssetRegistryModule.h"
@@ -250,10 +250,11 @@ void FCustomizableObjectEditorModule::StartupModule()
 	RegisterCustomDetails(PropertyModule, UCustomizableObjectNodeModifierClipMorph::StaticClass(), FOnGetDetailCustomizationInstance::CreateStatic(&FCustomizableObjectNodeModifierClipMorphDetails::MakeInstance));
 	RegisterCustomDetails(PropertyModule, UCustomizableObjectNodeModifierClipWithMesh::StaticClass(), FOnGetDetailCustomizationInstance::CreateStatic(&FCustomizableObjectNodeModifierClipWithMeshDetails::MakeInstance));
 	RegisterCustomDetails(PropertyModule, UCustomizableObjectNodeExternalPin::StaticClass(), FOnGetDetailCustomizationInstance::CreateStatic(&FCustomizableObjectNodeExternalPinDetails::MakeInstance));
-	RegisterCustomDetails(PropertyModule, UCustomizableObjectNodeMaterial::StaticClass(), FOnGetDetailCustomizationInstance::CreateStatic(&FCustomizableObjectNodeMaterialDetails::MakeInstance));
+	RegisterCustomDetails(PropertyModule, UCustomizableObjectNodeMaterial::StaticClass(), FOnGetDetailCustomizationInstance::CreateStatic(&FCustomizableObjectNodeDetails::MakeInstance));
 	RegisterCustomDetails(PropertyModule, UCustomizableObjectNodeSkeletalMesh::StaticClass(), FOnGetDetailCustomizationInstance::CreateStatic(&FCustomizableObjectNodeSkeletalMeshDetails::MakeInstance));
 	RegisterCustomDetails(PropertyModule, UCustomizableObjectNodeStaticMesh::StaticClass(), FOnGetDetailCustomizationInstance::CreateStatic(&FCustomizableObjectNodeDetails::MakeInstance));
 	RegisterCustomDetails(PropertyModule, UCustomizableObjectNodeTable::StaticClass(), FOnGetDetailCustomizationInstance::CreateStatic(&FCustomizableObjectNodeTableDetails::MakeInstance));
+	RegisterCustomDetails(PropertyModule, UCustomizableObjectNodeComponentMesh::StaticClass(), FOnGetDetailCustomizationInstance::CreateStatic(&FCustomizableObjectNodeComponentMeshDetails::MakeInstance));
 
 	// Other Objects
 	RegisterCustomDetails(PropertyModule, UCustomizableObject::StaticClass(), FOnGetDetailCustomizationInstance::CreateStatic(&FCustomizableObjectDetails::MakeInstance));
diff --git a/Engine/Plugins/Experimental/Mutable/Source/CustomizableObjectEditor/Private/MuCOE/EdGraphSchema_CustomizableObject.cpp b/Engine/Plugins/Experimental/Mutable/Source/CustomizableObjectEditor/Private/MuCOE/EdGraphSchema_CustomizableObject.cpp
--- a/Engine/Plugins/Experimental/Mutable/Source/CustomizableObjectEditor/Private/MuCOE/EdGraphSchema_CustomizableObject.cpp
+++ b/Engine/Plugins/Experimental/Mutable/Source/CustomizableObjectEditor/Private/MuCOE/EdGraphSchema_CustomizableObject.cpp
@@ -53,7 +53,7 @@
 #include "MuCOE/Nodes/CustomizableObjectNodeMeshVariation.h"
 #include "MuCOE/Nodes/CustomizableObjectNodeObjectChild.h"
 #include "MuCOE/Nodes/CustomizableObjectNodeObjectGroup.h"
-#include "MuCOE/Nodes/CustomizableObjectNodeComponentMesh.h"
+#include "MuCOE/Nodes/CustomizableObjectNodeComponentPassthroughMesh.h"
 #include "MuCOE/Nodes/CustomizableObjectNodeComponentVariation.h"
 #include "MuCOE/Nodes/CustomizableObjectNodeComponentSwitch.h"
 #include "MuCOE/Nodes/CustomizableObjectNodeProjectorConstant.h"
@@ -89,6 +89,7 @@
 #include "PropertyEditorModule.h"
 #include "MuCO/CustomizableObjectCustomVersion.h"
 #include "MuCOE/GraphTraversal.h"
+#include "MuCOE/Nodes/CustomizableObjectNodeComponentMeshAddTo.h"
 
 class IToolkit;
 
@@ -413,6 +414,8 @@ void UEdGraphSchema_CustomizableObject::GetGraphContextActions(FGraphContextMenu
 		UCustomizableObjectNode* ComponentTemplateNodes[]
 		{
 			ContextMenuBuilder.CreateTemplateNode<UCustomizableObjectNodeComponentMesh>(),
+			ContextMenuBuilder.CreateTemplateNode<UCustomizableObjectNodeComponentMeshAddTo>(),
+			ContextMenuBuilder.CreateTemplateNode<UCustomizableObjectNodeComponentPassthroughMesh>(),
 			ContextMenuBuilder.CreateTemplateNode<UCustomizableObjectNodeComponentVariation>(),
 			ContextMenuBuilder.CreateTemplateNode<UCustomizableObjectNodeComponentSwitch>(),
 		};
@@ -542,7 +545,7 @@ void UEdGraphSchema_CustomizableObject::GetGraphContextActions(FGraphContextMenu
 
 	{
 		// External Pin Nodes
-		TArray<FName> PinTypes({ PC_Material, PC_Modifier, PC_Mesh, PC_Image, PC_Projector, PC_GroupProjector, PC_Color, PC_Float, PC_Bool, PC_Enum, PC_Stack, PC_PassThroughImage, PC_MaterialAsset, PC_PoseAsset });
+		TArray<FName> PinTypes({ PC_Material, PC_Modifier, PC_Mesh, PC_Image, PC_Projector, PC_GroupProjector, PC_Color, PC_Float, PC_Bool, PC_Enum, PC_Stack, PC_PassThroughImage, PC_MaterialAsset, PC_PoseAsset, PC_Component });
 
 		// Add pin types from extensions
 		for (const FRegisteredCustomizableObjectPinType& PinType : ICustomizableObjectModule::Get().GetExtendedPinTypes())
diff --git a/Engine/Plugins/Experimental/Mutable/Source/CustomizableObjectEditor/Private/MuCOE/GenerateMutableSource/GenerateMutableSource.cpp b/Engine/Plugins/Experimental/Mutable/Source/CustomizableObjectEditor/Private/MuCOE/GenerateMutableSource/GenerateMutableSource.cpp
--- a/Engine/Plugins/Experimental/Mutable/Source/CustomizableObjectEditor/Private/MuCOE/GenerateMutableSource/GenerateMutableSource.cpp
+++ b/Engine/Plugins/Experimental/Mutable/Source/CustomizableObjectEditor/Private/MuCOE/GenerateMutableSource/GenerateMutableSource.cpp
@@ -2,6 +2,7 @@
 
 #include "MuCOE/GenerateMutableSource/GenerateMutableSource.h"
 
+#include "GenerateMutableSourceComponent.h"
 #include "Animation/AnimInstance.h"
 #include "Engine/SkeletalMeshSocket.h"
 #include "Engine/TextureLODSettings.h"
@@ -58,10 +59,6 @@
 #define LOCTEXT_NAMESPACE "CustomizableObjectEditor"
 
 
-// Forward declarations 
-bool AffectsCurrentComponent(const UEdGraphPin* Pin, FMutableGraphGenerationContext& GenerationContext);
-
-
 uint32 GetTypeHash(const FGraphCycleKey& Key)
 {
 	return HashCombine(GetTypeHash(&Key.Pin), GetTypeHash(Key.Id));
@@ -821,7 +818,7 @@ mu::NodeMeshApplyPosePtr CreateNodeMeshApplyPose(FMutableGraphGenerationContext&
 
 
 // Convert a CustomizableObject Source Graph into a mutable source graph  
-mu::Ptr<mu::NodeObject> GenerateMutableSource(const UEdGraphPin * Pin, FMutableGraphGenerationContext & GenerationContext, bool bPartialCompilation)
+mu::Ptr<mu::NodeObject> GenerateMutableSource(const UEdGraphPin * Pin, FMutableGraphGenerationContext & GenerationContext)
 {
 	MUTABLE_CPUPROFILER_SCOPE(GenerateMutableSource);
 
@@ -844,15 +841,6 @@ mu::Ptr<mu::NodeObject> GenerateMutableSource(const UEdGraphPin * Pin, FMutableG
 	
 	if (const UCustomizableObjectNodeObject* TypedNodeObj = Cast<UCustomizableObjectNodeObject>(Node))
 	{
-		// Add it to the current processing hierarchy if we have none
-		// \TODO: or special flag starting a separate component in the future.
-		bool bStartedComponent = false;
-		if (!GenerationContext.ComponentNewNode.Num())
-		{
-			GenerationContext.ComponentNewNode.Push(FMutableGraphGenerationContext::ObjectParent());
-			bStartedComponent = true;
-		}
-
 		mu::Ptr<mu::NodeObjectNew> ObjectNode = new mu::NodeObjectNew();
 		Result = ObjectNode;
 
@@ -863,105 +851,7 @@ mu::Ptr<mu::NodeObject> GenerateMutableSource(const UEdGraphPin * Pin, FMutableG
 			GenerationContext.Compiler->CompilerLog(FText::FromString(TEXT("Warning: Node has a duplicated GUID. A new ID has been generated, but cooked data will not be deterministic.")), Node, EMessageSeverity::Warning);
 		}
 		ObjectNode->SetUid(FinalGuid.ToString());
-
-
-		// Process the pins defining components implicitly (original method)
-		//-------------------------------------------------------------------
-
-		// LOD
-		const int32 NumLODs = TypedNodeObj->GetNumLODPins();
-
-		// Fill the basic LOD Settings
-		if (!GenerationContext.NumLODsInRoot)
-		{
-			check(!GenerationContext.ComponentInfos.IsEmpty());
-
-			// NumLODsInRoot
-			int32 MaxRefMeshLODs = 1;
-			for (int32 MeshIndex = 0; MeshIndex < GenerationContext.ComponentInfos.Num(); ++MeshIndex)
-			{
-				const USkeletalMesh* RefSkeletalMesh = GenerationContext.ComponentInfos[MeshIndex].RefSkeletalMesh;
-				if (RefSkeletalMesh && RefSkeletalMesh->GetLODNum() > MaxRefMeshLODs)
-				{
-					MaxRefMeshLODs = RefSkeletalMesh->GetLODNum();
-				}
-
-				if (TypedNodeObj->ComponentSettings.IsValidIndex(MeshIndex))
-				{
-					GenerationContext.ComponentInfos[MeshIndex].AccumulateBonesToRemovePerLOD(TypedNodeObj->ComponentSettings[MeshIndex], TypedNodeObj->NumLODs);
-				}
-			}
-
-			if (MaxRefMeshLODs < NumLODs)
-			{
-				FString Msg = FString::Printf(TEXT("The object has %d LODs but the reference mesh only %d. Resulting objects will have %d LODs."),
-					NumLODs, MaxRefMeshLODs, MaxRefMeshLODs);
-				GenerationContext.Compiler->CompilerLog(FText::FromString(Msg), Node, EMessageSeverity::Warning);
-				GenerationContext.NumLODsInRoot = MaxRefMeshLODs;
-			}
-			else
-			{
-				GenerationContext.NumLODsInRoot = NumLODs;
-			}
 		
-			USkeletalMesh* RefSkeletalMesh = GenerationContext.ComponentInfos[0].RefSkeletalMesh;
-			check(RefSkeletalMesh);
-
-			const FMutableLODSettings& LODSettings = GenerationContext.Object->LODSettings;
-
-			// Find the MinLOD available for the target platform
-			if (RefSkeletalMesh->IsMinLodQualityLevelEnable())
-			{
-				FSupportedQualityLevelArray SupportedQualityLevels = LODSettings.MinQualityLevelLOD.GetSupportedQualityLevels(*GenerationContext.Options.TargetPlatform->GetPlatformInfo().IniPlatformName.ToString());
-				
-				int32 MinValue = GenerationContext.NumLODsInRoot - 1;
-				for (int32& QL : SupportedQualityLevels)
-				{
-					// check if have data for the supported quality level or set to default.
-					if (LODSettings.MinQualityLevelLOD.IsQualityLevelValid(QL))
-					{
-						MinValue = FMath::Min(LODSettings.MinQualityLevelLOD.GetValueForQualityLevel(QL), MinValue);
-					}
-					else 
-					{
-						MinValue = LODSettings.MinQualityLevelLOD.GetDefault();
-						break;
-					}
-				}
-
-				GenerationContext.FirstLODAvailable = FMath::Max(0, MinValue);
-			}
-			else
-			{
-				GenerationContext.FirstLODAvailable = LODSettings.MinLOD.GetValueForPlatform(*GenerationContext.Options.TargetPlatform->IniPlatformName());
-			}
-
-			GenerationContext.FirstLODAvailable = FMath::Clamp(GenerationContext.FirstLODAvailable, 0, GenerationContext.NumLODsInRoot - 1);
-
-			// Find the streaming settings for the target platform
-			if (LODSettings.bOverrideLODStreamingSettings)
-			{
-				GenerationContext.bEnableLODStreaming = LODSettings.bEnableLODStreaming.GetValueForPlatform(*GenerationContext.Options.TargetPlatform->IniPlatformName());
-				GenerationContext.NumMaxLODsToStream = LODSettings.NumMaxStreamedLODs.GetValueForPlatform(*GenerationContext.Options.TargetPlatform->IniPlatformName());
-			}
-			else
-			{
-				for (int32 MeshIndex = 0; MeshIndex < GenerationContext.ComponentInfos.Num(); ++MeshIndex)
-				{
-					RefSkeletalMesh = GenerationContext.ComponentInfos[MeshIndex].RefSkeletalMesh;
-					check(RefSkeletalMesh);
-
-					GenerationContext.bEnableLODStreaming = GenerationContext.bEnableLODStreaming &&
-						RefSkeletalMesh->GetEnableLODStreaming(GenerationContext.Options.TargetPlatform);
-
-					GenerationContext.NumMaxLODsToStream = FMath::Min(static_cast<int32>(GenerationContext.NumMaxLODsToStream),
-						RefSkeletalMesh->GetMaxNumStreamedLODs(GenerationContext.Options.TargetPlatform));
-				}
-			}
-
-			GenerationContext.NumMaxLODsToStream = FMath::Clamp(GenerationContext.NumMaxLODsToStream, 0, GenerationContext.NumLODsInRoot - 1);
-		}
-
 		// States
 		int NumStates = TypedNodeObj->States.Num();
 		ObjectNode->SetStateCount(NumStates);
@@ -969,7 +859,7 @@ mu::Ptr<mu::NodeObject> GenerateMutableSource(const UEdGraphPin * Pin, FMutableG
 		// In a partial compilation we will filter the states of the root object
 		bool bFilterStates = true;
 
-		if (bPartialCompilation)
+		if (GenerationContext.bPartialCompilation)
 		{
 			if (!TypedNodeObj->ParentObject)
 			{
@@ -1009,181 +899,60 @@ mu::Ptr<mu::NodeObject> GenerateMutableSource(const UEdGraphPin * Pin, FMutableG
 
 			GenerationContext.StateUIDataMap.Add(State.Name, StateUIData);
 		}
+		
+		// Process components.
+		//-------------------------------------------------------------------
+		TArray<mu::Ptr<mu::NodeComponent>> ExplicitComponents;
 
-		// Update the current automatic LOD policy
-		ECustomizableObjectAutomaticLODStrategy LastLODStrategy = GenerationContext.CurrentAutoLODStrategy;
-		if (TypedNodeObj->AutoLODStrategy != ECustomizableObjectAutomaticLODStrategy::Inherited)
-		{
-			GenerationContext.CurrentAutoLODStrategy = TypedNodeObj->AutoLODStrategy;
-		}
-
-		// Mesh components per instance
-		if (!GenerationContext.NumMeshComponentsInRoot)
-		{
-			GenerationContext.NumMeshComponentsInRoot = GenerationContext.Object->GetPrivate()->MutableMeshComponents.Num();
-		}
-
-		const int32 NumLODsInRoot = GenerationContext.NumLODsInRoot;
-		const int32 NumMeshComponentsInRoot = GenerationContext.NumMeshComponentsInRoot;
-		if (GenerationContext.ComponentNewNode.Last().Components.Num() < NumMeshComponentsInRoot)
-		{
-			GenerationContext.ComponentNewNode.Last().Components.SetNum(NumMeshComponentsInRoot);
+		const UEdGraphPin* ComponentsPin = TypedNodeObj->ComponentsPin();
+		TArray<UEdGraphPin*> ConnectedComponentPins = FollowInputPinArray(*ComponentsPin);
+		for (const UEdGraphPin* ComponentNodePin : ConnectedComponentPins)
+		{		
+			mu::Ptr<mu::NodeComponent> ComponentNode = GenerateMutableSourceComponent(ComponentNodePin, GenerationContext);
+			ObjectNode->Components.Add(ComponentNode);
 		}
 
-		// Number of components defined in the traditional way: in the object node, and referenced in graph materials
-		int32 NumImplicitComponents = GenerationContext.NumMeshComponentsInRoot;
-
-		// Add implicit components
-		for (int32 CurrentComponent = 0; CurrentComponent < NumImplicitComponents; ++CurrentComponent)
+		// Process modifiers.
+		//-------------------------------------------------------------------
+		if (const UEdGraphPin* ModifierPin = TypedNodeObj->ModifiersPin())
 		{
-			FName ComponentName = GenerationContext.Object->GetPrivate()->MutableMeshComponents[CurrentComponent].Name;
-			GenerationContext.CurrentMeshComponent = ComponentName;
-
-			mu::Ptr<mu::NodeComponent> ComponentNode = Invoke([&GenerationContext, CurrentComponent]() -> mu::Ptr<mu::NodeComponent>
-				{
-					mu::Ptr<mu::NodeComponentNew>& ParentComponent =
-						GenerationContext.ComponentNewNode.Last().Components[CurrentComponent];
-
-					if (!ParentComponent)
-					{
-						ParentComponent = new mu::NodeComponentNew();
-						ParentComponent->Id = CurrentComponent;
-
-						return ParentComponent;
-					}
-
-					mu::Ptr<mu::NodeComponentEdit> EditComponent = new mu::NodeComponentEdit();
-					EditComponent->SetParent(ParentComponent.get());
-
-					return EditComponent;
-				});
-
-			ComponentNode->SetMessageContext(Node);
-			ObjectNode->Components.Add(ComponentNode);
-
-
-			int32 FirstLOD = -1;
-			for (int32 CurrentLOD = 0; CurrentLOD < NumLODsInRoot; ++CurrentLOD)
+			TArray<UEdGraphPin*> ConnectedModifierPins = FollowInputPinArray(*ModifierPin);
+			for (UEdGraphPin* const ChildNodePin : ConnectedModifierPins)
 			{
-				GenerationContext.CurrentLOD = CurrentLOD;
-
-				mu::Ptr<mu::NodeLOD> LODNode = new mu::NodeLOD();
-				ComponentNode->LODs.Add(LODNode);
-				LODNode->SetMessageContext(Node);
-
-		
-				const bool bUseAutomaticLods = GenerationContext.CurrentAutoLODStrategy == ECustomizableObjectAutomaticLODStrategy::AutomaticFromMesh;
-				FirstLOD = (CurrentLOD < NumLODs) && (FirstLOD == INDEX_NONE || !bUseAutomaticLods) ? CurrentLOD : FirstLOD;
-
-				if (FirstLOD < 0)
+				if (!Cast<UCustomizableObjectNodeModifierBase>(ChildNodePin->GetOwningNode()))
 				{
 					continue;
 				}
 
-				if (GenerationContext.CurrentLOD < GenerationContext.FirstLODAvailable)
-				{
-					continue;
-				} 
-			
-				// Generate all relevant LODs for this object up until the current LODIndex.
-				for (int32 LODIndex = FirstLOD; LODIndex <= CurrentLOD; ++LODIndex)
-				{
+				// Set it to "None" to indicate we don't care about component id.
+				FName OldCurrentMeshComponent = GenerationContext.CurrentMeshComponent;
+				GenerationContext.CurrentMeshComponent = FName();
 
-					const UEdGraphPin* LODPin = TypedNodeObj->LODPin(LODIndex);
-					if (!LODPin)
+				mu::Ptr<mu::NodeModifier> ModifierNode = GenerateMutableSourceModifier(ChildNodePin, GenerationContext);
+				for (int32 ComponentIndex = 0; ComponentIndex < ObjectNode->Components.Num(); ++ComponentIndex)
+				{
+					mu::Ptr<mu::NodeComponent> Component = ObjectNode->Components[ComponentIndex];
+					if (!Component)
 					{
 						continue;
 					}
-
-					GenerationContext.FromLOD = LODIndex;
-
-					TArray<UEdGraphPin*> ConnectedLODPins = FollowInputPinArray(*LODPin);
-
-					// Proccess non modifier nodes.
-					for (UEdGraphPin* const ChildNodePin : ConnectedLODPins)
-					{
-						// Modifiers are shared for all components and are processed per LOD and not component.
-						if (Cast<UCustomizableObjectNodeModifierBase>(ChildNodePin->GetOwningNode()))
-						{
-							continue;
-						}
-
-						if (!AffectsCurrentComponent(ChildNodePin, GenerationContext))
-						{
-							continue;
-						}
-
-						mu::Ptr<mu::NodeSurface> SurfaceNode = GenerateMutableSourceSurface(ChildNodePin, GenerationContext);
-						LODNode->Surfaces.Add(SurfaceNode);
-					}
-
-					// Process legacy modifiers.
-					for (UEdGraphPin* const ChildNodePin : ConnectedLODPins)
-					{
-						if (!Cast<UCustomizableObjectNodeModifierBase>(ChildNodePin->GetOwningNode()))
-						{
-							continue;
-						}
-
-						// Warn about legacy modifier connection
-						FString Msg = FString::Printf(TEXT("The object has legacy modifier connections (to material pins?) that should be updated."));
-						GenerationContext.Compiler->CompilerLog(FText::FromString(Msg), Node, EMessageSeverity::Warning);
-
-						// Set it to "None" to indicate we don't care about component id.
-						GenerationContext.CurrentMeshComponent = FName();
-
-						mu::Ptr<mu::NodeModifier> ModifierNode = GenerateMutableSourceModifier(ChildNodePin, GenerationContext);
-						LODNode->Modifiers.Add(ModifierNode);
-
-						GenerationContext.CurrentMeshComponent = ComponentName;
-					}
-				}
-
-				// Process modifiers.
-				const UEdGraphPin* ModifierPin = TypedNodeObj->ModifiersPin();
-				if (ModifierPin)
-				{
-					TArray<UEdGraphPin*> ConnectedModifierPins = FollowInputPinArray(*ModifierPin);
-					for (UEdGraphPin* const ChildNodePin : ConnectedModifierPins)
+					
+					for (int32 LODIndex = 0; LODIndex < Component->LODs.Num(); ++LODIndex)
 					{
-						if (!Cast<UCustomizableObjectNodeModifierBase>(ChildNodePin->GetOwningNode()))
+						mu::Ptr<mu::NodeLOD> LOD = Component->LODs[LODIndex];
+						if (!LOD)
 						{
 							continue;
 						}
-
-						// Set it to "None" to indicate we don't care about component id.
-						GenerationContext.CurrentMeshComponent = FName();
-
-						mu::Ptr<mu::NodeModifier> ModifierNode = GenerateMutableSourceModifier(ChildNodePin, GenerationContext);
-						LODNode->Modifiers.Add(ModifierNode);
-
-						GenerationContext.CurrentMeshComponent = ComponentName;
+						
+						LOD->Modifiers.Add(ModifierNode);
 					}
 				}
 
+				GenerationContext.CurrentMeshComponent = OldCurrentMeshComponent;
 			}
 		}
-
-
-		// Process the pins defining explicit components
-		//-------------------------------------------------------------------
-		TArray<mu::Ptr<mu::NodeComponent>> ExplicitComponents;
-
-		const UEdGraphPin* ComponentsPin = TypedNodeObj->ComponentsPin();
-		TArray<UEdGraphPin*> ConnectedComponentPins = FollowInputPinArray(*ComponentsPin);
-		for (const UEdGraphPin* ComponentNodePin : ConnectedComponentPins)
-		{		
-			// \TODO: Needed?
-			//GenerationContext.ComponentInfos.Add(RefSkeletalMesh);
-			//GenerationContext.ReferencedSkeletons.Add(RefSkeleton);
-
-			++GenerationContext.NumExplicitMeshComponents;
-
-			mu::Ptr<mu::NodeComponent> ComponentNode = GenerateMutableSourceComponent(ComponentNodePin, GenerationContext);
-			ObjectNode->Components.Add(ComponentNode);
-		}
-
-
+		
 		// Generate inputs to Object node pins added by extensions
 		//-------------------------------------------------------------------
 		for (const FRegisteredObjectNodeInputPin& ExtensionInputPin : ICustomizableObjectModule::Get().GetAdditionalObjectNodePins())
@@ -1226,17 +995,9 @@ mu::Ptr<mu::NodeObject> GenerateMutableSource(const UEdGraphPin * Pin, FMutableG
 		ObjectNode->Children.Reserve(ConnectedChildrenPins.Num());
 		for (int32 ChildIndex = 0; ChildIndex < ConnectedChildrenPins.Num(); ++ChildIndex)
 		{
-			mu::Ptr<mu::NodeObject> ChildNode = GenerateMutableSource(ConnectedChildrenPins[ChildIndex], GenerationContext, bPartialCompilation);
+			mu::Ptr<mu::NodeObject> ChildNode = GenerateMutableSource(ConnectedChildrenPins[ChildIndex], GenerationContext);
 			ObjectNode->Children.Add(ChildNode);
 		}
-
-		// Remove from the current processing hierarchy
-		if (bStartedComponent)
-		{
-			GenerationContext.ComponentNewNode.Pop();
-		}
-
-		GenerationContext.CurrentAutoLODStrategy = LastLODStrategy;
 	}
 
 	else if (const UCustomizableObjectNodeObjectGroup* TypedNodeGroup = Cast<UCustomizableObjectNodeObjectGroup>(Node))
@@ -1314,7 +1075,7 @@ mu::Ptr<mu::NodeObject> GenerateMutableSource(const UEdGraphPin * Pin, FMutableG
 			{
 				bAtLeastOneConnected = true;
 
-				ChildNode = GenerateMutableSource(ConnectedChildrenPins[ChildIndex], GenerationContext, bPartialCompilation);
+				ChildNode = GenerateMutableSource(ConnectedChildrenPins[ChildIndex], GenerationContext);
 				GroupNode->SetChild(ChildIndex, ChildNode.get());
 
 				if (CustomizableObjectNodeObject)
@@ -1422,7 +1183,7 @@ mu::Ptr<mu::NodeObject> GenerateMutableSource(const UEdGraphPin * Pin, FMutableG
 			{
 				bAtLeastOneConnected = true;
 
-				ChildNode = GenerateMutableSource(ExternalChildNode->OutputPin(), GenerationContext, bPartialCompilation);
+				ChildNode = GenerateMutableSource(ExternalChildNode->OutputPin(), GenerationContext);
 				GroupNode->SetChild(ChildIndex, ChildNode.get());
 
 				if (CustomizableObjectNodeObject)
@@ -1513,118 +1274,6 @@ mu::Ptr<mu::NodeObject> GenerateMutableSource(const UEdGraphPin * Pin, FMutableG
 	return Result;
 }
 
-bool AffectsCurrentComponent(const UEdGraphPin* Pin, FMutableGraphGenerationContext& GenerationContext)
-{
-	check(Pin);
-	RETURN_ON_CYCLE(*Pin, GenerationContext)
-	
-	FName ComponentName = FName();
-	const UCustomizableObjectNode* Node = CastChecked<UCustomizableObjectNode>(Pin->GetOwningNode());
-
-	// Parent material for extend and edit material nodes
-	UCustomizableObjectNodeMaterialBase* OriginalParentMaterialNode = nullptr;
-
-	if (const UCustomizableObjectNodeMaterialVariation* TypedNodeVar = Cast<UCustomizableObjectNodeMaterialVariation>(Node))
-	{
-		bool bAffectsCurrentComponent = false;
-
-		for (const UEdGraphPin* ConnectedPin : FollowInputPinArray(*TypedNodeVar->DefaultPin()))
-		{
-			if (!AffectsCurrentComponent(ConnectedPin, GenerationContext))
-			{
-				if (bAffectsCurrentComponent)
-				{
-					FString Msg = FString::Printf(TEXT("Error! One or more materials nodes linked to a material variation node have different component name"));
-					GenerationContext.Compiler->CompilerLog(FText::FromString(Msg), Node, EMessageSeverity::Error);
-				}
-
-				return false;
-			}
-		
-			bAffectsCurrentComponent = true;
-		}
-
-		for (int VariationIndex = 0; VariationIndex < TypedNodeVar->GetNumVariations(); ++VariationIndex)
-		{
-			if (UEdGraphPin* VariationPin = TypedNodeVar->VariationPin(VariationIndex))
-			{
-				for (const UEdGraphPin* ConnectedPin : FollowInputPinArray(*VariationPin))
-				{
-					if (!AffectsCurrentComponent(ConnectedPin, GenerationContext))
-					{
-						if (bAffectsCurrentComponent)
-						{
-							FString Msg = FString::Printf(TEXT("Error! One or more materials nodes linked to a material variation node have different component name"));
-							GenerationContext.Compiler->CompilerLog(FText::FromString(Msg), Node, EMessageSeverity::Error);
-						}
-
-						return false;
-					}
-
-					bAffectsCurrentComponent = true;
-				}
-			}
-		}
-
-		return bAffectsCurrentComponent;
-	}
-
-	else if (const UCustomizableObjectNodeMaterialSwitch* TypedNodeSwitch = Cast<UCustomizableObjectNodeMaterialSwitch>(Node))
-	{
-		bool bAffectsCurrentComponent = false;
-
-		int32 OptionCount = TypedNodeSwitch->GetNumElements();
-		for (int32 OptionIndex=0; OptionIndex<OptionCount; ++OptionIndex)
-		{
-			if (UEdGraphPin* OptionPin = TypedNodeSwitch->GetElementPin(OptionIndex))
-			{
-				for (const UEdGraphPin* ConnectedPin : FollowInputPinArray(*OptionPin))
-				{
-					if (!AffectsCurrentComponent(ConnectedPin, GenerationContext))
-					{
-						if (bAffectsCurrentComponent)
-						{
-							FString Msg = FString::Printf(TEXT("Error! One or more materials nodes linked to a material switch node have different component name"));
-							GenerationContext.Compiler->CompilerLog(FText::FromString(Msg), Node, EMessageSeverity::Error);
-						}
-
-						return false;
-					}
-
-					bAffectsCurrentComponent = true;
-				}
-			}
-		}
-
-		return bAffectsCurrentComponent;
-	}
-
-	else if (const UCustomizableObjectNodeMaterialBase* TypedNodeMat = Cast<UCustomizableObjectNodeMaterialBase>(Node))
-	{
-		ComponentName = TypedNodeMat->GetMeshComponentName();
-	}
-
-	else if (const UCustomizableObjectNodeModifierBase* TypedNodeModifier = Cast<UCustomizableObjectNodeModifierBase>(Node))
-	{
-		// Modifiers can affect any component depending on tags.
-		// TODO: Scan
-		return true;
-	}
-	else
-	{
-		unimplemented();
-		return false;
-	}
-
-	// Check if this material has a valid component
-	if (!GenerationContext.Object->GetPrivate()->MutableMeshComponents.ContainsByPredicate([&](const FMutableMeshComponentData& Component) { return Component.Name == ComponentName; }))
-	{
-		FString Msg = ("Error! Material Node with an invalid Component assigned.");
-		GenerationContext.Compiler->CompilerLog(FText::FromString(Msg), OriginalParentMaterialNode ? OriginalParentMaterialNode : Node, EMessageSeverity::Error);
-	}
-
-	return ComponentName == GenerationContext.CurrentMeshComponent;
-}
 
 int32 AddTagToMutableMeshUnique(mu::Mesh& MutableMesh, const FString& Tag)
 {
@@ -2125,20 +1774,20 @@ FMutableComponentInfo::FMutableComponentInfo(FName InComponentName, USkeletalMes
 }
 
 
-void FMutableComponentInfo::AccumulateBonesToRemovePerLOD(const FComponentSettings& ComponentSettings, int32 NumLODs)
+void FMutableComponentInfo::AccumulateBonesToRemovePerLOD(const TArray<FLODReductionSettings>& LODReductionSettings, int32 NumLODs)
 {
 	BonesToRemovePerLOD.SetNum(NumLODs);
 
 	TMap<FName, bool> BonesToRemove;
 
-	const int32 ComponentSettingsLODCount = ComponentSettings.LODReductionSettings.Num();
+	const int32 ComponentSettingsLODCount = LODReductionSettings.Num();
 	for (int32 LODIndex = 0; LODIndex < NumLODs; ++LODIndex)
 	{
 		if (LODIndex < ComponentSettingsLODCount)
 		{
-			const FLODReductionSettings& LODReductionSettings = ComponentSettings.LODReductionSettings[LODIndex];
+			const FLODReductionSettings& LODReductionSetting = LODReductionSettings[LODIndex];
 
-			for (const FBoneToRemove& Bone : LODReductionSettings.BonesToRemove)
+			for (const FBoneToRemove& Bone : LODReductionSetting.BonesToRemove)
 			{
 				if (bool* bOnlyRemoveChildren = BonesToRemove.Find(Bone.BoneName))
 				{
diff --git a/Engine/Plugins/Experimental/Mutable/Source/CustomizableObjectEditor/Private/MuCOE/GenerateMutableSource/GenerateMutableSource.h b/Engine/Plugins/Experimental/Mutable/Source/CustomizableObjectEditor/Private/MuCOE/GenerateMutableSource/GenerateMutableSource.h
--- a/Engine/Plugins/Experimental/Mutable/Source/CustomizableObjectEditor/Private/MuCOE/GenerateMutableSource/GenerateMutableSource.h
+++ b/Engine/Plugins/Experimental/Mutable/Source/CustomizableObjectEditor/Private/MuCOE/GenerateMutableSource/GenerateMutableSource.h
@@ -399,14 +399,16 @@ struct FMutableComponentInfo
 {
 	FMutableComponentInfo(FName InComponentName, USkeletalMesh* InRefSkeletalMesh);
 
-	void AccumulateBonesToRemovePerLOD(const FComponentSettings& ComponentSettings, int32 NumLODs);
+	void AccumulateBonesToRemovePerLOD(const TArray<FLODReductionSettings>& LODReductionSettings, int32 NumLODs);
 
 	FName ComponentName;
 
 	// Each component must have a reference SkeletalMesh with a valid Skeleton
 	USkeletalMesh* RefSkeletalMesh = nullptr;
 	USkeleton* RefSkeleton = nullptr;
-
+	
+	UCustomizableObjectNodeComponentMesh* NodeComponentMesh = nullptr;
+	
 	// Map to check skeleton compatibility
 	TMap<const USkeleton*, bool> SkeletonCompatibility;
 	
@@ -793,13 +795,6 @@ struct FMutableGraphGenerationContext
 
 	TArray<FAnimBpOverridePhysicsAssetsInfo> AnimBpOverridePhysicsAssetsInfo;
 
-	// Hierarchy of current ComponentNew nodes
-	struct ObjectParent
-	{
-		TArray<mu::Ptr<mu::NodeComponentNew>> Components;
-	};
-	TArray< ObjectParent > ComponentNewNode;
-
 	uint8 FromLOD = 0; // LOD to append to the CurrentLOD when using AutomaticLODs. 
 	uint8 CurrentLOD = 0;
 	FName CurrentMeshComponent;
@@ -810,14 +805,15 @@ struct FMutableGraphGenerationContext
 	uint8 NumLODsInRoot = 0;
 	uint8 NumMeshComponentsInRoot = 0;
 
-	/** Number of additional mesh components added with explict Component nodes to the root or another object node. */
-	uint8 NumExplicitMeshComponents = 0;
+	uint8 NumPassthroughMeshComponents = 0;
 
 	uint8 FirstLODAvailable = MAX_MESH_LOD_COUNT;
 	uint8 NumMaxLODsToStream = MAX_MESH_LOD_COUNT;
 
 	bool bEnableLODStreaming = true;
 
+	bool bPartialCompilation = false;
+
 	// Based on the last object visited.
 	ECustomizableObjectAutomaticLODStrategy CurrentAutoLODStrategy = ECustomizableObjectAutomaticLODStrategy::Manual;
 
@@ -977,9 +973,7 @@ namespace Private
 } //Private
 
 //
-mu::Ptr<mu::NodeObject> GenerateMutableSource(const class UEdGraphPin* Pin, FMutableGraphGenerationContext& GenerationContext, bool bPartialCompilation);
-mu::Ptr<mu::NodeComponent> GenerateMutableSourceComponent(const UEdGraphPin* Pin, FMutableGraphGenerationContext& GenerationContext);
-
+mu::Ptr<mu::NodeObject> GenerateMutableSource(const class UEdGraphPin* Pin, FMutableGraphGenerationContext& GenerationContext);
 
 /** Populate an array with all the information related to the reference skeletal meshes we might need in-game to generate instances */
 void PopulateReferenceSkeletalMeshesData(FMutableGraphGenerationContext& GenerationContext);
diff --git a/Engine/Plugins/Experimental/Mutable/Source/CustomizableObjectEditor/Private/MuCOE/GenerateMutableSource/GenerateMutableSourceComponent.cpp b/Engine/Plugins/Experimental/Mutable/Source/CustomizableObjectEditor/Private/MuCOE/GenerateMutableSource/GenerateMutableSourceComponent.cpp
--- a/Engine/Plugins/Experimental/Mutable/Source/CustomizableObjectEditor/Private/MuCOE/GenerateMutableSource/GenerateMutableSourceComponent.cpp
+++ b/Engine/Plugins/Experimental/Mutable/Source/CustomizableObjectEditor/Private/MuCOE/GenerateMutableSource/GenerateMutableSourceComponent.cpp
@@ -1,5 +1,9 @@
 // Copyright Epic Games, Inc. All Rights Reserved.
 
+#include "GenerateMutableSourceComponent.h"
+
+#include "GenerateMutableSourceModifier.h"
+#include "Factories/FbxMeshImportData.h"
 #include "MuCOE/GenerateMutableSource/GenerateMutableSource.h"
 
 #include "MuCOE/CustomizableObjectCompiler.h"
@@ -9,24 +13,121 @@
 #include "MuCOE/GraphTraversal.h"
 #include "MuCOE/MutableUtils.h"
 #include "MuCOE/Nodes/CustomizableObjectNodeComponentMesh.h"
+#include "MuCOE/Nodes/CustomizableObjectNodeComponentMeshAddTo.h"
+#include "MuCOE/Nodes/CustomizableObjectNodeComponentPassthroughMesh.h"
 #include "MuCOE/Nodes/CustomizableObjectNodeComponentSwitch.h"
 #include "MuCOE/Nodes/CustomizableObjectNodeComponentVariation.h"
+#include "MuCOE/Nodes/CustomizableObjectNodeModifierBase.h"
 #include "MuT/NodeComponent.h"
+#include "MuT/NodeComponentEdit.h"
 #include "MuT/NodeComponentNew.h"
 #include "MuT/NodeComponentSwitch.h"
 #include "MuT/NodeComponentVariation.h"
 #include "MuT/NodeLOD.h"
 #include "MuT/NodeSurfaceNew.h"
 #include "MuT/NodeMeshConstant.h"
+#include "MuT/NodeObjectNew.h"
 
 #include "Rendering/SkeletalMeshLODModel.h"
 #include "Rendering/SkeletalMeshModel.h"
 #include "Rendering/SkeletalMeshRenderData.h"
+#include "Interfaces/ITargetPlatform.h"
 
 
 #define LOCTEXT_NAMESPACE "CustomizableObjectEditor"
 
-mu::Ptr<mu::NodeComponent> GenerateMutableSourceComponent(const UEdGraphPin * Pin, FMutableGraphGenerationContext& GenerationContext)
+
+/** Generate LOD pins of the given NodeComponentBase (NodeComponent, NodeComponentExtend...).
+ * @param TypedComponentMesh Given component node.
+ * @param NodeComponent Core node to connect LOD generated pins. */
+void GenerateMutableSourceComponentMesh(FMutableGraphGenerationContext& GenerationContext, const UCustomizableObjectNodeComponentMeshBase& TypedComponentMesh, mu::Ptr<mu::NodeComponent> NodeComponent)
+{
+	int32 FirstLOD = -1;
+
+	const int32 NumLODsInRoot = GenerationContext.NumLODsInRoot;
+	for (int32 CurrentLOD = 0; CurrentLOD < NumLODsInRoot; ++CurrentLOD)
+	{
+		GenerationContext.CurrentLOD = CurrentLOD;
+
+		if (!NodeComponent->LODs.IsValidIndex(CurrentLOD))
+		{
+			NodeComponent->LODs.Add(new mu::NodeLOD());
+		}
+
+		mu::Ptr<mu::NodeLOD> LODNode = NodeComponent->LODs[CurrentLOD];
+
+		LODNode->SetMessageContext(&TypedComponentMesh);
+
+		const int32 NumLODs = TypedComponentMesh.LODPins.Num();
+
+		const bool bUseAutomaticLods = GenerationContext.CurrentAutoLODStrategy == ECustomizableObjectAutomaticLODStrategy::AutomaticFromMesh;
+		FirstLOD = (CurrentLOD < NumLODs) && (FirstLOD == INDEX_NONE || !bUseAutomaticLods) ? CurrentLOD : FirstLOD;
+
+		if (FirstLOD < 0)
+		{
+			continue;
+		}
+
+		if (GenerationContext.CurrentLOD < GenerationContext.FirstLODAvailable)
+		{
+			continue;
+		} 
+	
+		// Generate all relevant LODs for this object up until the current LODIndex.
+		for (int32 LODIndex = FirstLOD; LODIndex <= CurrentLOD; ++LODIndex)
+		{
+			if (!TypedComponentMesh.LODPins.IsValidIndex(LODIndex))
+			{
+				continue;
+			}
+			
+			const UEdGraphPin* LODPin = TypedComponentMesh.LODPins[LODIndex].Get();
+			check(LODPin);
+			
+			GenerationContext.FromLOD = LODIndex;
+
+			TArray<UEdGraphPin*> ConnectedLODPins = FollowInputPinArray(*LODPin);
+
+			// Proccess non modifier nodes.
+			for (UEdGraphPin* const ChildNodePin : ConnectedLODPins)
+			{
+				// Modifiers are shared for all components and are processed per LOD and not component.
+				if (Cast<UCustomizableObjectNodeModifierBase>(ChildNodePin->GetOwningNode()))
+				{
+					continue;
+				}
+				
+				mu::Ptr<mu::NodeSurface> SurfaceNode = GenerateMutableSourceSurface(ChildNodePin, GenerationContext);
+				LODNode->Surfaces.Add(SurfaceNode);
+			}
+
+			// Process legacy modifiers.
+			for (UEdGraphPin* const ChildNodePin : ConnectedLODPins)
+			{
+				if (!Cast<UCustomizableObjectNodeModifierBase>(ChildNodePin->GetOwningNode()))
+				{
+					continue;
+				}
+
+				// Warn about legacy modifier connection
+				FString Msg = FString::Printf(TEXT("The object has legacy modifier connections (to material pins?) that should be updated."));
+				GenerationContext.Compiler->CompilerLog(FText::FromString(Msg), &TypedComponentMesh, EMessageSeverity::Warning);
+
+				// Set it to "None" to indicate we don't care about component id.
+				FName OldCurrentMeshComponent = GenerationContext.CurrentMeshComponent;
+				GenerationContext.CurrentMeshComponent = FName();
+
+				mu::Ptr<mu::NodeModifier> ModifierNode = GenerateMutableSourceModifier(ChildNodePin, GenerationContext);
+				LODNode->Modifiers.Add(ModifierNode);
+
+				GenerationContext.CurrentMeshComponent = OldCurrentMeshComponent;
+			}
+		}
+	}
+}
+
+
+mu::Ptr<mu::NodeComponent> GenerateMutableSourceComponent(const UEdGraphPin* Pin, FMutableGraphGenerationContext& GenerationContext)
 {
 	check(Pin)
 	RETURN_ON_CYCLE(*Pin, GenerationContext)
@@ -35,7 +136,7 @@ mu::Ptr<mu::NodeComponent> GenerateMutableSourceComponent(const UEdGraphPin * Pi
 	
 	UCustomizableObjectNode* Node = CastChecked<UCustomizableObjectNode>(Pin->GetOwningNode());
 
-	const FGeneratedKey Key(reinterpret_cast<void*>(&GenerateMutableSourceComponent), *Pin, *Node, GenerationContext, true);
+	const FGeneratedKey Key(reinterpret_cast<void*>(&GenerateMutableSourceComponent), *Pin, *Node, GenerationContext, false);
 	if (const FGeneratedData* Generated = GenerationContext.Generated.Find(Key))
 	{
 		return static_cast<mu::NodeComponent*>(Generated->Node.get());
@@ -45,27 +146,212 @@ mu::Ptr<mu::NodeComponent> GenerateMutableSourceComponent(const UEdGraphPin * Pi
 	
 	if (const UCustomizableObjectNodeComponentMesh* TypedComponentMesh = Cast<UCustomizableObjectNodeComponentMesh>(Node))
 	{
+		int32 CurrentComponent = GenerationContext.Object->GetPrivate()->MutableMeshComponents.IndexOfByPredicate([&](const FMutableMeshComponentData& ComponentData)
+		{
+			return ComponentData.Name == TypedComponentMesh->ComponentName;
+		});
+		
+		if (CurrentComponent == INDEX_NONE)
+		{
+			GenerationContext.Compiler->CompilerLog(FText::Format(LOCTEXT("ComponentNotDeclared", "Component {0} not declared in Object Properties tab."), FText::FromName(TypedComponentMesh->ComponentName)), Node, EMessageSeverity::Error);
+			return {};
+		}
+
+		// Ensure that the CO has a valid AutoLODStrategy on the ActualRoot.
+		if (TypedComponentMesh->AutoLODStrategy == ECustomizableObjectAutomaticLODStrategy::Inherited)
+		{
+			GenerationContext.Compiler->CompilerLog(LOCTEXT("RootInheritsFromParent", "Error! Component LOD Strategy can't be set to 'Inherit from parent object'"), TypedComponentMesh, EMessageSeverity::Error);
+			return nullptr;
+		}
+		GenerationContext.CurrentAutoLODStrategy = TypedComponentMesh->AutoLODStrategy;
+		
+		mu::Ptr<mu::NodeObjectNew> ObjectNode = new mu::NodeObjectNew();
+
+		ObjectNode->SetName(TypedComponentMesh->ComponentName.ToString());
+		FGuid FinalGuid = GenerationContext.GetNodeIdUnique(TypedComponentMesh);
+		if (FinalGuid != TypedComponentMesh->NodeGuid)
+		{
+			GenerationContext.Compiler->CompilerLog(FText::FromString(TEXT("Warning: Node has a duplicated GUID. A new ID has been generated, but cooked data will not be deterministic.")), Node, EMessageSeverity::Warning);
+		}
+		ObjectNode->SetUid(FinalGuid.ToString());
+
+
+		// Process the pins defining components
+		//-------------------------------------------------------------------
+
+		// LOD
+		const int32 NumLODs = TypedComponentMesh->LODPins.Num();
+
+		// Fill the basic LOD Settings
+		if (!GenerationContext.NumLODsInRoot)
+		{
+			check(!GenerationContext.ComponentInfos.IsEmpty());
+
+			// NumLODsInRoot
+			int32 MaxRefMeshLODs = 1;
+			for (int32 MeshIndex = 0; MeshIndex < GenerationContext.ComponentInfos.Num(); ++MeshIndex)
+			{
+				const USkeletalMesh* RefSkeletalMesh = GenerationContext.ComponentInfos[MeshIndex].RefSkeletalMesh;
+				if (RefSkeletalMesh && RefSkeletalMesh->GetLODNum() > MaxRefMeshLODs)
+				{
+					MaxRefMeshLODs = RefSkeletalMesh->GetLODNum();
+				}
+	
+				GenerationContext.ComponentInfos[MeshIndex].AccumulateBonesToRemovePerLOD(TypedComponentMesh->LODReductionSettings, TypedComponentMesh->NumLODs);
+			}
+
+			if (MaxRefMeshLODs < NumLODs)
+			{
+				FString Msg = FString::Printf(TEXT("The object has %d LODs but the reference mesh only %d. Resulting objects will have %d LODs."),
+					NumLODs, MaxRefMeshLODs, MaxRefMeshLODs);
+				GenerationContext.Compiler->CompilerLog(FText::FromString(Msg), Node, EMessageSeverity::Warning);
+				GenerationContext.NumLODsInRoot = MaxRefMeshLODs;
+			}
+			else
+			{
+				GenerationContext.NumLODsInRoot = NumLODs;
+			}
+		
+			USkeletalMesh* RefSkeletalMesh = GenerationContext.ComponentInfos[0].RefSkeletalMesh;
+			check(RefSkeletalMesh);
+
+			const FMutableLODSettings& LODSettings = GenerationContext.Object->LODSettings;
+
+			// Find the MinLOD available for the target platform
+			if (RefSkeletalMesh->IsMinLodQualityLevelEnable())
+			{
+				FSupportedQualityLevelArray SupportedQualityLevels = LODSettings.MinQualityLevelLOD.GetSupportedQualityLevels(*GenerationContext.Options.TargetPlatform->GetPlatformInfo().IniPlatformName.ToString());
+				
+				int32 MinValue = GenerationContext.NumLODsInRoot - 1;
+				for (int32& QL : SupportedQualityLevels)
+				{
+					// check if have data for the supported quality level or set to default.
+					if (LODSettings.MinQualityLevelLOD.IsQualityLevelValid(QL))
+					{
+						MinValue = FMath::Min(LODSettings.MinQualityLevelLOD.GetValueForQualityLevel(QL), MinValue);
+					}
+					else 
+					{
+						MinValue = LODSettings.MinQualityLevelLOD.GetDefault();
+						break;
+					}
+				}
+
+				GenerationContext.FirstLODAvailable = FMath::Max(0, MinValue);
+			}
+			else
+			{
+				GenerationContext.FirstLODAvailable = LODSettings.MinLOD.GetValueForPlatform(*GenerationContext.Options.TargetPlatform->IniPlatformName());
+			}
+
+			GenerationContext.FirstLODAvailable = FMath::Clamp(GenerationContext.FirstLODAvailable, 0, GenerationContext.NumLODsInRoot - 1);
+
+			// Find the streaming settings for the target platform
+			if (LODSettings.bOverrideLODStreamingSettings)
+			{
+				GenerationContext.bEnableLODStreaming = LODSettings.bEnableLODStreaming.GetValueForPlatform(*GenerationContext.Options.TargetPlatform->IniPlatformName());
+				GenerationContext.NumMaxLODsToStream = LODSettings.NumMaxStreamedLODs.GetValueForPlatform(*GenerationContext.Options.TargetPlatform->IniPlatformName());
+			}
+			else
+			{
+				for (int32 MeshIndex = 0; MeshIndex < GenerationContext.ComponentInfos.Num(); ++MeshIndex)
+				{
+					RefSkeletalMesh = GenerationContext.ComponentInfos[MeshIndex].RefSkeletalMesh;
+					check(RefSkeletalMesh);
+
+					GenerationContext.bEnableLODStreaming = GenerationContext.bEnableLODStreaming &&
+						RefSkeletalMesh->GetEnableLODStreaming(GenerationContext.Options.TargetPlatform);
+
+					GenerationContext.NumMaxLODsToStream = FMath::Min(static_cast<int32>(GenerationContext.NumMaxLODsToStream),
+						RefSkeletalMesh->GetMaxNumStreamedLODs(GenerationContext.Options.TargetPlatform));
+				}
+			}
+
+			GenerationContext.NumMaxLODsToStream = FMath::Clamp(GenerationContext.NumMaxLODsToStream, 0, GenerationContext.NumLODsInRoot - 1);
+		}
+		
+		// Mesh components per instance
+		if (!GenerationContext.NumMeshComponentsInRoot)
+		{
+			GenerationContext.NumMeshComponentsInRoot = GenerationContext.Object->GetPrivate()->MutableMeshComponents.Num();
+		}
+		
+		mu::Ptr<mu::NodeComponentNew> NodeComponentNew = new mu::NodeComponentNew();
+		NodeComponentNew->Id = CurrentComponent;
+		
+		NodeComponentNew->SetMessageContext(Node);
+		ObjectNode->Components.Add(NodeComponentNew);
+
 		GenerationContext.CurrentMeshComponent = TypedComponentMesh->ComponentName;
+		GenerateMutableSourceComponentMesh(GenerationContext, *TypedComponentMesh, NodeComponentNew);
+		
+		Result = NodeComponentNew;
+	}
+
+	else if (const UCustomizableObjectNodeComponentMeshAddTo* TypedComponentMeshExtend = Cast<UCustomizableObjectNodeComponentMeshAddTo>(Node))
+	{
+		if (const FMutableComponentInfo* ComponentInfo = GenerationContext.ComponentInfos.FindByPredicate([&TypedComponentMeshExtend](const FMutableComponentInfo& Element)
+		{
+			return Element.ComponentName == TypedComponentMeshExtend->ParentComponentName &&
+				Element.NodeComponentMesh;
+		}))
+		{
+			UCustomizableObjectNodeComponentMesh* TypedParentComponentMesh = ComponentInfo->NodeComponentMesh;
+
+			if (TypedComponentMeshExtend->NumLODs > TypedParentComponentMesh->NumLODs)
+			{
+				FText Msg = FText::Format(LOCTEXT("ExtendMeshComponentLODs", "Add To Mesh Component can not have more LODs than its parent Mesh Component [{0}]."), FText::FromName(TypedComponentMeshExtend->ParentComponentName));
+				GenerationContext.Compiler->CompilerLog(Msg, TypedComponentMeshExtend, EMessageSeverity::Warning);
+			}
+
+			mu::Ptr<mu::NodeComponent> ParentNodeComponent = GenerateMutableSourceComponent(TypedParentComponentMesh->OutputPin.Get(), GenerationContext);
+
+			mu::Ptr<mu::NodeComponentEdit> NodeComponentEdit = new mu::NodeComponentEdit();
+			NodeComponentEdit->SetParent(ParentNodeComponent.get());
+			NodeComponentEdit->SetMessageContext(TypedComponentMeshExtend);
+		
+			GenerationContext.CurrentMeshComponent = TypedParentComponentMesh->ComponentName;
+
+			GenerationContext.CurrentAutoLODStrategy = TypedComponentMeshExtend->AutoLODStrategy == ECustomizableObjectAutomaticLODStrategy::Inherited ?
+				TypedParentComponentMesh->AutoLODStrategy :
+				TypedComponentMeshExtend->AutoLODStrategy;
+			
+			GenerateMutableSourceComponentMesh(GenerationContext, *TypedComponentMeshExtend, NodeComponentEdit);
 
-		if (TypedComponentMesh->ComponentName.IsNone())
+			Result = NodeComponentEdit;	
+		}
+		else
+		{
+			FText Msg = FText::Format(LOCTEXT("ExtendMeshComponent", "Can not find parent Mesh Component {0}."), FText::FromName(TypedComponentMeshExtend->ParentComponentName));
+			GenerationContext.Compiler->CompilerLog(Msg, TypedComponentMeshExtend, EMessageSeverity::Error);
+		}
+	}
+	
+	else if (const UCustomizableObjectNodeComponentPassthroughMesh* TypeComponentPassthroughMesh = Cast<UCustomizableObjectNodeComponentPassthroughMesh>(Node))
+	{
+		++GenerationContext.NumPassthroughMeshComponents;
+
+		GenerationContext.CurrentMeshComponent = TypeComponentPassthroughMesh->ComponentName;
+
+		if (TypeComponentPassthroughMesh->ComponentName.IsNone())
 		{
 			FString Msg = FString::Printf(TEXT("Invalid Component Name."));
-			GenerationContext.Compiler->CompilerLog(FText::FromString(Msg), TypedComponentMesh, EMessageSeverity::Warning);
+			GenerationContext.Compiler->CompilerLog(FText::FromString(Msg), TypeComponentPassthroughMesh, EMessageSeverity::Warning);
 			return nullptr;
 		}
 
-		if (!TypedComponentMesh->Mesh.IsValid())
+		if (!TypeComponentPassthroughMesh->Mesh.IsValid())
 		{
 			FString Msg = FString::Printf(TEXT("No mesh set for component node."));
-			GenerationContext.Compiler->CompilerLog(FText::FromString(Msg), TypedComponentMesh, EMessageSeverity::Warning);
+			GenerationContext.Compiler->CompilerLog(FText::FromString(Msg), TypeComponentPassthroughMesh, EMessageSeverity::Warning);
 			return nullptr;
 		}
 
-		USkeletalMesh* SkeletalMesh = Cast<USkeletalMesh>(TypedComponentMesh->Mesh.TryLoad());
+		USkeletalMesh* SkeletalMesh = Cast<USkeletalMesh>(TypeComponentPassthroughMesh->Mesh.TryLoad());
 		if (!SkeletalMesh)
 		{
 			FString Msg = FString::Printf(TEXT("Only SkeletalMeshes are supported in this node, for now."));
-			GenerationContext.Compiler->CompilerLog(FText::FromString(Msg), TypedComponentMesh, EMessageSeverity::Warning);
+			GenerationContext.Compiler->CompilerLog(FText::FromString(Msg), TypeComponentPassthroughMesh, EMessageSeverity::Warning);
 			return nullptr;
 		}
 
@@ -78,14 +364,14 @@ mu::Ptr<mu::NodeComponent> GenerateMutableSourceComponent(const UEdGraphPin * Pi
 			constexpr bool bIsReference = true;
 			TSoftClassPtr<UAnimInstance> AnimInstance;
 			mu::Ptr<mu::Mesh> MutableMesh = GenerateMutableMesh(SkeletalMesh, AnimInstance, 0, 0, 0, 0, MeshUniqueTags, 
-																GenerationContext, TypedComponentMesh, nullptr, bIsReference);
+																GenerationContext, TypeComponentPassthroughMesh, nullptr, bIsReference);
 
 			MeshNode->SetValue(MutableMesh);
 		}
 
 		// Create the component node
 		mu::Ptr<mu::NodeComponentNew> ComponentNode = new mu::NodeComponentNew;
-		ComponentNode->Id = GenerationContext.NumMeshComponentsInRoot - 1 + GenerationContext.NumExplicitMeshComponents; // Last root component id + Num explicit components
+		ComponentNode->Id = GenerationContext.NumMeshComponentsInRoot - 1 + GenerationContext.NumPassthroughMeshComponents; // Last root component id + Num explicit components
 
 		// Create a LOD for each pass-through mesh LOD.
 		const FSkeletalMeshModel* Model = SkeletalMesh->GetImportedModel();
@@ -100,7 +386,7 @@ mu::Ptr<mu::NodeComponent> GenerateMutableSourceComponent(const UEdGraphPin * Pi
 			for (int32 SectionIndex = 0; SectionIndex < SectionCount; ++SectionIndex)
 			{
 				// Is there a pin in the unreal node for this section?
-				if (UEdGraphPin* InMaterialPin = TypedComponentMesh->GetMaterialPin(LODIndex,SectionIndex))
+				if (UEdGraphPin* InMaterialPin = TypeComponentPassthroughMesh->GetMaterialPin(LODIndex,SectionIndex))
 				{
 					if (UEdGraphPin* ConnectedMaterialPin = FollowInputPin(*InMaterialPin))
 					{
@@ -125,8 +411,7 @@ mu::Ptr<mu::NodeComponent> GenerateMutableSourceComponent(const UEdGraphPin * Pi
 		GenerationContext.CurrentMeshComponent = FName();
 		Result = ComponentNode;
 	}
-
-
+	
 	else if (const UCustomizableObjectNodeComponentSwitch* TypedNodeSwitch = Cast<UCustomizableObjectNodeComponentSwitch>(Node))
 	{
 		// Using a lambda so control flow is easier to manage.
diff --git a/Engine/Plugins/Experimental/Mutable/Source/CustomizableObjectEditor/Private/MuCOE/GenerateMutableSource/GenerateMutableSourceComponent.h b/Engine/Plugins/Experimental/Mutable/Source/CustomizableObjectEditor/Private/MuCOE/GenerateMutableSource/GenerateMutableSourceComponent.h
new file mode 100644
--- /dev/null
+++ b/Engine/Plugins/Experimental/Mutable/Source/CustomizableObjectEditor/Private/MuCOE/GenerateMutableSource/GenerateMutableSourceComponent.h
@@ -0,0 +1,17 @@
+﻿// Copyright Epic Games, Inc. All Rights Reserved.
+
+#pragma once
+
+#include "MuR/Ptr.h"
+
+class UEdGraphPin;
+struct FMutableGraphGenerationContext;
+
+namespace mu
+{
+	class NodeComponent;
+}
+
+
+mu::Ptr<mu::NodeComponent> GenerateMutableSourceComponent(const UEdGraphPin* Pin, FMutableGraphGenerationContext& GenerationContext);
+
diff --git a/Engine/Plugins/Experimental/Mutable/Source/CustomizableObjectEditor/Private/MuCOE/GenerateMutableSource/GenerateMutableSourceSurface.cpp b/Engine/Plugins/Experimental/Mutable/Source/CustomizableObjectEditor/Private/MuCOE/GenerateMutableSource/GenerateMutableSourceSurface.cpp
--- a/Engine/Plugins/Experimental/Mutable/Source/CustomizableObjectEditor/Private/MuCOE/GenerateMutableSource/GenerateMutableSourceSurface.cpp
+++ b/Engine/Plugins/Experimental/Mutable/Source/CustomizableObjectEditor/Private/MuCOE/GenerateMutableSource/GenerateMutableSourceSurface.cpp
@@ -234,13 +234,7 @@ mu::Ptr<mu::NodeSurface> GenerateMutableSourceSurface(const UEdGraphPin * Pin, F
 
 	if (UCustomizableObjectNodeMaterialBase* TypedNodeMat = Cast<UCustomizableObjectNodeMaterialBase>(Node))
 	{
-		bool bGeneratingImplicitComponent = GenerationContext.ComponentMeshOverride.get()!=nullptr;
-
-		bool bValidForCurrentComponent = bGeneratingImplicitComponent || (TypedNodeMat->GetMeshComponentName() == GenerationContext.CurrentMeshComponent);
-		if (!bValidForCurrentComponent)
-		{
-			return Result;
-		}
+		bool bGeneratingImplicitComponent = GenerationContext.ComponentMeshOverride.get() != nullptr;
 
 		const UEdGraphPin* ConnectedMaterialPin = FollowInputPin(*TypedNodeMat->GetMeshPin());
 		// Warn when texture connections are improperly used by connecting them directly to material inputs when no layout is used
diff --git a/Engine/Plugins/Experimental/Mutable/Source/CustomizableObjectEditor/Private/MuCOE/Nodes/CustomizableObjectNode.cpp b/Engine/Plugins/Experimental/Mutable/Source/CustomizableObjectEditor/Private/MuCOE/Nodes/CustomizableObjectNode.cpp
--- a/Engine/Plugins/Experimental/Mutable/Source/CustomizableObjectEditor/Private/MuCOE/Nodes/CustomizableObjectNode.cpp
+++ b/Engine/Plugins/Experimental/Mutable/Source/CustomizableObjectEditor/Private/MuCOE/Nodes/CustomizableObjectNode.cpp
@@ -11,6 +11,7 @@
 #include "MuCOE/CustomizableObjectPin.h"
 #include "MuCOE/EdGraphSchema_CustomizableObject.h"
 #include "MuCOE/GraphTraversal.h"
+#include "MuCOE/Nodes/CustomizableObjectNodeComponentMeshAddTo.h"
 #include "MuCOE/Nodes/CustomizableObjectNodeObject.h"
 #include "MuCOE/Nodes/CustomizableObjectNodeObjectGroup.h"
 #include "MuCOE/RemapPins/CustomizableObjectNodeRemapPinsByName.h"
@@ -195,7 +196,7 @@ void UCustomizableObjectNode::ReconstructNode()
 }
 
 
-void UCustomizableObjectNode::ReconstructNode(UCustomizableObjectNodeRemapPins* RemapPinsAction)
+void UCustomizableObjectNode::FixupReconstructPins(UCustomizableObjectNodeRemapPins* RemapPinsAction, TFunction<void(UCustomizableObjectNode*, UCustomizableObjectNodeRemapPins*)> FuncAllocateDefaultPins)
 {
 	Modify();
 
@@ -221,7 +222,7 @@ void UCustomizableObjectNode::ReconstructNode(UCustomizableObjectNodeRemapPins*
 	TArray<UEdGraphPin*> OldPins(Pins); // We can not empty Pins at this point since it will break all FEdGraphPinReference during the reconstruction.
 
 	// Recreate the new pins
-	AllocateDefaultPins(RemapPinsAction);
+	FuncAllocateDefaultPins(this, RemapPinsAction);
 	
 	// Try to remap orphan and non orphan pins.
 	TArray<UEdGraphPin*> NewPins;
@@ -401,6 +402,12 @@ void UCustomizableObjectNode::BackwardsCompatibleFixup(int32 CustomizableObjectC
 }
 
 
+void UCustomizableObjectNode::ReconstructNode(UCustomizableObjectNodeRemapPins* RemapPinsMode)
+{
+	FixupReconstructPins(RemapPinsMode, [](UCustomizableObjectNode* Node, UCustomizableObjectNodeRemapPins* Action){ Node->AllocateDefaultPins(Action); });
+}
+
+
 bool UCustomizableObjectNode::CanConnect( const UEdGraphPin* InOwnedInputPin, const UEdGraphPin* InOutputPin, bool& bOutIsOtherNodeBlocklisted, bool& bOutArePinsCompatible) const
 {
 	bOutIsOtherNodeBlocklisted = false;
@@ -572,9 +579,12 @@ int32 UCustomizableObjectNode::GetLOD() const
 			{
 				for (UEdGraphPin* LinkedPin : FollowOutputPinArray(*Pin))
 				{
-					if (UCustomizableObjectNodeObject* CurrentCustomizableObjectNode = Cast<UCustomizableObjectNodeObject>(LinkedPin->GetOwningNode()))
+					if (UCustomizableObjectNodeComponentMeshBase* ComponentMesh = Cast<UCustomizableObjectNodeComponentMeshBase>(LinkedPin->GetOwningNode()))
 					{
-						return CurrentCustomizableObjectNode->GetLOD(LinkedPin);
+						return ComponentMesh->LODPins.IndexOfByPredicate([&](const FEdGraphPinReference& LODPin)
+						{
+							return LinkedPin == LODPin.Get();
+						});
 					}
 					else
 					{
diff --git a/Engine/Plugins/Experimental/Mutable/Source/CustomizableObjectEditor/Private/MuCOE/Nodes/CustomizableObjectNodeComponent.cpp b/Engine/Plugins/Experimental/Mutable/Source/CustomizableObjectEditor/Private/MuCOE/Nodes/CustomizableObjectNodeComponent.cpp
new file mode 100644
--- /dev/null
+++ b/Engine/Plugins/Experimental/Mutable/Source/CustomizableObjectEditor/Private/MuCOE/Nodes/CustomizableObjectNodeComponent.cpp
@@ -0,0 +1,9 @@
+﻿// Copyright Epic Games, Inc. All Rights Reserved.
+
+#include "MuCOE/Nodes/CustomizableObjectNodeComponent.h"
+
+
+bool UCustomizableObjectNodeComponent::IsAffectedByLOD() const
+{ 
+	return false; 
+}
diff --git a/Engine/Plugins/Experimental/Mutable/Source/CustomizableObjectEditor/Private/MuCOE/Nodes/CustomizableObjectNodeComponentMesh.cpp b/Engine/Plugins/Experimental/Mutable/Source/CustomizableObjectEditor/Private/MuCOE/Nodes/CustomizableObjectNodeComponentMesh.cpp
--- a/Engine/Plugins/Experimental/Mutable/Source/CustomizableObjectEditor/Private/MuCOE/Nodes/CustomizableObjectNodeComponentMesh.cpp
+++ b/Engine/Plugins/Experimental/Mutable/Source/CustomizableObjectEditor/Private/MuCOE/Nodes/CustomizableObjectNodeComponentMesh.cpp
@@ -2,282 +2,46 @@
 
 #include "MuCOE/Nodes/CustomizableObjectNodeComponentMesh.h"
 
-#include "Engine/SkeletalMesh.h"
-#include "Modules/ModuleManager.h"
-#include "MuCO/CustomizableObjectCustomVersion.h"
-#include "MuCOE/CustomizableObjectEditor.h"
-#include "MuCOE/CustomizableObjectEditorStyle.h"
-#include "MuCOE/CustomizableObjectLayout.h"
 #include "MuCOE/EdGraphSchema_CustomizableObject.h"
-#include "MuCOE/GraphTraversal.h"
-#include "MuCOE/MutableUtils.h"
-#include "MuCOE/Nodes/CustomizableObjectNodeLayoutBlocks.h"
-#include "MuCOE/RemapPins/CustomizableObjectNodeRemapPinsByNameDefaultPin.h"
-#include "MuCOE/UnrealEditorPortabilityHelpers.h"
-
-class UCustomizableObjectNodeRemapPinsByName;
-class UObject;
 
 
 #define LOCTEXT_NAMESPACE "CustomizableObjectEditor"
 
-const FName UCustomizableObjectNodeComponentMesh::OutputPinName(TEXT("Component"));
-
-
-void UCustomizableObjectNodeComponentMeshPinDataSection::Init(int32 InLODIndex, int32 InSectionIndex)
-{
-	LODIndex = InLODIndex;
-	SectionIndex = InSectionIndex;
-}
-
-
-int32 UCustomizableObjectNodeComponentMeshPinDataSection::GetLODIndex() const
-{
-	return LODIndex;
-}
-
-
-int32 UCustomizableObjectNodeComponentMeshPinDataSection::GetSectionIndex() const
-{
-	return SectionIndex;
-}
-
-
-bool UCustomizableObjectNodeComponentMeshPinDataSection::Equals(const UCustomizableObjectNodePinData& Other) const
-{
-	if (GetClass() != Other.GetClass())
-	{
-		return false;
-	}
-
-	const UCustomizableObjectNodeComponentMeshPinDataSection& OtherTyped = static_cast<const UCustomizableObjectNodeComponentMeshPinDataSection&>(Other);
-	if (LODIndex != OtherTyped.LODIndex ||
-		SectionIndex != OtherTyped.SectionIndex)
-	{
-		return false;
-	}
-
-	return Super::Equals(Other);
-}
-
 
 void UCustomizableObjectNodeComponentMesh::PostEditChangeProperty(FPropertyChangedEvent& PropertyChangedEvent)
 {
 	Super::PostEditChangeProperty(PropertyChangedEvent);
 
-	FProperty* PropertyThatChanged = PropertyChangedEvent.Property;
-	if (PropertyThatChanged && PropertyThatChanged->GetName() == TEXT("Mesh"))
-	{
-		ReconstructNode();
-	}
-}
-
-
-void UCustomizableObjectNodeComponentMesh::AllocateDefaultPins(UCustomizableObjectNodeRemapPins* RemapPins)
-{
-	const UEdGraphSchema_CustomizableObject* Schema = GetDefault<UEdGraphSchema_CustomizableObject>();
-	if (!Schema)
-	{
-		return;
-	}
-
-	UEdGraphPin* OutputPin = CustomCreatePin(EGPD_Output, Schema->PC_Component, OutputPinName);
-
-	if (!Mesh.IsValid())
+	const FProperty* PropertyThatChanged = PropertyChangedEvent.Property;
+	if (!PropertyThatChanged)
 	{
 		return;
 	}
 
-	// Support SkeletalMeshes only, for now.
-	USkeletalMesh* SkeletalMesh = Cast<USkeletalMesh>( Mesh.TryLoad() );
-	if (!SkeletalMesh)
+	if (PropertyThatChanged->GetFName() == GET_MEMBER_NAME_CHECKED(UCustomizableObjectNodeComponentMesh, NumLODs))
 	{
-		return;
-	}
-	
-	if (const FSkeletalMeshModel* ImportedModel = SkeletalMesh->GetImportedModel())
-	{
-		const int32 NumLODs = SkeletalMesh->GetLODNum();
-		for (int32 LODIndex = 0; LODIndex < NumLODs; ++LODIndex)
-		{
-			const int32 NumSections = ImportedModel->LODModels[LODIndex].Sections.Num();
-			for (int32 SectionIndex = 0; SectionIndex < NumSections; ++SectionIndex)
-			{
-				// Ignore disabled sections.
-				const bool bIsSectionDisabled = ImportedModel ->LODModels[LODIndex].Sections[SectionIndex].bDisabled;
-				if (bIsSectionDisabled)
-				{
-					continue;
-				}
-
-				UMaterialInterface* MaterialInterface = GetMaterialInterfaceFor(LODIndex, SectionIndex, ImportedModel);
-				
-				// Material for the section
-				{
-					UCustomizableObjectNodeComponentMeshPinDataMaterial* PinData = NewObject<UCustomizableObjectNodeComponentMeshPinDataMaterial>(this);
-					PinData->Init(LODIndex, SectionIndex);
-					
-					FString SectionFriendlyName = MaterialInterface ? MaterialInterface->GetName() : FString::Printf(TEXT("Section %i"), SectionIndex);
-					FString MaterialPinName = FString::Printf(TEXT("LOD %i - Section %i - Material"), LODIndex, SectionIndex);
-
-					UEdGraphPin* Pin = CustomCreatePin(EGPD_Input, Schema->PC_Material, FName(*MaterialPinName), PinData);
-					Pin->PinFriendlyName = FText::FromString(FString::Printf(TEXT("LOD %i - Section %i - %s"), LODIndex, SectionIndex, *SectionFriendlyName));
-					Pin->PinToolTip = MaterialPinName;
-				}
-				
-			}
-		}
+		LODReductionSettings.SetNum(NumLODs);
+		
+		ReconstructNode();
 	}
 }
 
 
 FText UCustomizableObjectNodeComponentMesh::GetNodeTitle(ENodeTitleType::Type TitleType) const
 {
-	USkeletalMesh* SkeletalMesh = Cast<USkeletalMesh>(Mesh.TryLoad());
-	if (SkeletalMesh)
-	{
-		FFormatNamedArguments Args;
-		Args.Add(TEXT("MeshName"), FText::FromString(SkeletalMesh->GetName()));
-
-		return FText::Format(LOCTEXT("ComponentMesh_Title", "{MeshName}\nMesh Component"), Args);
-	}
-	else
+	if (TitleType == ENodeTitleType::ListView)
 	{
 		return LOCTEXT("Component_Mesh", "Mesh Component");
 	}
-}
-
-
-FLinearColor UCustomizableObjectNodeComponentMesh::GetNodeTitleColor() const
-{
-	const UEdGraphSchema_CustomizableObject* Schema = GetDefault<UEdGraphSchema_CustomizableObject>();
-	return Schema->GetPinTypeColor(Schema->PC_Component);
-}
-
-
-void UCustomizableObjectNodeComponentMesh::GetPinSection(const UEdGraphPin& Pin, int32& OutLODIndex, int32& OutSectionIndex) const
-{
-	if (const UCustomizableObjectNodeComponentMeshPinDataSection* PinData = Cast<UCustomizableObjectNodeComponentMeshPinDataSection>(GetPinData(Pin)))
-	{
-		OutLODIndex = PinData->GetLODIndex();
-		OutSectionIndex = PinData->GetSectionIndex();
-		return;
-	}
-
-	OutLODIndex = -1;
-	OutSectionIndex = -1;
-}
-
-
-UMaterialInterface* UCustomizableObjectNodeComponentMesh::GetMaterialFor(const UEdGraphPin* Pin) const
-{
-	USkeletalMesh* SkeletalMesh = Cast<USkeletalMesh>(Mesh.TryLoad());
-	if (SkeletalMesh)
-	{
-		if (FSkeletalMaterial* SkeletalMaterial = GetSkeletalMaterialFor(*Pin))
-		{
-			return SkeletalMaterial->MaterialInterface;
-		}
-	}
-
-	return nullptr;
-}
-
-
-FSkeletalMaterial* UCustomizableObjectNodeComponentMesh::GetSkeletalMaterialFor(const UEdGraphPin& Pin) const
-{
-	int32 LODIndex;
-	int32 SectionIndex;
-	GetPinSection(Pin, LODIndex, SectionIndex);
-
-	return GetSkeletalMaterialFor(LODIndex, SectionIndex);
-}
-
-
-FText UCustomizableObjectNodeComponentMesh::GetTooltipText() const
-{
-	return LOCTEXT("Component_Mesh_Tooltip", "Define a new object component based on a Skeletal Mesh.");
-}
-
-
-UMaterialInterface* UCustomizableObjectNodeComponentMesh::GetMaterialInterfaceFor(const int32 LODIndex, const int32 SectionIndex, const FSkeletalMeshModel* ImportedModel) const
-{
-	if (FSkeletalMaterial* SkeletalMaterial = GetSkeletalMaterialFor(LODIndex, SectionIndex, ImportedModel))
-	{
-		return SkeletalMaterial->MaterialInterface;
-	}
-
-	return nullptr;
-}
-
-
-FSkeletalMaterial* UCustomizableObjectNodeComponentMesh::GetSkeletalMaterialFor(const int32 LODIndex, const int32 SectionIndex, const FSkeletalMeshModel* ImportedModel) const
-{
-	USkeletalMesh* SkeletalMesh = Cast<USkeletalMesh>(Mesh.TryLoad());
-	if (!SkeletalMesh)
-	{
-		return nullptr;
-	}
-
-	// We assume that LODIndex and MaterialIndex are valid for the imported model
-	int32 SkeletalMeshMaterialIndex = INDEX_NONE;
-	
-	// Check if we have lod info map to get the correct material index
-	if (const FSkeletalMeshLODInfo* LodInfo = SkeletalMesh->GetLODInfo(LODIndex))
-	{
-		if (LodInfo->LODMaterialMap.IsValidIndex(SectionIndex))
-		{
-			SkeletalMeshMaterialIndex = LodInfo->LODMaterialMap[SectionIndex];
-		}
-	}
-
-	// Only deduce index when the explicit mapping is not found or there is no remap
-	if (SkeletalMeshMaterialIndex == INDEX_NONE)
-	{
-	if (ImportedModel && ImportedModel->LODModels.IsValidIndex(LODIndex) && ImportedModel->LODModels[LODIndex].Sections.IsValidIndex(SectionIndex))
-	{
-		SkeletalMeshMaterialIndex = ImportedModel->LODModels[LODIndex].Sections[SectionIndex].MaterialIndex;
-	}
 	else
 	{
-		FSkeletalMeshModel* AuxImportedModel = SkeletalMesh->GetImportedModel();
-
-		if (AuxImportedModel)
-		{
-			if (AuxImportedModel->LODModels.IsValidIndex(LODIndex) && AuxImportedModel->LODModels[LODIndex].Sections.IsValidIndex(SectionIndex))
-			{
-				SkeletalMeshMaterialIndex = AuxImportedModel->LODModels[LODIndex].Sections[SectionIndex].MaterialIndex;
-			}
-		}
-	}
-	}
-	
-	if (SkeletalMesh->GetMaterials().IsValidIndex(SkeletalMeshMaterialIndex))
-	{
-		return &SkeletalMesh->GetMaterials()[SkeletalMeshMaterialIndex];
-	}
-	
-	return nullptr;
-}
-
-
+		FFormatNamedArguments Args;
+		Args.Add(TEXT("ComponentName"), FText::FromName(ComponentName));
 
-UEdGraphPin* UCustomizableObjectNodeComponentMesh::GetMaterialPin(const int32 LODIndex, const int32 SectionIndex) const
-{
-	for (UEdGraphPin* Pin : GetAllNonOrphanPins())
-	{
-		if (const UCustomizableObjectNodeComponentMeshPinDataSection* PinData = Cast<UCustomizableObjectNodeComponentMeshPinDataSection>(GetPinData(*Pin)))
-		{
-			if (PinData->GetLODIndex() == LODIndex &&
-				PinData->GetSectionIndex() == SectionIndex)
-			{
-				return Pin;
-			}
-		}
+		return FText::Format(LOCTEXT("ComponentMesh_Title", "{ComponentName}\nMesh Component"), Args);
 	}
-
-	return nullptr;
 }
 
 
 #undef LOCTEXT_NAMESPACE
+
diff --git a/Engine/Plugins/Experimental/Mutable/Source/CustomizableObjectEditor/Private/MuCOE/Nodes/CustomizableObjectNodeComponentMeshAddTo.cpp b/Engine/Plugins/Experimental/Mutable/Source/CustomizableObjectEditor/Private/MuCOE/Nodes/CustomizableObjectNodeComponentMeshAddTo.cpp
new file mode 100644
--- /dev/null
+++ b/Engine/Plugins/Experimental/Mutable/Source/CustomizableObjectEditor/Private/MuCOE/Nodes/CustomizableObjectNodeComponentMeshAddTo.cpp
@@ -0,0 +1,46 @@
+﻿// Copyright Epic Games, Inc. All Rights Reserved.
+
+#include "MuCOE/Nodes/CustomizableObjectNodeComponentMeshAddTo.h"
+
+#include "MuCOE/EdGraphSchema_CustomizableObject.h"
+#include "MuCOE/Nodes/CustomizableObjectNodeComponentMesh.h"
+
+
+#define LOCTEXT_NAMESPACE "CustomizableObjectEditor"
+
+
+void UCustomizableObjectNodeComponentMeshAddTo::PostEditChangeProperty(FPropertyChangedEvent& PropertyChangedEvent)
+{
+	Super::PostEditChangeProperty(PropertyChangedEvent);
+
+	const FProperty* PropertyThatChanged = PropertyChangedEvent.Property;
+	if (!PropertyThatChanged)
+	{
+		return;
+	}
+
+	if (PropertyThatChanged->GetFName() == GET_MEMBER_NAME_CHECKED(UCustomizableObjectNodeComponentMesh, NumLODs))
+	{
+		ReconstructNode();
+	}
+}
+
+
+FText UCustomizableObjectNodeComponentMeshAddTo::GetNodeTitle(ENodeTitleType::Type TitleType) const
+{
+	if (TitleType == ENodeTitleType::ListView)
+	{
+		return LOCTEXT("ComponentMeshAdd", "Add To Mesh Component");
+	}
+	else
+	{
+		FFormatNamedArguments Args;
+		Args.Add(TEXT("ComponentName"), FText::FromName(ParentComponentName));
+
+		return FText::Format(LOCTEXT("ComponentMeshAdd_Title", "{ComponentName}\nAdd To Mesh Component"), Args);
+	}
+}
+
+
+#undef LOCTEXT_NAMESPACE
+
diff --git a/Engine/Plugins/Experimental/Mutable/Source/CustomizableObjectEditor/Private/MuCOE/Nodes/CustomizableObjectNodeComponentMeshBase.cpp b/Engine/Plugins/Experimental/Mutable/Source/CustomizableObjectEditor/Private/MuCOE/Nodes/CustomizableObjectNodeComponentMeshBase.cpp
new file mode 100644
--- /dev/null
+++ b/Engine/Plugins/Experimental/Mutable/Source/CustomizableObjectEditor/Private/MuCOE/Nodes/CustomizableObjectNodeComponentMeshBase.cpp
@@ -0,0 +1,32 @@
+﻿// Copyright Epic Games, Inc. All Rights Reserved.
+
+#include "MuCOE/Nodes/CustomizableObjectNodeComponentMeshBase.h"
+
+#include "MuCOE/EdGraphSchema_CustomizableObject.h"
+
+
+void UCustomizableObjectNodeComponentMeshBase::AllocateDefaultPins(UCustomizableObjectNodeRemapPins* RemapPins)
+{
+	Super::AllocateDefaultPins(RemapPins);
+
+	const UEdGraphSchema_CustomizableObject* Schema = GetDefault<UEdGraphSchema_CustomizableObject>();
+
+	LODPins.Empty(NumLODs);
+	for (int32 LODIndex = 0; LODIndex < NumLODs; ++LODIndex)
+	{
+		FString LODName = FString::Printf(TEXT("LOD %d"), LODIndex);
+
+		UEdGraphPin* Pin = CustomCreatePin(EGPD_Input, Schema->PC_Material, FName(*LODName), true);
+		LODPins.Add(Pin);
+	}
+	
+	OutputPin = CustomCreatePin(EGPD_Output, Schema->PC_Component, TEXT("Component"));
+}
+
+
+FLinearColor UCustomizableObjectNodeComponentMeshBase::GetNodeTitleColor() const
+{
+	const UEdGraphSchema_CustomizableObject* Schema = GetDefault<UEdGraphSchema_CustomizableObject>();
+	return Schema->GetPinTypeColor(Schema->PC_Component);
+}
+
diff --git a/Engine/Plugins/Experimental/Mutable/Source/CustomizableObjectEditor/Private/MuCOE/Nodes/CustomizableObjectNodeComponentMeshDetails.h b/Engine/Plugins/Experimental/Mutable/Source/CustomizableObjectEditor/Private/MuCOE/Nodes/CustomizableObjectNodeComponentMeshDetails.h
new file mode 100644
--- /dev/null
+++ b/Engine/Plugins/Experimental/Mutable/Source/CustomizableObjectEditor/Private/MuCOE/Nodes/CustomizableObjectNodeComponentMeshDetails.h
@@ -0,0 +1,37 @@
+﻿// Copyright Epic Games, Inc. All Rights Reserved.
+
+#pragma once
+
+#include "CustomizableObjectNodeDetails.h"
+#include "Templates/SharedPointer.h"
+#include "UObject/WeakObjectPtr.h"
+
+class UCustomizableObjectNodeComponentMesh;
+class SWidget;
+
+
+/** Copy Material node details panel. Hides all properties from the inheret Material node. */
+class FCustomizableObjectNodeComponentMeshDetails : public FCustomizableObjectNodeDetails
+{
+public:
+	static TSharedRef<IDetailCustomization> MakeInstance();
+
+	// ILayoutDetails interface
+	virtual void CustomizeDetails(IDetailLayoutBuilder& DetailBuilder) override;
+	virtual void CustomizeDetails(const TSharedPtr<IDetailLayoutBuilder>& DetailBuilder) override;
+
+	// Own interface
+	
+	TSharedRef<SWidget> OnGenerateLODMenuForPicker();
+
+	FText GetCurrentLODName() const;
+
+	TSharedRef<SWidget>OnGenerateLODComboBoxForPicker();
+
+	void OnSelectedLODChanged(int32 NewLODIndex);
+
+private:
+	TWeakObjectPtr<UCustomizableObjectNodeComponentMesh> NodeComponentMesh;
+	
+	TWeakPtr<IDetailLayoutBuilder> WeakDetailsBuilder;
+};
\ No newline at end of file
diff --git a/Engine/Plugins/Experimental/Mutable/Source/CustomizableObjectEditor/Private/MuCOE/Nodes/CustomizableObjectNodeComponentPassthroughMesh.cpp b/Engine/Plugins/Experimental/Mutable/Source/CustomizableObjectEditor/Private/MuCOE/Nodes/CustomizableObjectNodeComponentPassthroughMesh.cpp
new file mode 100644
--- /dev/null
+++ b/Engine/Plugins/Experimental/Mutable/Source/CustomizableObjectEditor/Private/MuCOE/Nodes/CustomizableObjectNodeComponentPassthroughMesh.cpp
@@ -0,0 +1,283 @@
+﻿// Copyright Epic Games, Inc. All Rights Reserved.
+
+#include "MuCOE/Nodes/CustomizableObjectNodeComponentPassthroughMesh.h"
+
+#include "Engine/SkeletalMesh.h"
+#include "Modules/ModuleManager.h"
+#include "MuCO/CustomizableObjectCustomVersion.h"
+#include "MuCOE/CustomizableObjectEditor.h"
+#include "MuCOE/CustomizableObjectEditorStyle.h"
+#include "MuCOE/CustomizableObjectLayout.h"
+#include "MuCOE/EdGraphSchema_CustomizableObject.h"
+#include "MuCOE/GraphTraversal.h"
+#include "MuCOE/MutableUtils.h"
+#include "MuCOE/Nodes/CustomizableObjectNodeLayoutBlocks.h"
+#include "MuCOE/RemapPins/CustomizableObjectNodeRemapPinsByNameDefaultPin.h"
+#include "MuCOE/UnrealEditorPortabilityHelpers.h"
+
+class UCustomizableObjectNodeRemapPinsByName;
+class UObject;
+
+
+#define LOCTEXT_NAMESPACE "CustomizableObjectEditor"
+
+const FName UCustomizableObjectNodeComponentPassthroughMesh::OutputPinName(TEXT("Component"));
+
+
+void UCustomizableObjectNodeComponentMeshPinDataSection::Init(int32 InLODIndex, int32 InSectionIndex)
+{
+	LODIndex = InLODIndex;
+	SectionIndex = InSectionIndex;
+}
+
+
+int32 UCustomizableObjectNodeComponentMeshPinDataSection::GetLODIndex() const
+{
+	return LODIndex;
+}
+
+
+int32 UCustomizableObjectNodeComponentMeshPinDataSection::GetSectionIndex() const
+{
+	return SectionIndex;
+}
+
+
+bool UCustomizableObjectNodeComponentMeshPinDataSection::Equals(const UCustomizableObjectNodePinData& Other) const
+{
+	if (GetClass() != Other.GetClass())
+	{
+		return false;
+	}
+
+	const UCustomizableObjectNodeComponentMeshPinDataSection& OtherTyped = static_cast<const UCustomizableObjectNodeComponentMeshPinDataSection&>(Other);
+	if (LODIndex != OtherTyped.LODIndex ||
+		SectionIndex != OtherTyped.SectionIndex)
+	{
+		return false;
+	}
+
+	return Super::Equals(Other);
+}
+
+
+void UCustomizableObjectNodeComponentPassthroughMesh::PostEditChangeProperty(FPropertyChangedEvent& PropertyChangedEvent)
+{
+	Super::PostEditChangeProperty(PropertyChangedEvent);
+
+	FProperty* PropertyThatChanged = PropertyChangedEvent.Property;
+	if (PropertyThatChanged && PropertyThatChanged->GetName() == TEXT("Mesh"))
+	{
+		ReconstructNode();
+	}
+}
+
+
+void UCustomizableObjectNodeComponentPassthroughMesh::AllocateDefaultPins(UCustomizableObjectNodeRemapPins* RemapPins)
+{
+	const UEdGraphSchema_CustomizableObject* Schema = GetDefault<UEdGraphSchema_CustomizableObject>();
+	if (!Schema)
+	{
+		return;
+	}
+
+	UEdGraphPin* OutputPin = CustomCreatePin(EGPD_Output, Schema->PC_Component, OutputPinName);
+
+	if (!Mesh.IsValid())
+	{
+		return;
+	}
+
+	// Support SkeletalMeshes only, for now.
+	USkeletalMesh* SkeletalMesh = Cast<USkeletalMesh>( Mesh.TryLoad() );
+	if (!SkeletalMesh)
+	{
+		return;
+	}
+	
+	if (const FSkeletalMeshModel* ImportedModel = SkeletalMesh->GetImportedModel())
+	{
+		const int32 NumLODs = SkeletalMesh->GetLODNum();
+		for (int32 LODIndex = 0; LODIndex < NumLODs; ++LODIndex)
+		{
+			const int32 NumSections = ImportedModel->LODModels[LODIndex].Sections.Num();
+			for (int32 SectionIndex = 0; SectionIndex < NumSections; ++SectionIndex)
+			{
+				// Ignore disabled sections.
+				const bool bIsSectionDisabled = ImportedModel ->LODModels[LODIndex].Sections[SectionIndex].bDisabled;
+				if (bIsSectionDisabled)
+				{
+					continue;
+				}
+
+				UMaterialInterface* MaterialInterface = GetMaterialInterfaceFor(LODIndex, SectionIndex, ImportedModel);
+				
+				// Material for the section
+				{
+					UCustomizableObjectNodeComponentMeshPinDataMaterial* PinData = NewObject<UCustomizableObjectNodeComponentMeshPinDataMaterial>(this);
+					PinData->Init(LODIndex, SectionIndex);
+					
+					FString SectionFriendlyName = MaterialInterface ? MaterialInterface->GetName() : FString::Printf(TEXT("Section %i"), SectionIndex);
+					FString MaterialPinName = FString::Printf(TEXT("LOD %i - Section %i - Material"), LODIndex, SectionIndex);
+
+					UEdGraphPin* Pin = CustomCreatePin(EGPD_Input, Schema->PC_Material, FName(*MaterialPinName), PinData);
+					Pin->PinFriendlyName = FText::FromString(FString::Printf(TEXT("LOD %i - Section %i - %s"), LODIndex, SectionIndex, *SectionFriendlyName));
+					Pin->PinToolTip = MaterialPinName;
+				}
+				
+			}
+		}
+	}
+}
+
+
+FText UCustomizableObjectNodeComponentPassthroughMesh::GetNodeTitle(ENodeTitleType::Type TitleType) const
+{
+	USkeletalMesh* SkeletalMesh = Cast<USkeletalMesh>(Mesh.TryLoad());
+	if (SkeletalMesh)
+	{
+		FFormatNamedArguments Args;
+		Args.Add(TEXT("MeshName"), FText::FromString(SkeletalMesh->GetName()));
+
+		return FText::Format(LOCTEXT("ComponentMesh_Title", "{MeshName}\n Passthrough Mesh Component"), Args);
+	}
+	else
+	{
+		return LOCTEXT("Component_Mesh", "Passthrough Mesh Component");
+	}
+}
+
+
+FLinearColor UCustomizableObjectNodeComponentPassthroughMesh::GetNodeTitleColor() const
+{
+	const UEdGraphSchema_CustomizableObject* Schema = GetDefault<UEdGraphSchema_CustomizableObject>();
+	return Schema->GetPinTypeColor(Schema->PC_Component);
+}
+
+
+void UCustomizableObjectNodeComponentPassthroughMesh::GetPinSection(const UEdGraphPin& Pin, int32& OutLODIndex, int32& OutSectionIndex) const
+{
+	if (const UCustomizableObjectNodeComponentMeshPinDataSection* PinData = Cast<UCustomizableObjectNodeComponentMeshPinDataSection>(GetPinData(Pin)))
+	{
+		OutLODIndex = PinData->GetLODIndex();
+		OutSectionIndex = PinData->GetSectionIndex();
+		return;
+	}
+
+	OutLODIndex = -1;
+	OutSectionIndex = -1;
+}
+
+
+UMaterialInterface* UCustomizableObjectNodeComponentPassthroughMesh::GetMaterialFor(const UEdGraphPin* Pin) const
+{
+	USkeletalMesh* SkeletalMesh = Cast<USkeletalMesh>(Mesh.TryLoad());
+	if (SkeletalMesh)
+	{
+		if (FSkeletalMaterial* SkeletalMaterial = GetSkeletalMaterialFor(*Pin))
+		{
+			return SkeletalMaterial->MaterialInterface;
+		}
+	}
+
+	return nullptr;
+}
+
+
+FSkeletalMaterial* UCustomizableObjectNodeComponentPassthroughMesh::GetSkeletalMaterialFor(const UEdGraphPin& Pin) const
+{
+	int32 LODIndex;
+	int32 SectionIndex;
+	GetPinSection(Pin, LODIndex, SectionIndex);
+
+	return GetSkeletalMaterialFor(LODIndex, SectionIndex);
+}
+
+
+FText UCustomizableObjectNodeComponentPassthroughMesh::GetTooltipText() const
+{
+	return LOCTEXT("Component_Mesh_Tooltip", "Define a new object component based on a Skeletal Mesh.");
+}
+
+
+UMaterialInterface* UCustomizableObjectNodeComponentPassthroughMesh::GetMaterialInterfaceFor(const int32 LODIndex, const int32 SectionIndex, const FSkeletalMeshModel* ImportedModel) const
+{
+	if (FSkeletalMaterial* SkeletalMaterial = GetSkeletalMaterialFor(LODIndex, SectionIndex, ImportedModel))
+	{
+		return SkeletalMaterial->MaterialInterface;
+	}
+
+	return nullptr;
+}
+
+
+FSkeletalMaterial* UCustomizableObjectNodeComponentPassthroughMesh::GetSkeletalMaterialFor(const int32 LODIndex, const int32 SectionIndex, const FSkeletalMeshModel* ImportedModel) const
+{
+	USkeletalMesh* SkeletalMesh = Cast<USkeletalMesh>(Mesh.TryLoad());
+	if (!SkeletalMesh)
+	{
+		return nullptr;
+	}
+
+	// We assume that LODIndex and MaterialIndex are valid for the imported model
+	int32 SkeletalMeshMaterialIndex = INDEX_NONE;
+	
+	// Check if we have lod info map to get the correct material index
+	if (const FSkeletalMeshLODInfo* LodInfo = SkeletalMesh->GetLODInfo(LODIndex))
+	{
+		if (LodInfo->LODMaterialMap.IsValidIndex(SectionIndex))
+		{
+			SkeletalMeshMaterialIndex = LodInfo->LODMaterialMap[SectionIndex];
+		}
+	}
+
+	// Only deduce index when the explicit mapping is not found or there is no remap
+	if (SkeletalMeshMaterialIndex == INDEX_NONE)
+	{
+	if (ImportedModel && ImportedModel->LODModels.IsValidIndex(LODIndex) && ImportedModel->LODModels[LODIndex].Sections.IsValidIndex(SectionIndex))
+	{
+		SkeletalMeshMaterialIndex = ImportedModel->LODModels[LODIndex].Sections[SectionIndex].MaterialIndex;
+	}
+	else
+	{
+		FSkeletalMeshModel* AuxImportedModel = SkeletalMesh->GetImportedModel();
+
+		if (AuxImportedModel)
+		{
+			if (AuxImportedModel->LODModels.IsValidIndex(LODIndex) && AuxImportedModel->LODModels[LODIndex].Sections.IsValidIndex(SectionIndex))
+			{
+				SkeletalMeshMaterialIndex = AuxImportedModel->LODModels[LODIndex].Sections[SectionIndex].MaterialIndex;
+			}
+		}
+	}
+	}
+	
+	if (SkeletalMesh->GetMaterials().IsValidIndex(SkeletalMeshMaterialIndex))
+	{
+		return &SkeletalMesh->GetMaterials()[SkeletalMeshMaterialIndex];
+	}
+	
+	return nullptr;
+}
+
+
+
+UEdGraphPin* UCustomizableObjectNodeComponentPassthroughMesh::GetMaterialPin(const int32 LODIndex, const int32 SectionIndex) const
+{
+	for (UEdGraphPin* Pin : GetAllNonOrphanPins())
+	{
+		if (const UCustomizableObjectNodeComponentMeshPinDataSection* PinData = Cast<UCustomizableObjectNodeComponentMeshPinDataSection>(GetPinData(*Pin)))
+		{
+			if (PinData->GetLODIndex() == LODIndex &&
+				PinData->GetSectionIndex() == SectionIndex)
+			{
+				return Pin;
+			}
+		}
+	}
+
+	return nullptr;
+}
+
+
+#undef LOCTEXT_NAMESPACE
diff --git a/Engine/Plugins/Experimental/Mutable/Source/CustomizableObjectEditor/Private/MuCOE/Nodes/CustomizableObjectNodeCopyMaterial.cpp b/Engine/Plugins/Experimental/Mutable/Source/CustomizableObjectEditor/Private/MuCOE/Nodes/CustomizableObjectNodeCopyMaterial.cpp
--- a/Engine/Plugins/Experimental/Mutable/Source/CustomizableObjectEditor/Private/MuCOE/Nodes/CustomizableObjectNodeCopyMaterial.cpp
+++ b/Engine/Plugins/Experimental/Mutable/Source/CustomizableObjectEditor/Private/MuCOE/Nodes/CustomizableObjectNodeCopyMaterial.cpp
@@ -218,19 +218,6 @@ bool UCustomizableObjectNodeCopyMaterial::IsReuseMaterialBetweenLODs() const
 }
 
 
-FName UCustomizableObjectNodeCopyMaterial::GetMeshComponentName() const
-{
-	if (UCustomizableObjectNodeMaterial* NodeMaterial = GetMaterialNode())
-	{
-		return NodeMaterial->GetMeshComponentName();
-	}
-	else
-	{
-		return FName();
-	}
-}
-
-
 TArray<FString>* UCustomizableObjectNodeCopyMaterial::GetEnableTags()
 {
 	if (UCustomizableObjectNodeMaterial* NodeMaterial = GetMaterialNode())
diff --git a/Engine/Plugins/Experimental/Mutable/Source/CustomizableObjectEditor/Private/MuCOE/Nodes/CustomizableObjectNodeMaterial.cpp b/Engine/Plugins/Experimental/Mutable/Source/CustomizableObjectEditor/Private/MuCOE/Nodes/CustomizableObjectNodeMaterial.cpp
--- a/Engine/Plugins/Experimental/Mutable/Source/CustomizableObjectEditor/Private/MuCOE/Nodes/CustomizableObjectNodeMaterial.cpp
+++ b/Engine/Plugins/Experimental/Mutable/Source/CustomizableObjectEditor/Private/MuCOE/Nodes/CustomizableObjectNodeMaterial.cpp
@@ -623,7 +623,7 @@ void UCustomizableObjectNodeMaterial::BackwardsCompatibleFixup(int32 Customizabl
 
 	if (CustomizableObjectCustomVersion == FCustomizableObjectCustomVersion::NewComponentOptions)
 	{
-		MeshComponentName = FName(FString::FromInt(MeshComponentIndex_DEPRECATED));
+		MeshComponentName_DEPRECATED = FName(FString::FromInt(MeshComponentIndex_DEPRECATED));
 	}
 
 	if (CustomizableObjectCustomVersion == FCustomizableObjectCustomVersion::NodeMaterialTypedImagePins)
@@ -1163,27 +1163,6 @@ bool UCustomizableObjectNodeMaterial::CustomRemovePin(UEdGraphPin& Pin)
 }
 
 
-void UCustomizableObjectNodeMaterial::ReconstructNode(UCustomizableObjectNodeRemapPins* RemapPinsMode)
-{
-	Super::ReconstructNode(RemapPinsMode);
-
-	// When a material node is created, the first component is set as its mesh component
-	if (MeshComponentName.IsNone())
-	{
-		if (UCustomizableObject* CustomizableObject = Cast<UCustomizableObject>(GetOutermostObject()))
-		{
-			if (UCustomizableObject* RootObject = GetRootObject(CustomizableObject))
-			{
-				if (RootObject->GetPrivate()->MutableMeshComponents.Num())
-				{
-					MeshComponentName = RootObject->GetPrivate()->MutableMeshComponents[0].Name;
-				}
-			}
-		}
-	}
-}
-
-
 void UCustomizableObjectNodeMaterial::SetMaterial(UMaterialInterface* InMaterial)
 {
 	Material = InMaterial;
@@ -1202,12 +1181,6 @@ bool UCustomizableObjectNodeMaterial::IsReuseMaterialBetweenLODs() const
 }
 
 
-FName UCustomizableObjectNodeMaterial::GetMeshComponentName() const
-{
-	return MeshComponentName;
-}
-
-
 TArray<FString>* UCustomizableObjectNodeMaterial::GetEnableTags()
 {
 	return &Tags;
@@ -1507,10 +1480,4 @@ void UCustomizableObjectNodeMaterial::BreakExistingConnectionsPostConnection(UEd
 }
 
 
-void UCustomizableObjectNodeMaterial::SetComponentName(const FName& Name)
-{
-	MeshComponentName = Name;
-}
-
-
 #undef LOCTEXT_NAMESPACE
diff --git a/Engine/Plugins/Experimental/Mutable/Source/CustomizableObjectEditor/Private/MuCOE/Nodes/CustomizableObjectNodeMaterialDetails.cpp b/Engine/Plugins/Experimental/Mutable/Source/CustomizableObjectEditor/Private/MuCOE/Nodes/CustomizableObjectNodeMaterialDetails.cpp
deleted file mode 100644
diff --git a/Engine/Plugins/Experimental/Mutable/Source/CustomizableObjectEditor/Private/MuCOE/Nodes/CustomizableObjectNodeMaterialDetails.h b/Engine/Plugins/Experimental/Mutable/Source/CustomizableObjectEditor/Private/MuCOE/Nodes/CustomizableObjectNodeMaterialDetails.h
deleted file mode 100644
diff --git a/Engine/Plugins/Experimental/Mutable/Source/CustomizableObjectEditor/Private/MuCOE/Nodes/CustomizableObjectNodeObject.cpp b/Engine/Plugins/Experimental/Mutable/Source/CustomizableObjectEditor/Private/MuCOE/Nodes/CustomizableObjectNodeObject.cpp
--- a/Engine/Plugins/Experimental/Mutable/Source/CustomizableObjectEditor/Private/MuCOE/Nodes/CustomizableObjectNodeObject.cpp
+++ b/Engine/Plugins/Experimental/Mutable/Source/CustomizableObjectEditor/Private/MuCOE/Nodes/CustomizableObjectNodeObject.cpp
@@ -11,29 +11,33 @@
 #include "MuCOE/GraphTraversal.h"
 #include "MuCOE/ICustomizableObjectEditor.h"
 #include "MuCOE/Nodes/CustomizableObjectNodeMaterial.h"
+#include "MuCOE/Nodes/CustomizableObjectNodeComponentMesh.h"
 #include "MuCOE/Nodes/CustomizableObjectNodeMaterialVariation.h"
 #include "Misc/UObjectToken.h"
 #include "Logging/MessageLog.h"
 #include "Containers/Queue.h"
 #include "MuCO/CustomizableObjectCustomVersion.h"
+#include "MuCOE/CustomizableObjectEditor_Deprecated.h"
+#include "MuCOE/Nodes/CustomizableObjectNodeComponentMeshAddTo.h"
 #include "MuCOE/Nodes/CustomizableObjectNodeCopyMaterial.h"
+#include "MuCOE/Nodes/CustomizableObjectNodeMaterialSwitch.h"
 
 class UCustomizableObjectNodeRemapPins;
 
 #define LOCTEXT_NAMESPACE "CustomizableObjectEditor"
 
+
 const FName UCustomizableObjectNodeObject::ChildrenPinName(TEXT("Children"));
 const FName UCustomizableObjectNodeObject::ComponentsPinName(TEXT("Components"));
 const FName UCustomizableObjectNodeObject::ModifiersPinName(TEXT("Modifiers"));
 const FName UCustomizableObjectNodeObject::OutputPinName(TEXT("Object"));
-const TCHAR* UCustomizableObjectNodeObject::LODPinNamePrefix = TEXT("LOD ");
+
 
 UCustomizableObjectNodeObject::UCustomizableObjectNodeObject()
 	: Super()
 {
 	bIsBase = true;
 	ObjectName = "Unnamed Object";
-	NumLODs = 1;
 	Identifier = FGuid::NewGuid();
 }
 
@@ -97,23 +101,23 @@ void UCustomizableObjectNodeObject::BackwardsCompatibleFixup(int32 CustomizableO
 	if (CustomizableObjectCustomVersion == FCustomizableObjectCustomVersion::NewComponentOptions)
 	{
 		// Like we did in the CO components, we use the index of the component as the name of the component
-		for (int32 ComponentIndex = 0; ComponentIndex < ComponentSettings.Num(); ++ComponentIndex)
+		for (int32 ComponentIndex = 0; ComponentIndex < ComponentSettings_DEPRECATED.Num(); ++ComponentIndex)
 		{
-			ComponentSettings[ComponentIndex].ComponentName = FString::FromInt(ComponentIndex);
+			ComponentSettings_DEPRECATED[ComponentIndex].ComponentName = FString::FromInt(ComponentIndex);
 		}
 	}
 
 	if (CustomizableObjectCustomVersion == FCustomizableObjectCustomVersion::MovedCompatibilityFromPostBackwardsCompatibleFixup)
 	{
 		// Fix up ComponentSettings. Only root nodes
-		if (ComponentSettings.IsEmpty() && bIsBase && !ParentObject)
+		if (ComponentSettings_DEPRECATED.IsEmpty() && bIsBase && !ParentObject)
 		{
 			FComponentSettings ComponentSettingsTemplate;
-			ComponentSettingsTemplate.LODReductionSettings.SetNum(NumLODs);
+			ComponentSettingsTemplate.LODReductionSettings.SetNum(NumLODs_DEPRECATED);
 
 			if (UCustomizableObject* CurrentObject = Cast<UCustomizableObject>(GetOutermostObject()))
 			{
-				ComponentSettings.Init(ComponentSettingsTemplate, CurrentObject->GetPrivate()->MutableMeshComponents.Num());
+				ComponentSettings_DEPRECATED.Init(ComponentSettingsTemplate, CurrentObject->GetPrivate()->MutableMeshComponents.Num());
 			}
 		}
 	}
@@ -129,6 +133,277 @@ void UCustomizableObjectNodeObject::BackwardsCompatibleFixup(int32 CustomizableO
 			CustomCreatePin(EGPD_Input, Schema->PC_Modifier, ModifiersPinName, true);
 		}
 	}
+	
+	if (CustomizableObjectCustomVersion == FCustomizableObjectCustomVersion::NodeComponentMesh)
+	{
+		auto NodeComponentMeshBaseAllocateDefaultPins = [](UCustomizableObjectNode* Node, UCustomizableObjectNodeRemapPins*)
+		{
+			UCustomizableObjectNodeComponentMeshBase* NodeComponentMesh = CastChecked<UCustomizableObjectNodeComponentMeshBase>(Node);
+						
+			const UEdGraphSchema_CustomizableObject* Schema = GetDefault<UEdGraphSchema_CustomizableObject>();
+
+			NodeComponentMesh->LODPins.Empty(NodeComponentMesh->NumLODs);
+			for (int32 NodeComponentLODIndex = 0; NodeComponentLODIndex < NodeComponentMesh->NumLODs; ++NodeComponentLODIndex)
+			{
+				FString LODName = FString::Printf(TEXT("LOD %d"), NodeComponentLODIndex);
+
+				UEdGraphPin* Pin = NodeComponentMesh->CustomCreatePin(EGPD_Input, Schema->PC_Material, FName(*LODName), true);
+				NodeComponentMesh->LODPins.Add(Pin);
+			}
+						
+			NodeComponentMesh->OutputPin = NodeComponentMesh->CustomCreatePin(EGPD_Output, Schema->PC_Component, TEXT("Component"));
+		};
+
+		auto NodeObjectAllocateDefaultPins = [](UCustomizableObjectNode* Node, UCustomizableObjectNodeRemapPins*)
+		{
+			UCustomizableObjectNodeObject* NodeObject = CastChecked<UCustomizableObjectNodeObject>(Node);
+
+			const UEdGraphSchema_CustomizableObject* Schema = GetDefault<UEdGraphSchema_CustomizableObject>();
+
+			NodeObject->CustomCreatePin(EGPD_Input, Schema->PC_Component, ComponentsPinName, true);
+	
+			NodeObject->CustomCreatePin(EGPD_Input, Schema->PC_Modifier, ModifiersPinName, true);
+
+			NodeObject->CustomCreatePin(EGPD_Input, Schema->PC_Object, ChildrenPinName, true);
+	
+			for (const FRegisteredObjectNodeInputPin& Pin : ICustomizableObjectModule::Get().GetAdditionalObjectNodePins())
+			{
+				// Use the global pin name here to prevent extensions using the same pin names from
+				// interfering with each other.
+				//
+				// This also prevents extension pins from clashing with the built-in pins from this node,
+				// such as "Object".
+				UEdGraphPin* GraphPin = NodeObject->CustomCreatePin(EGPD_Input, Pin.InputPin.PinType, Pin.GlobalPinName, Pin.InputPin.bIsArray);
+
+				GraphPin->PinFriendlyName = Pin.InputPin.DisplayName;
+			}
+
+			UEdGraphPin* OutputPin = NodeObject->CustomCreatePin(EGPD_Output, Schema->PC_Object, OutputPinName);
+
+			if (NodeObject->bIsBase)
+			{
+				OutputPin->bHidden = true;
+			}
+		};
+		
+		const UEdGraphSchema_CustomizableObject* Schema = GetDefault<UEdGraphSchema_CustomizableObject>();
+		
+		bool bMoved = false;
+
+		int32 NodesCreated = 0;
+
+		if (!ComponentsPin()) // Some old nodes do not have the component pin.
+		{
+			CustomCreatePin(EGPD_Input, Schema->PC_Component, ComponentsPinName, true);
+		}
+		
+		if (!ParentObject && // Is a root object
+			bIsBase)
+		{
+			if (!bMoved)
+			{
+				bMoved = true;
+				NodePosX += 400; // Move it a bit to make space for the new component nodes.
+			}
+
+			for (const FMutableMeshComponentData& MeshComponent : GetRootObject(*this)->GetPrivate()->MutableMeshComponents)
+			{
+				UCustomizableObjectNodeComponentMesh* DefaultObject = UCustomizableObjectNodeComponentMesh::StaticClass()->GetDefaultObject<UCustomizableObjectNodeComponentMesh>();
+				UEdGraphNode* Node = FCustomizableObjectSchemaAction_NewNode::CreateNode(GetGraph(), ComponentsPin(), FVector2D(NodePosX - 300.0,  NodePosY + 200.0 * NodesCreated), DefaultObject);
+				UCustomizableObjectNodeComponentMesh* NodeComponentMesh = CastChecked<UCustomizableObjectNodeComponentMesh>(Node);
+
+				++NodesCreated;
+				
+				NodeComponentMesh->ComponentName = MeshComponent.Name;
+				NodeComponentMesh->NumLODs = NumLODs_DEPRECATED;
+				if (FComponentSettings* Result = ComponentSettings_DEPRECATED.FindByPredicate([&](const FComponentSettings& Settings)
+				{
+					return Settings.ComponentName == NodeComponentMesh->ComponentName;
+				}))
+				{
+					NodeComponentMesh->LODReductionSettings = Result->LODReductionSettings;
+				}
+
+				NodeComponentMesh->FixupReconstructPins(CreateRemapPinsByName(), NodeComponentMeshBaseAllocateDefaultPins);
+				
+				ComponentsPin()->MakeLinkTo(NodeComponentMesh->OutputPin.Get());
+			}
+		}
+
+		TMap<FName, UCustomizableObjectNodeComponentMeshAddTo*> ExistingNodeComponentMeshAddTo;
+		
+		// Create all NodeComponentMeshAddTo.
+		for (int32 LODIndex = 0; LODIndex < NumLODs_DEPRECATED; ++LODIndex)
+		{
+			UEdGraphPin* OldLODPin = FindPin(FString::Printf(TEXT("%s%d "), TEXT("LOD "), LODIndex));
+
+			TArray<UEdGraphPin*> CopyLinkedPins = OldLODPin->LinkedTo;
+			for (UEdGraphPin* LinkedPin : CopyLinkedPins) // Import/Exports/Reroute not supported.
+			{
+				UEdGraphNode* LinkedNode = LinkedPin->GetOwningNode();
+
+				auto CreateNodeComponent = [&, this](const FName& MeshComponentName)
+				{
+					if (UCustomizableObjectNodeComponentMeshAddTo** Result = ExistingNodeComponentMeshAddTo.Find(MeshComponentName))
+					{
+						return *Result;
+					}
+					
+					if (!bMoved)
+					{
+						bMoved = true;
+						NodePosX += 400; // Move it a bit to make space for the new component nodes.
+					}
+					
+					UCustomizableObjectNodeComponentMeshAddTo* DefaultObject = UCustomizableObjectNodeComponentMeshAddTo::StaticClass()->GetDefaultObject<UCustomizableObjectNodeComponentMeshAddTo>();
+					UEdGraphNode* Node = FCustomizableObjectSchemaAction_NewNode::CreateNode(GetGraph(), ComponentsPin(), FVector2D(NodePosX - 300.0,  NodePosY + 200.0 * NodesCreated), DefaultObject);
+					UCustomizableObjectNodeComponentMeshAddTo* NodeComponentMeshAddTo = CastChecked<UCustomizableObjectNodeComponentMeshAddTo>(Node);
+
+					++NodesCreated;
+
+					NodeComponentMeshAddTo->ParentComponentName = MeshComponentName;
+					NodeComponentMeshAddTo->NumLODs = NumLODs_DEPRECATED;
+					NodeComponentMeshAddTo->AutoLODStrategy = AutoLODStrategy_DEPRECATED;
+					
+					// Create LOD pins.
+					NodeComponentMeshAddTo->FixupReconstructPins(CreateRemapPinsByName(), NodeComponentMeshBaseAllocateDefaultPins);
+
+					ComponentsPin()->MakeLinkTo(NodeComponentMeshAddTo->OutputPin.Get());
+
+					ExistingNodeComponentMeshAddTo.Add(NodeComponentMeshAddTo->ParentComponentName ,NodeComponentMeshAddTo);
+					
+					return NodeComponentMeshAddTo;
+				};
+
+				
+				UCustomizableObjectNodeComponentMeshAddTo* NodeComponentMeshAddTo = nullptr;
+				bool bFixNode = false;
+				
+				if (UCustomizableObjectNodeMaterial* NodeMaterial = Cast<UCustomizableObjectNodeMaterial>(LinkedNode))
+				{
+					bFixNode = true;
+
+					NodeComponentMeshAddTo = CreateNodeComponent(NodeMaterial->MeshComponentName_DEPRECATED);
+				}
+				else if (UCustomizableObjectNodeMaterialSwitch* NodeMaterialSwitch = Cast<UCustomizableObjectNodeMaterialSwitch>(LinkedNode))
+				{
+					bFixNode = true;
+
+					[&]() // Lambda to ease the control flow.
+					{
+						if (!NodeMaterialSwitch->GetNumElements()) // We should at least have a component to know were to connect this material. If not, not supported.
+						{
+							return;
+						}
+
+						FName ComponentName;
+						bool bFirst = true;
+						for (int32 ElementIndex = 0; ElementIndex < NodeMaterialSwitch->GetNumElements(); ++ElementIndex)
+						{
+							if (UEdGraphPin* ConnectedPin = FollowInputPin(*NodeMaterialSwitch->GetElementPin(ElementIndex)))
+							{
+								if (UCustomizableObjectNodeMaterialBase* FirstNodeMaterialBase = Cast<UCustomizableObjectNodeMaterialBase>(ConnectedPin->GetOwningNode()))
+								{
+									if (UCustomizableObjectNodeMaterial* FirstNodeMaterial = FirstNodeMaterialBase->GetMaterialNode())
+									{
+										if (bFirst)
+										{
+											bFirst = false;
+											ComponentName = FirstNodeMaterial->MeshComponentName_DEPRECATED;
+										}
+										else
+										{
+											if (ComponentName != FirstNodeMaterial->MeshComponentName_DEPRECATED) // All components must match. If not, not supported.
+											{
+												return;
+											}
+										}
+									}
+								}
+							}
+						}
+
+						NodeComponentMeshAddTo = CreateNodeComponent(ComponentName);
+					}();
+				}
+				else if (UCustomizableObjectNodeMaterialVariation* NodeMaterialVariation = Cast<UCustomizableObjectNodeMaterialVariation>(LinkedNode))
+				{
+					bFixNode = true;
+
+					[&]() // Lambda to ease the control flow.
+					{
+						if (!NodeMaterialVariation->GetNumVariations()) // We should at least have a component to know were to connect this material. If not, not supported.
+						{
+							return;
+						}
+
+						FName ComponentName;
+						bool bFirst = true;
+
+						if (UEdGraphPin* ConnectedPin = FollowInputPin(*NodeMaterialVariation->DefaultPin()))
+						{
+							if (UCustomizableObjectNodeMaterialBase* FirstNodeMaterialBase = Cast<UCustomizableObjectNodeMaterialBase>(ConnectedPin->GetOwningNode()))
+							{
+								if (UCustomizableObjectNodeMaterial* FirstNodeMaterial = FirstNodeMaterialBase->GetMaterialNode())
+								{
+									bFirst = false;
+									ComponentName = FirstNodeMaterial->MeshComponentName_DEPRECATED;
+								}
+							}
+						}
+						
+						for (int32 ElementIndex = 0; ElementIndex < NodeMaterialVariation->GetNumVariations(); ++ElementIndex)
+						{
+							if (UEdGraphPin* ConnectedPin = FollowInputPin(*NodeMaterialVariation->VariationPin(ElementIndex)))
+							{
+								if (UCustomizableObjectNodeMaterialBase* FirstNodeMaterialBase = Cast<UCustomizableObjectNodeMaterialBase>(ConnectedPin->GetOwningNode()))
+								{
+									if (UCustomizableObjectNodeMaterial* FirstNodeMaterial = FirstNodeMaterialBase->GetMaterialNode())
+									{
+										if (bFirst)
+										{
+											bFirst = false;
+											ComponentName = FirstNodeMaterial->MeshComponentName_DEPRECATED;
+										}
+										else
+										{
+											if (ComponentName != FirstNodeMaterial->MeshComponentName_DEPRECATED) // All components must match. If not, not supported.
+											{
+												return;
+											}
+										}
+									}
+								}
+							}
+						}
+
+						if (!bFirst)
+						{
+							NodeComponentMeshAddTo = CreateNodeComponent(ComponentName);
+						}
+					}();
+				}
+
+				if (bFixNode && !NodeComponentMeshAddTo)
+				{
+					FString Msg = FString::Printf(TEXT("A Object node has a legacy connection to a node [%s] without automatic upgrade support. Manual update is probably needed."), *LinkedNode->GetName());
+					FCustomizableObjectEditorLogger::CreateLog(FText::FromString(Msg))
+						.Severity(EMessageSeverity::Warning)
+						.Context(*this)
+						.BaseObject(true)
+						.Log();
+				}
+
+				if (NodeComponentMeshAddTo)
+				{
+					LinkedPin->MakeLinkTo(NodeComponentMeshAddTo->LODPins[LODIndex].Get());
+					LinkedPin->BreakLinkTo(OldLODPin);
+				}
+			}
+		}
+		
+		FixupReconstructPins(CreateRemapPinsByName(), NodeObjectAllocateDefaultPins);
+	}
 }
 
 
@@ -149,11 +424,11 @@ void UCustomizableObjectNodeObject::PostEditChangeProperty(FPropertyChangedEvent
 	FProperty* PropertyThatChanged = PropertyChangedEvent.Property;
 	if (PropertyThatChanged && PropertyThatChanged->GetName() == TEXT("NumLODs"))
 	{
-		NumLODs = FMath::Clamp(NumLODs, 1, 64);
+		NumLODs_DEPRECATED = FMath::Clamp(NumLODs_DEPRECATED, 1, 64);
 
-		for (int32 CompSetIndex = 0; CompSetIndex < ComponentSettings.Num(); ++CompSetIndex)
+		for (int32 CompSetIndex = 0; CompSetIndex < ComponentSettings_DEPRECATED.Num(); ++CompSetIndex)
 		{
-			ComponentSettings[CompSetIndex].LODReductionSettings.SetNum(NumLODs);
+			ComponentSettings_DEPRECATED[CompSetIndex].LODReductionSettings.SetNum(NumLODs_DEPRECATED);
 		}
 
 		ReconstructNode();
@@ -167,25 +442,12 @@ void UCustomizableObjectNodeObject::AllocateDefaultPins(UCustomizableObjectNodeR
 {
 	const UEdGraphSchema_CustomizableObject* Schema = GetDefault<UEdGraphSchema_CustomizableObject>();
 
-	// NOTE: Ensure all built-in pins are handled in UCustomizableObjectNodeObject::IsBuiltInPin
-
-	for (int32 i = 0; i < NumLODs; ++i)
-	{
-		FString LODName = FString::Printf(TEXT("%s%d "), LODPinNamePrefix, i);
-
-		UEdGraphPin* Pin = CustomCreatePin(EGPD_Input, Schema->PC_Material, FName(*LODName), true);
-		Pin->bDefaultValueIsIgnored = true;
-	}
-
-	UEdGraphPin* ComponentsPin = CustomCreatePin(EGPD_Input, Schema->PC_Component, ComponentsPinName, true);
-	ComponentsPin->bDefaultValueIsIgnored = true;
-
-	UEdGraphPin* ModifiersPin = CustomCreatePin(EGPD_Input, Schema->PC_Modifier, ModifiersPinName, true);
-	ModifiersPin->bDefaultValueIsIgnored = true;
-
-	UEdGraphPin* ChildrenPin = CustomCreatePin(EGPD_Input, Schema->PC_Object, ChildrenPinName, true);
-	ChildrenPin->bDefaultValueIsIgnored = true;
+	CustomCreatePin(EGPD_Input, Schema->PC_Component, ComponentsPinName, true);
+	
+	CustomCreatePin(EGPD_Input, Schema->PC_Modifier, ModifiersPinName, true);
 
+	CustomCreatePin(EGPD_Input, Schema->PC_Object, ChildrenPinName, true);
+	
 	for (const FRegisteredObjectNodeInputPin& Pin : ICustomizableObjectModule::Get().GetAdditionalObjectNodePins())
 	{
 		// Use the global pin name here to prevent extensions using the same pin names from
@@ -253,20 +515,6 @@ void UCustomizableObjectNodeObject::PrepareForCopying()
 }
 
 
-int32 UCustomizableObjectNodeObject::GetLOD(UEdGraphPin* Pin) const
-{
-	for (int32 LOD = 0; LOD < GetNumLODPins(); ++LOD)
-	{
-		if (Pin == LODPin(LOD))
-		{
-			return LOD;
-		}
-	}
-
-	return -1;
-}
-
-
 bool UCustomizableObjectNodeObject::CanUserDeleteNode() const
 {
 	return !bIsBase;
@@ -279,48 +527,6 @@ bool UCustomizableObjectNodeObject::CanDuplicateNode() const
 }
 
 
-TArray<UCustomizableObjectNodeMaterialBase*> UCustomizableObjectNodeObject::GetMaterialNodes(const int LOD) const
-{
-	TArray<UCustomizableObjectNodeMaterialBase*> Result;
-
-	TQueue<UEdGraphNode*> PotentialCustomizableNodeObjects;
-
-	for (const UEdGraphPin* LinkedPin : FollowInputPinArray(*LODPin(LOD)))
-	{
-		PotentialCustomizableNodeObjects.Enqueue(LinkedPin->GetOwningNode());
-	}
-
-	UEdGraphNode* CurrentElement;
-	while (PotentialCustomizableNodeObjects.Dequeue(CurrentElement))
-	{
-		if (UCustomizableObjectNodeMaterialBase* CurrentMaterialNode = Cast<UCustomizableObjectNodeMaterialBase>(CurrentElement))
-		{
-			Result.Add(CurrentMaterialNode);
-		}
-		else if (UCustomizableObjectNodeMaterialVariation* CurrentMaterialVariationNode = Cast<UCustomizableObjectNodeMaterialVariation>(CurrentElement))
-		{
-			// Case of material variation. It's not a material, but a node that further references any material, add all its inputs that could be a material
-			for (int numMaterialPin = 0; numMaterialPin < CurrentMaterialVariationNode->GetNumVariations(); ++numMaterialPin)
-			{
-				const UEdGraphPin* VariationPin = CurrentMaterialVariationNode->VariationPin(numMaterialPin);
-				for (const UEdGraphPin* LinkedPin : FollowInputPinArray(*VariationPin))
-				{
-					PotentialCustomizableNodeObjects.Enqueue(LinkedPin->GetOwningNode());
-				}
-			}
-
-			const UEdGraphPin* DefaultPin = CurrentMaterialVariationNode->DefaultPin();
-			for (const UEdGraphPin* LinkedPin : FollowInputPinArray(*DefaultPin))
-			{
-				PotentialCustomizableNodeObjects.Enqueue(LinkedPin->GetOwningNode());
-			}
-		} 
-	}
-
-	return Result;
-}
-
-
 void UCustomizableObjectNodeObject::PostBackwardsCompatibleFixup()
 {
 	Super::PostBackwardsCompatibleFixup();
@@ -380,12 +586,5 @@ bool UCustomizableObjectNodeObject::IsSingleOutputNode() const
 	return true;
 }
 
-bool UCustomizableObjectNodeObject::IsBuiltInPin(FName PinName)
-{
-	return PinName == ChildrenPinName
-		|| PinName == ComponentsPinName
-		|| PinName == OutputPinName
-		|| PinName.ToString().StartsWith(LODPinNamePrefix);
-}
 
 #undef LOCTEXT_NAMESPACE
diff --git a/Engine/Plugins/Experimental/Mutable/Source/CustomizableObjectEditor/Private/MuCOE/Nodes/CustomizableObjectNodeObjectDetails.cpp b/Engine/Plugins/Experimental/Mutable/Source/CustomizableObjectEditor/Private/MuCOE/Nodes/CustomizableObjectNodeObjectDetails.cpp
--- a/Engine/Plugins/Experimental/Mutable/Source/CustomizableObjectEditor/Private/MuCOE/Nodes/CustomizableObjectNodeObjectDetails.cpp
+++ b/Engine/Plugins/Experimental/Mutable/Source/CustomizableObjectEditor/Private/MuCOE/Nodes/CustomizableObjectNodeObjectDetails.cpp
@@ -168,13 +168,9 @@ void FCustomizableObjectNodeObjectDetails::CustomizeDetails( IDetailLayoutBuilde
 	}
 
 	IDetailCategoryBuilder& StatesCategory = DetailBuilder.EditCategory("States");
-	IDetailCategoryBuilder& RuntimeParameters = DetailBuilder.EditCategory("States Runtime Parameters");
 	IDetailCategoryBuilder& ExternalCategory = DetailBuilder.EditCategory("AttachedToExternalObject");
 	IDetailCategoryBuilder& RealTimeMorphTargets = DetailBuilder.EditCategory("RealTime Morph Targets");
-	IDetailCategoryBuilder& LODCustomSettings = DetailBuilder.EditCategory("LOD Custom Settings");
-
-	//StatesCategory.CategoryIcon( "ActorClassIcon.CustomizableObject" );
-
+	
 	if (BaseObjectNode.IsValid())
 	{
 		// Properties
@@ -183,35 +179,13 @@ void FCustomizableObjectNodeObjectDetails::CustomizeDetails( IDetailLayoutBuilde
 		//TSharedRef<IPropertyHandle> ComponentsProperty = DetailBuilder.GetProperty("Components");
 		TSharedRef<IPropertyHandle> ParentObjectProperty = DetailBuilder.GetProperty("ParentObject");
 		TSharedRef<IPropertyHandle> LODsProperty = DetailBuilder.GetProperty("NumLODs");
-		TSharedRef<IPropertyHandle> ComponentSettingsProperty= DetailBuilder.GetProperty("ComponentSettings");
-
-		// Index of the component shown in the Bones to edit widget
-		int32 CurrentComponentIndex = 0;
-
-		for (int32 ComponentIndex = 0; ComponentIndex < BaseObjectNode->ComponentSettings.Num(); ++ComponentIndex)
-		{
-			if (BaseObjectNode->CurrentComponent == BaseObjectNode->ComponentSettings[ComponentIndex].ComponentName)
-			{
-				CurrentComponentIndex = ComponentIndex;
-				break;
-			}
-		}
-
-		FName BonesToRemovePropertyPath = FName("ComponentSettings[" + FString::FromInt(CurrentComponentIndex) + "].LODReductionSettings[" + FString::FromInt(BaseObjectNode->CurrentLOD) + "].BonesToRemove");
-		TSharedRef<IPropertyHandle> BonesToRemoveProperty = DetailBuilder.GetProperty(BonesToRemovePropertyPath);
-
+		
 		// Callbacks
 		StatesProperty->SetOnPropertyValueChanged(FSimpleDelegate::CreateSP(this, &FCustomizableObjectNodeObjectDetails::OnStatesPropertyChanged));
 
-		//TODO(Max UE-215837)
-		//ComponentsProperty->SetOnPropertyValueChanged(FSimpleDelegate::CreateSP(this, &FCustomizableObjectNodeObjectDetails::OnStatesPropertyChanged));
-		//NumComponentsProperty->SetOnPropertyValueChangedWithData(TDelegate<void(const FPropertyChangedEvent&)>::CreateSP(this, &FCustomizableObjectNodeObjectDetails::OnNumComponentsOrLODsChanged));
-		LODsProperty->SetOnPropertyValueChangedWithData(TDelegate<void(const FPropertyChangedEvent&)>::CreateSP(this, &FCustomizableObjectNodeObjectDetails::OnNumComponentsOrLODsChanged));
-
 		// Hidden Properties
 		DetailBuilder.HideProperty("ParentObjectGroupId");
 		DetailBuilder.HideProperty("ParentObject");
-		DetailBuilder.HideProperty("ComponentSettings");
 
 		GroupNodeComboOptions.Empty();
 
@@ -271,11 +245,8 @@ void FCustomizableObjectNodeObjectDetails::CustomizeDetails( IDetailLayoutBuilde
 						]
 					]
 				];
-
-				DetailBuilder.HideProperty("NumMeshComponents");
+				
 				DetailBuilder.HideProperty(StatesProperty);
-				//TODO(Max UE-215837)
-				//DetailBuilder.HideProperty(ComponentsProperty);
 			}
             else
             {
@@ -286,42 +257,7 @@ void FCustomizableObjectNodeObjectDetails::CustomizeDetails( IDetailLayoutBuilde
                     SNew(SCustomizableObjectNodeSkeletalMeshRTMorphTargetOverride)
                     .Node(BaseObjectNode.Get())
                 ];
-
-				// Component Settings Category ----------
-
-				// Component Picker
-				LODCustomSettings.AddCustomRow((LOCTEXT("ComponentCustomModeSelect", "Select Component")))
-				.NameContent()
-				[
-					SNew(STextBlock)
-					.Text(LOCTEXT("ComponentCustomSettingsSelectTitle", "Component"))
-					.Font(IDetailLayoutBuilder::GetDetailFont())
-					.ToolTipText(LOCTEXT("ComponentCustomSettingsSelectTooltip", "Select the component to edit."))
-				]
-				.ValueContent()
-				.VAlign(VAlign_Center)
-				[
-					OnGenerateComponentComboBoxForPicker()
-				];
-				
-				// LOD Picker
-				LODCustomSettings.AddCustomRow((LOCTEXT("LODCustomModeSelect", "Select LOD")))
-				.NameContent()
-				[
-					SNew(STextBlock)
-					.Text(LOCTEXT("LODCustomSettingsSelectTitle", "LOD"))
-					.Font(IDetailLayoutBuilder::GetDetailFont())
-					.ToolTipText(LOCTEXT("LODCustomSettingsSelectTooltip", "Select the component's LOD to edit."))
-				]
-				.ValueContent()
-				.VAlign(VAlign_Center)
-				[
-					OnGenerateLODComboBoxForPicker()
-				];
-
-				// Bones to remove widget
-				LODCustomSettings.AddProperty(BonesToRemoveProperty);
-			}
+            }
 		}
 		else
 		{
@@ -429,153 +365,6 @@ void FCustomizableObjectNodeObjectDetails::OnStatesPropertyChanged()
 }
 
 
-TSharedRef<SWidget> FCustomizableObjectNodeObjectDetails::OnGenerateComponentComboBoxForPicker()
-{
-	return SNew(SComboButton)
-		.OnGetMenuContent(this, &FCustomizableObjectNodeObjectDetails::OnGenerateComponentMenuForPicker)
-		.VAlign(VAlign_Center)
-		.ContentPadding(0)
-		.ButtonContent()
-		[
-			SNew(STextBlock)
-			.Font(IDetailLayoutBuilder::GetDetailFont())
-			.Text(this, &FCustomizableObjectNodeObjectDetails::GetCurrentComponentName)
-		];
-}
-
-
-TSharedRef<SWidget> FCustomizableObjectNodeObjectDetails::OnGenerateLODComboBoxForPicker()
-{
-	return SNew(SComboButton)
-		.OnGetMenuContent(this, &FCustomizableObjectNodeObjectDetails::OnGenerateLODMenuForPicker)
-		.VAlign(VAlign_Center)
-		.ButtonContent()
-		[
-			SNew(STextBlock)
-			.Font(IDetailLayoutBuilder::GetDetailFont())
-			.Text(this, &FCustomizableObjectNodeObjectDetails::GetCurrentLODName)
-		];
-}
-
-
-TSharedRef<SWidget> FCustomizableObjectNodeObjectDetails::OnGenerateComponentMenuForPicker()
-{
-	if (BaseObjectNode.IsValid())
-	{
-		if (const UCustomizableObject* ParentObject = Cast<UCustomizableObject>(BaseObjectNode->GetOutermostObject()))
-		{
-			int32 NumComponents = ParentObject->GetPrivate()->MutableMeshComponents.Num();
-			FMenuBuilder MenuBuilder(true, NULL);
-
-			for (int32 ComponentIndex = 0; ComponentIndex < NumComponents; ++ComponentIndex)
-			{
-				FText ComponentString = FText::FromName(ParentObject->GetPrivate()->MutableMeshComponents[ComponentIndex].Name);
-				FUIAction Action(FExecuteAction::CreateSP(this, &FCustomizableObjectNodeObjectDetails::OnSelectedComponentChanged, ParentObject->GetPrivate()->MutableMeshComponents[ComponentIndex].Name));
-				MenuBuilder.AddMenuEntry(ComponentString, FText::GetEmpty(), FSlateIcon(), Action);
-			}
-
-			return MenuBuilder.MakeWidget();
-		}
-	}
-
-	return SNullWidget::NullWidget;
-}
-
-
-TSharedRef<SWidget> FCustomizableObjectNodeObjectDetails::OnGenerateLODMenuForPicker()
-{
-	if (BaseObjectNode.IsValid())
-	{
-		int32 NumLODs = BaseObjectNode->NumLODs;
-		FMenuBuilder MenuBuilder(true, NULL);
-
-		for (int32 LODIndex = 0; LODIndex < NumLODs; ++LODIndex)
-		{
-			FText LODString = FText::FromString((TEXT("LOD ") + FString::FromInt(LODIndex)));
-			FUIAction Action(FExecuteAction::CreateSP(this, &FCustomizableObjectNodeObjectDetails::OnSelectedLODChanged, LODIndex));
-			MenuBuilder.AddMenuEntry(LODString, FText::GetEmpty(), FSlateIcon(), Action);
-		}
-
-		return MenuBuilder.MakeWidget();
-	}
-
-	return SNullWidget::NullWidget;
-}
-
-
-void FCustomizableObjectNodeObjectDetails::OnSelectedComponentChanged(const FName NewComponentSelected)
-{
-	if (BaseObjectNode.IsValid())
-	{
-		BaseObjectNode->CurrentComponent = NewComponentSelected;
-		BaseObjectNode->CurrentLOD = 0;
-	}
-	
-	DetailBuilderPtr->ForceRefreshDetails();
-}
-
-
-void FCustomizableObjectNodeObjectDetails::OnSelectedLODChanged(int32 NewLODIndex)
-{
-	if (BaseObjectNode.IsValid())
-	{
-		BaseObjectNode->CurrentLOD = NewLODIndex;
-	}
-
-	DetailBuilderPtr->ForceRefreshDetails();
-}
-
-
-FText FCustomizableObjectNodeObjectDetails::GetCurrentComponentName() const
-{
-	FText ComponentText;
-
-	if (BaseObjectNode.IsValid())
-	{
-		ComponentText = FText::FromName(BaseObjectNode->CurrentComponent);
-	}
-
-	return ComponentText;
-}
-
-
-FText FCustomizableObjectNodeObjectDetails::GetCurrentLODName() const
-{
-	FText LODText;
-
-	if (BaseObjectNode.IsValid())
-	{
-		LODText = FText::FromString(FString(TEXT("LOD ")) + FString::FromInt(BaseObjectNode->CurrentLOD));
-	}
-
-	return LODText;
-}
-
-
-void FCustomizableObjectNodeObjectDetails::OnNumComponentsOrLODsChanged(const FPropertyChangedEvent& PropertyChangedEvent)
-{
-	if (DetailBuilderPtr && PropertyChangedEvent.ChangeType == EPropertyChangeType::ValueSet)
-	{
-		if (const UCustomizableObject* ParentObject = Cast<UCustomizableObject>(BaseObjectNode->GetOutermostObject()))
-		{
-			if (ParentObject->GetPrivate()->MutableMeshComponents.FindByPredicate([&](const FMutableMeshComponentData& Component) { return Component.Name == BaseObjectNode->CurrentComponent; }) == nullptr)
-			{
-				BaseObjectNode->CurrentComponent = ParentObject->GetPrivate()->MutableMeshComponents.Last().Name;
-
-				// Reset the LOD selection
-				BaseObjectNode->CurrentLOD = 0;
-			}
-			else
-			{
-				BaseObjectNode->CurrentLOD = FMath::Min(BaseObjectNode->NumLODs - 1, BaseObjectNode->CurrentLOD);
-			}
-		}
-	
-		DetailBuilderPtr->ForceRefreshDetails();
-	}
-}
-
-
 void FCustomizableObjectNodeObjectDetails::FillParameterNamesArray()
 {
 	BaseObjectNode->ParameterNames.Empty();
diff --git a/Engine/Plugins/Experimental/Mutable/Source/CustomizableObjectEditor/Private/MuCOE/Nodes/CustomizableObjectNodeObjectDetails.h b/Engine/Plugins/Experimental/Mutable/Source/CustomizableObjectEditor/Private/MuCOE/Nodes/CustomizableObjectNodeObjectDetails.h
--- a/Engine/Plugins/Experimental/Mutable/Source/CustomizableObjectEditor/Private/MuCOE/Nodes/CustomizableObjectNodeObjectDetails.h
+++ b/Engine/Plugins/Experimental/Mutable/Source/CustomizableObjectEditor/Private/MuCOE/Nodes/CustomizableObjectNodeObjectDetails.h
@@ -101,31 +101,9 @@ private:
 	TArray< class UCustomizableObjectNodeMaterial* > ParentOptionNode;
 
 	void OnGroupNodeComboBoxSelectionChanged(TSharedPtr<FString> Selection, ESelectInfo::Type SelectInfo, TSharedRef<IPropertyHandle> ParentProperty);
-
-	// Component and LOD combobox generation callbacks.
-	TSharedRef<SWidget> OnGenerateComponentComboBoxForPicker();
-	TSharedRef<SWidget> OnGenerateLODComboBoxForPicker();
-
-	// Component and LOD menu generation callbacks
-	TSharedRef<SWidget> OnGenerateComponentMenuForPicker();
-	TSharedRef<SWidget> OnGenerateLODMenuForPicker();
-
-	// Component and LOD OnSelect callbacks
-	void OnSelectedComponentChanged(const FName NewComponentSelected);
-	void OnSelectedLODChanged(int32 NewLODIndex);
-
-	// Component and LOD name generation callbacks
-	FText GetCurrentComponentName() const;
-	FText GetCurrentLODName() const;
-
-	// Calback to force the refresh of the detail if the number of components or LODs changes
-	void OnNumComponentsOrLODsChanged(const FPropertyChangedEvent& PropertyChangedEvent);
-
+	
 	// Refreshes the Details View when the states variable has been updated
 	void OnStatesPropertyChanged();
-
-	FIntPoint GetGridSize() const;
-	TArray<FIntRect> GetBlocks() const;
-
+	
 	void FillParameterNamesArray();
 };
diff --git a/Engine/Plugins/Experimental/Mutable/Source/CustomizableObjectEditor/Private/MuCOE/Nodes/FCustomizableObjectNodeComponentMeshDetails.cpp b/Engine/Plugins/Experimental/Mutable/Source/CustomizableObjectEditor/Private/MuCOE/Nodes/FCustomizableObjectNodeComponentMeshDetails.cpp
new file mode 100644
--- /dev/null
+++ b/Engine/Plugins/Experimental/Mutable/Source/CustomizableObjectEditor/Private/MuCOE/Nodes/FCustomizableObjectNodeComponentMeshDetails.cpp
@@ -0,0 +1,132 @@
+﻿// Copyright Epic Games, Inc. All Rights Reserved.
+
+#include "CustomizableObjectNodeComponentMeshDetails.h"
+
+#include "DetailCategoryBuilder.h"
+#include "DetailLayoutBuilder.h"
+#include "DetailWidgetRow.h"
+#include "MuCOE/Nodes/CustomizableObjectNodeComponentMesh.h"
+#include "Widgets/Text/STextBlock.h"
+#include "Widgets/Input/SComboButton.h"
+#include "Widgets/Input/SComboButton.h"
+#include "Framework/MultiBox/MultiBoxBuilder.h"
+
+
+#define LOCTEXT_NAMESPACE "CustomizableObjectDetails"
+
+
+TSharedRef<IDetailCustomization> FCustomizableObjectNodeComponentMeshDetails::MakeInstance()
+{
+	return MakeShareable(new FCustomizableObjectNodeComponentMeshDetails);
+}
+
+
+void FCustomizableObjectNodeComponentMeshDetails::CustomizeDetails(IDetailLayoutBuilder& DetailBuilder)
+{
+	FCustomizableObjectNodeDetails::CustomizeDetails(DetailBuilder);
+
+	
+	const IDetailsView* DetailsView = DetailBuilder.GetDetailsView();
+	if (DetailsView->GetSelectedObjects().IsEmpty())
+	{
+		return;
+	}
+
+	NodeComponentMesh = Cast<UCustomizableObjectNodeComponentMesh>(DetailsView->GetSelectedObjects()[0].Get());
+	if (!NodeComponentMesh.IsValid())
+	{
+		return;
+	}
+
+	// LOD Picker
+	IDetailCategoryBuilder& LODCustomSettings = DetailBuilder.EditCategory("LOD Custom Settings");
+
+	LODCustomSettings.AddCustomRow((LOCTEXT("LODCustomModeSelect", "Select LOD")))
+	.NameContent()
+	[
+		SNew(STextBlock)
+		.Text(LOCTEXT("LODCustomSettingsSelectTitle", "LOD"))
+		.Font(IDetailLayoutBuilder::GetDetailFont())
+		.ToolTipText(LOCTEXT("LODCustomSettingsSelectTooltip", "Select the component's LOD to edit."))
+	]
+	.ValueContent()
+	.VAlign(VAlign_Center)
+	[
+		OnGenerateLODComboBoxForPicker()
+	];
+
+	DetailBuilder.HideProperty("LODReductionSettings");
+
+	FName BonesToRemovePropertyPath = FName("LODReductionSettings[" + FString::FromInt(NodeComponentMesh.Get()->SelectedLOD) + "].BonesToRemove");
+	TSharedRef<IPropertyHandle> BonesToRemoveProperty = DetailBuilder.GetProperty(BonesToRemovePropertyPath);
+	
+	// Bones to remove widget
+	LODCustomSettings.AddProperty(BonesToRemoveProperty);
+}
+
+
+void FCustomizableObjectNodeComponentMeshDetails::CustomizeDetails(const TSharedPtr<IDetailLayoutBuilder>& DetailBuilder)
+{
+	FCustomizableObjectNodeDetails::CustomizeDetails(DetailBuilder);
+
+	WeakDetailsBuilder = DetailBuilder;
+}
+
+
+TSharedRef<SWidget> FCustomizableObjectNodeComponentMeshDetails::OnGenerateLODComboBoxForPicker()
+{
+	return SNew(SComboButton)
+		.OnGetMenuContent(this, &FCustomizableObjectNodeComponentMeshDetails::OnGenerateLODMenuForPicker)
+		.VAlign(VAlign_Center)
+		.ButtonContent()
+		[
+			SNew(STextBlock)
+			.Font(IDetailLayoutBuilder::GetDetailFont())
+			.Text(this, &FCustomizableObjectNodeComponentMeshDetails::GetCurrentLODName)
+		];
+}
+
+
+TSharedRef<SWidget> FCustomizableObjectNodeComponentMeshDetails::OnGenerateLODMenuForPicker()
+{
+	if (NodeComponentMesh.IsValid())
+	{
+		int32 NumLODs = NodeComponentMesh->NumLODs;
+		FMenuBuilder MenuBuilder(true, NULL);
+
+		for (int32 LODIndex = 0; LODIndex < NumLODs; ++LODIndex)
+		{
+			FText LODString = FText::FromString((TEXT("LOD ") + FString::FromInt(LODIndex)));
+			FUIAction Action(FExecuteAction::CreateSP(this, &FCustomizableObjectNodeComponentMeshDetails::OnSelectedLODChanged, LODIndex));
+			MenuBuilder.AddMenuEntry(LODString, FText::GetEmpty(), FSlateIcon(), Action);
+		}
+
+		return MenuBuilder.MakeWidget();
+	}
+
+	return SNullWidget::NullWidget;
+}
+
+
+void FCustomizableObjectNodeComponentMeshDetails::OnSelectedLODChanged(int32 NewLODIndex)
+{
+	NodeComponentMesh.Get()->SelectedLOD = NewLODIndex;
+
+	WeakDetailsBuilder.Pin()->ForceRefreshDetails();
+}
+
+
+FText FCustomizableObjectNodeComponentMeshDetails::GetCurrentLODName() const
+{
+	FText LODText;
+
+	if (NodeComponentMesh.IsValid())
+	{
+		LODText = FText::FromString(FString(TEXT("LOD ")) + FString::FromInt(NodeComponentMesh.Get()->SelectedLOD));
+	}
+
+	return LODText;
+}
+
+#undef LOCTEXT_NAMESPACE
+
diff --git a/Engine/Plugins/Experimental/Mutable/Source/CustomizableObjectEditor/Private/MuCOE/Widgets/CustomizableObjectLODReductionSettings.cpp b/Engine/Plugins/Experimental/Mutable/Source/CustomizableObjectEditor/Private/MuCOE/Widgets/CustomizableObjectLODReductionSettings.cpp
--- a/Engine/Plugins/Experimental/Mutable/Source/CustomizableObjectEditor/Private/MuCOE/Widgets/CustomizableObjectLODReductionSettings.cpp
+++ b/Engine/Plugins/Experimental/Mutable/Source/CustomizableObjectEditor/Private/MuCOE/Widgets/CustomizableObjectLODReductionSettings.cpp
@@ -37,7 +37,7 @@ void FCustomizableObjectLODReductionSettings::CustomizeHeader(TSharedRef<IProper
 	if (BoneNameProperty.IsValid() && BoneNameProperty->IsValidHandle() 
 		&& IncludeBoneProperty.IsValid() && IncludeBoneProperty->IsValidHandle())
 	{
-		ObjectNode = GetObjectNode();
+		NodeComponentMesh = GetObjectNode();
 
 		bool bOnlyRemoveChildrenProperty = false;
 		IncludeBoneProperty->GetValue(bOnlyRemoveChildrenProperty);
@@ -107,7 +107,7 @@ void FCustomizableObjectLODReductionSettings::CustomizeHeader(TSharedRef<IProper
 }
 
 
-UCustomizableObjectNodeObject* FCustomizableObjectLODReductionSettings::GetObjectNode()
+UCustomizableObjectNodeComponentMesh* FCustomizableObjectLODReductionSettings::GetObjectNode()
 {
 	if (!BoneNameProperty->IsValidHandle())
 	{
@@ -119,7 +119,7 @@ UCustomizableObjectNodeObject* FCustomizableObjectLODReductionSettings::GetObjec
 
 	if (Objects.Num())
 	{
-		return Cast<UCustomizableObjectNodeObject>(Objects[0]);
+		return Cast<UCustomizableObjectNodeComponentMesh>(Objects[0]);
 	}
 
 	return nullptr;
@@ -146,14 +146,14 @@ const struct FReferenceSkeleton& FCustomizableObjectLODReductionSettings::GetRef
 	static FReferenceSkeleton DummySkeleton;
 	USkeletalMesh* SkeletalMesh = nullptr;
 
-	if (ObjectNode)
+	if (NodeComponentMesh)
 	{
-		if (const UCustomizableObject* CustomizableObject = Cast<UCustomizableObject>(ObjectNode->GetCustomizableObjectGraph()->GetOuter()))
+		if (const UCustomizableObject* CustomizableObject = Cast<UCustomizableObject>(NodeComponentMesh->GetCustomizableObjectGraph()->GetOuter()))
 		{
 			if (const FMutableMeshComponentData* Component = CustomizableObject->GetPrivate()->MutableMeshComponents.FindByPredicate(
 				[&](const FMutableMeshComponentData& Component) 
 				{ 
-					return Component.Name == ObjectNode->CurrentComponent; 
+					return Component.Name == NodeComponentMesh->ComponentName; 
 				}))
 			{
 				SkeletalMesh = Component->ReferenceSkeletalMesh;
diff --git a/Engine/Plugins/Experimental/Mutable/Source/CustomizableObjectEditor/Private/MuCOE/Widgets/CustomizableObjectLODReductionSettings.h b/Engine/Plugins/Experimental/Mutable/Source/CustomizableObjectEditor/Private/MuCOE/Widgets/CustomizableObjectLODReductionSettings.h
--- a/Engine/Plugins/Experimental/Mutable/Source/CustomizableObjectEditor/Private/MuCOE/Widgets/CustomizableObjectLODReductionSettings.h
+++ b/Engine/Plugins/Experimental/Mutable/Source/CustomizableObjectEditor/Private/MuCOE/Widgets/CustomizableObjectLODReductionSettings.h
@@ -4,6 +4,7 @@
 
 #include "IPropertyTypeCustomization.h"
 
+class UCustomizableObjectNodeComponentMesh;
 class IPropertyHandle;
 class IPropertyHandleArray;
 class SBoneSelectionWidget;
@@ -37,12 +38,12 @@ private:
 	void OnBoneCheckBoxChanged(ECheckBoxState NewState);
 
 	// Returns the object node of the property represented in the widget
-	UCustomizableObjectNodeObject* GetObjectNode();
+	UCustomizableObjectNodeComponentMesh* GetObjectNode();
 
 private:
 
 	// Pointer to the object node
-	UCustomizableObjectNodeObject* ObjectNode;
+	UCustomizableObjectNodeComponentMesh* NodeComponentMesh;
 
 	// Property to change after bone has been picked
 	TSharedPtr<IPropertyHandle> BoneNameProperty;
diff --git a/Engine/Plugins/Experimental/Mutable/Source/CustomizableObjectEditor/Public/MuCOE/CustomizableObjectEditor_Deprecated.h b/Engine/Plugins/Experimental/Mutable/Source/CustomizableObjectEditor/Public/MuCOE/CustomizableObjectEditor_Deprecated.h
--- a/Engine/Plugins/Experimental/Mutable/Source/CustomizableObjectEditor/Public/MuCOE/CustomizableObjectEditor_Deprecated.h
+++ b/Engine/Plugins/Experimental/Mutable/Source/CustomizableObjectEditor/Public/MuCOE/CustomizableObjectEditor_Deprecated.h
@@ -9,6 +9,8 @@
 #include "CustomizableObjectEditor_Deprecated.generated.h"
 
 class UTexture2D;
+struct FLODReductionSettings;
+
 
 // Place to hide all deprecated data structures. They are still needed for deserialization backwards compatibility.
 
@@ -176,6 +178,21 @@ struct CUSTOMIZABLEOBJECTEDITOR_API FCustomizableObjectMeshVariation
 };
 
 
+// UCustomizableObjectNodeObject
+// Deprecated, do not use!
+USTRUCT()
+struct FComponentSettings
+{
+	GENERATED_USTRUCT_BODY()
+	
+	UPROPERTY()
+	FString ComponentName;
+
+	UPROPERTY()
+	TArray<FLODReductionSettings> LODReductionSettings;
+};
+
+
 #if UE_ENABLE_INCLUDE_ORDER_DEPRECATED_IN_5_2
 #include "Engine/Texture2D.h"
 #endif
diff --git a/Engine/Plugins/Experimental/Mutable/Source/CustomizableObjectEditor/Public/MuCOE/Nodes/CustomizableObjectNode.h b/Engine/Plugins/Experimental/Mutable/Source/CustomizableObjectEditor/Public/MuCOE/Nodes/CustomizableObjectNode.h
--- a/Engine/Plugins/Experimental/Mutable/Source/CustomizableObjectEditor/Public/MuCOE/Nodes/CustomizableObjectNode.h
+++ b/Engine/Plugins/Experimental/Mutable/Source/CustomizableObjectEditor/Public/MuCOE/Nodes/CustomizableObjectNode.h
@@ -113,6 +113,14 @@ public:
 	 * When called, it is guaranteed that all nodes in this graph will have executed the PostLoad function. */
 	virtual void BackwardsCompatibleFixup(int32 CustomizableObjectCustomVersion);
 
+	/** Reconstructs the node pins using the provided AllocateDefaultPins function.
+	 * Safe to call in BackwardsCompatibleFixup.
+	 * Notice that this function will not call the node AllocateDefaultPins function nor the ReconstructNode overrides.
+	 *
+	 * Since this function is used BackwardsCompatibleFixup, it should not have major changes (or they will break existing fixups).
+	 * If changes are needed, duplicate this function so old fixups remain unchanged. */
+	void FixupReconstructPins(UCustomizableObjectNodeRemapPins* RemapPinsAction, TFunction<void(UCustomizableObjectNode*, UCustomizableObjectNodeRemapPins*)> AllocateDefaultPins);
+	
 	/** Add post load work here.
 	 * When called, it is guaranteed that all nodes in this graph will have executed the BackwardsCompatibleFixup function.
 	 * Notice that no compatibility code should not go here. Use BackwardsCompatibleFixup instead. */
diff --git a/Engine/Plugins/Experimental/Mutable/Source/CustomizableObjectEditor/Public/MuCOE/Nodes/CustomizableObjectNodeComponent.h b/Engine/Plugins/Experimental/Mutable/Source/CustomizableObjectEditor/Public/MuCOE/Nodes/CustomizableObjectNodeComponent.h
--- a/Engine/Plugins/Experimental/Mutable/Source/CustomizableObjectEditor/Public/MuCOE/Nodes/CustomizableObjectNodeComponent.h
+++ b/Engine/Plugins/Experimental/Mutable/Source/CustomizableObjectEditor/Public/MuCOE/Nodes/CustomizableObjectNodeComponent.h
@@ -15,6 +15,7 @@ class CUSTOMIZABLEOBJECTEDITOR_API UCustomizableObjectNodeComponent : public UCu
 public:
 	GENERATED_BODY()
 
-
+	// UCustomizableObjectNode interface
+	virtual bool IsAffectedByLOD() const override;
 };
 
diff --git a/Engine/Plugins/Experimental/Mutable/Source/CustomizableObjectEditor/Public/MuCOE/Nodes/CustomizableObjectNodeComponentMesh.h b/Engine/Plugins/Experimental/Mutable/Source/CustomizableObjectEditor/Public/MuCOE/Nodes/CustomizableObjectNodeComponentMesh.h
--- a/Engine/Plugins/Experimental/Mutable/Source/CustomizableObjectEditor/Public/MuCOE/Nodes/CustomizableObjectNodeComponentMesh.h
+++ b/Engine/Plugins/Experimental/Mutable/Source/CustomizableObjectEditor/Public/MuCOE/Nodes/CustomizableObjectNodeComponentMesh.h
@@ -2,107 +2,66 @@
 
 #pragma once
 
-#include "MuCOE/Nodes/CustomizableObjectNodeComponent.h"
-#include "MuCOE/RemapPins/CustomizableObjectNodeRemapPinsByNameDefaultPin.h"
-#include "SGraphNode.h"
-#include "GameplayTagContainer.h"
+#include "CustomizableObjectNodeComponentMeshBase.h"
+#include "MuR/System.h"
 
 #include "CustomizableObjectNodeComponentMesh.generated.h"
 
-namespace ENodeTitleType { enum Type : int; }
-
-class FArchive;
-class FAssetThumbnail;
-class FAssetThumbnailPool;
-class FSkeletalMeshModel;
-class ISinglePropertyView;
-class UAnimInstance;
-class UCustomizableObjectLayout;
-class UCustomizableObjectNodeRemapPins;
-class UCustomizableObjectNodeRemapPinsByName;
-class UMaterialInterface;
-class UObject;
-class USkeletalMesh;
-class UTexture2D;
-struct FPropertyChangedEvent;
-struct FSkeletalMaterial;
-
-
-/** PinData of a pin that belongs to a Skeletal Mesh Section. */
-UCLASS()
-class UCustomizableObjectNodeComponentMeshPinDataSection : public UCustomizableObjectNodePinData
-{
-	GENERATED_BODY()
 
-public:
-	void Init(int32 InLODIndex, int32 InSectionIndex);
-
-	int32 GetLODIndex() const;
+UENUM()
+enum class ECustomizableObjectSelectionOverride : uint8
+{
+	NoOverride = 0 UMETA(DisplayName = "No Override"),
+	Disable    = 1 UMETA(DisplayName = "Disable"    ),
+	Enable     = 2 UMETA(DisplayName = "Enable"     )
+};
 
-	int32 GetSectionIndex() const;
 
-protected:
-	virtual bool Equals(const UCustomizableObjectNodePinData& Other) const override;
+USTRUCT()
+struct FBoneToRemove
+{
+	GENERATED_USTRUCT_BODY()
 
-private:
-	UPROPERTY()
-	int32 LODIndex = -1;
+	UPROPERTY(EditAnywhere, Category = CustomizableObject)
+	bool bOnlyRemoveChildren = false;
 
-	UPROPERTY()
-	int32 SectionIndex = -1;
+	UPROPERTY(EditAnywhere, Category = CustomizableObject)
+	FName BoneName;
 };
 
 
-/** PinData of a Mesh pin. */
-UCLASS()
-class UCustomizableObjectNodeComponentMeshPinDataMaterial : public UCustomizableObjectNodeComponentMeshPinDataSection
+USTRUCT()
+struct FLODReductionSettings
 {
-	GENERATED_BODY()
+	GENERATED_USTRUCT_BODY()
+
+	/** Selects which bones will be removed from the final skeleton
+	* BoneName: Name of the bone that will be removed. Its children will be removed too.
+	* Remove Only Children: If true, only the children of the selected bone will be removed. The selected bone will remain.
+	*/
+	UPROPERTY(EditAnywhere, Category = CustomizableObject)
+	TArray<FBoneToRemove> BonesToRemove;
 };
 
 
 UCLASS()
-class CUSTOMIZABLEOBJECTEDITOR_API UCustomizableObjectNodeComponentMesh : public UCustomizableObjectNodeComponent
+class CUSTOMIZABLEOBJECTEDITOR_API UCustomizableObjectNodeComponentMesh : public UCustomizableObjectNodeComponentMeshBase
 {
-public:
 	GENERATED_BODY()
 
-	UPROPERTY(EditAnywhere, Category = Mesh)
-	FName ComponentName;
-
-	UPROPERTY(EditAnywhere, Category = Mesh, Meta = (DisplayName = Mesh, AllowedClasses = "/Script/Engine.StaticMesh, /Script/Engine.SkeletalMesh"))
-	FSoftObjectPath Mesh;
-
-	// UObject interface.
+public:
+	// UObject interface
 	virtual void PostEditChangeProperty(FPropertyChangedEvent& PropertyChangedEvent) override;
-	
+
 	// UEdGraphNode interface
 	virtual FText GetNodeTitle(ENodeTitleType::Type TitleType) const override;
-	virtual FLinearColor GetNodeTitleColor() const override;
-	virtual FText GetTooltipText() const override;
-
-	// UCustomizableObjectNode interface
-	virtual void AllocateDefaultPins(UCustomizableObjectNodeRemapPins* RemapPins) override;
-	//virtual UCustomizableObjectNodeRemapPins* CreateRemapPinsDefault() const override;
-
-	// UCustomizableObjectNodeComponent interface
-
-
-	// Own interface
 	
-	void GetPinSection(const UEdGraphPin& Pin, int32& OutLODIndex, int32& OutSectionIndex) const;
-
-	/** Find the pin for a given lod and section. */
-	UEdGraphPin* GetMaterialPin(const int32 LODIndex, const int32 SectionIndex) const;
-
-	/** Returns the material associated to the given output pin. */
-	UMaterialInterface* GetMaterialFor(const UEdGraphPin* Pin) const;
-	FSkeletalMaterial* GetSkeletalMaterialFor(const UEdGraphPin& Pin) const;
-
-private:
-	UMaterialInterface* GetMaterialInterfaceFor(const int32 LODIndex, const int32 MaterialIndex, const FSkeletalMeshModel* ImportedModel = nullptr) const;
-	FSkeletalMaterial* GetSkeletalMaterialFor(const int32 LODIndex, const int32 MaterialIndex, const FSkeletalMeshModel* ImportedModel = nullptr) const;
-
-	static const FName OutputPinName;
+	UPROPERTY(EditAnywhere, Category = ComponentMesh)
+	FName ComponentName;
+	
+	UPROPERTY(EditAnywhere, Category = ComponentMesh)
+	TArray<FLODReductionSettings> LODReductionSettings;
 
+	/** Details selected LOD. */
+	int32 SelectedLOD = 0;
 };
diff --git a/Engine/Plugins/Experimental/Mutable/Source/CustomizableObjectEditor/Public/MuCOE/Nodes/CustomizableObjectNodeComponentMeshAddTo.h b/Engine/Plugins/Experimental/Mutable/Source/CustomizableObjectEditor/Public/MuCOE/Nodes/CustomizableObjectNodeComponentMeshAddTo.h
new file mode 100644
--- /dev/null
+++ b/Engine/Plugins/Experimental/Mutable/Source/CustomizableObjectEditor/Public/MuCOE/Nodes/CustomizableObjectNodeComponentMeshAddTo.h
@@ -0,0 +1,23 @@
+﻿// Copyright Epic Games, Inc. All Rights Reserved.
+
+#pragma once
+
+#include "MuCOE/Nodes/CustomizableObjectNodeComponentMeshBase.h"
+
+#include "CustomizableObjectNodeComponentMeshAddTo.generated.h"
+
+UCLASS()
+class CUSTOMIZABLEOBJECTEDITOR_API UCustomizableObjectNodeComponentMeshAddTo : public UCustomizableObjectNodeComponentMeshBase
+{
+	GENERATED_BODY()
+
+public:
+	// UObject interface
+	virtual void PostEditChangeProperty(FPropertyChangedEvent& PropertyChangedEvent) override;
+
+	// UEdGraphNode interface
+	virtual FText GetNodeTitle(ENodeTitleType::Type TitleType) const override;
+	
+	UPROPERTY(EditAnywhere, Category = ComponentMesh)
+	FName ParentComponentName;
+};
\ No newline at end of file
diff --git a/Engine/Plugins/Experimental/Mutable/Source/CustomizableObjectEditor/Public/MuCOE/Nodes/CustomizableObjectNodeComponentMeshBase.h b/Engine/Plugins/Experimental/Mutable/Source/CustomizableObjectEditor/Public/MuCOE/Nodes/CustomizableObjectNodeComponentMeshBase.h
new file mode 100644
--- /dev/null
+++ b/Engine/Plugins/Experimental/Mutable/Source/CustomizableObjectEditor/Public/MuCOE/Nodes/CustomizableObjectNodeComponentMeshBase.h
@@ -0,0 +1,43 @@
+﻿// Copyright Epic Games, Inc. All Rights Reserved.
+
+#pragma once
+
+#include "MuCOE/Nodes/CustomizableObjectNodeComponent.h"
+
+#include "CustomizableObjectNodeComponentMeshBase.generated.h"
+
+
+UENUM()
+enum class ECustomizableObjectAutomaticLODStrategy : uint8
+{
+	// Use the same strategy than the parent object. If root, then use "Manual".
+	Inherited = 0 UMETA(DisplayName = "Inherit from parent object"),
+	// Don't try to generate LODs automatically for the child nodes. Only the ones tha explicitely define them will be used.
+	Manual = 1 UMETA(DisplayName = "Only manually created LODs"),
+	// Try to generate the same material structure than LOd 0 if the source meshes have LODs.
+	AutomaticFromMesh = 2 UMETA(DisplayName = "Automatic from mesh")
+};
+
+
+UCLASS()
+class CUSTOMIZABLEOBJECTEDITOR_API UCustomizableObjectNodeComponentMeshBase : public UCustomizableObjectNodeComponent
+{
+	GENERATED_BODY()
+	
+public:
+	// UEdGraphNode interface
+	virtual void AllocateDefaultPins(UCustomizableObjectNodeRemapPins* RemapPins) override;
+	virtual FLinearColor GetNodeTitleColor() const override;
+
+	UPROPERTY(EditAnywhere, Category = ComponentMesh)
+	int32 NumLODs = 1;
+	
+	UPROPERTY(EditAnywhere, Category = ComponentMesh)
+	ECustomizableObjectAutomaticLODStrategy AutoLODStrategy = ECustomizableObjectAutomaticLODStrategy::AutomaticFromMesh;
+	
+	UPROPERTY()
+	TArray<FEdGraphPinReference> LODPins;
+	
+	UPROPERTY()
+	FEdGraphPinReference OutputPin;
+};
\ No newline at end of file
diff --git a/Engine/Plugins/Experimental/Mutable/Source/CustomizableObjectEditor/Public/MuCOE/Nodes/CustomizableObjectNodeComponentPassthroughMesh.h b/Engine/Plugins/Experimental/Mutable/Source/CustomizableObjectEditor/Public/MuCOE/Nodes/CustomizableObjectNodeComponentPassthroughMesh.h
new file mode 100644
--- /dev/null
+++ b/Engine/Plugins/Experimental/Mutable/Source/CustomizableObjectEditor/Public/MuCOE/Nodes/CustomizableObjectNodeComponentPassthroughMesh.h
@@ -0,0 +1,104 @@
+﻿// Copyright Epic Games, Inc. All Rights Reserved.
+
+#pragma once
+
+#include "MuCOE/Nodes/CustomizableObjectNodeComponent.h"
+#include "MuCOE/RemapPins/CustomizableObjectNodeRemapPinsByNameDefaultPin.h"
+#include "SGraphNode.h"
+#include "GameplayTagContainer.h"
+
+#include "CustomizableObjectNodeComponentPassthroughMesh.generated.h"
+
+namespace ENodeTitleType { enum Type : int; }
+
+class FArchive;
+class FAssetThumbnail;
+class FAssetThumbnailPool;
+class FSkeletalMeshModel;
+class ISinglePropertyView;
+class UAnimInstance;
+class UCustomizableObjectLayout;
+class UCustomizableObjectNodeRemapPins;
+class UCustomizableObjectNodeRemapPinsByName;
+class UMaterialInterface;
+class UObject;
+class USkeletalMesh;
+class UTexture2D;
+struct FPropertyChangedEvent;
+struct FSkeletalMaterial;
+
+
+/** PinData of a pin that belongs to a Skeletal Mesh Section. */
+UCLASS()
+class UCustomizableObjectNodeComponentMeshPinDataSection : public UCustomizableObjectNodePinData
+{
+	GENERATED_BODY()
+
+public:
+	void Init(int32 InLODIndex, int32 InSectionIndex);
+
+	int32 GetLODIndex() const;
+
+	int32 GetSectionIndex() const;
+
+protected:
+	virtual bool Equals(const UCustomizableObjectNodePinData& Other) const override;
+
+private:
+	UPROPERTY()
+	int32 LODIndex = -1;
+
+	UPROPERTY()
+	int32 SectionIndex = -1;
+};
+
+
+/** PinData of a Mesh pin. */
+UCLASS()
+class UCustomizableObjectNodeComponentMeshPinDataMaterial : public UCustomizableObjectNodeComponentMeshPinDataSection
+{
+	GENERATED_BODY()
+};
+
+
+UCLASS()
+class CUSTOMIZABLEOBJECTEDITOR_API UCustomizableObjectNodeComponentPassthroughMesh : public UCustomizableObjectNodeComponent
+{
+public:
+	GENERATED_BODY()
+
+	UPROPERTY(EditAnywhere, Category = PassthroughMesh)
+	FName ComponentName;
+
+	UPROPERTY(EditAnywhere, Category = PassthroughMesh, Meta = (DisplayName = Mesh, AllowedClasses = "/Script/Engine.StaticMesh, /Script/Engine.SkeletalMesh"))
+	FSoftObjectPath Mesh;
+
+	// UObject interface.
+	virtual void PostEditChangeProperty(FPropertyChangedEvent& PropertyChangedEvent) override;
+	
+	// UEdGraphNode interface
+	virtual FText GetNodeTitle(ENodeTitleType::Type TitleType) const override;
+	virtual FLinearColor GetNodeTitleColor() const override;
+	virtual FText GetTooltipText() const override;
+
+	// UCustomizableObjectNode interface
+	virtual void AllocateDefaultPins(UCustomizableObjectNodeRemapPins* RemapPins) override;
+
+	// Own interface
+	
+	void GetPinSection(const UEdGraphPin& Pin, int32& OutLODIndex, int32& OutSectionIndex) const;
+
+	/** Find the pin for a given lod and section. */
+	UEdGraphPin* GetMaterialPin(const int32 LODIndex, const int32 SectionIndex) const;
+
+	/** Returns the material associated to the given output pin. */
+	UMaterialInterface* GetMaterialFor(const UEdGraphPin* Pin) const;
+	FSkeletalMaterial* GetSkeletalMaterialFor(const UEdGraphPin& Pin) const;
+
+private:
+	UMaterialInterface* GetMaterialInterfaceFor(const int32 LODIndex, const int32 MaterialIndex, const FSkeletalMeshModel* ImportedModel = nullptr) const;
+	FSkeletalMaterial* GetSkeletalMaterialFor(const int32 LODIndex, const int32 MaterialIndex, const FSkeletalMeshModel* ImportedModel = nullptr) const;
+
+	static const FName OutputPinName;
+
+};
diff --git a/Engine/Plugins/Experimental/Mutable/Source/CustomizableObjectEditor/Public/MuCOE/Nodes/CustomizableObjectNodeCopyMaterial.h b/Engine/Plugins/Experimental/Mutable/Source/CustomizableObjectEditor/Public/MuCOE/Nodes/CustomizableObjectNodeCopyMaterial.h
--- a/Engine/Plugins/Experimental/Mutable/Source/CustomizableObjectEditor/Public/MuCOE/Nodes/CustomizableObjectNodeCopyMaterial.h
+++ b/Engine/Plugins/Experimental/Mutable/Source/CustomizableObjectEditor/Public/MuCOE/Nodes/CustomizableObjectNodeCopyMaterial.h
@@ -35,7 +35,6 @@ public:
 	// UCustomizableObjectNodeMaterialBase interface
 	virtual UMaterialInterface* GetMaterial() const override;
 	virtual bool IsReuseMaterialBetweenLODs() const override;
-	virtual FName GetMeshComponentName() const override;
 	virtual TArray<FString>* GetEnableTags() override;
 	virtual UEdGraphPin* GetMaterialAssetPin() const override;
 	virtual int32 GetNumParameters(EMaterialParameterType Type) const override;
diff --git a/Engine/Plugins/Experimental/Mutable/Source/CustomizableObjectEditor/Public/MuCOE/Nodes/CustomizableObjectNodeMaterial.h b/Engine/Plugins/Experimental/Mutable/Source/CustomizableObjectEditor/Public/MuCOE/Nodes/CustomizableObjectNodeMaterial.h
--- a/Engine/Plugins/Experimental/Mutable/Source/CustomizableObjectEditor/Public/MuCOE/Nodes/CustomizableObjectNodeMaterial.h
+++ b/Engine/Plugins/Experimental/Mutable/Source/CustomizableObjectEditor/Public/MuCOE/Nodes/CustomizableObjectNodeMaterial.h
@@ -132,7 +132,6 @@ public:
 	virtual bool IsNodeOutDatedAndNeedsRefresh() override;
 	virtual FString GetRefreshMessage() const override;
 	virtual TSharedPtr<IDetailsView> CustomizePinDetails(const UEdGraphPin& Pin) const override;
-	void ReconstructNode(UCustomizableObjectNodeRemapPins* RemapPinsMode);
 	virtual TArray<FString>* GetEnableTags() override;
 
 	
@@ -141,7 +140,6 @@ public:
 	virtual UEdGraphPin* OutputPin() const override;
 	virtual UMaterialInterface* GetMaterial() const override;
 	virtual bool IsReuseMaterialBetweenLODs() const override;
-	virtual FName GetMeshComponentName() const override;
 	virtual int32 GetNumParameters(EMaterialParameterType Type) const override;
 	virtual FNodeMaterialParameterId GetParameterId(EMaterialParameterType Type, int32 ParameterIndex) const override;
 	virtual FName GetParameterName(EMaterialParameterType Type, int32 ParameterIndex) const override;
@@ -164,7 +162,6 @@ public:
 	// Own Interface
 
 	void SetMaterial(UMaterialInterface* InMaterial);
-	void SetComponentName(const FName& Name);
 
 	static bool HasParameter(const UMaterialInterface* InMaterial, const FNodeMaterialParameterId& ParameterId);
 	static int32 GetParameterLayerIndex(const UMaterialInterface* InMaterial, EMaterialParameterType Type, int32 ParameterIndex);
@@ -188,10 +185,12 @@ private:
 	UPROPERTY()
 	int32 MeshComponentIndex_DEPRECATED = 0;
 
+public:
 	/** Selects which Mesh component of the Instance this material belongs to */
-	UPROPERTY(EditAnywhere, Category = CustomizableObject)
-	FName MeshComponentName;
+	UPROPERTY()
+	FName MeshComponentName_DEPRECATED;
 
+private:
 	/** Materials will be reused between LODs when possible unless explicitly disabled. */
 	UPROPERTY(EditAnywhere, Category = CustomizableObject)
 	bool bReuseMaterialBetweenLODs = true;
diff --git a/Engine/Plugins/Experimental/Mutable/Source/CustomizableObjectEditor/Public/MuCOE/Nodes/CustomizableObjectNodeMaterialBase.h b/Engine/Plugins/Experimental/Mutable/Source/CustomizableObjectEditor/Public/MuCOE/Nodes/CustomizableObjectNodeMaterialBase.h
--- a/Engine/Plugins/Experimental/Mutable/Source/CustomizableObjectEditor/Public/MuCOE/Nodes/CustomizableObjectNodeMaterialBase.h
+++ b/Engine/Plugins/Experimental/Mutable/Source/CustomizableObjectEditor/Public/MuCOE/Nodes/CustomizableObjectNodeMaterialBase.h
@@ -67,8 +67,6 @@ public:
 	virtual UMaterialInterface* GetMaterial() const PURE_VIRTUAL(UCustomizableObjectNodeMaterial::GetMaterial, return {}; );
 	
 	virtual bool IsReuseMaterialBetweenLODs() const PURE_VIRTUAL(UCustomizableObjectNodeMaterial::IsReuseMaterialBetweenLODs, return {}; );
-
-	virtual FName GetMeshComponentName() const PURE_VIRTUAL(UCustomizableObjectNodeMaterial::GetComponentName, return {}; );
 		
 	virtual UEdGraphPin* GetMeshPin() const PURE_VIRTUAL(UCustomizableObjectNodeMaterial::GetMeshPin, return {}; );
 
diff --git a/Engine/Plugins/Experimental/Mutable/Source/CustomizableObjectEditor/Public/MuCOE/Nodes/CustomizableObjectNodeObject.h b/Engine/Plugins/Experimental/Mutable/Source/CustomizableObjectEditor/Public/MuCOE/Nodes/CustomizableObjectNodeObject.h
--- a/Engine/Plugins/Experimental/Mutable/Source/CustomizableObjectEditor/Public/MuCOE/Nodes/CustomizableObjectNodeObject.h
+++ b/Engine/Plugins/Experimental/Mutable/Source/CustomizableObjectEditor/Public/MuCOE/Nodes/CustomizableObjectNodeObject.h
@@ -2,6 +2,8 @@
 
 #pragma once
 
+#include "CustomizableObjectNodeComponentMesh.h"
+#include "CustomizableObjectNodeComponentMeshBase.h"
 #include "MuCOE/Nodes/CustomizableObjectNode.h"
 #include "MuR/System.h"
 
@@ -15,45 +17,7 @@ class UCustomizableObjectNodeRemapPins;
 class UObject;
 struct FPropertyChangedEvent;
 struct FSoftObjectPath;
-
-USTRUCT()
-struct FBoneToRemove
-{
-	GENERATED_USTRUCT_BODY()
-
-	UPROPERTY(EditAnywhere, Category = CustomizableObject)
-	bool bOnlyRemoveChildren = false;
-
-	UPROPERTY(EditAnywhere, Category = CustomizableObject)
-	FName BoneName;
-};
-
-
-USTRUCT()
-struct FLODReductionSettings
-{
-	GENERATED_USTRUCT_BODY()
-
-	/** Selects which bones will be removed from the final skeleton
-	* BoneName: Name of the bone that will be removed. Its children will be removed too.
-	* Remove Only Children: If true, only the children of the selected bone will be removed. The selected bone will remain.
-	*/
-	UPROPERTY(EditAnywhere, Category = CustomizableObject)
-	TArray<FBoneToRemove> BonesToRemove;
-};
-
-
-USTRUCT()
-struct FComponentSettings
-{
-	GENERATED_USTRUCT_BODY()
-	
-	UPROPERTY(VisibleInstanceOnly, Category = CustomizableObject)
-	FString ComponentName;
-
-	UPROPERTY(EditAnywhere, Category = CustomizableObject)
-	TArray<FLODReductionSettings> LODReductionSettings;
-};
+struct FComponentSettings;
 
 
 USTRUCT()
@@ -111,25 +75,6 @@ private:
 };
 
 
-UENUM()
-enum class ECustomizableObjectAutomaticLODStrategy : uint8
-{
-	// Use the same strategy than the parent object. If root, then use "Manual".
-	Inherited = 0 UMETA(DisplayName = "Inherit from parent object"),
-	// Don't try to generate LODs automatically for the child nodes. Only the ones tha explicitely define them will be used.
-	Manual = 1 UMETA(DisplayName = "Only manually created LODs"),
-	// Try to generate the same material structure than LOd 0 if the source meshes have LODs.
-	AutomaticFromMesh = 2 UMETA(DisplayName = "Automatic from mesh")
-};
-
-UENUM()
-enum class ECustomizableObjectSelectionOverride : uint8
-{
-    NoOverride = 0 UMETA(DisplayName = "No Override"),
-    Disable    = 1 UMETA(DisplayName = "Disable"    ),
-    Enable     = 2 UMETA(DisplayName = "Enable"     )
-};
-
 USTRUCT()
 struct FRealTimeMorphSelectionOverride
 {
@@ -158,6 +103,7 @@ struct FRealTimeMorphSelectionOverride
     TArray<ECustomizableObjectSelectionOverride> Override;
 };
 
+
 UCLASS()
 class CUSTOMIZABLEOBJECTEDITOR_API UCustomizableObjectNodeObject : public UCustomizableObjectNode
 {
@@ -172,11 +118,11 @@ public:
 	UPROPERTY(EditAnywhere, Category = UI, meta = (DisplayName = "Parameter UI Metadata"))
 	FMutableParamUIMetadata ParamUIMetadata;
 
-	UPROPERTY(EditAnywhere, Category = CustomizableObject)
-	int32 NumLODs;
+	UPROPERTY()
+	int32 NumLODs_DEPRECATED = 1;
 
-	UPROPERTY(EditAnywhere, Category = CustomizableObject)
-	ECustomizableObjectAutomaticLODStrategy AutoLODStrategy = ECustomizableObjectAutomaticLODStrategy::AutomaticFromMesh;
+	UPROPERTY()
+	ECustomizableObjectAutomaticLODStrategy AutoLODStrategy_DEPRECATED = ECustomizableObjectAutomaticLODStrategy::AutomaticFromMesh;
 
 	UPROPERTY()
 	int32 NumMeshComponents_DEPRECATED = 1;
@@ -193,21 +139,11 @@ public:
 	UPROPERTY()
 	FGuid Identifier;
 
-    // Soft references SkeletalMeshes found in the previous compilation.
-    // Only populated if the node is the root.
-    UPROPERTY()
-    TArray<FSoftObjectPath> ReferencedSkeletalMeshes;
-
-    // Information about the realtime morph targets usage. It indexes to ReferncedSkeletakMeshes array
-    // so it is need to keep them syncronized. 
-    // This overrides the per skeletal mesh node selection 
     UPROPERTY()
     TArray<FRealTimeMorphSelectionOverride> RealTimeMorphSelectionOverrides;
 
-	// Array of bones to remove from the mesh.All influences assigned to these bones will be transferred to the closest valid bone.
-	// Selected per component and LOD. Bones will be accumulated down the line.
-	UPROPERTY(EditAnywhere, Category = CustomizableObject)
-	TArray<FComponentSettings> ComponentSettings;
+	UPROPERTY()
+	TArray<FComponentSettings> ComponentSettings_DEPRECATED;
 	
 	// UObject interface.
 	virtual void PostEditChangeProperty(FPropertyChangedEvent& PropertyChangedEvent) override;
@@ -231,28 +167,7 @@ public:
 	// Own interface
 	UPROPERTY()
 	bool bIsBase;
-
-	UEdGraphPin* LODPin(int32 LODIndex) const
-	{
-		FString LODName = FString::Printf(TEXT("%s%d "), LODPinNamePrefix, LODIndex);
-		return FindPin(LODName);
-	}
-
-	int32 GetNumLODPins() const
-	{
-		int32 Count = 0;
-
-		for (UEdGraphPin* Pin : GetAllNonOrphanPins())
-		{
-			if (Pin->GetName().StartsWith(LODPinNamePrefix))
-			{
-				Count++;
-			}
-		}
-
-		return Count;
-	}
-
+	
 	UEdGraphPin* ComponentsPin() const
 	{
 		return FindPin(ComponentsPinName);
@@ -273,24 +188,12 @@ public:
 		return FindPin(OutputPinName);
 	}
 
-	/** Return the LOD which a LOD pin references to. Return -1 if a pin does not belong to any LOD. */
-	int32 GetLOD(UEdGraphPin* Pin) const;
-
+	
 	virtual bool CanUserDeleteNode() const override;
 	virtual bool CanDuplicateNode() const override;
-
-	/** Get all Material Nodes int his Customizable Object which belong to the given LOD.
-	 *
-	 * @param LOD LOD which materials have to belong to.
-	 */
-	TArray<UCustomizableObjectNodeMaterialBase*> GetMaterialNodes(int LOD) const;
-
+	
 	bool IsSingleOutputNode() const override;
-
-	// Node Details Support
-	FName CurrentComponent;
-	int32 CurrentLOD = 0;
-
+	
 	// Array filled in the Details of the node to store all the parameter names of a CO graph (full tree)
 	TArray<FString> ParameterNames;
 
@@ -299,8 +202,5 @@ private:
 	static const FName ComponentsPinName;
 	static const FName ModifiersPinName;
 	static const FName OutputPinName;
-	static const TCHAR* LODPinNamePrefix;
-
-	static bool IsBuiltInPin(FName PinName);
 };
 
