[mutable] Fixed out of bounds array access error when compiling in the COEditor and accessing a newly added float parameter.
#tests Editor
#rnx

diff --git a/Engine/Plugins/Experimental/Mutable/Source/CustomizableObject/Private/MuCO/CustomizableObject.cpp b/Engine/Plugins/Experimental/Mutable/Source/CustomizableObject/Private/MuCO/CustomizableObject.cpp
--- a/Engine/Plugins/Experimental/Mutable/Source/CustomizableObject/Private/MuCO/CustomizableObject.cpp
+++ b/Engine/Plugins/Experimental/Mutable/Source/CustomizableObject/Private/MuCO/CustomizableObject.cpp
@@ -1335,6 +1335,16 @@ FString UCustomizableObject::GetStateParameterName(int32 StateIndex, int32 Param
 #if WITH_EDITORONLY_DATA
 void UCustomizableObjectPrivate::PostCompile()
 {
+	for (TObjectIterator<UCustomizableObjectInstance> It; It; ++It)
+	{
+		if (It->GetCustomizableObject() == this->GetPublic())
+		{
+			// This cannot be bound to the PostCompileDelegate below because the CO Editor binds to it too and the order of broadcast is indeterminate.
+			// The Instance's OnPostCompile() must happen before all the other bindings.
+			It->GetPrivate()->OnPostCompile();
+		}
+	}
+
 	PostCompileDelegate.Broadcast();
 }
 #endif
diff --git a/Engine/Plugins/Experimental/Mutable/Source/CustomizableObject/Private/MuCO/CustomizableObjectInstance.cpp b/Engine/Plugins/Experimental/Mutable/Source/CustomizableObject/Private/MuCO/CustomizableObjectInstance.cpp
--- a/Engine/Plugins/Experimental/Mutable/Source/CustomizableObject/Private/MuCO/CustomizableObjectInstance.cpp
+++ b/Engine/Plugins/Experimental/Mutable/Source/CustomizableObject/Private/MuCO/CustomizableObjectInstance.cpp
@@ -347,14 +347,12 @@ void UCustomizableInstancePrivate::BindObjectDelegates(UCustomizableObject*  Cur
 	// Unbind callback from the previous CO
 	if (CurrentCustomizableObject)
 	{
-		CurrentCustomizableObject->GetPrivate()->PostCompileDelegate.RemoveAll(this);
 		CurrentCustomizableObject->GetPrivate()->Status.GetOnStateChangedDelegate().RemoveAll(this);
 	}
 
 	// Bind callback to the new CO
 	if (NewCustomizableObject)
 	{
-		NewCustomizableObject->GetPrivate()->PostCompileDelegate.AddUObject(this, &UCustomizableInstancePrivate::OnPostCompile);
 		NewCustomizableObject->GetPrivate()->Status.GetOnStateChangedDelegate().AddUObject(this, &UCustomizableInstancePrivate::OnObjectStatusChanged);
 	}
 }
