[mutable] Remove warnings when trying to upgrade a node that didn't have parent in the first place.
#rnx

diff --git a/Engine/Plugins/Experimental/Mutable/Source/CustomizableObjectEditor/Private/MuCOE/Nodes/CustomizableObjectNodeModifierEditMeshSectionBase.cpp b/Engine/Plugins/Experimental/Mutable/Source/CustomizableObjectEditor/Private/MuCOE/Nodes/CustomizableObjectNodeModifierEditMeshSectionBase.cpp
--- a/Engine/Plugins/Experimental/Mutable/Source/CustomizableObjectEditor/Private/MuCOE/Nodes/CustomizableObjectNodeModifierEditMeshSectionBase.cpp
+++ b/Engine/Plugins/Experimental/Mutable/Source/CustomizableObjectEditor/Private/MuCOE/Nodes/CustomizableObjectNodeModifierEditMeshSectionBase.cpp
@@ -57,9 +57,7 @@ void UCustomizableObjectNodeModifierEditMeshSectionBase::BackwardsCompatibleFixu
 		}
 		else
 		{
-			// Conversion failed?
-			ensure(false);
-			UE_LOG(LogMutable, Warning, TEXT("[%s] UCustomizableObjectNodeModifierEditMeshSectionBase version upgrade failed."), *GetOutermost()->GetName());
+			UE_LOG(LogMutable, Log, TEXT("[%s] UCustomizableObjectNodeModifierEditMeshSectionBase has no parent. It will not be upgraded."), *GetOutermost()->GetName());
 		}
 	}
 }
diff --git a/Engine/Plugins/Experimental/Mutable/Source/CustomizableObjectEditor/Private/MuCOE/Nodes/CustomizableObjectNodeModifierExtendMeshSection.cpp b/Engine/Plugins/Experimental/Mutable/Source/CustomizableObjectEditor/Private/MuCOE/Nodes/CustomizableObjectNodeModifierExtendMeshSection.cpp
--- a/Engine/Plugins/Experimental/Mutable/Source/CustomizableObjectEditor/Private/MuCOE/Nodes/CustomizableObjectNodeModifierExtendMeshSection.cpp
+++ b/Engine/Plugins/Experimental/Mutable/Source/CustomizableObjectEditor/Private/MuCOE/Nodes/CustomizableObjectNodeModifierExtendMeshSection.cpp
@@ -100,9 +100,8 @@ void UCustomizableObjectNodeModifierExtendMeshSection::BackwardsCompatibleFixup(
 		}
 		else
 		{
-			// Conversion failed?
-			ensure(false);
-			UE_LOG(LogMutable, Warning, TEXT("[%s] UCustomizableObjectNodeModifierEditMeshSectionBase version upgrade failed."), *GetOutermost()->GetName());
+			// No parent, no conversion.
+			UE_LOG(LogMutable, Log, TEXT("[%s] UCustomizableObjectNodeModifierExtendMeshSection has no parent. It will not be upgraded."), *GetOutermost()->GetName());
 		}
 
 		ReconstructNode();
diff --git a/Engine/Plugins/Experimental/Mutable/Source/CustomizableObjectEditor/Private/MuCOE/Nodes/CustomizableObjectNodeModifierRemoveMeshBlocks.cpp b/Engine/Plugins/Experimental/Mutable/Source/CustomizableObjectEditor/Private/MuCOE/Nodes/CustomizableObjectNodeModifierRemoveMeshBlocks.cpp
--- a/Engine/Plugins/Experimental/Mutable/Source/CustomizableObjectEditor/Private/MuCOE/Nodes/CustomizableObjectNodeModifierRemoveMeshBlocks.cpp
+++ b/Engine/Plugins/Experimental/Mutable/Source/CustomizableObjectEditor/Private/MuCOE/Nodes/CustomizableObjectNodeModifierRemoveMeshBlocks.cpp
@@ -63,9 +63,7 @@ void UCustomizableObjectNodeModifierRemoveMeshBlocks::BackwardsCompatibleFixup(i
 		}
 		else
 		{
-			// Conversion failed?
-			ensure(false);
-			UE_LOG(LogMutable, Warning, TEXT("[%s] UCustomizableObjectNodeModifierRemoveMeshBlocks version upgrade failed."), *GetOutermost()->GetName());
+			UE_LOG(LogMutable, Log, TEXT("[%s] UCustomizableObjectNodeModifierRemoveMeshBlocks has no parent. It will not be upgraded."), *GetOutermost()->GetName());
 		}
 	}
 
