[Mutable] Fix "Detaching from existing linker"

Removed StaticDuplicateObject from inside PostLoads.

#rnx

diff --git a/Engine/Plugins/Experimental/Mutable/Source/CustomizableObjectEditor/Private/MuCOE/Nodes/CustomizableObjectNodeObject.cpp b/Engine/Plugins/Experimental/Mutable/Source/CustomizableObjectEditor/Private/MuCOE/Nodes/CustomizableObjectNodeObject.cpp
--- a/Engine/Plugins/Experimental/Mutable/Source/CustomizableObjectEditor/Private/MuCOE/Nodes/CustomizableObjectNodeObject.cpp
+++ b/Engine/Plugins/Experimental/Mutable/Source/CustomizableObjectEditor/Private/MuCOE/Nodes/CustomizableObjectNodeObject.cpp
@@ -21,6 +21,7 @@
 #include "MuCOE/Nodes/CustomizableObjectNodeComponentMeshAddTo.h"
 #include "MuCOE/Nodes/CustomizableObjectNodeCopyMaterial.h"
 #include "MuCOE/Nodes/CustomizableObjectNodeMaterialSwitch.h"
+#include "Settings/EditorStyleSettings.h"
 
 class UCustomizableObjectNodeRemapPins;
 
@@ -185,6 +186,53 @@ void UCustomizableObjectNodeObject::BackwardsCompatibleFixup(int32 CustomizableO
 				OutputPin->bHidden = true;
 			}
 		};
+
+		auto CreateNewNode = [](class UEdGraph* ParentGraph, UEdGraphPin* FromPin, const FVector2D Location, class UEdGraphNode* InNodeTemplate)
+		{
+			// UE Code from FSchemaAction_NewNode::CreateNode(...). Overlap calculations performed before AutowireNewNode(...).
+
+			const int32 NodeDistance = 60;
+			
+			// Duplicate template node to create new node
+			UEdGraphNode* ResultNode = DuplicateObject<UEdGraphNode>(InNodeTemplate, ParentGraph);
+
+			ResultNode->SetFlags(RF_Transactional);
+
+			ParentGraph->AddNode(ResultNode, true);
+
+			ResultNode->CreateNewGuid();
+			ResultNode->PostPlacedNewNode();
+			if (UCustomizableObjectNode* TypedResultNode = Cast<UCustomizableObjectNode>(ResultNode))
+			{
+				TypedResultNode->BeginConstruct();
+				TypedResultNode->PostBackwardsCompatibleFixup();
+			}
+			ResultNode->ReconstructNode(); // Mutable node lifecycle always starts at ReconstructNode.
+
+			// For input pins, new node will generally overlap node being dragged off
+			// Work out if we want to visually push away from connected node
+			int32 XLocation = Location.X;
+			if (FromPin && FromPin->Direction == EGPD_Input)
+			{
+				UEdGraphNode* PinNode = FromPin->GetOwningNode();
+				const float XDelta = FMath::Abs(PinNode->NodePosX - Location.X);
+
+				if (XDelta < NodeDistance)
+				{
+					// Set location to edge of current node minus the max move distance
+					// to force node to push off from connect node enough to give selection handle
+					XLocation = PinNode->NodePosX - NodeDistance;
+				}
+			}
+
+			ResultNode->AutowireNewNode(FromPin);
+
+			ResultNode->NodePosX = XLocation;
+			ResultNode->NodePosY = Location.Y;
+			ResultNode->SnapToGrid(GetDefault<UEditorStyleSettings>()->GridSnapSize);
+
+			return ResultNode;
+		};
 		
 		const UEdGraphSchema_CustomizableObject* Schema = GetDefault<UEdGraphSchema_CustomizableObject>();
 		
@@ -208,8 +256,8 @@ void UCustomizableObjectNodeObject::BackwardsCompatibleFixup(int32 CustomizableO
 
 			for (const FMutableMeshComponentData& MeshComponent : GetRootObject(*this)->GetPrivate()->MutableMeshComponents)
 			{
-				UCustomizableObjectNodeComponentMesh* DefaultObject = UCustomizableObjectNodeComponentMesh::StaticClass()->GetDefaultObject<UCustomizableObjectNodeComponentMesh>();
-				UEdGraphNode* Node = FCustomizableObjectSchemaAction_NewNode::CreateNode(GetGraph(), ComponentsPin(), FVector2D(NodePosX - 300.0,  NodePosY + 200.0 * NodesCreated), DefaultObject);
+				UCustomizableObjectNodeComponentMesh* NewNode = NewObject<UCustomizableObjectNodeComponentMesh>(this);
+				UEdGraphNode* Node = CreateNewNode(GetGraph(), ComponentsPin(), FVector2D(NodePosX - 300.0,  NodePosY + 200.0 * NodesCreated), NewNode);
 				UCustomizableObjectNodeComponentMesh* NodeComponentMesh = CastChecked<UCustomizableObjectNodeComponentMesh>(Node);
 
 				++NodesCreated;
@@ -255,8 +303,8 @@ void UCustomizableObjectNodeObject::BackwardsCompatibleFixup(int32 CustomizableO
 						NodePosX += 400; // Move it a bit to make space for the new component nodes.
 					}
 					
-					UCustomizableObjectNodeComponentMeshAddTo* DefaultObject = UCustomizableObjectNodeComponentMeshAddTo::StaticClass()->GetDefaultObject<UCustomizableObjectNodeComponentMeshAddTo>();
-					UEdGraphNode* Node = FCustomizableObjectSchemaAction_NewNode::CreateNode(GetGraph(), ComponentsPin(), FVector2D(NodePosX - 300.0,  NodePosY + 200.0 * NodesCreated), DefaultObject);
+					UCustomizableObjectNodeComponentMeshAddTo* NewNode = NewObject<UCustomizableObjectNodeComponentMeshAddTo>(this);
+					UEdGraphNode* Node = CreateNewNode(GetGraph(), ComponentsPin(), FVector2D(NodePosX - 300.0,  NodePosY + 200.0 * NodesCreated), NewNode);
 					UCustomizableObjectNodeComponentMeshAddTo* NodeComponentMeshAddTo = CastChecked<UCustomizableObjectNodeComponentMeshAddTo>(Node);
 
 					++NodesCreated;
