[mutable] Uniformize "Required Tags" properties on modifier hierarchy.
Implement better details widget for tags and morph selection in modifiers.
#rnx

diff --git a/Engine/Plugins/Experimental/Mutable/Source/CustomizableObject/Internal/MuCO/CustomizableObjectCustomVersion.h b/Engine/Plugins/Experimental/Mutable/Source/CustomizableObject/Internal/MuCO/CustomizableObjectCustomVersion.h
--- a/Engine/Plugins/Experimental/Mutable/Source/CustomizableObject/Internal/MuCO/CustomizableObjectCustomVersion.h
+++ b/Engine/Plugins/Experimental/Mutable/Source/CustomizableObject/Internal/MuCO/CustomizableObjectCustomVersion.h
@@ -137,11 +137,13 @@ struct CUSTOMIZABLEOBJECT_API FCustomizableObjectCustomVersion
 		MovedCompatibilityFromPostBackwardsCompatibleFixup,
 
 		FixModifierPin, 
-		
+
 		ConvertEditAndExtendToModifiers,
 
 		NodeComponentMesh, // Moved LOD pins from NodeObject to NodeComponentMesh.
 		
+		UnifyRequiredTags,
+
 		// -----<new versions can be added above this line>-------------------------------------------------
 		VersionPlusOne,
 		LatestVersion = VersionPlusOne - 1
diff --git a/Engine/Plugins/Experimental/Mutable/Source/CustomizableObjectEditor/Private/MuCOE/CustomizableObjectEditorModule.cpp b/Engine/Plugins/Experimental/Mutable/Source/CustomizableObjectEditor/Private/MuCOE/CustomizableObjectEditorModule.cpp
--- a/Engine/Plugins/Experimental/Mutable/Source/CustomizableObjectEditor/Private/MuCOE/CustomizableObjectEditorModule.cpp
+++ b/Engine/Plugins/Experimental/Mutable/Source/CustomizableObjectEditor/Private/MuCOE/CustomizableObjectEditorModule.cpp
@@ -37,6 +37,7 @@
 #include "MuCOE/Nodes/CustomizableObjectNodeModifierEditMeshSection.h"
 #include "MuCOE/Nodes/CustomizableObjectNodeModifierEditMeshSectionDetails.h"
 #include "MuCOE/Nodes/CustomizableObjectNodeModifierExtendMeshSection.h"
+#include "MuCOE/Nodes/CustomizableObjectNodeModifierExtendMeshSectionDetails.h"
 #include "MuCOE/Nodes/CustomizableObjectNodeModifierMorphMeshSection.h"
 #include "MuCOE/Nodes/CustomizableObjectNodeModifierMorphMeshSectionDetails.h"
 #include "MuCOE/Nodes/CustomizableObjectNodeModifierRemoveMeshBlocks.h"
@@ -47,6 +48,8 @@
 #include "MuCOE/Nodes/CustomizableObjectNodeModifierClipMorphDetails.h"
 #include "MuCOE/Nodes/CustomizableObjectNodeModifierClipWithMesh.h"
 #include "MuCOE/Nodes/CustomizableObjectNodeModifierClipWithMeshDetails.h"
+#include "MuCOE/Nodes/CustomizableObjectNodeModifierClipWithUVMask.h"
+#include "MuCOE/Nodes/CustomizableObjectNodeModifierClipDeform.h"
 #include "MuCOE/Nodes/CustomizableObjectNodeMeshMorph.h"
 #include "MuCOE/Nodes/CustomizableObjectNodeMeshMorphDetails.h"
 #include "MuCOE/Nodes/CustomizableObjectNodeMeshReshapeCommon.h"
@@ -239,16 +242,19 @@ void FCustomizableObjectEditorModule::StartupModule()
 	FPropertyEditorModule& PropertyModule = FModuleManager::LoadModuleChecked<FPropertyEditorModule>("PropertyEditor");
 	RegisterCustomDetails(PropertyModule, UCustomizableObjectNodeLayoutBlocks::StaticClass(), FOnGetDetailCustomizationInstance::CreateStatic(&FCustomizableObjectNodeLayoutBlocksDetails::MakeInstance));
 	RegisterCustomDetails(PropertyModule, UCustomizableObjectNodeModifierEditMeshSection::StaticClass(), FOnGetDetailCustomizationInstance::CreateStatic(&FCustomizableObjectNodeModifierEditMeshSectionDetails::MakeInstance));
+	RegisterCustomDetails(PropertyModule, UCustomizableObjectNodeModifierExtendMeshSection::StaticClass(), FOnGetDetailCustomizationInstance::CreateStatic(&FCustomizableObjectNodeModifierExtendMeshSectionDetails::MakeInstance));
 	RegisterCustomDetails(PropertyModule, UCustomizableObjectNodeModifierRemoveMesh::StaticClass(), FOnGetDetailCustomizationInstance::CreateStatic(&FCustomizableObjectNodeModifierRemoveMeshDetails::MakeInstance));
 	RegisterCustomDetails(PropertyModule, UCustomizableObjectNodeModifierRemoveMeshBlocks::StaticClass(), FOnGetDetailCustomizationInstance::CreateStatic(&FCustomizableObjectNodeModifierRemoveMeshBlocksDetails::MakeInstance));
 	RegisterCustomDetails(PropertyModule, UCustomizableObjectNodeModifierMorphMeshSection::StaticClass(), FOnGetDetailCustomizationInstance::CreateStatic(&FCustomizableObjectNodeModifierMorphMeshSectionDetails::MakeInstance));
+	RegisterCustomDetails(PropertyModule, UCustomizableObjectNodeModifierClipMorph::StaticClass(), FOnGetDetailCustomizationInstance::CreateStatic(&FCustomizableObjectNodeModifierClipMorphDetails::MakeInstance));
+	RegisterCustomDetails(PropertyModule, UCustomizableObjectNodeModifierClipWithMesh::StaticClass(), FOnGetDetailCustomizationInstance::CreateStatic(&FCustomizableObjectNodeModifierClipWithMeshDetails::MakeInstance));
+	RegisterCustomDetails(PropertyModule, UCustomizableObjectNodeModifierClipWithUVMask::StaticClass(), FOnGetDetailCustomizationInstance::CreateStatic(&FCustomizableObjectNodeModifierBaseDetails::MakeInstance));
+	RegisterCustomDetails(PropertyModule, UCustomizableObjectNodeModifierClipDeform::StaticClass(), FOnGetDetailCustomizationInstance::CreateStatic(&FCustomizableObjectNodeModifierBaseDetails::MakeInstance));
 	RegisterCustomDetails(PropertyModule, UCustomizableObjectNodeObject::StaticClass(), FOnGetDetailCustomizationInstance::CreateStatic(&FCustomizableObjectNodeObjectDetails::MakeInstance));
 	RegisterCustomDetails(PropertyModule, UCustomizableObjectNodeObjectGroup::StaticClass(), FOnGetDetailCustomizationInstance::CreateStatic(&FCustomizableObjectNodeObjectGroupDetails::MakeInstance));
 	RegisterCustomDetails(PropertyModule, UCustomizableObjectNodeProjectorParameter::StaticClass(), FOnGetDetailCustomizationInstance::CreateStatic(&FCustomizableObjectNodeProjectorParameterDetails::MakeInstance));
 	RegisterCustomDetails(PropertyModule, UCustomizableObjectNodeProjectorConstant::StaticClass(), FOnGetDetailCustomizationInstance::CreateStatic(&FCustomizableObjectNodeProjectorParameterDetails::MakeInstance));
 	RegisterCustomDetails(PropertyModule, UCustomizableObjectNodeMeshMorph::StaticClass(), FOnGetDetailCustomizationInstance::CreateStatic(&FCustomizableObjectNodeMeshMorphDetails::MakeInstance));
-	RegisterCustomDetails(PropertyModule, UCustomizableObjectNodeModifierClipMorph::StaticClass(), FOnGetDetailCustomizationInstance::CreateStatic(&FCustomizableObjectNodeModifierClipMorphDetails::MakeInstance));
-	RegisterCustomDetails(PropertyModule, UCustomizableObjectNodeModifierClipWithMesh::StaticClass(), FOnGetDetailCustomizationInstance::CreateStatic(&FCustomizableObjectNodeModifierClipWithMeshDetails::MakeInstance));
 	RegisterCustomDetails(PropertyModule, UCustomizableObjectNodeExternalPin::StaticClass(), FOnGetDetailCustomizationInstance::CreateStatic(&FCustomizableObjectNodeExternalPinDetails::MakeInstance));
 	RegisterCustomDetails(PropertyModule, UCustomizableObjectNodeMaterial::StaticClass(), FOnGetDetailCustomizationInstance::CreateStatic(&FCustomizableObjectNodeDetails::MakeInstance));
 	RegisterCustomDetails(PropertyModule, UCustomizableObjectNodeSkeletalMesh::StaticClass(), FOnGetDetailCustomizationInstance::CreateStatic(&FCustomizableObjectNodeSkeletalMeshDetails::MakeInstance));
diff --git a/Engine/Plugins/Experimental/Mutable/Source/CustomizableObjectEditor/Private/MuCOE/GenerateMutableSource/GenerateMutableSourceModifier.cpp b/Engine/Plugins/Experimental/Mutable/Source/CustomizableObjectEditor/Private/MuCOE/GenerateMutableSource/GenerateMutableSourceModifier.cpp
--- a/Engine/Plugins/Experimental/Mutable/Source/CustomizableObjectEditor/Private/MuCOE/GenerateMutableSource/GenerateMutableSourceModifier.cpp
+++ b/Engine/Plugins/Experimental/Mutable/Source/CustomizableObjectEditor/Private/MuCOE/GenerateMutableSource/GenerateMutableSourceModifier.cpp
@@ -77,7 +77,7 @@ mu::Ptr<mu::NodeModifier> GenerateMutableSourceModifier(const UEdGraphPin * Pin,
 		ClipNode->SetVertexSelectionBone(GenerationContext.GetBoneUnique(TypedNodeClip->BoneName), TypedNodeClip->MaxEffectRadius);
 
 		ClipNode->MultipleTagsPolicy = TypedNodeClip->MultipleTagPolicy;
-		for (const FString& Tag : TypedNodeClip->Tags)
+		for (const FString& Tag : TypedNodeClip->RequiredTags)
 		{
 			 ClipNode->RequiredTags.Add(Tag);
 		}
@@ -129,7 +129,7 @@ mu::Ptr<mu::NodeModifier> GenerateMutableSourceModifier(const UEdGraphPin * Pin,
 		}
 	
 		ClipNode->MultipleTagsPolicy = TypedNodeClipDeform->MultipleTagPolicy;
-		for (const FString& Tag : TypedNodeClipDeform->Tags)
+		for (const FString& Tag : TypedNodeClipDeform->RequiredTags)
 		{
 			ClipNode->RequiredTags.Add(Tag);
 		}		
@@ -202,7 +202,7 @@ mu::Ptr<mu::NodeModifier> GenerateMutableSourceModifier(const UEdGraphPin * Pin,
 		}
 
 		ClipNode->MultipleTagsPolicy = TypedNodeClipMesh->MultipleTagPolicy;
-		for (const FString& Tag : TypedNodeClipMesh->Tags)
+		for (const FString& Tag : TypedNodeClipMesh->RequiredTags)
 		{
 			 ClipNode->RequiredTags.Add(Tag);
 		}
@@ -259,7 +259,7 @@ mu::Ptr<mu::NodeModifier> GenerateMutableSourceModifier(const UEdGraphPin * Pin,
 		ClipNode->LayoutIndex = TypedNodeClipUVMask->UVChannelForMask;
 
 		ClipNode->MultipleTagsPolicy = TypedNodeClipUVMask->MultipleTagPolicy;
-		for (const FString& Tag : TypedNodeClipUVMask->Tags)
+		for (const FString& Tag : TypedNodeClipUVMask->RequiredTags)
 		{
 			ClipNode->RequiredTags.Add(Tag);
 		}
diff --git a/Engine/Plugins/Experimental/Mutable/Source/CustomizableObjectEditor/Private/MuCOE/GraphTraversal.cpp b/Engine/Plugins/Experimental/Mutable/Source/CustomizableObjectEditor/Private/MuCOE/GraphTraversal.cpp
--- a/Engine/Plugins/Experimental/Mutable/Source/CustomizableObjectEditor/Private/MuCOE/GraphTraversal.cpp
+++ b/Engine/Plugins/Experimental/Mutable/Source/CustomizableObjectEditor/Private/MuCOE/GraphTraversal.cpp
@@ -437,6 +437,8 @@ const UEdGraphPin* FindMeshBaseSource(const UEdGraphPin& Pin, const bool bOnlyLo
 		Pin.PinType.PinCategory == UEdGraphSchema_CustomizableObject::PC_PassThroughMesh
 		||
 		Pin.PinType.PinCategory == UEdGraphSchema_CustomizableObject::PC_Material
+		||
+		Pin.PinType.PinCategory == UEdGraphSchema_CustomizableObject::PC_Modifier
 	);
 	
 	const UEdGraphNode* Node = Pin.GetOwningNode();
diff --git a/Engine/Plugins/Experimental/Mutable/Source/CustomizableObjectEditor/Private/MuCOE/Nodes/CustomizableObjectNodeMeshClipDeform.cpp b/Engine/Plugins/Experimental/Mutable/Source/CustomizableObjectEditor/Private/MuCOE/Nodes/CustomizableObjectNodeMeshClipDeform.cpp
--- a/Engine/Plugins/Experimental/Mutable/Source/CustomizableObjectEditor/Private/MuCOE/Nodes/CustomizableObjectNodeMeshClipDeform.cpp
+++ b/Engine/Plugins/Experimental/Mutable/Source/CustomizableObjectEditor/Private/MuCOE/Nodes/CustomizableObjectNodeMeshClipDeform.cpp
@@ -3,6 +3,7 @@
 #include "MuCOE/Nodes/CustomizableObjectNodeModifierClipDeform.h"
 
 #include "MuCOE/EdGraphSchema_CustomizableObject.h"
+#include "MuCO/CustomizableObjectCustomVersion.h"
 
 class UCustomizableObjectNodeRemapPins;
 
@@ -21,6 +22,20 @@ void UCustomizableObjectNodeModifierClipDeform::AllocateDefaultPins(UCustomizabl
 }
 
 
+void UCustomizableObjectNodeModifierClipDeform::BackwardsCompatibleFixup(int32 CustomizableObjectCustomVersion)
+{
+	Super::BackwardsCompatibleFixup(CustomizableObjectCustomVersion);
+
+	const UEdGraphSchema_CustomizableObject* Schema = GetDefault<UEdGraphSchema_CustomizableObject>();
+
+	if (CustomizableObjectCustomVersion == FCustomizableObjectCustomVersion::UnifyRequiredTags)
+	{
+		RequiredTags = Tags_DEPRECATED;
+		Tags_DEPRECATED.Empty();
+	}
+}
+
+
 UEdGraphPin* UCustomizableObjectNodeModifierClipDeform::ClipShapePin() const
 {
 	return FindPin(TEXT("Clip Shape"), EGPD_Input);
diff --git a/Engine/Plugins/Experimental/Mutable/Source/CustomizableObjectEditor/Private/MuCOE/Nodes/CustomizableObjectNodeMeshClipMorph.cpp b/Engine/Plugins/Experimental/Mutable/Source/CustomizableObjectEditor/Private/MuCOE/Nodes/CustomizableObjectNodeMeshClipMorph.cpp
--- a/Engine/Plugins/Experimental/Mutable/Source/CustomizableObjectEditor/Private/MuCOE/Nodes/CustomizableObjectNodeMeshClipMorph.cpp
+++ b/Engine/Plugins/Experimental/Mutable/Source/CustomizableObjectEditor/Private/MuCOE/Nodes/CustomizableObjectNodeMeshClipMorph.cpp
@@ -163,6 +163,20 @@ void UCustomizableObjectNodeModifierClipMorph::PinConnectionListChanged(UEdGraph
 }
 
 
+void UCustomizableObjectNodeModifierClipMorph::BackwardsCompatibleFixup(int32 CustomizableObjectCustomVersion)
+{
+	Super::BackwardsCompatibleFixup(CustomizableObjectCustomVersion);
+
+	const UEdGraphSchema_CustomizableObject* Schema = GetDefault<UEdGraphSchema_CustomizableObject>();
+
+	if (CustomizableObjectCustomVersion == FCustomizableObjectCustomVersion::UnifyRequiredTags)
+	{
+		RequiredTags = Tags_DEPRECATED;
+		Tags_DEPRECATED.Empty();
+	}
+}
+
+
 FText UCustomizableObjectNodeModifierClipMorph::GetTooltipText() const
 {
 	return LOCTEXT("Clip_Mesh_Morph_Tooltip", "Defines a cutting plane on a bone to cut tagged Materials that go past it, while morphing the mesh after the cut to blend in more naturally.\nIt only cuts and morphs mesh that receives some influence of that bone or other descendant bones.");
diff --git a/Engine/Plugins/Experimental/Mutable/Source/CustomizableObjectEditor/Private/MuCOE/Nodes/CustomizableObjectNodeMeshClipWithMesh.cpp b/Engine/Plugins/Experimental/Mutable/Source/CustomizableObjectEditor/Private/MuCOE/Nodes/CustomizableObjectNodeMeshClipWithMesh.cpp
--- a/Engine/Plugins/Experimental/Mutable/Source/CustomizableObjectEditor/Private/MuCOE/Nodes/CustomizableObjectNodeMeshClipWithMesh.cpp
+++ b/Engine/Plugins/Experimental/Mutable/Source/CustomizableObjectEditor/Private/MuCOE/Nodes/CustomizableObjectNodeMeshClipWithMesh.cpp
@@ -5,6 +5,7 @@
 #include "MuCOE/CustomizableObjectGraph.h"
 #include "MuCOE/EdGraphSchema_CustomizableObject.h"
 #include "MuCOE/ICustomizableObjectEditor.h"
+#include "MuCO/CustomizableObjectCustomVersion.h"
 
 class UCustomizableObjectNodeRemapPins;
 
@@ -56,6 +57,20 @@ void UCustomizableObjectNodeModifierClipWithMesh::AllocateDefaultPins(UCustomiza
 }
 
 
+void UCustomizableObjectNodeModifierClipWithMesh::BackwardsCompatibleFixup(int32 CustomizableObjectCustomVersion)
+{
+	Super::BackwardsCompatibleFixup(CustomizableObjectCustomVersion);
+
+	const UEdGraphSchema_CustomizableObject* Schema = GetDefault<UEdGraphSchema_CustomizableObject>();
+
+	if (CustomizableObjectCustomVersion == FCustomizableObjectCustomVersion::UnifyRequiredTags)
+	{
+		RequiredTags = Tags_DEPRECATED;
+		Tags_DEPRECATED.Empty();
+	}
+}
+
+
 void UCustomizableObjectNodeModifierClipWithMesh::ReconstructNode(UCustomizableObjectNodeRemapPins* RemapPinsMode)
 {
 	Super::ReconstructNode(RemapPinsMode);
diff --git a/Engine/Plugins/Experimental/Mutable/Source/CustomizableObjectEditor/Private/MuCOE/Nodes/CustomizableObjectNodeModifierBase.cpp b/Engine/Plugins/Experimental/Mutable/Source/CustomizableObjectEditor/Private/MuCOE/Nodes/CustomizableObjectNodeModifierBase.cpp
--- a/Engine/Plugins/Experimental/Mutable/Source/CustomizableObjectEditor/Private/MuCOE/Nodes/CustomizableObjectNodeModifierBase.cpp
+++ b/Engine/Plugins/Experimental/Mutable/Source/CustomizableObjectEditor/Private/MuCOE/Nodes/CustomizableObjectNodeModifierBase.cpp
@@ -23,7 +23,7 @@ UEdGraphPin* UCustomizableObjectNodeModifierBase::OutputPin() const
 }
 
 
-bool UCustomizableObjectNodeModifierBase::IsApplicableTo(UCustomizableObjectNode* Candidate) const
+bool UCustomizableObjectNodeModifierBase::IsApplicableTo(UCustomizableObjectNode* Candidate)
 {
 	if (!Candidate)
 	{
@@ -31,14 +31,13 @@ bool UCustomizableObjectNodeModifierBase::IsApplicableTo(UCustomizableObjectNode
 	}
 
 	const TArray<FString>* EnabledTags = Candidate->GetEnableTags();
-	const TArray<FString>* RequiredTags = GetRequiredTags();
-	if (EnabledTags && RequiredTags)
+	if (EnabledTags)
 	{
-		switch ( GetMultipleTagsPolicy() )
+		switch ( MultipleTagPolicy )
 		{
 		case EMutableMultipleTagPolicy::OnlyOneRequired:
 		{
-			for (const FString& RequiredTag : *RequiredTags)
+			for (const FString& RequiredTag : RequiredTags)
 			{
 				if (EnabledTags->Contains(RequiredTag))
 				{
@@ -50,7 +49,7 @@ bool UCustomizableObjectNodeModifierBase::IsApplicableTo(UCustomizableObjectNode
 
 		case EMutableMultipleTagPolicy::AllRequired:
 		{
-			for (const FString& RequiredTag : *RequiredTags)
+			for (const FString& RequiredTag : RequiredTags)
 			{
 				if (!EnabledTags->Contains(RequiredTag))
 				{
@@ -71,7 +70,7 @@ bool UCustomizableObjectNodeModifierBase::IsApplicableTo(UCustomizableObjectNode
 }
 
 
-void UCustomizableObjectNodeModifierBase::GetPossiblyModifiedNodes(TArray<UCustomizableObjectNode*>& CandidateNodes) const
+void UCustomizableObjectNodeModifierBase::GetPossiblyModifiedNodes(TArray<UCustomizableObjectNode*>& CandidateNodes)
 {
 	// Scan all potential receivers
 	UCustomizableObject* ThisNodeObject = GetRootObject(*this);
@@ -264,22 +263,18 @@ void UCustomizableObjectNodeModifierBase::PostBackwardsCompatibleFixup()
 	}
 
 	// Apply backwards compatibility auto-generated tags to external objects.
-	const TArray<FString>* RequiredTags = GetRequiredTags();
-	if (RequiredTags)
+	for (const FLegacyTag& Tag : LegacyBackportsRequiredTags)
 	{
-		for (const FLegacyTag& Tag : LegacyBackportsRequiredTags)
+		// If we still have it in this node
+		if (RequiredTags.Contains(Tag.Tag))
 		{
-			// If we still have it in this node
-			if (RequiredTags->Contains(Tag.Tag))
+			UCustomizableObjectNode* ParentNode = GetCustomizableObjectExternalNode<UCustomizableObjectNode>(Tag.ParentObject.Get(), Tag.ParentNode);
+			if (ParentNode)
 			{
-				UCustomizableObjectNode* ParentNode = GetCustomizableObjectExternalNode<UCustomizableObjectNode>(Tag.ParentObject.Get(), Tag.ParentNode);
-				if (ParentNode)
+				TArray<FString>* NodeEnableTags = ParentNode->GetEnableTags();
+				if (NodeEnableTags)
 				{
-					TArray<FString>* NodeEnableTags = ParentNode->GetEnableTags();
-					if (NodeEnableTags)
-					{
-						NodeEnableTags->AddUnique(Tag.Tag);
-					}
+					NodeEnableTags->AddUnique(Tag.Tag);
 				}
 			}
 		}
diff --git a/Engine/Plugins/Experimental/Mutable/Source/CustomizableObjectEditor/Private/MuCOE/Nodes/CustomizableObjectNodeModifierBaseDetails.cpp b/Engine/Plugins/Experimental/Mutable/Source/CustomizableObjectEditor/Private/MuCOE/Nodes/CustomizableObjectNodeModifierBaseDetails.cpp
new file mode 100644
--- /dev/null
+++ b/Engine/Plugins/Experimental/Mutable/Source/CustomizableObjectEditor/Private/MuCOE/Nodes/CustomizableObjectNodeModifierBaseDetails.cpp
@@ -0,0 +1,269 @@
+// Copyright Epic Games, Inc. All Rights Reserved.
+
+#include "MuCOE/Nodes/CustomizableObjectNodeModifierBaseDetails.h"
+
+#include "MuCOE/CustomizableObjectEditorUtilities.h"
+#include "MuCOE/CustomizableObjectEditorStyle.h"
+#include "MuCOE/GraphTraversal.h"
+#include "MuCOE/Nodes/CustomizableObjectNodeModifierBase.h"
+#include "MuCOE/Nodes/CustomizableObjectNodeMaterial.h"
+#include "MuCOE/Nodes/CustomizableObjectNodeSkeletalMesh.h"
+#include "MuCOE/UnrealEditorPortabilityHelpers.h"
+#include "PropertyCustomizationHelpers.h"
+#include "Widgets/Input/STextComboBox.h"
+#include "Widgets/Input/SSearchBox.h"
+#include "Widgets/Input/SButton.h"
+#include "DetailLayoutBuilder.h"
+#include "IDetailsView.h"
+#include "Engine/SkeletalMesh.h"
+#include "Misc/Attribute.h"
+
+
+#define LOCTEXT_NAMESPACE "CustomizableObjectDetails"
+
+
+TSharedRef<IDetailCustomization> FCustomizableObjectNodeModifierBaseDetails::MakeInstance()
+{
+	return MakeShareable( new FCustomizableObjectNodeModifierBaseDetails );
+}
+
+
+void FCustomizableObjectNodeModifierBaseDetails::CustomizeDetails( IDetailLayoutBuilder& DetailBuilder )
+{
+	FCustomizableObjectNodeDetails::CustomizeDetails(DetailBuilder);
+
+	const IDetailsView* DetailsView = DetailBuilder.GetDetailsView();
+	if ( DetailsView->GetSelectedObjects().Num() )
+	{
+		Node = Cast<UCustomizableObjectNodeModifierBase>( DetailsView->GetSelectedObjects()[0].Get() );
+	}
+
+	if (!Node)
+	{
+		return;
+	}
+	
+	// Move modifier conditions to the top.
+	IDetailCategoryBuilder& ModifierCategory = DetailBuilder.EditCategory("Modifier");
+	ModifierCategory.SetSortOrder(-10000);
+
+	// Add the required tags widget
+	{
+		RefreshTagOptions();
+
+		RequiredTagsPropertyHandle = DetailBuilder.GetProperty(GET_MEMBER_NAME_CHECKED(UCustomizableObjectNodeModifierBase, RequiredTags), UCustomizableObjectNodeModifierBase::StaticClass());
+		DetailBuilder.HideProperty(RequiredTagsPropertyHandle);
+
+		RequiredTagsPropertyHandle->SetOnPropertyValueChanged(FSimpleDelegate::CreateSP(this, &FCustomizableObjectNodeModifierBaseDetails::OnRequiredTagsPropertyChanged));
+		RequiredTagsPropertyHandle->SetOnChildPropertyValueChanged(FSimpleDelegate::CreateSP(this, &FCustomizableObjectNodeModifierBaseDetails::OnRequiredTagsPropertyChanged));
+
+		TagsPolicyPropertyHandle = DetailBuilder.GetProperty(GET_MEMBER_NAME_CHECKED(UCustomizableObjectNodeModifierBase, MultipleTagPolicy), UCustomizableObjectNodeModifierBase::StaticClass());
+		TagsPolicyPropertyHandle->SetOnPropertyValueChanged(FSimpleDelegate::CreateSP(this, &FCustomizableObjectNodeModifierBaseDetails::OnRequiredTagsPropertyChanged));
+
+		ModifierCategory.AddCustomRow(FText::FromString(TEXT("Required Tags")))
+			.PropertyHandleList({ RequiredTagsPropertyHandle })
+			.NameContent()
+				.VAlign(VAlign_Fill)
+			[
+				SNew(SVerticalBox)
+					+ SVerticalBox::Slot()
+					.VAlign(VAlign_Top)
+					.Padding(FMargin(0,4.0f,0,4.0f))
+					[
+						SNew(STextBlock)
+							.Text(LOCTEXT("tableodifierDetails_RequiredTags", "Required Tags"))
+							.Font(IDetailLayoutBuilder::GetDetailFont())
+					]
+			]
+			.ValueContent()
+				.HAlign(HAlign_Fill)
+			[
+				SNew(SVerticalBox)
+
+					// Header with the "add tag" UI
+					+ SVerticalBox::Slot()
+					.AutoHeight()
+					.HAlign(HAlign_Fill)
+					[
+						SNew(SHorizontalBox)
+
+						+ SHorizontalBox::Slot()
+						.FillWidth(10.0f)
+						.HAlign(HAlign_Fill)
+						.VAlign(VAlign_Center)
+						[
+							SAssignNew(this->TagCombo, SMutableSearchComboBox)
+								.ButtonStyle(FAppStyle::Get(), "NoBorder")
+								.MenuButtonBrush(FAppStyle::GetBrush(TEXT("Icons.PlusCircle")))
+								.OptionsSource(&TagComboOptionsSource)
+								.AllowAddNewOptions(true)
+								.OnSelectionChanged(this, &FCustomizableObjectNodeModifierBaseDetails::OnTagComboBoxSelectionChanged)
+						]
+					]
+
+					// List of tags
+					+ SVerticalBox::Slot()
+					.VAlign(VAlign_Fill)
+					.HAlign(HAlign_Fill)
+					[
+						SAssignNew(this->TagList, SListView<TSharedPtr<FTagUIData>>)
+							.ListItemsSource(&CurrentTagsSource)
+							.OnGenerateRow(this, &FCustomizableObjectNodeModifierBaseDetails::GenerateTagMenuItemRow)
+							.SelectionMode(ESelectionMode::None)							
+					]
+
+					// No-tags warning: shown only if there are no required tags defined
+					+ SVerticalBox::Slot()
+					.AutoHeight()
+					.HAlign(HAlign_Fill)
+					[
+						SNew(STextBlock)
+							.Text(LOCTEXT("ModifierDetails_NoRequiredTagsWarning", "Warning: There are no required tags, so this modifier will not do anything."))
+							.AutoWrapText(true)
+							.Font(IDetailLayoutBuilder::GetDetailFont())
+							.Visibility_Lambda([&]() 
+								{
+									return (!Node || Node->RequiredTags.IsEmpty()) ? EVisibility::Visible : EVisibility::Collapsed;
+								})
+					]
+			];
+
+	}
+}
+
+
+void FCustomizableObjectNodeModifierBaseDetails::OnRequiredTagsPropertyChanged()
+{
+	// This seems necessary to detect the "Reset to default" actions.
+	RefreshTagOptions();
+}
+
+
+void FCustomizableObjectNodeModifierBaseDetails::RefreshTagOptions()
+{
+	// Tag combo options
+	{
+		TagComboOptionsSource.SetNum(0,EAllowShrinking::No);
+
+		// Scan all potential receivers
+		UCustomizableObject* ThisNodeObject = GetRootObject(*Node);
+		UCustomizableObject* RootObject = GetRootObject(ThisNodeObject);
+
+		TSet<UCustomizableObject*> AllCustomizableObject;
+		GetAllObjectsInGraph(RootObject, AllCustomizableObject);
+
+		for (const UCustomizableObject* CustObject : AllCustomizableObject)
+		{
+			if (CustObject)
+			{
+				for (const TObjectPtr<UEdGraphNode>& CandidateNode : CustObject->GetPrivate()->GetSource()->Nodes)
+				{
+					UCustomizableObjectNode* Typed = Cast<UCustomizableObjectNode>(CandidateNode);
+					if (Typed)
+					{
+						TArray<FString>* EnableTags = Typed->GetEnableTags();
+						if (EnableTags)
+						{
+							for (const FString& Tag : *EnableTags)
+							{
+								if (!Tag.IsEmpty())
+								{
+									TagComboOptionsSource.AddUnique(MakeShared<FString>(Tag));
+								}
+							}
+						}
+					}
+				}
+			}
+		}
+
+		// TODO: add material nodes and options to create tags for them, show hierarchy, ...
+	}
+
+	// Current Tags
+	if (Node)
+	{
+		CurrentTagsSource.SetNum(0, EAllowShrinking::No);
+
+		for (const FString& Tag : Node->RequiredTags)
+		{
+			TSharedPtr<FTagUIData> Data = MakeShared<FTagUIData>();
+			Data->DisplayName = Tag;
+			Data->Tag = Tag;
+			CurrentTagsSource.Add(Data);
+		}
+	}
+
+	if (TagList)
+	{
+		TagList->RequestListRefresh();
+	}
+}
+
+
+TSharedRef<ITableRow> FCustomizableObjectNodeModifierBaseDetails::GenerateTagMenuItemRow(TSharedPtr<FTagUIData> InItem, const TSharedRef<STableViewBase>& OwnerTable)
+{
+	if (!InItem)
+	{
+		return SNew(SComboRow<TSharedPtr<FString>>, OwnerTable)
+			[
+				SNew(STextBlock)
+					.Text(FText::FromString("No item."))
+					.Font(IDetailLayoutBuilder::GetDetailFont())
+			];
+	}
+
+	return SNew(SComboRow<TSharedPtr<FString>>, OwnerTable)
+		[
+			SNew(SHorizontalBox)
+
+				+ SHorizontalBox::Slot()
+				.HAlign(HAlign_Fill)
+				.VAlign(VAlign_Center)
+				[
+					SNew(STextBlock)
+						.Text(FText::FromString(InItem->DisplayName))
+						.Font(IDetailLayoutBuilder::GetDetailFont())
+				]
+
+				+ SHorizontalBox::Slot()
+				.AutoWidth()
+				[
+					SNew(SButton)
+						.ButtonStyle(FAppStyle::Get(), "NoBorder")
+						.ToolTipText(LOCTEXT("RemoveModifierTag","Remove this tag from the modifier."))
+						.OnClicked_Lambda([this,InItem]() 
+							{
+								if (Node && InItem )
+								{
+									if (Node->RequiredTags.Contains(InItem->Tag))
+									{
+										Node->RequiredTags.Remove(InItem->Tag);
+										Node->Modify();
+										OnRequiredTagsPropertyChanged();
+									}
+								}
+								return FReply::Handled();
+							})
+						[
+							SNew(SImage)
+								.Image(FAppStyle::GetBrush(TEXT("Icons.MinusCircle")))
+						]
+				]
+		];
+}
+
+
+void FCustomizableObjectNodeModifierBaseDetails::OnTagComboBoxSelectionChanged(const FText& NewText)
+{
+	FString Tag = NewText.ToString();
+	if (Node && !Node->RequiredTags.Contains(Tag))
+	{
+		Node->RequiredTags.Add(Tag);
+		Node->Modify();
+		OnRequiredTagsPropertyChanged();
+	}
+}
+
+
+#undef LOCTEXT_NAMESPACE
diff --git a/Engine/Plugins/Experimental/Mutable/Source/CustomizableObjectEditor/Private/MuCOE/Nodes/CustomizableObjectNodeModifierBaseDetails.h b/Engine/Plugins/Experimental/Mutable/Source/CustomizableObjectEditor/Private/MuCOE/Nodes/CustomizableObjectNodeModifierBaseDetails.h
new file mode 100644
--- /dev/null
+++ b/Engine/Plugins/Experimental/Mutable/Source/CustomizableObjectEditor/Private/MuCOE/Nodes/CustomizableObjectNodeModifierBaseDetails.h
@@ -0,0 +1,59 @@
+// Copyright Epic Games, Inc. All Rights Reserved.
+
+#pragma once
+
+#include "CustomizableObjectNodeDetails.h"
+#include "MuCOE/SMutableSearchComboBox.h"
+#include "IDetailCustomization.h"
+#include "Types/SlateEnums.h"
+#include "Widgets/Input/SComboButton.h"
+#include "Widgets/Input/SEditableTextBox.h"
+#include "Widgets/Views/SListView.h"
+#include "Widgets/Views/ITableRow.h"
+
+namespace ESelectInfo { enum Type : int; }
+
+class FString;
+class IDetailLayoutBuilder;
+class IPropertyHandle;
+class UCustomizableObjectNodeObject;
+
+
+class FCustomizableObjectNodeModifierBaseDetails : public FCustomizableObjectNodeDetails
+{
+public:
+	// Makes a new instance of this detail layout class for a specific detail view requesting it 
+	static TSharedRef<IDetailCustomization> MakeInstance();
+
+	// ILayoutDetails interface
+	void CustomizeDetails( IDetailLayoutBuilder& DetailBuilder ) override;
+
+protected:
+
+	/** Called when a required tags or policy property has changed, potentially changing the modified nodes. */
+	virtual void OnRequiredTagsPropertyChanged();
+
+private:
+
+	class UCustomizableObjectNodeModifierBase* Node = nullptr;
+
+	/** */
+	TSharedPtr<IPropertyHandle> RequiredTagsPropertyHandle;
+	TSharedPtr<IPropertyHandle> TagsPolicyPropertyHandle;
+	TSharedPtr<SMutableSearchComboBox> TagCombo;
+	TArray< TSharedPtr<FString> > TagComboOptionsSource;
+	void OnTagComboBoxSelectionChanged(const FText& NewText);
+
+	struct FTagUIData
+	{
+		FString Tag;
+		FString DisplayName;
+	};
+	TSharedPtr<SListView<TSharedPtr<FTagUIData>>> TagList;
+	TArray< TSharedPtr<FTagUIData> > CurrentTagsSource;
+
+	/** */
+	void RefreshTagOptions();
+	TSharedRef<ITableRow> GenerateTagMenuItemRow(TSharedPtr<FTagUIData> InItem, const TSharedRef<STableViewBase>& OwnerTable);
+
+};
diff --git a/Engine/Plugins/Experimental/Mutable/Source/CustomizableObjectEditor/Private/MuCOE/Nodes/CustomizableObjectNodeModifierClipMorphDetails.cpp b/Engine/Plugins/Experimental/Mutable/Source/CustomizableObjectEditor/Private/MuCOE/Nodes/CustomizableObjectNodeModifierClipMorphDetails.cpp
--- a/Engine/Plugins/Experimental/Mutable/Source/CustomizableObjectEditor/Private/MuCOE/Nodes/CustomizableObjectNodeModifierClipMorphDetails.cpp
+++ b/Engine/Plugins/Experimental/Mutable/Source/CustomizableObjectEditor/Private/MuCOE/Nodes/CustomizableObjectNodeModifierClipMorphDetails.cpp
@@ -29,7 +29,7 @@ TSharedRef<IDetailCustomization> FCustomizableObjectNodeModifierClipMorphDetails
 
 void FCustomizableObjectNodeModifierClipMorphDetails::CustomizeDetails(IDetailLayoutBuilder& DetailBuilder)
 {
-	FCustomizableObjectNodeDetails::CustomizeDetails(DetailBuilder);
+	FCustomizableObjectNodeModifierBaseDetails::CustomizeDetails(DetailBuilder);
 
 	Node = nullptr;
 	DetailBuilderPtr = &DetailBuilder;
diff --git a/Engine/Plugins/Experimental/Mutable/Source/CustomizableObjectEditor/Private/MuCOE/Nodes/CustomizableObjectNodeModifierClipMorphDetails.h b/Engine/Plugins/Experimental/Mutable/Source/CustomizableObjectEditor/Private/MuCOE/Nodes/CustomizableObjectNodeModifierClipMorphDetails.h
--- a/Engine/Plugins/Experimental/Mutable/Source/CustomizableObjectEditor/Private/MuCOE/Nodes/CustomizableObjectNodeModifierClipMorphDetails.h
+++ b/Engine/Plugins/Experimental/Mutable/Source/CustomizableObjectEditor/Private/MuCOE/Nodes/CustomizableObjectNodeModifierClipMorphDetails.h
@@ -2,7 +2,7 @@
 
 #pragma once
 
-#include "CustomizableObjectNodeDetails.h"
+#include "CustomizableObjectNodeModifierBaseDetails.h"
 #include "IDetailCustomization.h"
 #include "UObject/ObjectMacros.h"
 
@@ -14,7 +14,7 @@ class IDetailLayoutBuilder;
 class IPropertyHandle;
 
 
-class FCustomizableObjectNodeModifierClipMorphDetails : public FCustomizableObjectNodeDetails
+class FCustomizableObjectNodeModifierClipMorphDetails : public FCustomizableObjectNodeModifierBaseDetails
 {
 public:
 	/** Makes a new instance of this detail layout class for a specific detail view requesting it */
diff --git a/Engine/Plugins/Experimental/Mutable/Source/CustomizableObjectEditor/Private/MuCOE/Nodes/CustomizableObjectNodeModifierClipWithMeshDetails.cpp b/Engine/Plugins/Experimental/Mutable/Source/CustomizableObjectEditor/Private/MuCOE/Nodes/CustomizableObjectNodeModifierClipWithMeshDetails.cpp
--- a/Engine/Plugins/Experimental/Mutable/Source/CustomizableObjectEditor/Private/MuCOE/Nodes/CustomizableObjectNodeModifierClipWithMeshDetails.cpp
+++ b/Engine/Plugins/Experimental/Mutable/Source/CustomizableObjectEditor/Private/MuCOE/Nodes/CustomizableObjectNodeModifierClipWithMeshDetails.cpp
@@ -31,7 +31,7 @@ TSharedRef<IDetailCustomization> FCustomizableObjectNodeModifierClipWithMeshDeta
 
 void FCustomizableObjectNodeModifierClipWithMeshDetails::CustomizeDetails(IDetailLayoutBuilder& DetailBuilder)
 {
-	FCustomizableObjectNodeDetails::CustomizeDetails(DetailBuilder);
+	FCustomizableObjectNodeModifierBaseDetails::CustomizeDetails(DetailBuilder);
 
 	Node = nullptr;
 	DetailBuilderPtr = &DetailBuilder;
@@ -241,7 +241,7 @@ TSharedPtr<FString> FCustomizableObjectNodeModifierClipWithMeshDetails::SetIniti
 			Node->bUseMaterials = true;
 		}
 
-		if (Node->Tags.Num() > 0)
+		if (Node->RequiredTags.Num() > 0)
 		{
 			InitialMethod = ClippingMethods[1];
 			Node->bUseTags = true;
@@ -286,7 +286,7 @@ void FCustomizableObjectNodeModifierClipWithMeshDetails::OnClippingMethodComboBo
 		{
 			Node->bUseMaterials = true;
 			Node->bUseTags = false;
-			Node->Tags.Empty();
+			Node->RequiredTags.Empty();
 		}
 		else
 		{
@@ -324,7 +324,7 @@ FReply FCustomizableObjectNodeModifierClipWithMeshDetails::OnAddTagPressed()
 {
 	if (Node)
 	{
-		Node->Tags.Add("");
+		Node->RequiredTags.Add("");
 		DetailBuilderPtr->ForceRefreshDetails();
 
 		return FReply::Handled();
diff --git a/Engine/Plugins/Experimental/Mutable/Source/CustomizableObjectEditor/Private/MuCOE/Nodes/CustomizableObjectNodeModifierClipWithMeshDetails.h b/Engine/Plugins/Experimental/Mutable/Source/CustomizableObjectEditor/Private/MuCOE/Nodes/CustomizableObjectNodeModifierClipWithMeshDetails.h
--- a/Engine/Plugins/Experimental/Mutable/Source/CustomizableObjectEditor/Private/MuCOE/Nodes/CustomizableObjectNodeModifierClipWithMeshDetails.h
+++ b/Engine/Plugins/Experimental/Mutable/Source/CustomizableObjectEditor/Private/MuCOE/Nodes/CustomizableObjectNodeModifierClipWithMeshDetails.h
@@ -2,7 +2,7 @@
 
 #pragma once
 
-#include "CustomizableObjectNodeDetails.h"
+#include "CustomizableObjectNodeModifierBaseDetails.h"
 #include "IDetailCustomization.h"
 #include "Widgets/SCompoundWidget.h"
 
@@ -17,7 +17,7 @@ class UCustomizableObjectNodeModifierClipWithMesh;
 struct FAssetData;
 
 
-class FCustomizableObjectNodeModifierClipWithMeshDetails : public FCustomizableObjectNodeDetails
+class FCustomizableObjectNodeModifierClipWithMeshDetails : public FCustomizableObjectNodeModifierBaseDetails
 {
 public:
 	/** Makes a new instance of this detail layout class for a specific detail view requesting it */
diff --git a/Engine/Plugins/Experimental/Mutable/Source/CustomizableObjectEditor/Private/MuCOE/Nodes/CustomizableObjectNodeModifierClipWithUVMask.cpp b/Engine/Plugins/Experimental/Mutable/Source/CustomizableObjectEditor/Private/MuCOE/Nodes/CustomizableObjectNodeModifierClipWithUVMask.cpp
--- a/Engine/Plugins/Experimental/Mutable/Source/CustomizableObjectEditor/Private/MuCOE/Nodes/CustomizableObjectNodeModifierClipWithUVMask.cpp
+++ b/Engine/Plugins/Experimental/Mutable/Source/CustomizableObjectEditor/Private/MuCOE/Nodes/CustomizableObjectNodeModifierClipWithUVMask.cpp
@@ -5,13 +5,13 @@
 #include "MuCOE/CustomizableObjectGraph.h"
 #include "MuCOE/EdGraphSchema_CustomizableObject.h"
 #include "MuCOE/ICustomizableObjectEditor.h"
+#include "MuCO/CustomizableObjectCustomVersion.h"
 
 class UCustomizableObjectNodeRemapPins;
 
 #define LOCTEXT_NAMESPACE "CustomizableObjectEditor"
 
 
-
 void UCustomizableObjectNodeModifierClipWithUVMask::AllocateDefaultPins(UCustomizableObjectNodeRemapPins* RemapPins)
 {
 	const UEdGraphSchema_CustomizableObject* Schema = GetDefault<UEdGraphSchema_CustomizableObject>();
@@ -42,6 +42,20 @@ void UCustomizableObjectNodeModifierClipWithUVMask::PinConnectionListChanged(UEd
 }
 
 
+void UCustomizableObjectNodeModifierClipWithUVMask::BackwardsCompatibleFixup(int32 CustomizableObjectCustomVersion)
+{
+	Super::BackwardsCompatibleFixup(CustomizableObjectCustomVersion);
+
+	const UEdGraphSchema_CustomizableObject* Schema = GetDefault<UEdGraphSchema_CustomizableObject>();
+
+	if (CustomizableObjectCustomVersion == FCustomizableObjectCustomVersion::UnifyRequiredTags)
+	{
+		RequiredTags = Tags_DEPRECATED;
+		Tags_DEPRECATED.Empty();
+	}
+}
+
+
 UEdGraphPin* UCustomizableObjectNodeModifierClipWithUVMask::ClipMaskPin() const
 {
 	return FindPin(TEXT("Clip Mask"));
diff --git a/Engine/Plugins/Experimental/Mutable/Source/CustomizableObjectEditor/Private/MuCOE/Nodes/CustomizableObjectNodeModifierEditMeshSectionDetails.cpp b/Engine/Plugins/Experimental/Mutable/Source/CustomizableObjectEditor/Private/MuCOE/Nodes/CustomizableObjectNodeModifierEditMeshSectionDetails.cpp
--- a/Engine/Plugins/Experimental/Mutable/Source/CustomizableObjectEditor/Private/MuCOE/Nodes/CustomizableObjectNodeModifierEditMeshSectionDetails.cpp
+++ b/Engine/Plugins/Experimental/Mutable/Source/CustomizableObjectEditor/Private/MuCOE/Nodes/CustomizableObjectNodeModifierEditMeshSectionDetails.cpp
@@ -30,7 +30,7 @@ TSharedRef<IDetailCustomization> FCustomizableObjectNodeModifierEditMeshSectionD
 
 void FCustomizableObjectNodeModifierEditMeshSectionDetails::CustomizeDetails( IDetailLayoutBuilder& DetailBuilder )
 {
-	FCustomizableObjectNodeDetails::CustomizeDetails(DetailBuilder);
+	FCustomizableObjectNodeModifierBaseDetails::CustomizeDetails(DetailBuilder);
 
 	Node = nullptr;
 	const IDetailsView* DetailsView = DetailBuilder.GetDetailsView();
@@ -136,6 +136,13 @@ void FCustomizableObjectNodeModifierEditMeshSectionDetails::CustomizeDetails( ID
 }
 
 
+void FCustomizableObjectNodeModifierEditMeshSectionDetails::OnRequiredTagsPropertyChanged()
+{
+	FCustomizableObjectNodeModifierBaseDetails::OnRequiredTagsPropertyChanged();
+	UpdateLayout();
+}
+
+
 void FCustomizableObjectNodeModifierEditMeshSectionDetails::UpdateLayout()
 {
 	MUTABLE_CPUPROFILER_SCOPE(FCustomizableObjectNodeModifierEditMeshSectionDetails_UpdateLayout);
diff --git a/Engine/Plugins/Experimental/Mutable/Source/CustomizableObjectEditor/Private/MuCOE/Nodes/CustomizableObjectNodeModifierEditMeshSectionDetails.h b/Engine/Plugins/Experimental/Mutable/Source/CustomizableObjectEditor/Private/MuCOE/Nodes/CustomizableObjectNodeModifierEditMeshSectionDetails.h
--- a/Engine/Plugins/Experimental/Mutable/Source/CustomizableObjectEditor/Private/MuCOE/Nodes/CustomizableObjectNodeModifierEditMeshSectionDetails.h
+++ b/Engine/Plugins/Experimental/Mutable/Source/CustomizableObjectEditor/Private/MuCOE/Nodes/CustomizableObjectNodeModifierEditMeshSectionDetails.h
@@ -3,7 +3,7 @@
 #pragma once
 
 
-#include "CustomizableObjectNodeDetails.h"
+#include "CustomizableObjectNodeModifierBaseDetails.h"
 #include "IDetailCustomization.h"
 #include "Types/SlateEnums.h"
 
@@ -13,7 +13,7 @@ class UCustomizableObjectNodeModifierEditMeshSection;
 class SCustomizableObjectNodeLayoutBlocksEditor;
 
 
-class FCustomizableObjectNodeModifierEditMeshSectionDetails : public FCustomizableObjectNodeDetails
+class FCustomizableObjectNodeModifierEditMeshSectionDetails : public FCustomizableObjectNodeModifierBaseDetails
 {
 public:
 	/** Makes a new instance of this detail layout class for a specific detail view requesting it */
@@ -22,6 +22,9 @@ public:
 	/** ILayoutDetails interface */
 	virtual void CustomizeDetails(IDetailLayoutBuilder& DetailBuilder) override;
 
+	/** FCustomizableObjectNodeModifierBaseDetails interface. */
+	virtual void OnRequiredTagsPropertyChanged() override;
+
 private:
 	UCustomizableObjectNodeModifierEditMeshSection* Node;
 
diff --git a/Engine/Plugins/Experimental/Mutable/Source/CustomizableObjectEditor/Private/MuCOE/Nodes/CustomizableObjectNodeModifierExtendMeshSectionDetails.cpp b/Engine/Plugins/Experimental/Mutable/Source/CustomizableObjectEditor/Private/MuCOE/Nodes/CustomizableObjectNodeModifierExtendMeshSectionDetails.cpp
new file mode 100644
--- /dev/null
+++ b/Engine/Plugins/Experimental/Mutable/Source/CustomizableObjectEditor/Private/MuCOE/Nodes/CustomizableObjectNodeModifierExtendMeshSectionDetails.cpp
@@ -0,0 +1,49 @@
+// Copyright Epic Games, Inc. All Rights Reserved.
+
+#include "MuCOE/Nodes/CustomizableObjectNodeModifierExtendMeshSectionDetails.h"
+
+#include "MuCOE/CustomizableObjectEditorUtilities.h"
+#include "MuCOE/CustomizableObjectEditorStyle.h"
+#include "MuCOE/GraphTraversal.h"
+#include "MuCOE/Nodes/CustomizableObjectNodeMaterial.h"
+#include "MuCOE/Nodes/CustomizableObjectNodeModifierMorphMeshSection.h"
+#include "MuCOE/Nodes/CustomizableObjectNodeModifierExtendMeshSection.h"
+#include "MuCOE/Nodes/CustomizableObjectNodeSkeletalMesh.h"
+#include "MuCOE/UnrealEditorPortabilityHelpers.h"
+#include "PropertyCustomizationHelpers.h"
+#include "Widgets/Input/STextComboBox.h"
+#include "Widgets/Input/SSearchBox.h"
+#include "DetailLayoutBuilder.h"
+#include "IDetailsView.h"
+#include "Engine/SkeletalMesh.h"
+#include "Misc/Attribute.h"
+
+
+#define LOCTEXT_NAMESPACE "CustomizableObjectDetails"
+
+
+TSharedRef<IDetailCustomization> FCustomizableObjectNodeModifierExtendMeshSectionDetails::MakeInstance()
+{
+	return MakeShareable( new FCustomizableObjectNodeModifierExtendMeshSectionDetails );
+}
+
+
+void FCustomizableObjectNodeModifierExtendMeshSectionDetails::CustomizeDetails( IDetailLayoutBuilder& DetailBuilder )
+{
+	FCustomizableObjectNodeModifierBaseDetails::CustomizeDetails(DetailBuilder);
+
+	const IDetailsView* DetailsView = DetailBuilder.GetDetailsView();
+	if ( DetailsView->GetSelectedObjects().Num() )
+	{
+		Node = Cast<UCustomizableObjectNodeModifierExtendMeshSection>( DetailsView->GetSelectedObjects()[0].Get() );
+	}
+
+	if (!Node)
+	{
+		return;
+	}
+	
+}
+
+
+#undef LOCTEXT_NAMESPACE
diff --git a/Engine/Plugins/Experimental/Mutable/Source/CustomizableObjectEditor/Private/MuCOE/Nodes/CustomizableObjectNodeModifierExtendMeshSectionDetails.h b/Engine/Plugins/Experimental/Mutable/Source/CustomizableObjectEditor/Private/MuCOE/Nodes/CustomizableObjectNodeModifierExtendMeshSectionDetails.h
new file mode 100644
--- /dev/null
+++ b/Engine/Plugins/Experimental/Mutable/Source/CustomizableObjectEditor/Private/MuCOE/Nodes/CustomizableObjectNodeModifierExtendMeshSectionDetails.h
@@ -0,0 +1,36 @@
+// Copyright Epic Games, Inc. All Rights Reserved.
+
+#pragma once
+
+#include "CustomizableObjectNodeModifierBaseDetails.h"
+#include "MuCOE/SMutableSearchComboBox.h"
+#include "IDetailCustomization.h"
+#include "Types/SlateEnums.h"
+#include "Widgets/Input/SComboButton.h"
+#include "Widgets/Input/SEditableTextBox.h"
+#include "Widgets/Views/SListView.h"
+#include "Widgets/Views/ITableRow.h"
+
+namespace ESelectInfo { enum Type : int; }
+
+class FString;
+class IDetailLayoutBuilder;
+class IPropertyHandle;
+class UCustomizableObjectNodeObject;
+
+
+class FCustomizableObjectNodeModifierExtendMeshSectionDetails : public FCustomizableObjectNodeModifierBaseDetails
+{
+public:
+	// Makes a new instance of this detail layout class for a specific detail view requesting it 
+	static TSharedRef<IDetailCustomization> MakeInstance();
+
+	// ILayoutDetails interface
+	void CustomizeDetails( IDetailLayoutBuilder& DetailBuilder ) override;
+
+private:
+
+	class UCustomizableObjectNodeModifierExtendMeshSection* Node = nullptr;
+
+
+};
diff --git a/Engine/Plugins/Experimental/Mutable/Source/CustomizableObjectEditor/Private/MuCOE/Nodes/CustomizableObjectNodeModifierMorphMeshSectionDetails.cpp b/Engine/Plugins/Experimental/Mutable/Source/CustomizableObjectEditor/Private/MuCOE/Nodes/CustomizableObjectNodeModifierMorphMeshSectionDetails.cpp
--- a/Engine/Plugins/Experimental/Mutable/Source/CustomizableObjectEditor/Private/MuCOE/Nodes/CustomizableObjectNodeModifierMorphMeshSectionDetails.cpp
+++ b/Engine/Plugins/Experimental/Mutable/Source/CustomizableObjectEditor/Private/MuCOE/Nodes/CustomizableObjectNodeModifierMorphMeshSectionDetails.cpp
@@ -3,6 +3,7 @@
 #include "MuCOE/Nodes/CustomizableObjectNodeModifierMorphMeshSectionDetails.h"
 
 #include "MuCOE/CustomizableObjectEditorUtilities.h"
+#include "MuCOE/CustomizableObjectEditorStyle.h"
 #include "MuCOE/GraphTraversal.h"
 #include "MuCOE/Nodes/CustomizableObjectNodeMaterial.h"
 #include "MuCOE/Nodes/CustomizableObjectNodeModifierMorphMeshSection.h"
@@ -29,7 +30,7 @@ TSharedRef<IDetailCustomization> FCustomizableObjectNodeModifierMorphMeshSection
 
 void FCustomizableObjectNodeModifierMorphMeshSectionDetails::CustomizeDetails( IDetailLayoutBuilder& DetailBuilder )
 {
-	FCustomizableObjectNodeDetails::CustomizeDetails(DetailBuilder);
+	FCustomizableObjectNodeModifierBaseDetails::CustomizeDetails(DetailBuilder);
 
 	const IDetailsView* DetailsView = DetailBuilder.GetDetailsView();
 	if ( DetailsView->GetSelectedObjects().Num() )
@@ -37,28 +38,30 @@ void FCustomizableObjectNodeModifierMorphMeshSectionDetails::CustomizeDetails( I
 		Node = Cast<UCustomizableObjectNodeModifierMorphMeshSection>( DetailsView->GetSelectedObjects()[0].Get() );
 	}
 
+	if (!Node)
+	{
+		return;
+	}
 
 	// This property is not relevant for this node
 	DetailBuilder.HideProperty(GET_MEMBER_NAME_CHECKED(UCustomizableObjectNodeModifierWithMaterial, ReferenceMaterial), UCustomizableObjectNodeModifierWithMaterial::StaticClass());
 	
-	IDetailCategoryBuilder& ModifierCategory = DetailBuilder.EditCategory("Modifier");
-	ModifierCategory.SetSortOrder(-10000);
+	// Add a morph selection widget.
+	IDetailCategoryBuilder& MorphCategory = DetailBuilder.EditCategory( "Morph" );
 
-	IDetailCategoryBuilder& BlocksCategory = DetailBuilder.EditCategory( "Customizable Object" );
-
-	if (Node)
+	// Scan for hint morph names
 	{
-		// Scan for hint morph names
+		TArray<UCustomizableObjectNode*> CandidateNodes;
+		Node->GetPossiblyModifiedNodes(CandidateNodes);
+
+		// For now just show the first one
+		for (UCustomizableObjectNode* Candidate : CandidateNodes)
 		{
-			TArray<UCustomizableObjectNode*> CandidateNodes;
-			Node->GetPossiblyModifiedNodes(CandidateNodes);
+			USkeletalMesh* SkeletalMesh = nullptr;
 
-			// For now just show the first one
-			for (UCustomizableObjectNode* Candidate : CandidateNodes)
+			if (UCustomizableObjectNodeMaterialBase* MaterialNode = Cast<UCustomizableObjectNodeMaterialBase>(Candidate))
 			{
-				USkeletalMesh* SkeletalMesh = nullptr;
-
-				if (UCustomizableObjectNodeMaterialBase* MaterialNode = Cast<UCustomizableObjectNodeMaterialBase>(Candidate))
+				if (MaterialNode->OutputPin())
 				{
 					const UEdGraphPin* SourceMeshPin = FindMeshBaseSource(*MaterialNode->OutputPin(), false);
 					if (SourceMeshPin)
@@ -70,8 +73,11 @@ void FCustomizableObjectNodeModifierMorphMeshSectionDetails::CustomizeDetails( I
 						}
 					}
 				}
+			}
 
-				else if (UCustomizableObjectNodeModifierExtendMeshSection* ExtendNode = Cast<UCustomizableObjectNodeModifierExtendMeshSection>(Candidate))
+			else if (UCustomizableObjectNodeModifierExtendMeshSection* ExtendNode = Cast<UCustomizableObjectNodeModifierExtendMeshSection>(Candidate))
+			{
+				if (ExtendNode->OutputPin())
 				{
 					const UEdGraphPin* SourceMeshPin = FindMeshBaseSource(*ExtendNode->OutputPin(), false);
 					if (SourceMeshPin)
@@ -83,209 +89,79 @@ void FCustomizableObjectNodeModifierMorphMeshSectionDetails::CustomizeDetails( I
 						}
 					}
 				}
+			}
 
-				if (SkeletalMesh)
+			if (SkeletalMesh)
+			{
+				const TArray<TObjectPtr<UMorphTarget>>& Morphs = SkeletalMesh->GetMorphTargets();
+				for (TObjectPtr<UMorphTarget> Morph : Morphs)
 				{
-					const TArray<TObjectPtr<UMorphTarget>>& Morphs = SkeletalMesh->GetMorphTargets();
-					for (TObjectPtr<UMorphTarget> Morph : Morphs)
+					if (Morph)
 					{
 						FString MorphTargetName = Morph->GetName();
-						OptionsSource.AddUnique(MakeShared<FString>(MorphTargetName));
+						MorphOptionsSource.AddUnique(MakeShared<FString>(MorphTargetName));
 					}
 				}
 			}
-
-			// TODO: Add all morphs if no candidate is found? Add both in different menu sections?
 		}
-		FilteredOptionsSource = OptionsSource;
 
-		// Morph target selection
-		TSharedRef<IPropertyHandle> MorphTargetNameProperty = DetailBuilder.GetProperty("MorphTargetName");
-		BlocksCategory.AddCustomRow(LOCTEXT("MorphMaterialDetails_Target", "Target"))
-			[
-				SNew(SProperty, MorphTargetNameProperty)
-					.ShouldDisplayName(false)
-					.CustomWidget()
-					[
-						SNew(SBorder)
-							.BorderImage(UE_MUTABLE_GET_BRUSH("NoBorder"))
-							.Padding(FMargin(0.0f, 0.0f, 10.0f, 0.0f))
-							[
-								SNew(SHorizontalBox)
+		// TODO: Add all morphs if no candidate is found? Add both in different menu sections?
+	}
 
-								+ SHorizontalBox::Slot()
-								.FillWidth(10.0f)
-								.VAlign(VAlign_Center)
-								[
-									SNew(STextBlock)
-										.Text(LOCTEXT("MorphMaterialDetails_MorphTarget", "Morph Target"))
-										.Font(IDetailLayoutBuilder::GetDetailFont())
-								]
+	TSharedRef<IPropertyHandle> MorphTargetNameProperty = DetailBuilder.GetProperty("MorphTargetName");
+	MorphCategory.AddCustomRow(LOCTEXT("MorphMaterialDetails_Target", "Target"))
+		[
+			SNew(SProperty, MorphTargetNameProperty)
+				.ShouldDisplayName(false)
+				.CustomWidget()
+				[
+					SNew(SBorder)
+						.BorderImage(UE_MUTABLE_GET_BRUSH("NoBorder"))
+						.Padding(FMargin(0.0f, 0.0f, 10.0f, 0.0f))
+						[
+							SNew(SHorizontalBox)
+
+							+ SHorizontalBox::Slot()
+							.FillWidth(10.0f)
+							.VAlign(VAlign_Center)
+							[
+								SNew(STextBlock)
+									.Text(LOCTEXT("MorphMaterialDetails_MorphTarget", "Morph Target"))
+									.Font(IDetailLayoutBuilder::GetDetailFont())
+							]
 
-								+ SHorizontalBox::Slot()
-								.FillWidth(10.0f)
-								.HAlign(HAlign_Fill)
-								.VAlign(VAlign_Center)
+							+ SHorizontalBox::Slot()
+							.FillWidth(10.0f)
+							.HAlign(HAlign_Fill)
+							.VAlign(VAlign_Center)
+							[
+								SAssignNew(this->MorphCombo , SMutableSearchComboBox)
+								.ButtonStyle(FAppStyle::Get(), "NoBorder")
+								.OptionsSource(&MorphOptionsSource)
+								.OnSelectionChanged(this, &FCustomizableObjectNodeModifierMorphMeshSectionDetails::OnMorphTargetComboBoxSelectionChanged)
+								.Content()
 								[
-									SAssignNew(this->MorphCombo , SComboButton)
-									.ButtonStyle(FAppStyle::Get(), "NoBorder")
-									.ButtonContent()
-									[
-										SNew(SEditableTextBox)
-										.Font(IDetailLayoutBuilder::GetDetailFont())
-										.Text_Lambda([&]() { return Node ? FText::FromString(Node->MorphTargetName) : FText();})
-										.OnTextChanged(this, &FCustomizableObjectNodeModifierMorphMeshSectionDetails::OnMorphTargetComboBoxSelectionChanged)
-									]
-									.MenuContent()
-									[
-										SNew(SVerticalBox)
-
-											+ SVerticalBox::Slot()
-											.AutoHeight()
-											[
-												SAssignNew(this->SearchField, SEditableTextBox)
-													.HintText(LOCTEXT("Search", "Search"))
-													.OnTextChanged(this, &FCustomizableObjectNodeModifierMorphMeshSectionDetails::OnSearchTextChanged)
-													.OnTextCommitted(this, &FCustomizableObjectNodeModifierMorphMeshSectionDetails::OnSearchTextCommitted)
-											]
-
-											+ SVerticalBox::Slot()
-											[
-												SAssignNew(this->ComboListView, SListView< TSharedPtr<FString> >)
-													.ListItemsSource(&FilteredOptionsSource)
-													.OnGenerateRow(this, &FCustomizableObjectNodeModifierMorphMeshSectionDetails::GenerateMenuItemRow)
-													.OnSelectionChanged(this, &FCustomizableObjectNodeModifierMorphMeshSectionDetails::OnSelectionChanged)
-													.OnKeyDownHandler(this, &FCustomizableObjectNodeModifierMorphMeshSectionDetails::OnKeyDownHandler)
-													.SelectionMode(ESelectionMode::Single)
-											]
-									]
+									SNew(SEditableTextBox)
+									.Font(IDetailLayoutBuilder::GetDetailFont())
+									.Text_Lambda([&]() { return Node ? FText::FromString(Node->MorphTargetName) : FText();})
+									.OnTextChanged(this, &FCustomizableObjectNodeModifierMorphMeshSectionDetails::OnMorphTargetComboBoxSelectionChanged)
 								]
 							]
-					]
-			];
-	}
-	else
-	{
-		BlocksCategory.AddCustomRow( LOCTEXT("MorphMaterialDetails_Node", "Node") )
-		[
-			SNew( STextBlock )
-			.Text( LOCTEXT( "MorphMaterialDetails_NodeNotFound", "Node not found" ) )
+						]
+				]
 		];
-	}
-}
-
-
-void FCustomizableObjectNodeModifierMorphMeshSectionDetails::OnSearchTextChanged(const FText& ChangedText)
-{
-	SearchText = ChangedText;
-
-	RefreshOptions();
-}
-
-void FCustomizableObjectNodeModifierMorphMeshSectionDetails::OnSearchTextCommitted(const FText& InText, ETextCommit::Type InCommitType)
-{
-	if ((InCommitType == ETextCommit::Type::OnEnter) && FilteredOptionsSource.Num() > 0)
-	{
-		ComboListView->SetSelection(FilteredOptionsSource[0], ESelectInfo::OnKeyPress);
-	}
 }
 
 
 void FCustomizableObjectNodeModifierMorphMeshSectionDetails::OnMorphTargetComboBoxSelectionChanged(const FText& NewText)
-{	
-	if (Node && Node->MorphTargetName!=NewText.ToString())
+{
+	if (Node && Node->MorphTargetName != NewText.ToString())
 	{
 		Node->MorphTargetName = NewText.ToString();
 		Node->Modify();
 	}
 }
 
-void FCustomizableObjectNodeModifierMorphMeshSectionDetails::RefreshOptions()
-{
-	// Need to refresh filtered list whenever options change
-	FilteredOptionsSource.Reset();
-
-	if (SearchText.IsEmpty())
-	{
-		FilteredOptionsSource.Append(OptionsSource);
-	}
-	else
-	{
-		TArray<FString> SearchTokens;
-		SearchText.ToString().ParseIntoArrayWS(SearchTokens);
-
-		for (const TSharedPtr<FString>& Option : OptionsSource)
-		{
-			bool bAllTokensMatch = true;
-			for (const FString& SearchToken : SearchTokens)
-			{
-				if (Option->Find(SearchToken, ESearchCase::Type::IgnoreCase) == INDEX_NONE)
-				{
-					bAllTokensMatch = false;
-					break;
-				}
-			}
-
-			if (bAllTokensMatch)
-			{
-				FilteredOptionsSource.Add(Option);
-			}
-		}
-	}
-
-	ComboListView->RequestListRefresh();
-}
-
-
-TSharedRef<ITableRow> FCustomizableObjectNodeModifierMorphMeshSectionDetails::GenerateMenuItemRow(TSharedPtr<FString> InItem, const TSharedRef<STableViewBase>& OwnerTable)
-{
-	return SNew(SComboRow<TSharedPtr<FString>>, OwnerTable)
-		[
-			SNew(STextBlock).Text(FText::FromString(*InItem))
-		];
-}
-
-
-void FCustomizableObjectNodeModifierMorphMeshSectionDetails::OnSelectionChanged(TSharedPtr<FString> ProposedSelection, ESelectInfo::Type SelectInfo)
-{
-	if (!ProposedSelection || !Node)
-	{
-		return;
-	}
-
-	// Ensure that the proposed selection is different from selected
-	if (*ProposedSelection != Node->MorphTargetName)
-	{
-		Node->MorphTargetName = *ProposedSelection;
-	}
-
-	// close combo as long as the selection wasn't from navigation
-	if (SelectInfo != ESelectInfo::OnNavigation)
-	{
-		MorphCombo->SetIsOpen(false);
-	}
-	else
-	{
-		ComboListView->RequestScrollIntoView(ProposedSelection, 0);
-	}
-}
-
-
-FReply FCustomizableObjectNodeModifierMorphMeshSectionDetails::OnKeyDownHandler(const FGeometry& MyGeometry, const FKeyEvent& InKeyEvent)
-{
-	if (InKeyEvent.GetKey() == EKeys::Enter)
-	{
-		// Select the first selected item on hitting enter
-		TArray<TSharedPtr<FString>> SelectedItems = ComboListView->GetSelectedItems();
-		if (SelectedItems.Num() > 0)
-		{
-			OnSelectionChanged(SelectedItems[0], ESelectInfo::OnKeyPress);
-			return FReply::Handled();
-		}
-	}
-
-	return FReply::Unhandled();
-}
 
 
 #undef LOCTEXT_NAMESPACE
diff --git a/Engine/Plugins/Experimental/Mutable/Source/CustomizableObjectEditor/Private/MuCOE/Nodes/CustomizableObjectNodeModifierMorphMeshSectionDetails.h b/Engine/Plugins/Experimental/Mutable/Source/CustomizableObjectEditor/Private/MuCOE/Nodes/CustomizableObjectNodeModifierMorphMeshSectionDetails.h
--- a/Engine/Plugins/Experimental/Mutable/Source/CustomizableObjectEditor/Private/MuCOE/Nodes/CustomizableObjectNodeModifierMorphMeshSectionDetails.h
+++ b/Engine/Plugins/Experimental/Mutable/Source/CustomizableObjectEditor/Private/MuCOE/Nodes/CustomizableObjectNodeModifierMorphMeshSectionDetails.h
@@ -2,7 +2,8 @@
 
 #pragma once
 
-#include "CustomizableObjectNodeDetails.h"
+#include "CustomizableObjectNodeModifierBaseDetails.h"
+#include "MuCOE/SMutableSearchComboBox.h"
 #include "IDetailCustomization.h"
 #include "Types/SlateEnums.h"
 #include "Widgets/Input/SComboButton.h"
@@ -17,7 +18,8 @@ class IDetailLayoutBuilder;
 class IPropertyHandle;
 class UCustomizableObjectNodeObject;
 
-class FCustomizableObjectNodeModifierMorphMeshSectionDetails : public FCustomizableObjectNodeDetails
+
+class FCustomizableObjectNodeModifierMorphMeshSectionDetails : public FCustomizableObjectNodeModifierBaseDetails
 {
 public:
 	// Makes a new instance of this detail layout class for a specific detail view requesting it 
@@ -27,27 +29,12 @@ public:
 	void CustomizeDetails( IDetailLayoutBuilder& DetailBuilder ) override;
 
 private:
+
 	class UCustomizableObjectNodeModifierMorphMeshSection* Node = nullptr;
 
 	/** */
-	TSharedPtr<SComboButton> MorphCombo;
-	/** The search field used for the combox box's contents */
-	TSharedPtr<SEditableTextBox> SearchField;
-	/** The ListView that we pop up; visualized the available options. */
-	TSharedPtr< SListView< TSharedPtr<FString> > > ComboListView;
-	/** */
-	TArray< TSharedPtr<FString> > OptionsSource;
-	/** Updated whenever search text is changed */
-	FText SearchText;
-	/** Filtered list that is actually displayed */
-	TArray< TSharedPtr<FString> > FilteredOptionsSource;
-
+	TSharedPtr<SMutableSearchComboBox> MorphCombo;
+	TArray< TSharedPtr<FString> > MorphOptionsSource;
 	void OnMorphTargetComboBoxSelectionChanged(const FText& NewText);
-	void OnSearchTextChanged(const FText& ChangedText);
-	void OnSearchTextCommitted(const FText& InText, ETextCommit::Type InCommitType);
-	TSharedRef<ITableRow> GenerateMenuItemRow(TSharedPtr<FString> InItem, const TSharedRef<STableViewBase>& OwnerTable);
-	void RefreshOptions();
-	void OnSelectionChanged(TSharedPtr<FString> ProposedSelection, ESelectInfo::Type SelectInfo);
-	FReply OnKeyDownHandler(const FGeometry& MyGeometry, const FKeyEvent& InKeyEvent);
 
 };
diff --git a/Engine/Plugins/Experimental/Mutable/Source/CustomizableObjectEditor/Private/MuCOE/Nodes/CustomizableObjectNodeModifierRemoveMeshBlocksDetails.cpp b/Engine/Plugins/Experimental/Mutable/Source/CustomizableObjectEditor/Private/MuCOE/Nodes/CustomizableObjectNodeModifierRemoveMeshBlocksDetails.cpp
--- a/Engine/Plugins/Experimental/Mutable/Source/CustomizableObjectEditor/Private/MuCOE/Nodes/CustomizableObjectNodeModifierRemoveMeshBlocksDetails.cpp
+++ b/Engine/Plugins/Experimental/Mutable/Source/CustomizableObjectEditor/Private/MuCOE/Nodes/CustomizableObjectNodeModifierRemoveMeshBlocksDetails.cpp
@@ -28,7 +28,7 @@ TSharedRef<IDetailCustomization> FCustomizableObjectNodeModifierRemoveMeshBlocks
 
 void FCustomizableObjectNodeModifierRemoveMeshBlocksDetails::CustomizeDetails(IDetailLayoutBuilder& DetailBuilder)
 {
-	FCustomizableObjectNodeDetails::CustomizeDetails(DetailBuilder);
+	FCustomizableObjectNodeModifierBaseDetails::CustomizeDetails(DetailBuilder);
 
 	Node = nullptr;
 	const IDetailsView* DetailsView = DetailBuilder.GetDetailsView();
@@ -40,9 +40,6 @@ void FCustomizableObjectNodeModifierRemoveMeshBlocksDetails::CustomizeDetails(ID
 	// This property is not relevant for this node
 	DetailBuilder.HideProperty(GET_MEMBER_NAME_CHECKED(UCustomizableObjectNodeModifierWithMaterial, ReferenceMaterial), UCustomizableObjectNodeModifierWithMaterial::StaticClass());
 
-	IDetailCategoryBuilder& ModifierCategory = DetailBuilder.EditCategory("Modifier");
-	ModifierCategory.SetSortOrder(-10000);
-
 	IDetailCategoryBuilder& LayoutCategory = DetailBuilder.EditCategory("LayoutOptions");
 	LayoutCategory.SetSortOrder(10000);
 
@@ -141,6 +138,13 @@ void FCustomizableObjectNodeModifierRemoveMeshBlocksDetails::CustomizeDetails(ID
 }
 
 
+void FCustomizableObjectNodeModifierRemoveMeshBlocksDetails::OnRequiredTagsPropertyChanged()
+{
+	FCustomizableObjectNodeModifierBaseDetails::OnRequiredTagsPropertyChanged();
+	UpdateLayout();
+}
+
+
 void FCustomizableObjectNodeModifierRemoveMeshBlocksDetails::UpdateLayout()
 {
 	// Try to find the parent layout, because we want to show its UVs in the widget
diff --git a/Engine/Plugins/Experimental/Mutable/Source/CustomizableObjectEditor/Private/MuCOE/Nodes/CustomizableObjectNodeModifierRemoveMeshBlocksDetails.h b/Engine/Plugins/Experimental/Mutable/Source/CustomizableObjectEditor/Private/MuCOE/Nodes/CustomizableObjectNodeModifierRemoveMeshBlocksDetails.h
--- a/Engine/Plugins/Experimental/Mutable/Source/CustomizableObjectEditor/Private/MuCOE/Nodes/CustomizableObjectNodeModifierRemoveMeshBlocksDetails.h
+++ b/Engine/Plugins/Experimental/Mutable/Source/CustomizableObjectEditor/Private/MuCOE/Nodes/CustomizableObjectNodeModifierRemoveMeshBlocksDetails.h
@@ -3,7 +3,7 @@
 #pragma once
 
 
-#include "CustomizableObjectNodeDetails.h"
+#include "CustomizableObjectNodeModifierBaseDetails.h"
 #include "IDetailCustomization.h"
 #include "Types/SlateEnums.h"
 
@@ -13,7 +13,7 @@ class UCustomizableObjectNodeModifierRemoveMeshBlocks;
 class SCustomizableObjectNodeLayoutBlocksEditor;
 
 
-class FCustomizableObjectNodeModifierRemoveMeshBlocksDetails : public FCustomizableObjectNodeDetails
+class FCustomizableObjectNodeModifierRemoveMeshBlocksDetails : public FCustomizableObjectNodeModifierBaseDetails
 {
 public:
 	/** Makes a new instance of this detail layout class for a specific detail view requesting it */
@@ -22,6 +22,9 @@ public:
 	/** ILayoutDetails interface */
 	virtual void CustomizeDetails(IDetailLayoutBuilder& DetailBuilder) override;
 
+	/** FCustomizableObjectNodeModifierBaseDetails interface. */
+	virtual void OnRequiredTagsPropertyChanged() override;
+
 private:
 	UCustomizableObjectNodeModifierRemoveMeshBlocks* Node;
 
diff --git a/Engine/Plugins/Experimental/Mutable/Source/CustomizableObjectEditor/Private/MuCOE/SCustomizableObjectEditorTagExplorer.cpp b/Engine/Plugins/Experimental/Mutable/Source/CustomizableObjectEditor/Private/MuCOE/SCustomizableObjectEditorTagExplorer.cpp
--- a/Engine/Plugins/Experimental/Mutable/Source/CustomizableObjectEditor/Private/MuCOE/SCustomizableObjectEditorTagExplorer.cpp
+++ b/Engine/Plugins/Experimental/Mutable/Source/CustomizableObjectEditor/Private/MuCOE/SCustomizableObjectEditorTagExplorer.cpp
@@ -186,14 +186,11 @@ void SCustomizableObjectEditorTagExplorer::FillTagInformation(const UCustomizabl
 
 			if (UCustomizableObjectNodeModifierBase* TypedNodeModifier = Cast<UCustomizableObjectNodeModifierBase>(Node))
 			{
-				const TArray<FString>* RequiredTags = TypedNodeModifier->GetRequiredTags();
-				if (RequiredTags)
+				const TArray<FString>& RequiredTags = TypedNodeModifier->RequiredTags;
+				for (const FString& CurrentTag : RequiredTags)
 				{
-					for (const FString& CurrentTag : *RequiredTags)
-					{
-						NodeTags.Add(CurrentTag, TypedNodeModifier);
-						Tags.AddUnique(CurrentTag);
-					}
+					NodeTags.Add(CurrentTag, TypedNodeModifier);
+					Tags.AddUnique(CurrentTag);
 				}
 			}
 		}
diff --git a/Engine/Plugins/Experimental/Mutable/Source/CustomizableObjectEditor/Private/MuCOE/SMutableSearchComboBox.cpp b/Engine/Plugins/Experimental/Mutable/Source/CustomizableObjectEditor/Private/MuCOE/SMutableSearchComboBox.cpp
new file mode 100644
--- /dev/null
+++ b/Engine/Plugins/Experimental/Mutable/Source/CustomizableObjectEditor/Private/MuCOE/SMutableSearchComboBox.cpp
@@ -0,0 +1,275 @@
+// Copyright Epic Games, Inc. All Rights Reserved.
+
+#include "MuCOE/SMutableSearchComboBox.h"
+
+#include "MuCOE/CustomizableObjectEditorUtilities.h"
+#include "MuCOE/CustomizableObjectEditorStyle.h"
+#include "MuCOE/GraphTraversal.h"
+#include "MuCOE/UnrealEditorPortabilityHelpers.h"
+#include "PropertyCustomizationHelpers.h"
+#include "Widgets/Input/SComboBox.h"
+#include "Widgets/Input/STextComboBox.h"
+#include "Widgets/Input/SSearchBox.h"
+#include "Engine/SkeletalMesh.h"
+#include "Misc/Attribute.h"
+#include "DetailLayoutBuilder.h"	// For font: TODO: move to argument
+
+
+#define LOCTEXT_NAMESPACE "CustomizableObjectDetails"
+
+
+void SMutableSearchComboBox::Construct(const FArguments& InArgs)
+{
+	check(InArgs._ComboBoxStyle);
+
+	ItemStyle = InArgs._ItemStyle;
+	MenuRowPadding = InArgs._ComboBoxStyle->MenuRowPadding;
+
+	bAllowAddNewOptions = InArgs._AllowAddNewOptions;
+
+	// Work out which values we should use based on whether we were given an override, or should use the style's version
+	OurComboButtonStyle = InArgs._ComboBoxStyle->ComboButtonStyle;
+	if (InArgs._MenuButtonBrush)
+	{
+		OurComboButtonStyle.DownArrowImage = *InArgs._MenuButtonBrush;
+	}
+	const FButtonStyle* const OurButtonStyle = InArgs._ButtonStyle ? InArgs._ButtonStyle : &OurComboButtonStyle.ButtonStyle;
+
+	this->OnSelectionChanged = InArgs._OnSelectionChanged;
+	this->OnGenerateWidget = InArgs._OnGenerateWidget;
+
+	OptionsSource = InArgs._OptionsSource;
+	RefreshOptions();
+
+	TAttribute<EVisibility> SearchVisibility = InArgs._SearchVisibility;
+	const EVisibility CurrentSearchVisibility = SearchVisibility.Get();
+
+	TSharedRef<SWidget> ComboBoxMenuContent =
+		SNew(SBox)
+		.MaxDesiredHeight(InArgs._MaxListHeight)
+		[
+			SNew(SVerticalBox)
+
+				+ SVerticalBox::Slot()
+				.AutoHeight()
+				[
+					SAssignNew(this->SearchField, SEditableTextBox)
+						.HintText(bAllowAddNewOptions ? LOCTEXT("SearchOrAdd", "Search or add...") : LOCTEXT("Search", "Search..."))
+						.OnTextChanged(this, &SMutableSearchComboBox::OnSearchTextChanged)
+						.OnTextCommitted(this, &SMutableSearchComboBox::OnSearchTextCommitted)
+						.Visibility(SearchVisibility)
+				]
+
+				+ SVerticalBox::Slot()
+				[
+					SAssignNew(this->ComboListView, SComboListType)
+						.ListItemsSource(&FilteredOptionsSource)
+						.OnGenerateRow(this, &SMutableSearchComboBox::GenerateMenuItemRow)
+						.OnSelectionChanged(this, &SMutableSearchComboBox::OnSelectionChanged_Internal)
+						.OnKeyDownHandler(this, &SMutableSearchComboBox::OnKeyDownHandler)
+						.SelectionMode(ESelectionMode::Single)
+				]
+		];
+
+	// Set up content
+	TSharedPtr<SWidget> ButtonContent = InArgs._Content.Widget;
+	if (InArgs._Content.Widget == SNullWidget::NullWidget)
+	{
+		SAssignNew(ButtonContent, STextBlock);
+	}
+
+
+	SComboButton::Construct(SComboButton::FArguments()
+		.ComboButtonStyle(&OurComboButtonStyle)
+		.ButtonStyle(OurButtonStyle)
+		.Method(InArgs._Method)
+		.ButtonContent()
+		[
+			ButtonContent.ToSharedRef()
+		]
+		.MenuContent()
+		[
+			ComboBoxMenuContent
+		]
+		.ContentPadding(InArgs._ContentPadding)
+		.ForegroundColor(InArgs._ForegroundColor)
+		.OnMenuOpenChanged(this, &SMutableSearchComboBox::OnMenuOpenChanged)
+		.IsFocusable(true)
+	);
+
+	if (CurrentSearchVisibility == EVisibility::Visible)
+	{
+		SetMenuContentWidgetToFocus(SearchField);
+	}
+	else
+	{
+		SetMenuContentWidgetToFocus(ComboListView);
+	}
+}
+
+
+void SMutableSearchComboBox::RefreshOptions()
+{
+	// Need to refresh filtered list whenever options change
+	FilteredOptionsSource.Reset();
+
+	if (SearchText.IsEmpty())
+	{
+		for (const TSharedPtr<FString>& Option : *OptionsSource)
+		{
+			if (Option)
+			{
+				FFilteredOption Data;
+				Data.ActualOption = *Option;
+				Data.DisplayOption = *Option;
+				FilteredOptionsSource.Add(MakeShared<FFilteredOption>(Data));
+			}
+		}
+	}
+	else
+	{
+		TArray<FString> SearchTokens;
+		SearchText.ToString().ParseIntoArrayWS(SearchTokens);
+
+		for (const TSharedPtr<FString>& Option : *OptionsSource)
+		{
+			bool bAllTokensMatch = true;
+			for (const FString& SearchToken : SearchTokens)
+			{
+				if (Option->Find(SearchToken, ESearchCase::Type::IgnoreCase) == INDEX_NONE)
+				{
+					bAllTokensMatch = false;
+					break;
+				}
+			}
+
+			if (bAllTokensMatch)
+			{
+				FFilteredOption Data;
+				Data.ActualOption = *Option;
+				Data.DisplayOption = *Option;
+				FilteredOptionsSource.Add(MakeShared<FFilteredOption>(Data));
+			}
+		}
+
+		bool bFullMatch = false;
+		FString SearchString = SearchText.ToString();
+		for (const TSharedPtr<FString>& Option : *OptionsSource)		
+		{
+			if (*Option == SearchString)
+			{
+				bFullMatch = true; 
+				break;
+			}
+		}
+
+		if ( bAllowAddNewOptions && !SearchText.IsEmpty() )
+		{
+			FFilteredOption Data;
+			Data.ActualOption = SearchString;
+			Data.DisplayOption = FString::Printf(TEXT("Add new (%s)"), *SearchString);
+			FilteredOptionsSource.Add(MakeShared<FFilteredOption>(Data));
+		}
+	}
+
+	if (ComboListView)
+	{
+		ComboListView->RequestListRefresh();
+	}
+}
+
+
+TSharedRef<ITableRow> SMutableSearchComboBox::GenerateMenuItemRow(TSharedPtr<FFilteredOption> InItem, const TSharedRef<STableViewBase>& OwnerTable)
+{
+	// TODO
+	//if (OnGenerateWidget.IsBound())
+	//{
+	//	return SNew(SComboRow<TSharedPtr<FString>>, OwnerTable)
+	//		.Style(ItemStyle)
+	//		.Padding(MenuRowPadding)
+	//		[
+	//			OnGenerateWidget.Execute(InItem)
+	//		];
+	//}
+	//else
+	{
+		// Just add the text
+		return SNew(SComboRow<TSharedPtr<FString>>, OwnerTable)
+			[
+				SNew(STextBlock)
+					.Text(FText::FromString(InItem ? InItem->DisplayOption : FString(TEXT("Invalid item."))))
+					.Font(IDetailLayoutBuilder::GetDetailFont())
+			];
+	}
+}
+
+
+void SMutableSearchComboBox::OnMenuOpenChanged(bool bOpen)
+{
+	if (bOpen == false)
+	{
+		// Set focus back to ComboBox for users focusing the ListView that just closed
+		FSlateApplication::Get().ForEachUser([this](FSlateUser& User)
+			{
+				TSharedRef<SWidget> ThisRef = this->AsShared();
+				if (User.IsWidgetInFocusPath(this->ComboListView))
+				{
+					User.SetFocus(ThisRef);
+				}
+			});
+
+	}
+}
+
+
+void SMutableSearchComboBox::OnSelectionChanged_Internal(TSharedPtr<FFilteredOption> ProposedSelection, ESelectInfo::Type SelectInfo)
+{
+	if (!ProposedSelection)
+	{
+		return;
+	}
+
+	// close combo as long as the selection wasn't from navigation
+	if (SelectInfo != ESelectInfo::OnNavigation)
+	{		
+		OnSelectionChanged.ExecuteIfBound(ProposedSelection ? FText::FromString(ProposedSelection->ActualOption) : FText());
+		this->SetIsOpen(false);
+	}
+}
+
+
+void SMutableSearchComboBox::OnSearchTextChanged(const FText& ChangedText)
+{
+	SearchText = ChangedText;
+
+	RefreshOptions();
+}
+
+
+void SMutableSearchComboBox::OnSearchTextCommitted(const FText& InText, ETextCommit::Type InCommitType)
+{
+	if ((InCommitType == ETextCommit::Type::OnEnter) && FilteredOptionsSource.Num() > 0)
+	{
+		ComboListView->SetSelection(FilteredOptionsSource[0], ESelectInfo::OnKeyPress);
+	}
+}
+
+
+FReply SMutableSearchComboBox::OnKeyDownHandler(const FGeometry& MyGeometry, const FKeyEvent& InKeyEvent)
+{
+	if (InKeyEvent.GetKey() == EKeys::Enter)
+	{
+		// Select the first selected item on hitting enter
+		TArray<TSharedPtr<FFilteredOption>> SelectedItems = ComboListView->GetSelectedItems();
+		if (SelectedItems.Num() > 0)
+		{
+			OnSelectionChanged_Internal(SelectedItems[0], ESelectInfo::OnKeyPress);
+			return FReply::Handled();
+		}
+	}
+
+	return FReply::Unhandled();
+}
+
+
+#undef LOCTEXT_NAMESPACE
diff --git a/Engine/Plugins/Experimental/Mutable/Source/CustomizableObjectEditor/Private/MuCOE/SMutableSearchComboBox.h b/Engine/Plugins/Experimental/Mutable/Source/CustomizableObjectEditor/Private/MuCOE/SMutableSearchComboBox.h
new file mode 100644
--- /dev/null
+++ b/Engine/Plugins/Experimental/Mutable/Source/CustomizableObjectEditor/Private/MuCOE/SMutableSearchComboBox.h
@@ -0,0 +1,145 @@
+// Copyright Epic Games, Inc. All Rights Reserved.
+
+#pragma once
+
+#include "Types/SlateEnums.h"
+#include "Widgets/Input/SComboButton.h"
+#include "Widgets/Input/SEditableTextBox.h"
+#include "Widgets/Views/SListView.h"
+#include "Widgets/Views/ITableRow.h"
+
+namespace ESelectInfo { enum Type : int; }
+
+class FString;
+class IDetailLayoutBuilder;
+class IPropertyHandle;
+class UCustomizableObjectNodeObject;
+
+/** */
+class SMutableSearchComboBox : public SComboButton
+{
+public:
+	struct FFilteredOption
+	{
+		FString ActualOption;
+		FString DisplayOption;
+	};
+
+	/** Type of list used for showing menu options. */
+	typedef SListView< TSharedPtr<FFilteredOption> > SComboListType;
+
+	/** Delegate type used to generate widgets that represent Options */
+	typedef typename TSlateDelegates< TSharedPtr<FString> >::FOnGenerateWidget FOnGenerateWidget;
+
+	SLATE_BEGIN_ARGS(SMutableSearchComboBox)
+		: _Content()
+		, _ComboBoxStyle(&FAppStyle::Get().GetWidgetStyle<FComboBoxStyle>("ComboBox"))
+		, _ButtonStyle(nullptr)
+		, _ItemStyle(&FAppStyle::Get().GetWidgetStyle<FTableRowStyle>("ComboBox.Row"))
+		, _ContentPadding(_ComboBoxStyle->ContentPadding)
+		, _ForegroundColor(FSlateColor::UseStyle())
+		, _OptionsSource()
+		, _OnSelectionChanged()
+		, _OnGenerateWidget()
+		, _MenuButtonBrush(nullptr)
+		, _Method()
+		, _AllowAddNewOptions(false)
+		, _MaxListHeight(450.0f)
+		, _SearchVisibility()
+		{}
+
+		/** Slot for this button's content (optional) */
+		SLATE_DEFAULT_SLOT(FArguments, Content)
+
+		SLATE_STYLE_ARGUMENT(FComboBoxStyle, ComboBoxStyle)
+
+		/** The visual style of the button part of the combo box (overrides ComboBoxStyle) */
+		SLATE_STYLE_ARGUMENT(FButtonStyle, ButtonStyle)
+
+		SLATE_STYLE_ARGUMENT(FTableRowStyle, ItemStyle)
+
+		SLATE_ATTRIBUTE(FMargin, ContentPadding)
+		SLATE_ATTRIBUTE(FSlateColor, ForegroundColor)
+
+		SLATE_ARGUMENT(const TArray< TSharedPtr<FString> >*, OptionsSource)
+		SLATE_EVENT(FOnTextChanged, OnSelectionChanged)
+		SLATE_EVENT(FOnGenerateWidget, OnGenerateWidget)
+
+		SLATE_ARGUMENT(const FSlateBrush*, MenuButtonBrush)
+
+		SLATE_ARGUMENT(TOptional<EPopupMethod>, Method)
+
+		SLATE_ARGUMENT(bool, AllowAddNewOptions)
+
+		/** The max height of the combo box menu */
+		SLATE_ARGUMENT(float, MaxListHeight)
+
+		/** Allow setting the visibility of the search box dynamically */
+		SLATE_ATTRIBUTE(EVisibility, SearchVisibility)
+
+	SLATE_END_ARGS()
+
+	/**
+	 * Construct the widget from a declaration
+	 *
+	 * @param InArgs   Declaration from which to construct the combo box
+	 */
+	void Construct(const FArguments& InArgs);
+
+	/**
+	 * Requests a list refresh after updating options
+	 */
+	void RefreshOptions();
+
+private:
+
+	/** Generate a row for the InItem in the combo box's list (passed in as OwnerTable). Do this by calling the user-specified OnGenerateWidget */
+	TSharedRef<ITableRow> GenerateMenuItemRow(TSharedPtr<FFilteredOption> InItem, const TSharedRef<STableViewBase>& OwnerTable);
+
+	/** Called if the menu is closed */
+	void OnMenuOpenChanged(bool bOpen);
+
+	/** Invoked when the selection in the list changes */
+	void OnSelectionChanged_Internal(TSharedPtr<FFilteredOption> ProposedSelection, ESelectInfo::Type SelectInfo);
+
+	/** Invoked when the search text changes */
+	void OnSearchTextChanged(const FText& ChangedText);
+
+	/** Sets the current selection to the first valid match when user presses enter in the filter box */
+	void OnSearchTextCommitted(const FText& InText, ETextCommit::Type InCommitType);
+
+	FReply OnKeyDownHandler(const FGeometry& MyGeometry, const FKeyEvent& InKeyEvent);
+
+	/** The item style to use. */
+	const FTableRowStyle* ItemStyle;
+
+	/** The padding around each menu row */
+	FMargin MenuRowPadding;
+
+private:
+	/** Delegate that is invoked when the selected item in the combo box changes */
+	FOnTextChanged OnSelectionChanged;
+	/** The search field used for the combox box's contents */
+	TSharedPtr<SEditableTextBox> SearchField;
+	/** The ListView that we pop up; visualized the available options. */
+	TSharedPtr< SComboListType > ComboListView;
+	/** The Scrollbar used in the ListView. */
+	TSharedPtr< SScrollBar > CustomScrollbar;
+	/** Delegate to invoke when we need to visualize an option as a widget. */
+	FOnGenerateWidget OnGenerateWidget;
+
+	/** Updated whenever search text is changed */
+	FText SearchText;
+
+	/** Source data for this combo box */
+	const TArray< TSharedPtr<FString> >* OptionsSource;
+
+	/** Filtered list that is actually displayed */
+	TArray< TSharedPtr<FFilteredOption> > FilteredOptionsSource;
+
+	/** Copied to replace the image. */
+	FComboButtonStyle OurComboButtonStyle;
+
+	bool bAllowAddNewOptions = false;
+};
+
diff --git a/Engine/Plugins/Experimental/Mutable/Source/CustomizableObjectEditor/Public/MuCOE/Nodes/CustomizableObjectNodeModifierBase.h b/Engine/Plugins/Experimental/Mutable/Source/CustomizableObjectEditor/Public/MuCOE/Nodes/CustomizableObjectNodeModifierBase.h
--- a/Engine/Plugins/Experimental/Mutable/Source/CustomizableObjectEditor/Public/MuCOE/Nodes/CustomizableObjectNodeModifierBase.h
+++ b/Engine/Plugins/Experimental/Mutable/Source/CustomizableObjectEditor/Public/MuCOE/Nodes/CustomizableObjectNodeModifierBase.h
@@ -3,6 +3,7 @@
 #pragma once
 
 #include "MuCOE/Nodes/CustomizableObjectNode.h"
+#include "MuT/NodeModifier.h"
 
 #include "CustomizableObjectNodeModifierBase.generated.h"
 
@@ -33,6 +34,16 @@ class CUSTOMIZABLEOBJECTEDITOR_API UCustomizableObjectNodeModifierBase : public
 
 	GENERATED_BODY()
 
+public:
+
+	/** Materials in all other objects that activate this tags will be affected by this modifier. */
+	UPROPERTY(EditAnywhere, BlueprintReadWrite, Category = Modifier)
+	TArray<FString> RequiredTags;
+
+	/** Policy to use required tags in case more than one is added. */
+	UPROPERTY(EditAnywhere, Category = Modifier)
+	EMutableMultipleTagPolicy MultipleTagPolicy = EMutableMultipleTagPolicy::OnlyOneRequired;
+
 public:
 
 	// EdGraphNode interface
@@ -44,18 +55,15 @@ public:
 
 	// Own interface
 	virtual UEdGraphPin* OutputPin() const;
-	virtual const TArray<FString>* GetRequiredTags() const PURE_VIRTUAL(UCustomizableObjectNodeModifierBase::GetRequiredTags, return nullptr; );
-
-	virtual EMutableMultipleTagPolicy GetMultipleTagsPolicy() const PURE_VIRTUAL(UCustomizableObjectNodeModifierBase::GetMultipleTagsPolicy, return {}; );
 
 	/** Return true if the tags and policy in this node would make this node modify the given parameter node. */
-	bool IsApplicableTo(UCustomizableObjectNode* Candidate) const;
+	bool IsApplicableTo(UCustomizableObjectNode* Candidate);
 
 	/** Get the list of all nodes of this object that may be affected by this modifier. This is an approximate
 	* query: it includes all nodes that may be affected by this modifier but it may include nodes that will not 
 	* be modified by this modifier because of other conditions.
 	*/
-	void GetPossiblyModifiedNodes(TArray<UCustomizableObjectNode*>&) const;
+	void GetPossiblyModifiedNodes(TArray<UCustomizableObjectNode*>&);
 
 	/** Generate a tag that is unique and stable for the given node. */
 	static FString MakeNodeAutoTag(UEdGraphNode*);
diff --git a/Engine/Plugins/Experimental/Mutable/Source/CustomizableObjectEditor/Public/MuCOE/Nodes/CustomizableObjectNodeModifierClipDeform.h b/Engine/Plugins/Experimental/Mutable/Source/CustomizableObjectEditor/Public/MuCOE/Nodes/CustomizableObjectNodeModifierClipDeform.h
--- a/Engine/Plugins/Experimental/Mutable/Source/CustomizableObjectEditor/Public/MuCOE/Nodes/CustomizableObjectNodeModifierClipDeform.h
+++ b/Engine/Plugins/Experimental/Mutable/Source/CustomizableObjectEditor/Public/MuCOE/Nodes/CustomizableObjectNodeModifierClipDeform.h
@@ -27,12 +27,8 @@ class CUSTOMIZABLEOBJECTEDITOR_API UCustomizableObjectNodeModifierClipDeform : p
 
 public:
 
-	UPROPERTY(EditAnywhere, BlueprintReadWrite, Category = MeshClipDeform)
-	TArray<FString> Tags;
-	
-	/** Policy to use tags in case more than one is added. */
-	UPROPERTY(EditAnywhere, Category = MeshClipDeform)
-	EMutableMultipleTagPolicy MultipleTagPolicy = EMutableMultipleTagPolicy::OnlyOneRequired;
+	UPROPERTY()
+	TArray<FString> Tags_DEPRECATED;
 
 	UPROPERTY(EditAnywhere, Category = MeshClipDeform)
 	EShapeBindingMethod BindingMethod;
@@ -45,10 +41,7 @@ public:
 
 	// UCustomizableObjectNode interface
 	virtual void AllocateDefaultPins(UCustomizableObjectNodeRemapPins* RemapPins) override;
-
-	// UCustomizableObjectNodeModifierBase interface
-	virtual const TArray<FString>* GetRequiredTags() const override { return &Tags; }
-	virtual EMutableMultipleTagPolicy GetMultipleTagsPolicy() const override { return MultipleTagPolicy; }
+	virtual void BackwardsCompatibleFixup(int32 CustomizableObjectCustomVersion) override;
 
 	// Own interface
 
diff --git a/Engine/Plugins/Experimental/Mutable/Source/CustomizableObjectEditor/Public/MuCOE/Nodes/CustomizableObjectNodeModifierClipMorph.h b/Engine/Plugins/Experimental/Mutable/Source/CustomizableObjectEditor/Public/MuCOE/Nodes/CustomizableObjectNodeModifierClipMorph.h
--- a/Engine/Plugins/Experimental/Mutable/Source/CustomizableObjectEditor/Public/MuCOE/Nodes/CustomizableObjectNodeModifierClipMorph.h
+++ b/Engine/Plugins/Experimental/Mutable/Source/CustomizableObjectEditor/Public/MuCOE/Nodes/CustomizableObjectNodeModifierClipMorph.h
@@ -24,12 +24,8 @@ public:
 	UPROPERTY(EditAnywhere, Category = MeshToClipAndMorph)
 	FName BoneName;
 
-	UPROPERTY(EditAnywhere, BlueprintReadWrite, Category = MeshToClipAndMorph)
-	TArray<FString> Tags;
-
-	/** Policy to use tags in case more than one is added. */
-	UPROPERTY(EditAnywhere, Category = MeshToClipAndMorph)
-	EMutableMultipleTagPolicy MultipleTagPolicy = EMutableMultipleTagPolicy::OnlyOneRequired;
+	UPROPERTY()
+	TArray<FString> Tags_DEPRECATED;
 
 	UPROPERTY(EditAnywhere, Category = MeshToClipAndMorph)
 	uint32 ReferenceSkeletonIndex = 0;
@@ -87,10 +83,7 @@ public:
 	// UCustomizableObjectNode interface
 	virtual void AllocateDefaultPins(UCustomizableObjectNodeRemapPins* RemapPins) override;
 	virtual void PinConnectionListChanged(UEdGraphPin* Pin) override;
-
-	// UCustomizableObjectNodeModifierBase interface
-	virtual const TArray<FString>* GetRequiredTags() const override { return &Tags; }
-	virtual EMutableMultipleTagPolicy GetMultipleTagsPolicy() const override { return MultipleTagPolicy; }
+	virtual void BackwardsCompatibleFixup(int32 CustomizableObjectCustomVersion) override;
 
 	// Own interface
 	UEdGraphPin* OutputPin() const;
diff --git a/Engine/Plugins/Experimental/Mutable/Source/CustomizableObjectEditor/Public/MuCOE/Nodes/CustomizableObjectNodeModifierClipWithMesh.h b/Engine/Plugins/Experimental/Mutable/Source/CustomizableObjectEditor/Public/MuCOE/Nodes/CustomizableObjectNodeModifierClipWithMesh.h
--- a/Engine/Plugins/Experimental/Mutable/Source/CustomizableObjectEditor/Public/MuCOE/Nodes/CustomizableObjectNodeModifierClipWithMesh.h
+++ b/Engine/Plugins/Experimental/Mutable/Source/CustomizableObjectEditor/Public/MuCOE/Nodes/CustomizableObjectNodeModifierClipWithMesh.h
@@ -23,11 +23,8 @@ public:
 
 	UCustomizableObjectNodeModifierClipWithMesh();
 
-	UPROPERTY(EditAnywhere, BlueprintReadWrite, Category = MeshToClip)
-	TArray<FString> Tags;
-
-	UPROPERTY(EditAnywhere, Category = CustomizableObjectToClip)
-	EMutableMultipleTagPolicy MultipleTagPolicy = EMutableMultipleTagPolicy::OnlyOneRequired;
+	UPROPERTY()
+	TArray<FString> Tags_DEPRECATED;
 
 	//!< If assigned, then a material inside this CO will be clipped by this node.
     //!< If several materials with the same name, all are considered (to cover all LOD levels)
@@ -59,10 +56,7 @@ public:
 	virtual void PinConnectionListChanged(UEdGraphPin* Pin) override;
 	virtual void UpdateReferencedNodeId(const FGuid& NewGuid) override;
 	virtual void BeginPostDuplicate(bool bDuplicateForPIE) override;
-
-	// UCustomizableObjectNodeModifierBase interface
-	virtual const TArray<FString>* GetRequiredTags() const override { return &Tags; }
-	virtual EMutableMultipleTagPolicy GetMultipleTagsPolicy() const override { return MultipleTagPolicy; }
+	virtual void BackwardsCompatibleFixup(int32 CustomizableObjectCustomVersion) override;
 
 	// Own interface
 	UEdGraphPin* OutputPin() const;
diff --git a/Engine/Plugins/Experimental/Mutable/Source/CustomizableObjectEditor/Public/MuCOE/Nodes/CustomizableObjectNodeModifierClipWithUVMask.h b/Engine/Plugins/Experimental/Mutable/Source/CustomizableObjectEditor/Public/MuCOE/Nodes/CustomizableObjectNodeModifierClipWithUVMask.h
--- a/Engine/Plugins/Experimental/Mutable/Source/CustomizableObjectEditor/Public/MuCOE/Nodes/CustomizableObjectNodeModifierClipWithUVMask.h
+++ b/Engine/Plugins/Experimental/Mutable/Source/CustomizableObjectEditor/Public/MuCOE/Nodes/CustomizableObjectNodeModifierClipWithUVMask.h
@@ -21,12 +21,9 @@ class CUSTOMIZABLEOBJECTEDITOR_API UCustomizableObjectNodeModifierClipWithUVMask
 public:
 
 	/** Materials in all other objects that activate this tags will be clipped with this UV mask. */
-	UPROPERTY(EditAnywhere, BlueprintReadWrite, Category = MeshToClip)
-	TArray<FString> Tags;
+	UPROPERTY()
+	TArray<FString> Tags_DEPRECATED;
 
-	/** Policy to use tags in case more than one is added. */
-	UPROPERTY(EditAnywhere, Category = MeshToClip)
-	EMutableMultipleTagPolicy MultipleTagPolicy = EMutableMultipleTagPolicy::OnlyOneRequired;
 
 	/** UV channel index that will be used to get the UVs to apply the clipping mask to. */
 	UPROPERTY(EditAnywhere, BlueprintReadWrite, Category = MeshToClip)
@@ -41,10 +38,8 @@ public:
 	// UCustomizableObjectNode interface
 	virtual void AllocateDefaultPins(UCustomizableObjectNodeRemapPins* RemapPins) override;
 	virtual void PinConnectionListChanged(UEdGraphPin* Pin) override;
+	virtual void BackwardsCompatibleFixup(int32 CustomizableObjectCustomVersion) override;
 
-	// UCustomizableObjectNodeModifierBase interface
-	virtual const TArray<FString>* GetRequiredTags() const override { return &Tags; }
-	virtual EMutableMultipleTagPolicy GetMultipleTagsPolicy() const override { return MultipleTagPolicy; }
 
 	// Own interface
 
diff --git a/Engine/Plugins/Experimental/Mutable/Source/CustomizableObjectEditor/Public/MuCOE/Nodes/CustomizableObjectNodeModifierWithMaterial.h b/Engine/Plugins/Experimental/Mutable/Source/CustomizableObjectEditor/Public/MuCOE/Nodes/CustomizableObjectNodeModifierWithMaterial.h
--- a/Engine/Plugins/Experimental/Mutable/Source/CustomizableObjectEditor/Public/MuCOE/Nodes/CustomizableObjectNodeModifierWithMaterial.h
+++ b/Engine/Plugins/Experimental/Mutable/Source/CustomizableObjectEditor/Public/MuCOE/Nodes/CustomizableObjectNodeModifierWithMaterial.h
@@ -4,7 +4,6 @@
 
 #include "MuCOE/Nodes/CustomizableObjectNodeModifierBase.h"
 #include "MuCOE/Nodes/CustomizableObjectNodeMaterialBase.h"
-#include "MuT/NodeModifier.h"
 
 #include "CustomizableObjectNodeModifierWithMaterial.generated.h"
 
@@ -19,14 +18,6 @@ class CUSTOMIZABLEOBJECTEDITOR_API UCustomizableObjectNodeModifierWithMaterial :
 
 public:
 
-	/** Materials in all other objects that activate this tags will be affected by this modifier. */
-	UPROPERTY(EditAnywhere, BlueprintReadWrite, Category = Modifier)
-	TArray<FString> RequiredTags;
-
-	/** Policy to use required tags in case more than one is added. */
-	UPROPERTY(EditAnywhere, Category = Modifier)
-	EMutableMultipleTagPolicy MultipleTagPolicy = EMutableMultipleTagPolicy::OnlyOneRequired;
-
 	/** Reference material that defines the structure of the data to be extended.
 	* The Sections modified with this modifier are supposed to have the same texture parameters, but don't need to have exactly
 	* the ReferenceMaterial set.
@@ -46,10 +37,6 @@ public:
 	// UCustomizableObjectNode interface
 	virtual bool IsNodeOutDatedAndNeedsRefresh() override;
 
-	// UCustomizableObjectNodeModifierBase interface
-	virtual const TArray<FString>* GetRequiredTags() const override { return &RequiredTags;  }
-	virtual EMutableMultipleTagPolicy GetMultipleTagsPolicy() const override { return MultipleTagPolicy;  }
-
 	// Own interface
 
 	/** */
