[mutable] Fix possible crash with legacy variation nodes.
#rnx

diff --git a/Engine/Plugins/Experimental/Mutable/Source/CustomizableObjectEditor/Private/MuCOE/Nodes/CustomizableObjectNodeObject.cpp b/Engine/Plugins/Experimental/Mutable/Source/CustomizableObjectEditor/Private/MuCOE/Nodes/CustomizableObjectNodeObject.cpp
--- a/Engine/Plugins/Experimental/Mutable/Source/CustomizableObjectEditor/Private/MuCOE/Nodes/CustomizableObjectNodeObject.cpp
+++ b/Engine/Plugins/Experimental/Mutable/Source/CustomizableObjectEditor/Private/MuCOE/Nodes/CustomizableObjectNodeObject.cpp
@@ -402,22 +402,26 @@ void UCustomizableObjectNodeObject::BackwardsCompatibleFixup(int32 CustomizableO
 						
 						for (int32 ElementIndex = 0; ElementIndex < NodeMaterialVariation->GetNumVariations(); ++ElementIndex)
 						{
-							if (UEdGraphPin* ConnectedPin = FollowInputPin(*NodeMaterialVariation->VariationPin(ElementIndex)))
+							UEdGraphPin* VariationPin = NodeMaterialVariation->VariationPin(ElementIndex);
+							if (VariationPin)
 							{
-								if (UCustomizableObjectNodeMaterialBase* FirstNodeMaterialBase = Cast<UCustomizableObjectNodeMaterialBase>(ConnectedPin->GetOwningNode()))
+								if (UEdGraphPin* ConnectedPin = FollowInputPin(*VariationPin))
 								{
-									if (UCustomizableObjectNodeMaterial* FirstNodeMaterial = FirstNodeMaterialBase->GetMaterialNode())
+									if (UCustomizableObjectNodeMaterialBase* FirstNodeMaterialBase = Cast<UCustomizableObjectNodeMaterialBase>(ConnectedPin->GetOwningNode()))
 									{
-										if (bFirst)
-										{
-											bFirst = false;
-											ComponentName = FirstNodeMaterial->MeshComponentName_DEPRECATED;
-										}
-										else
+										if (UCustomizableObjectNodeMaterial* FirstNodeMaterial = FirstNodeMaterialBase->GetMaterialNode())
 										{
-											if (ComponentName != FirstNodeMaterial->MeshComponentName_DEPRECATED) // All components must match. If not, not supported.
+											if (bFirst)
 											{
-												return;
+												bFirst = false;
+												ComponentName = FirstNodeMaterial->MeshComponentName_DEPRECATED;
+											}
+											else
+											{
+												if (ComponentName != FirstNodeMaterial->MeshComponentName_DEPRECATED) // All components must match. If not, not supported.
+												{
+													return;
+												}
 											}
 										}
 									}
diff --git a/Engine/Plugins/Experimental/Mutable/Source/CustomizableObjectEditor/Private/MuCOE/Nodes/CustomizableObjectNodeVariation.cpp b/Engine/Plugins/Experimental/Mutable/Source/CustomizableObjectEditor/Private/MuCOE/Nodes/CustomizableObjectNodeVariation.cpp
--- a/Engine/Plugins/Experimental/Mutable/Source/CustomizableObjectEditor/Private/MuCOE/Nodes/CustomizableObjectNodeVariation.cpp
+++ b/Engine/Plugins/Experimental/Mutable/Source/CustomizableObjectEditor/Private/MuCOE/Nodes/CustomizableObjectNodeVariation.cpp
@@ -74,7 +74,11 @@ UEdGraphPin* UCustomizableObjectNodeVariation::DefaultPin() const
 
 UEdGraphPin* UCustomizableObjectNodeVariation::VariationPin(int32 Index) const
 {
-	return VariationsPins[Index].Get();
+	if (VariationsPins.IsValidIndex(Index))
+	{
+		return VariationsPins[Index].Get();
+	}
+	return nullptr;
 }
 
 
