[mutable] Redesign modifiers in core node hierarchy to move them to object nodes.
#rnx

diff --git a/Engine/Plugins/Experimental/Mutable/Source/CustomizableObjectEditor/Private/MuCOE/GenerateMutableSource/GenerateMutableSource.cpp b/Engine/Plugins/Experimental/Mutable/Source/CustomizableObjectEditor/Private/MuCOE/GenerateMutableSource/GenerateMutableSource.cpp
--- a/Engine/Plugins/Experimental/Mutable/Source/CustomizableObjectEditor/Private/MuCOE/GenerateMutableSource/GenerateMutableSource.cpp
+++ b/Engine/Plugins/Experimental/Mutable/Source/CustomizableObjectEditor/Private/MuCOE/GenerateMutableSource/GenerateMutableSource.cpp
@@ -929,25 +929,7 @@ mu::Ptr<mu::NodeObject> GenerateMutableSource(const UEdGraphPin * Pin, FMutableG
 				GenerationContext.CurrentMeshComponent = FName();
 
 				mu::Ptr<mu::NodeModifier> ModifierNode = GenerateMutableSourceModifier(ChildNodePin, GenerationContext);
-				for (int32 ComponentIndex = 0; ComponentIndex < ObjectNode->Components.Num(); ++ComponentIndex)
-				{
-					mu::Ptr<mu::NodeComponent> Component = ObjectNode->Components[ComponentIndex];
-					if (!Component)
-					{
-						continue;
-					}
-					
-					for (int32 LODIndex = 0; LODIndex < Component->LODs.Num(); ++LODIndex)
-					{
-						mu::Ptr<mu::NodeLOD> LOD = Component->LODs[LODIndex];
-						if (!LOD)
-						{
-							continue;
-						}
-						
-						LOD->Modifiers.Add(ModifierNode);
-					}
-				}
+				ObjectNode->Modifiers.Add(ModifierNode);
 
 				GenerationContext.CurrentMeshComponent = OldCurrentMeshComponent;
 			}
diff --git a/Engine/Plugins/Experimental/Mutable/Source/CustomizableObjectEditor/Private/MuCOE/GenerateMutableSource/GenerateMutableSourceComponent.cpp b/Engine/Plugins/Experimental/Mutable/Source/CustomizableObjectEditor/Private/MuCOE/GenerateMutableSource/GenerateMutableSourceComponent.cpp
--- a/Engine/Plugins/Experimental/Mutable/Source/CustomizableObjectEditor/Private/MuCOE/GenerateMutableSource/GenerateMutableSourceComponent.cpp
+++ b/Engine/Plugins/Experimental/Mutable/Source/CustomizableObjectEditor/Private/MuCOE/GenerateMutableSource/GenerateMutableSourceComponent.cpp
@@ -40,7 +40,7 @@
 /** Generate LOD pins of the given NodeComponentBase (NodeComponent, NodeComponentExtend...).
  * @param TypedComponentMesh Given component node.
  * @param NodeComponent Core node to connect LOD generated pins. */
-void GenerateMutableSourceComponentMesh(FMutableGraphGenerationContext& GenerationContext, const UCustomizableObjectNodeComponentMeshBase& TypedComponentMesh, mu::Ptr<mu::NodeComponent> NodeComponent)
+void GenerateMutableSourceComponentMesh(FMutableGraphGenerationContext& GenerationContext, const UCustomizableObjectNodeComponentMeshBase& TypedComponentMesh, mu::Ptr<mu::NodeComponent> NodeComponent, mu::Ptr<mu::NodeObjectNew> NodeObject)
 {
 	int32 FirstLOD = -1;
 
@@ -109,18 +109,26 @@ void GenerateMutableSourceComponentMesh(FMutableGraphGenerationContext& Generati
 					continue;
 				}
 
-				// Warn about legacy modifier connection
-				FString Msg = FString::Printf(TEXT("The object has legacy modifier connections (to material pins?) that should be updated."));
-				GenerationContext.Compiler->CompilerLog(FText::FromString(Msg), &TypedComponentMesh, EMessageSeverity::Warning);
+				if (NodeObject)
+				{
+					// Warn about legacy modifier connection
+					FString Msg = FString::Printf(TEXT("The object has legacy modifier connections (to material pins?) that should be updated."));
+					GenerationContext.Compiler->CompilerLog(FText::FromString(Msg), &TypedComponentMesh, EMessageSeverity::Warning);
 
-				// Set it to "None" to indicate we don't care about component id.
-				FName OldCurrentMeshComponent = GenerationContext.CurrentMeshComponent;
-				GenerationContext.CurrentMeshComponent = FName();
+					// Set it to "None" to indicate we don't care about component id.
+					FName OldCurrentMeshComponent = GenerationContext.CurrentMeshComponent;
+					GenerationContext.CurrentMeshComponent = FName();
 
-				mu::Ptr<mu::NodeModifier> ModifierNode = GenerateMutableSourceModifier(ChildNodePin, GenerationContext);
-				LODNode->Modifiers.Add(ModifierNode);
+					mu::Ptr<mu::NodeModifier> ModifierNode = GenerateMutableSourceModifier(ChildNodePin, GenerationContext);
+					NodeObject->Modifiers.AddUnique(ModifierNode);
 
-				GenerationContext.CurrentMeshComponent = OldCurrentMeshComponent;
+					GenerationContext.CurrentMeshComponent = OldCurrentMeshComponent;
+				}
+				else
+				{
+					FString Msg = FString::Printf(TEXT("The object has legacy modifier connections that cannot be generated. Their connections should be updated."));
+					GenerationContext.Compiler->CompilerLog(FText::FromString(Msg), &TypedComponentMesh, EMessageSeverity::Warning);
+				}
 			}
 		}
 	}
@@ -283,7 +291,7 @@ mu::Ptr<mu::NodeComponent> GenerateMutableSourceComponent(const UEdGraphPin* Pin
 		ObjectNode->Components.Add(NodeComponentNew);
 
 		GenerationContext.CurrentMeshComponent = TypedComponentMesh->ComponentName;
-		GenerateMutableSourceComponentMesh(GenerationContext, *TypedComponentMesh, NodeComponentNew);
+		GenerateMutableSourceComponentMesh(GenerationContext, *TypedComponentMesh, NodeComponentNew, ObjectNode);
 		
 		Result = NodeComponentNew;
 	}
@@ -316,7 +324,7 @@ mu::Ptr<mu::NodeComponent> GenerateMutableSourceComponent(const UEdGraphPin* Pin
 				TypedParentComponentMesh->AutoLODStrategy :
 				TypedComponentMeshExtend->AutoLODStrategy;
 			
-			GenerateMutableSourceComponentMesh(GenerationContext, *TypedComponentMeshExtend, NodeComponentEdit);
+			GenerateMutableSourceComponentMesh(GenerationContext, *TypedComponentMeshExtend, NodeComponentEdit, nullptr);
 
 			Result = NodeComponentEdit;	
 		}
diff --git a/Engine/Plugins/Experimental/Mutable/Source/CustomizableObjectEditor/Private/MuCOE/SMutableGraphViewer.cpp b/Engine/Plugins/Experimental/Mutable/Source/CustomizableObjectEditor/Private/MuCOE/SMutableGraphViewer.cpp
--- a/Engine/Plugins/Experimental/Mutable/Source/CustomizableObjectEditor/Private/MuCOE/SMutableGraphViewer.cpp
+++ b/Engine/Plugins/Experimental/Mutable/Source/CustomizableObjectEditor/Private/MuCOE/SMutableGraphViewer.cpp
@@ -318,6 +318,11 @@ void SMutableGraphViewer::GetChildrenForInfo(TSharedPtr<FMutableGraphTreeElement
 			AddChildFunc(ObjectNew->Components[l].get(), TEXT("COMP") );
 		}
 
+		for (int32 Modifier = 0; Modifier < ObjectNew->Modifiers.Num(); Modifier++)
+		{
+			AddChildFunc(ObjectNew->Modifiers[Modifier].get(), FString::Printf(TEXT("MOD [%d]"), Modifier));
+		}
+
 		for (int32 l = 0; l < ObjectNew->Children.Num(); ++l)
 		{
 			AddChildFunc(ObjectNew->Children[l].get(), TEXT("CHILD"));
@@ -420,10 +425,6 @@ void SMutableGraphViewer::GetChildrenForInfo(TSharedPtr<FMutableGraphTreeElement
 		{
 			AddChildFunc(LodVar->Surfaces[SurfaceIndex].get(), FString::Printf(TEXT("SURFACE [%d]"), SurfaceIndex));
 		}
-		for (int32 Modifier = 0; Modifier < LodVar->Modifiers.Num(); Modifier++)
-		{
-			AddChildFunc(LodVar->Modifiers[Modifier].get(), FString::Printf(TEXT("MOD [%d]"), Modifier));
-		}
 	}
 	
 	else if (ParentNode->GetType() == mu::NodeComponentNew::GetStaticType())
diff --git a/Engine/Plugins/Experimental/Mutable/Source/MutableTools/Private/MuT/CodeGenerator.cpp b/Engine/Plugins/Experimental/Mutable/Source/MutableTools/Private/MuT/CodeGenerator.cpp
--- a/Engine/Plugins/Experimental/Mutable/Source/MutableTools/Private/MuT/CodeGenerator.cpp
+++ b/Engine/Plugins/Experimental/Mutable/Source/MutableTools/Private/MuT/CodeGenerator.cpp
@@ -1866,10 +1866,11 @@ namespace mu
 		for (const FirstPassGenerator::FModifier& m: FirstPass.Modifiers)
 		{
 			// Correct LOD?
-			if (m.LOD != LOD)
-			{
-				continue;
-			}
+			// TODO: LOD restrictions inside the modifiers data.
+			//if (m.LOD != LOD)
+			//{
+			//	continue;
+			//}
 
 			// Correct stage?
 			if (m.Node->bApplyBeforeNormalOperations != bModifiersForBeforeOperations)
diff --git a/Engine/Plugins/Experimental/Mutable/Source/MutableTools/Private/MuT/CodeGenerator_FirstPass.cpp b/Engine/Plugins/Experimental/Mutable/Source/MutableTools/Private/MuT/CodeGenerator_FirstPass.cpp
--- a/Engine/Plugins/Experimental/Mutable/Source/MutableTools/Private/MuT/CodeGenerator_FirstPass.cpp
+++ b/Engine/Plugins/Experimental/Mutable/Source/MutableTools/Private/MuT/CodeGenerator_FirstPass.cpp
@@ -195,14 +195,11 @@ namespace mu
 	//---------------------------------------------------------------------------------------------
 	void FirstPassGenerator::Generate_Modifier(const NodeModifier* InNode)
 	{
-		check(CurrentLOD>=0);
-
 		// Add the data about this modifier
 		FModifier thisData;
 		thisData.Node = InNode;
 		thisData.ObjectCondition = CurrentCondition.Last().ObjectCondition;
 		thisData.StateCondition = CurrentStateCondition.Last();
-		thisData.LOD = CurrentLOD;
 		thisData.PositiveTags = CurrentPositiveTags;
 		thisData.NegativeTags = CurrentNegativeTags;
 		Modifiers.Add(thisData);
@@ -556,14 +553,6 @@ namespace mu
 				Generate_Generic(c.get());
 			}
 		}
-
-		for (const Ptr<NodeModifier>& c : InNode->Modifiers)
-		{
-			if (c)
-			{
-				Generate_Modifier(c.get());
-			}
-		}
 	}
 
 
@@ -599,6 +588,15 @@ namespace mu
 			}
 		}
 
+		// Process the modifiers
+		for (const Ptr<NodeModifier>& c : InNode->Modifiers)
+		{
+			if (c)
+			{
+				Generate_Modifier(c.get());
+			}
+		}
+
 		// Process the children
 		for (const Ptr<NodeObject>& Child : InNode->Children)
 		{
diff --git a/Engine/Plugins/Experimental/Mutable/Source/MutableTools/Private/MuT/CodeGenerator_FirstPass.h b/Engine/Plugins/Experimental/Mutable/Source/MutableTools/Private/MuT/CodeGenerator_FirstPass.h
--- a/Engine/Plugins/Experimental/Mutable/Source/MutableTools/Private/MuT/CodeGenerator_FirstPass.h
+++ b/Engine/Plugins/Experimental/Mutable/Source/MutableTools/Private/MuT/CodeGenerator_FirstPass.h
@@ -160,9 +160,6 @@ namespace mu
 
             // This is filled in CodeGenerator_SecondPass.
             StateCondition StateCondition;
-
-            //
-            int32 LOD = 0;
         };
 		TArray<FModifier> Modifiers;
 
diff --git a/Engine/Plugins/Experimental/Mutable/Source/MutableTools/Public/MuT/NodeLOD.h b/Engine/Plugins/Experimental/Mutable/Source/MutableTools/Public/MuT/NodeLOD.h
--- a/Engine/Plugins/Experimental/Mutable/Source/MutableTools/Public/MuT/NodeLOD.h
+++ b/Engine/Plugins/Experimental/Mutable/Source/MutableTools/Public/MuT/NodeLOD.h
@@ -28,7 +28,6 @@ namespace mu
 		//-----------------------------------------------------------------------------------------
 
 		TArray<Ptr<NodeSurface>> Surfaces;
-		TArray<Ptr<NodeModifier>> Modifiers;
 
 	protected:
 
diff --git a/Engine/Plugins/Experimental/Mutable/Source/MutableTools/Public/MuT/NodeObjectNew.h b/Engine/Plugins/Experimental/Mutable/Source/MutableTools/Public/MuT/NodeObjectNew.h
--- a/Engine/Plugins/Experimental/Mutable/Source/MutableTools/Public/MuT/NodeObjectNew.h
+++ b/Engine/Plugins/Experimental/Mutable/Source/MutableTools/Public/MuT/NodeObjectNew.h
@@ -29,6 +29,9 @@ namespace mu
 		/** Components defined in the object. */
 		TArray<Ptr<NodeComponent>> Components;
 
+		/** Modifiers defined in the object. */
+		TArray<Ptr<NodeModifier>> Modifiers;
+
 		/** Children objects. */
 		TArray<Ptr<NodeObject>> Children;
 
