[mutable] Name a local namespace to avoid its members to leak on the global scope and cause compile errors due to unity builds.
#rnx

diff --git a/Engine/Plugins/Experimental/Mutable/Source/CustomizableObjectEditor/Private/MuCOE/EdGraphSchema_CustomizableObject.cpp b/Engine/Plugins/Experimental/Mutable/Source/CustomizableObjectEditor/Private/MuCOE/EdGraphSchema_CustomizableObject.cpp
--- a/Engine/Plugins/Experimental/Mutable/Source/CustomizableObjectEditor/Private/MuCOE/EdGraphSchema_CustomizableObject.cpp
+++ b/Engine/Plugins/Experimental/Mutable/Source/CustomizableObjectEditor/Private/MuCOE/EdGraphSchema_CustomizableObject.cpp
@@ -99,7 +99,7 @@ class IToolkit;
 
 #define SNAP_GRID (16) // @todo ensure this is the same as SNodePanel::GetSnapGridSize()
 
-namespace 
+namespace MutablePrivate
 {
 	// Maximum distance a drag can be off a node edge to require 'push off' from node
 	const int32 NodeDistance = 60;
@@ -179,11 +179,11 @@ UEdGraphNode* FCustomizableObjectSchemaAction_NewNode::CreateNode(class UEdGraph
 		UEdGraphNode* PinNode = FromPin->GetOwningNode();
 		const float XDelta = FMath::Abs(PinNode->NodePosX - Location.X);
 
-		if (XDelta < NodeDistance)
+		if (XDelta < MutablePrivate::NodeDistance)
 		{
 			// Set location to edge of current node minus the max move distance
 			// to force node to push off from connect node enough to give selection handle
-			XLocation = PinNode->NodePosX - NodeDistance;
+			XLocation = PinNode->NodePosX - MutablePrivate::NodeDistance;
 		}
 	}
 
