[Mutable] Merged trivial NodeComponentMesh and NodeComponentMeshAddTo nodes.

#rnx

diff --git a/Engine/Plugins/Experimental/Mutable/Source/CustomizableObject/Internal/MuCO/CustomizableObjectCustomVersion.h b/Engine/Plugins/Experimental/Mutable/Source/CustomizableObject/Internal/MuCO/CustomizableObjectCustomVersion.h
--- a/Engine/Plugins/Experimental/Mutable/Source/CustomizableObject/Internal/MuCO/CustomizableObjectCustomVersion.h
+++ b/Engine/Plugins/Experimental/Mutable/Source/CustomizableObject/Internal/MuCO/CustomizableObjectCustomVersion.h
@@ -141,6 +141,8 @@ struct CUSTOMIZABLEOBJECT_API FCustomizableObjectCustomVersion
 		ConvertEditAndExtendToModifiers,
 
 		NodeComponentMesh, // Moved LOD pins from NodeObject to NodeComponentMesh.
+
+		MergeNodeComponents,
 		
 		UnifyRequiredTags,
 
diff --git a/Engine/Plugins/Experimental/Mutable/Source/CustomizableObjectEditor/Private/MuCOE/Nodes/CustomizableObjectNodeObject.cpp b/Engine/Plugins/Experimental/Mutable/Source/CustomizableObjectEditor/Private/MuCOE/Nodes/CustomizableObjectNodeObject.cpp
--- a/Engine/Plugins/Experimental/Mutable/Source/CustomizableObjectEditor/Private/MuCOE/Nodes/CustomizableObjectNodeObject.cpp
+++ b/Engine/Plugins/Experimental/Mutable/Source/CustomizableObjectEditor/Private/MuCOE/Nodes/CustomizableObjectNodeObject.cpp
@@ -456,6 +456,58 @@ void UCustomizableObjectNodeObject::BackwardsCompatibleFixup(int32 CustomizableO
 		
 		FixupReconstructPins(CreateRemapPinsByName(), NodeObjectAllocateDefaultPins);
 	}
+
+	if (CustomizableObjectCustomVersion == FCustomizableObjectCustomVersion::MergeNodeComponents)
+	{
+		if (UEdGraphPin* PinComponents = ComponentsPin())
+		{
+			// Find all Mesh Component nodes.
+			TMap<FName, UCustomizableObjectNodeComponentMesh*> NodeComponentMeshes;
+			for (UEdGraphPin* Pin : PinComponents->LinkedTo)
+			{
+				if (UCustomizableObjectNodeComponentMesh* NodeComponentMesh = Cast<UCustomizableObjectNodeComponentMesh>(Pin->GetOwningNode()))
+				{
+					NodeComponentMeshes.Emplace(NodeComponentMesh->ComponentName, NodeComponentMesh);
+				}
+			}
+
+			TArray<UCustomizableObjectNodeComponentMeshAddTo*> NodesToRemove;
+
+			// Find all Add To Mesh Component nodes.
+			for (UEdGraphPin* Pin : PinComponents->LinkedTo)
+			{
+				if (UCustomizableObjectNodeComponentMeshAddTo* NodeComponentMeshAddTo = Cast<UCustomizableObjectNodeComponentMeshAddTo>(Pin->GetOwningNode()))
+				{
+					UCustomizableObjectNodeComponentMesh** Result = NodeComponentMeshes.Find(NodeComponentMeshAddTo->ParentComponentName);
+					if (!Result)
+					{
+						continue;
+					}
+
+					UCustomizableObjectNodeComponentMesh* NodeComponentMesh = *Result;
+					if (NodeComponentMesh->NumLODs != NodeComponentMeshAddTo->NumLODs)
+					{
+						continue;
+					}
+				
+					for (int32 LODIndex = 0; LODIndex < NodeComponentMesh->NumLODs; ++LODIndex)
+					{
+						for (UEdGraphPin* LinkedPin : NodeComponentMeshAddTo->LODPins[LODIndex].Get()->LinkedTo)
+						{
+							LinkedPin->MakeLinkTo(NodeComponentMesh->LODPins[LODIndex].Get());
+						}
+					}
+
+					NodesToRemove.Add(NodeComponentMeshAddTo);
+				}
+			}
+
+			for (UCustomizableObjectNodeComponentMeshAddTo* Node : NodesToRemove)
+			{
+				GetGraph()->RemoveNode(Node);
+			}
+		}
+	}
 }
 
 
