[Mutable] Remove Swizzle operations restrictions around first channel.

- The first channel was used as reference in some cases and its presence was mandatory. Now this restriction is lifted, the first availabe channel is used instead.

#rnx

diff --git a/Engine/Plugins/Experimental/Mutable/Source/MutableRuntime/Private/MuR/CodeRunnerTasked.cpp b/Engine/Plugins/Experimental/Mutable/Source/MutableRuntime/Private/MuR/CodeRunnerTasked.cpp
--- a/Engine/Plugins/Experimental/Mutable/Source/MutableRuntime/Private/MuR/CodeRunnerTasked.cpp
+++ b/Engine/Plugins/Experimental/Mutable/Source/MutableRuntime/Private/MuR/CodeRunnerTasked.cpp
@@ -1432,22 +1432,33 @@ namespace mu
 		MUTABLE_CPUPROFILER_SCOPE(FImageSwizzleTask_Prepare);
 		bOutFailed = false;
 
-		for (int i = 0; i < MUTABLE_OP_MAX_SWIZZLE_CHANNELS; ++i)
+		int32 FirstValidSourceIndex = -1;
+		for (int32 SourceIndex = 0; SourceIndex < MUTABLE_OP_MAX_SWIZZLE_CHANNELS; ++SourceIndex)
+		{
+			if (Args.sources[SourceIndex])
+			{
+				FirstValidSourceIndex = SourceIndex;
+				break;
+			}
+		}
+
+		for (int32 SourceIndex = 0; SourceIndex < MUTABLE_OP_MAX_SWIZZLE_CHANNELS; ++SourceIndex)
 		{
-			if (Args.sources[i])
+			if (Args.sources[SourceIndex])
 			{
-				Sources[i] = Runner->LoadImage({ Args.sources[i], Op.ExecutionIndex, Op.ExecutionOptions });
+				Sources[SourceIndex] = Runner->LoadImage({ Args.sources[SourceIndex], Op.ExecutionIndex, Op.ExecutionOptions });
 			}
 		}
 
 		// Shortcuts
-		if (!Sources[0])
+		if (FirstValidSourceIndex < 0 || !Sources[FirstValidSourceIndex])
 		{
-			for (int i = 0; i < MUTABLE_OP_MAX_SWIZZLE_CHANNELS; ++i)
+			for (int32 SourceIndex = 0; SourceIndex < MUTABLE_OP_MAX_SWIZZLE_CHANNELS; ++SourceIndex)
 			{
-				Runner->Release(Sources[i]);
+				Runner->Release(Sources[SourceIndex]);
 			}
 			Runner->StoreImage(Op, nullptr);
+
 			return false;
 		}
 
@@ -1456,38 +1467,39 @@ namespace mu
 
 		FImageOperator ImOp = MakeImageOperator(Runner);
 
-		int32 ResultLODs = Sources[0]->GetLODCount();
+		int32 ResultLODs = Sources[FirstValidSourceIndex]->GetLODCount();
 
 		// Be defensive: ensure formats are uncompressed
-		for (int i = 0; i < MUTABLE_OP_MAX_SWIZZLE_CHANNELS; ++i)
+		for (int32 SourceIndex = 0; SourceIndex < MUTABLE_OP_MAX_SWIZZLE_CHANNELS; ++SourceIndex)
 		{
-			if (Sources[i] && Sources[i]->GetFormat()!=GetUncompressedFormat(Sources[i]->GetFormat()))
+			if (Sources[SourceIndex] && Sources[SourceIndex]->GetFormat() != GetUncompressedFormat(Sources[SourceIndex]->GetFormat()))
 			{
 				MUTABLE_CPUPROFILER_SCOPE(ImageFormat_ForSwizzle);
 
-				EImageFormat UncompressedFormat = GetUncompressedFormat(Sources[i]->GetFormat());
-				Ptr<Image> Formatted = Runner->CreateImage(Sources[i]->GetSizeX(), Sources[i]->GetSizeY(), 1, UncompressedFormat, EInitializationType::NotInitialized);
+				EImageFormat UncompressedFormat = GetUncompressedFormat(Sources[SourceIndex]->GetFormat());
+				Ptr<Image> Formatted = Runner->CreateImage(Sources[SourceIndex]->GetSizeX(), Sources[SourceIndex]->GetSizeY(), 1, UncompressedFormat, EInitializationType::NotInitialized);
 				bool bSuccess = false;
 				int32 ImageCompressionQuality = 4; // TODO
-				ImOp.ImagePixelFormat(bSuccess, ImageCompressionQuality, Formatted.get(), Sources[i].get());
+				ImOp.ImagePixelFormat(bSuccess, ImageCompressionQuality, Formatted.get(), Sources[SourceIndex].get());
 				check(bSuccess); // Decompression cannot fail
-				Runner->Release(Sources[i]);
-				Sources[i] = Formatted;
+				Runner->Release(Sources[SourceIndex]);
+				Sources[SourceIndex] = Formatted;
 				ResultLODs = 1;
 			}
 		}
 
+		const FImageSize ResultSize = Sources[FirstValidSourceIndex]->GetSize();
 
 		// Be defensive: ensure image sizes match.
-		for (int i = 1; i < MUTABLE_OP_MAX_SWIZZLE_CHANNELS; ++i)
+		for (int32 SourceIndex = FirstValidSourceIndex + 1; SourceIndex < MUTABLE_OP_MAX_SWIZZLE_CHANNELS; ++SourceIndex)
 		{
-			if (Sources[i] && Sources[i]->GetSize() != Sources[0]->GetSize())
+			if (Sources[SourceIndex] && ResultSize != Sources[SourceIndex]->GetSize())
 			{
 				MUTABLE_CPUPROFILER_SCOPE(ImageResize_ForSwizzle);
-				Ptr<Image> Resized = Runner->CreateImage(Sources[0]->GetSizeX(), Sources[0]->GetSizeY(), 1, Sources[i]->GetFormat(), EInitializationType::NotInitialized);
-				ImOp.ImageResizeLinear(Resized.get(), 0, Sources[i].get());
-				Runner->Release(Sources[i]);
-				Sources[i] = Resized;
+				Ptr<Image> Resized = Runner->CreateImage(ResultSize.X, ResultSize.Y, 1, Sources[SourceIndex]->GetFormat(), EInitializationType::NotInitialized);
+				ImOp.ImageResizeLinear(Resized.get(), 0, Sources[SourceIndex].get());
+				Runner->Release(Sources[SourceIndex]);
+				Sources[SourceIndex] = Resized;
 				ResultLODs = 1;
 			}
 		}
@@ -1501,7 +1513,7 @@ namespace mu
 			}
 		}
 
-		Result = Runner->CreateImage(Sources[0]->GetSizeX(), Sources[0]->GetSizeY(), ResultLODs, Args.format, EInitializationType::Black);
+		Result = Runner->CreateImage(ResultSize.X, ResultSize.Y, ResultLODs, Args.format, EInitializationType::Black);
 		return true;
 	}
 
@@ -1518,9 +1530,9 @@ namespace mu
 	bool FImageSwizzleTask::Complete(CodeRunner* Runner)
 	{
 		// This runs in the Runner thread
-		for (int i = 0; i < MUTABLE_OP_MAX_SWIZZLE_CHANNELS; ++i)
+		for (int32 SourceIndex = 0; SourceIndex < MUTABLE_OP_MAX_SWIZZLE_CHANNELS; ++SourceIndex)
 		{
-			Runner->Release(Sources[i]);
+			Runner->Release(Sources[SourceIndex]);
 		}
 
 		// \TODO: If Result LODs differ from Source[0]'s, rebuild mips?
diff --git a/Engine/Plugins/Experimental/Mutable/Source/MutableRuntime/Private/MuR/OpImageSwizzle.cpp b/Engine/Plugins/Experimental/Mutable/Source/MutableRuntime/Private/MuR/OpImageSwizzle.cpp
--- a/Engine/Plugins/Experimental/Mutable/Source/MutableRuntime/Private/MuR/OpImageSwizzle.cpp
+++ b/Engine/Plugins/Experimental/Mutable/Source/MutableRuntime/Private/MuR/OpImageSwizzle.cpp
@@ -1,6 +1,7 @@
 // Copyright Epic Games, Inc. All Rights Reserved.
 
 #include "MuR/ImagePrivate.h"
+#include "MuR/Operations.h"
 #include "Async/ParallelFor.h"
 
 namespace mu
@@ -9,11 +10,6 @@ namespace mu
 	{
 		MUTABLE_CPUPROFILER_SCOPE(ImageSwizzle);
 
-		if (!Sources[0])
-		{
-			return;
-		}
-
 		EImageFormat Format = Result->GetFormat();
 
 		// LODs may not match due to bugs, only precess the common avaliable lODs.
@@ -222,16 +218,29 @@ namespace mu
 		}
 	}
 
-	Ptr<Image> FImageOperator::ImageSwizzle( EImageFormat Format, const Ptr<const Image> Sources[], const uint8 Channels[] )
+	Ptr<Image> FImageOperator::ImageSwizzle(EImageFormat Format, const Ptr<const Image> Sources[], const uint8 Channels[])
 	{
 		MUTABLE_CPUPROFILER_SCOPE(ImageSwizzle);
 
-		if (!Sources[0])
+		int32 FirstValidSourceIndex = -1;
+		for (int32 SourceIndex = 0; SourceIndex < MUTABLE_OP_MAX_SWIZZLE_CHANNELS; ++SourceIndex)
+		{
+			if (Sources[SourceIndex])
+			{
+				FirstValidSourceIndex = SourceIndex;
+				break;
+			}
+		}
+
+		if (FirstValidSourceIndex < 0)
 		{
 			return nullptr;
 		}
 
-		Ptr<Image> Dest = CreateImage(Sources[0]->GetSizeX(), Sources[0]->GetSizeY(), Sources[0]->GetLODCount(), Format, EInitializationType::Black);
+		const FImageSize ResultSize = Sources[FirstValidSourceIndex]->GetSize();
+		const int32 ResultNumLODs = Sources[FirstValidSourceIndex]->GetLODCount();
+
+		Ptr<Image> Dest = CreateImage(ResultSize.X, ResultSize.Y, ResultNumLODs, Format, EInitializationType::Black);
 
 		mu::ImageSwizzle(Dest.get(), Sources, Channels);
 
diff --git a/Engine/Plugins/Experimental/Mutable/Source/MutableTools/Private/MuT/ASTOpImageSwizzle.cpp b/Engine/Plugins/Experimental/Mutable/Source/MutableTools/Private/MuT/ASTOpImageSwizzle.cpp
--- a/Engine/Plugins/Experimental/Mutable/Source/MutableTools/Private/MuT/ASTOpImageSwizzle.cpp
+++ b/Engine/Plugins/Experimental/Mutable/Source/MutableTools/Private/MuT/ASTOpImageSwizzle.cpp
@@ -825,11 +825,11 @@ namespace mu
 		// Swizzle of RGB from a saturate + A from a different source
 		// This can be optimized to apply the saturate after the swizzle, since it doesn't touch A
 		if (!at
-			&&
-			Sources[0]->GetOpType() == OP_TYPE::IM_SATURATE 
 			&&
 			Sources[0]
 			&&
+			Sources[0]->GetOpType() == OP_TYPE::IM_SATURATE 
+			&&
 			(Sources[0] == Sources[1]) && (Sources[0] == Sources[2])
 			&&
 			// Actually it would be enough with all the RGB channels to be present in any order
@@ -1028,9 +1028,19 @@ namespace mu
 			}
 		}
 
-		if (Sources[0].child())
+		int32 FirstValidSourceIndex = -1;
+		for (int32 SourceIndex = 0; SourceIndex < MUTABLE_OP_MAX_SWIZZLE_CHANNELS; ++SourceIndex)
+		{
+			if (Sources[SourceIndex].child())
+			{
+				FirstValidSourceIndex = SourceIndex;
+				break;	
+			}
+		}
+
+		if (FirstValidSourceIndex >= 0)
 		{
-			res = Sources[0]->GetImageDesc(returnBestOption, context);
+			res = Sources[FirstValidSourceIndex]->GetImageDesc(returnBestOption, context);
 			res.m_format = Format;
 			check(res.m_format != EImageFormat::IF_NONE);
 		}
diff --git a/Engine/Plugins/Experimental/Mutable/Source/MutableTools/Private/MuT/CodeGenerator_Image.cpp b/Engine/Plugins/Experimental/Mutable/Source/MutableTools/Private/MuT/CodeGenerator_Image.cpp
--- a/Engine/Plugins/Experimental/Mutable/Source/MutableTools/Private/MuT/CodeGenerator_Image.cpp
+++ b/Engine/Plugins/Experimental/Mutable/Source/MutableTools/Private/MuT/CodeGenerator_Image.cpp
@@ -766,17 +766,17 @@ namespace mu
 
         // This node always produces a swizzle operation and sometimes it may produce a pixelformat
 		// operation to compress the Result
-        Ptr<ASTOpImageSwizzle> op = new ASTOpImageSwizzle();
+        Ptr<ASTOpImageSwizzle> SwizzleOp = new ASTOpImageSwizzle();
 
 		// Format
-		EImageFormat compressedFormat = EImageFormat::IF_NONE;
+		EImageFormat CompressedFormat = EImageFormat::IF_NONE;
 
 		switch (node.m_format)
 		{
         case EImageFormat::IF_BC1:
         case EImageFormat::IF_ASTC_4x4_RGB_LDR:
-            compressedFormat = node.m_format;
-            op->Format = node.m_sources[3] ? EImageFormat::IF_RGBA_UBYTE : EImageFormat::IF_RGB_UBYTE;
+            CompressedFormat = node.m_format;
+            SwizzleOp->Format = node.m_sources[3] ? EImageFormat::IF_RGBA_UBYTE : EImageFormat::IF_RGB_UBYTE;
 			break;
 
 		case EImageFormat::IF_BC2:
@@ -784,24 +784,24 @@ namespace mu
 		case EImageFormat::IF_BC6:
         case EImageFormat::IF_BC7:
         case EImageFormat::IF_ASTC_4x4_RGBA_LDR:
-            compressedFormat = node.m_format;
-             op->Format = EImageFormat::IF_RGBA_UBYTE;
+            CompressedFormat = node.m_format;
+            SwizzleOp->Format = EImageFormat::IF_RGBA_UBYTE;
 			break;
 
 		case EImageFormat::IF_BC4:
-			compressedFormat = node.m_format;
-            op->Format = EImageFormat::IF_L_UBYTE;
+			CompressedFormat = node.m_format;
+            SwizzleOp->Format = EImageFormat::IF_L_UBYTE;
 			break;
 
 		case EImageFormat::IF_BC5:
         case EImageFormat::IF_ASTC_4x4_RG_LDR:
-            compressedFormat = node.m_format;
+            CompressedFormat = node.m_format;
 			// TODO: Should be RG
-            op->Format = EImageFormat::IF_RGB_UBYTE;
+            SwizzleOp->Format = EImageFormat::IF_RGB_UBYTE;
 			break;
 
 		default:
-            op->Format = node.m_format;
+            SwizzleOp->Format = node.m_format;
 			break;
 
 		}
@@ -812,59 +812,63 @@ namespace mu
 		check(node.m_sources.Num() == node.m_sourceChannels.Num());
 
 		// First source, for reference in the size
-        Ptr<ASTOp> first;
-		FImageDesc FirstDesc;
-		for (int32 t = 0; t<node.m_sources.Num(); ++t)
+        Ptr<ASTOp> FirstValid;
+		FImageDesc FirstValidDesc;
+		int32 FirstValidSourceIndex = -1;
+
+		check(MUTABLE_OP_MAX_SWIZZLE_CHANNELS >= node.m_sources.Num());
+		for (int32 SourceIndex = 0; SourceIndex < node.m_sources.Num(); ++SourceIndex)
 		{
-			if (node.m_sources[t])
+			if (node.m_sources[SourceIndex])
 			{
 				FImageGenerationResult BaseResult;
-				GenerateImage(Options, BaseResult, node.m_sources[t]);
-                Ptr<ASTOp> source = BaseResult.op;
+				GenerateImage(Options, BaseResult, node.m_sources[SourceIndex]);
+                Ptr<ASTOp> Source = BaseResult.op;
 
-				source = GenerateImageUncompressed(source);
+				Source = GenerateImageUncompressed(Source);
 
-				if (!source)
+				if (!Source)
 				{
 					// TODO: Warn?
-					source = GenerateMissingImageCode(TEXT("Swizzle channel"), EImageFormat::IF_L_UBYTE, InNode->GetMessageContext(), Options);
+					Source = GenerateMissingImageCode(TEXT("Swizzle channel"), EImageFormat::IF_L_UBYTE, InNode->GetMessageContext(), Options);
 				}
 
-                Ptr<ASTOp> sizedSource;
-				if (first && FirstDesc.m_size[0])
+                Ptr<ASTOp> SizedSource;
+				if (FirstValid && FirstValidDesc.m_size[0])
 				{
-					sizedSource = GenerateImageSize(source, FIntVector2(FirstDesc.m_size));
+					SizedSource = GenerateImageSize(Source, FIntVector2(FirstValidDesc.m_size));
 				}
 				else
 				{
-					first = source;
-					sizedSource = source;
-					FirstDesc = first->GetImageDesc();
+					FirstValid = Source;
+					SizedSource = Source;
+					FirstValidDesc = FirstValid->GetImageDesc();
+					FirstValidSourceIndex = SourceIndex;
 				}
 
-                op->Sources[t] = sizedSource;
-                op->SourceChannels[t] = (uint8)node.m_sourceChannels[t];
+                SwizzleOp->Sources[SourceIndex] = SizedSource;
+                SwizzleOp->SourceChannels[SourceIndex] = (uint8)node.m_sourceChannels[SourceIndex];
 			}
 		}
 
-		// At least one source is required
-        if (!op->Sources[0])
+        if (FirstValidSourceIndex < 0)
 		{
-            Ptr<ASTOp> source = GenerateMissingImageCode(TEXT("First swizzle image"), EImageFormat::IF_RGBA_UBYTE, InNode->GetMessageContext(), Options);
-            op->Sources[0] = source;
+            Ptr<ASTOp> Source = GenerateMissingImageCode(TEXT("First swizzle image"), EImageFormat::IF_RGBA_UBYTE, InNode->GetMessageContext(), Options);
+            SwizzleOp->Sources[0] = Source;
+
 		}
 
-        Ptr<ASTOp> resultOp = op;
+        Ptr<ASTOp> ResultOp = SwizzleOp;
 
-		if (compressedFormat != EImageFormat::IF_NONE)
+		if (CompressedFormat != EImageFormat::IF_NONE)
 		{
-            Ptr<ASTOpImagePixelFormat> fop = new ASTOpImagePixelFormat();
-            fop->Source = resultOp;
-            fop->Format = compressedFormat;
-			resultOp = fop;
+            Ptr<ASTOpImagePixelFormat> FormatOp = new ASTOpImagePixelFormat();
+            FormatOp->Source = ResultOp;
+            FormatOp->Format = CompressedFormat;
+			ResultOp = FormatOp;
 		}
 
-        Result.op = resultOp;
+        Result.op = ResultOp;
 	}
 
 
