[mutable] Added protection for dereferencing of null pointer during instance compilation.

#rnx

diff --git a/Engine/Plugins/Experimental/Mutable/Source/CustomizableObjectEditor/Private/MuCOE/CustomizableObjectEditorViewportClient.cpp b/Engine/Plugins/Experimental/Mutable/Source/CustomizableObjectEditor/Private/MuCOE/CustomizableObjectEditorViewportClient.cpp
--- a/Engine/Plugins/Experimental/Mutable/Source/CustomizableObjectEditor/Private/MuCOE/CustomizableObjectEditorViewportClient.cpp
+++ b/Engine/Plugins/Experimental/Mutable/Source/CustomizableObjectEditor/Private/MuCOE/CustomizableObjectEditorViewportClient.cpp
@@ -771,6 +771,11 @@ void FCustomizableObjectEditorViewportClient::OnInstanceUpdate(UCustomizableObje
 	// Configure the initial orbital position of the camera
 	if (!bIsCameraSetup)
 	{
+		if (!Actor.IsValid())
+		{
+			return;
+		}
+		
 		bIsCameraSetup = true;
 		
 		FVector Center;
@@ -909,13 +914,16 @@ bool FCustomizableObjectEditorViewportClient::InputKey(const FInputKeyEventArgs&
 	{
 		if (EventArgs.Key == EKeys::F)
 		{
-			FVector Center;
-			FVector Extents;
-			Actor.Get()->GetActorBounds(false, Center, Extents, true);
-	
-			FocusViewportOnBox(FBox(Center - Extents, Center + Extents), true);
-			
-			return true;
+			if (Actor.IsValid())
+			{
+				FVector Center;
+				FVector Extents;
+				Actor.Get()->GetActorBounds(false, Center, Extents, true);
+
+				FocusViewportOnBox(FBox(Center - Extents, Center + Extents), true);
+                return true;
+			}
+		
 		}
 		else if (WidgetType != EWidgetType::Hidden) // Do not change the type when hidden.
 		{
