[mutable] Rename Material nodes and pins to Mesh Section when appropiate.
#rnx

diff --git a/Engine/Plugins/Experimental/Mutable/Source/CustomizableObject/Internal/MuCO/CustomizableObjectCustomVersion.h b/Engine/Plugins/Experimental/Mutable/Source/CustomizableObject/Internal/MuCO/CustomizableObjectCustomVersion.h
--- a/Engine/Plugins/Experimental/Mutable/Source/CustomizableObject/Internal/MuCO/CustomizableObjectCustomVersion.h
+++ b/Engine/Plugins/Experimental/Mutable/Source/CustomizableObject/Internal/MuCO/CustomizableObjectCustomVersion.h
@@ -146,6 +146,8 @@ struct CUSTOMIZABLEOBJECT_API FCustomizableObjectCustomVersion
 		
 		UnifyRequiredTags,
 
+		MaterialPinsRename, // Moved from "Material" to "Mesh Section"
+
 		// -----<new versions can be added above this line>-------------------------------------------------
 		VersionPlusOne,
 		LatestVersion = VersionPlusOne - 1
diff --git a/Engine/Plugins/Experimental/Mutable/Source/CustomizableObjectEditor/Private/MuCOE/EdGraphSchema_CustomizableObject.cpp b/Engine/Plugins/Experimental/Mutable/Source/CustomizableObjectEditor/Private/MuCOE/EdGraphSchema_CustomizableObject.cpp
--- a/Engine/Plugins/Experimental/Mutable/Source/CustomizableObjectEditor/Private/MuCOE/EdGraphSchema_CustomizableObject.cpp
+++ b/Engine/Plugins/Experimental/Mutable/Source/CustomizableObjectEditor/Private/MuCOE/EdGraphSchema_CustomizableObject.cpp
@@ -1349,7 +1349,7 @@ FText UEdGraphSchema_CustomizableObject::GetPinCategoryName(const FName& PinCate
 	}
 	else if (PinCategory == UEdGraphSchema_CustomizableObject::PC_Material)
 	{
-		return LOCTEXT("Material_Pin_Category", "Material");
+		return LOCTEXT("MeshSection_Pin_Category", "Mesh Section");
 	}
 	else if (PinCategory == UEdGraphSchema_CustomizableObject::PC_Modifier)
 	{
diff --git a/Engine/Plugins/Experimental/Mutable/Source/CustomizableObjectEditor/Private/MuCOE/Nodes/CustomizableObjectNodeCopyMaterial.cpp b/Engine/Plugins/Experimental/Mutable/Source/CustomizableObjectEditor/Private/MuCOE/Nodes/CustomizableObjectNodeCopyMaterial.cpp
--- a/Engine/Plugins/Experimental/Mutable/Source/CustomizableObjectEditor/Private/MuCOE/Nodes/CustomizableObjectNodeCopyMaterial.cpp
+++ b/Engine/Plugins/Experimental/Mutable/Source/CustomizableObjectEditor/Private/MuCOE/Nodes/CustomizableObjectNodeCopyMaterial.cpp
@@ -9,6 +9,7 @@
 #include "MuCOE/Nodes/CustomizableObjectNodeObject.h"
 #include "MuCOE/Nodes/CustomizableObjectNodeSkeletalMesh.h"
 #include "MuCOE/Nodes/CustomizableObjectNodeMaterial.h"
+#include "MuCO/CustomizableObjectCustomVersion.h"
 
 #define LOCTEXT_NAMESPACE "CustomizableObjectEditor"
 
@@ -17,15 +18,15 @@
 const static FString MeshPinName = TEXT("Mesh_Input_Pin");
 
 /** Material input pin key */
-const static FString MaterialPinName = TEXT("Material_Input_Pin");
+const static FString MeshSectionPinName = TEXT("MeshSection_Input_Pin");
 
 /** Material output pin key*/
-const static FString OutputPinName = TEXT("Material_Output_Pin");
+const static FString OutputPinName = TEXT("MeshSection_Output_Pin");
 
 
 FText UCustomizableObjectNodeCopyMaterial::GetNodeTitle(ENodeTitleType::Type TitleType) const
 {
-	return LOCTEXT("Copy_Material", "Copy Material");
+	return LOCTEXT("Copy_MeshSection", "Copy Mesh Section");
 }
 
 
@@ -38,13 +39,38 @@ void UCustomizableObjectNodeCopyMaterial::AllocateDefaultPins(UCustomizableObjec
 	Pin->PinFriendlyName = LOCTEXT("MeshPin", "Mesh");
 	Pin->bDefaultValueIsIgnored = true;
 
-	Pin = CustomCreatePin(EGPD_Input, Schema->PC_Material, FName(MaterialPinName));
-	Pin->PinFriendlyName = LOCTEXT("BaseMaterialPin", "Base Material");;
+	Pin = CustomCreatePin(EGPD_Input, Schema->PC_Material, FName(MeshSectionPinName));
+	Pin->PinFriendlyName = LOCTEXT("BaseMeshSectionPin", "Base Mesh Section");;
 	Pin->bDefaultValueIsIgnored = true;
 
 	// Output pins
 	Pin = CustomCreatePin(EGPD_Output, Schema->PC_Material, FName(OutputPinName));
-	Pin->PinFriendlyName = LOCTEXT("MaterialPin", "Material");
+	Pin->PinFriendlyName = LOCTEXT("MeshSectionPin", "Mesh Section");
+}
+
+
+void UCustomizableObjectNodeCopyMaterial::BackwardsCompatibleFixup(int32 CustomizableObjectCustomVersion)
+{
+	Super::BackwardsCompatibleFixup(CustomizableObjectCustomVersion);
+
+	if (CustomizableObjectCustomVersion == FCustomizableObjectCustomVersion::MaterialPinsRename)
+	{
+		UEdGraphPin* BaseMaterialPin = FindPin(TEXT("Material_Input_Pin"), EEdGraphPinDirection::EGPD_Input);
+		if (BaseMaterialPin)
+		{
+			const FString PinFriendlyName = TEXT("Base Mesh Section");
+			BaseMaterialPin->PinName = FName(MeshSectionPinName);
+			BaseMaterialPin->PinFriendlyName = FText::FromString(PinFriendlyName);
+		}
+
+		UEdGraphPin* MaterialPin = FindPin(TEXT("Material_Output_Pin"), EEdGraphPinDirection::EGPD_Output);
+		if (MaterialPin)
+		{
+			const FString PinFriendlyName = TEXT("Mesh Section");
+			MaterialPin->PinName = FName(OutputPinName);
+			MaterialPin->PinFriendlyName = FText::FromString(PinFriendlyName);
+		}
+	}
 }
 
 
@@ -82,7 +108,15 @@ TArray<UCustomizableObjectLayout*> UCustomizableObjectNodeCopyMaterial::GetLayou
 
 UEdGraphPin* UCustomizableObjectNodeCopyMaterial::OutputPin() const
 {
-	return FindPin(OutputPinName);
+	UEdGraphPin* Pin = FindPin(OutputPinName, EEdGraphPinDirection::EGPD_Output);
+
+	// Legacy name
+	if (!Pin)
+	{
+		Pin = FindPin(TEXT("Material_Output_Pin"), EEdGraphPinDirection::EGPD_Output);
+	}
+
+	return Pin;
 }
 
 
@@ -99,9 +133,17 @@ bool UCustomizableObjectNodeCopyMaterial::RealMaterialDataHasChanged() const
 }
 
 
-UEdGraphPin* UCustomizableObjectNodeCopyMaterial::GetMaterialPin() const
+UEdGraphPin* UCustomizableObjectNodeCopyMaterial::GetMeshSectionPin() const
 {
-	return FindPin(MaterialPinName);
+	UEdGraphPin* Pin = FindPin(MeshSectionPinName, EEdGraphPinDirection::EGPD_Input);
+
+	// Legacy name
+	if (!Pin)
+	{
+		Pin = FindPin(TEXT("Material_Input_Pin"), EEdGraphPinDirection::EGPD_Input);
+	}
+
+	return Pin;
 }
 
 
@@ -127,7 +169,7 @@ UCustomizableObjectNodeMaterial* UCustomizableObjectNodeCopyMaterial::GetMateria
 {
 	UCustomizableObjectNodeMaterial* Result = nullptr;
 
-	UEdGraphPin* MaterialPin = GetMaterialPin();
+	UEdGraphPin* MaterialPin = GetMeshSectionPin();
 	if (const UEdGraphPin* ConnectedPin = FollowInputPin(*MaterialPin))
 	{
 		return Cast<UCustomizableObjectNodeMaterial>(ConnectedPin->GetOwningNode());
@@ -144,7 +186,7 @@ bool UCustomizableObjectNodeCopyMaterial::CanConnect(const UEdGraphPin* InOwnedI
 		return false;
 	}
 
-	if (InOwnedInputPin == GetMaterialPin())
+	if (InOwnedInputPin == GetMeshSectionPin())
 	{
 		const UEdGraphNode* OuputPinOwningNode = InOutputPin->GetOwningNode();
 		return (OuputPinOwningNode->IsA(UCustomizableObjectNodeMaterial::StaticClass()) && !OuputPinOwningNode->IsA(UCustomizableObjectNodeCopyMaterial::StaticClass()))
@@ -395,10 +437,4 @@ FText UCustomizableObjectNodeCopyMaterial::GetTooltipText() const
 }
 
 
-FLinearColor UCustomizableObjectNodeCopyMaterial::GetNodeTitleColor() const
-{
-	const UEdGraphSchema_CustomizableObject* Schema = GetDefault<UEdGraphSchema_CustomizableObject>();
-	return Schema->GetPinTypeColor(Schema->PC_Material);
-}
-
 #undef LOCTEXT_NAMESPACE
diff --git a/Engine/Plugins/Experimental/Mutable/Source/CustomizableObjectEditor/Private/MuCOE/Nodes/CustomizableObjectNodeMaterial.cpp b/Engine/Plugins/Experimental/Mutable/Source/CustomizableObjectEditor/Private/MuCOE/Nodes/CustomizableObjectNodeMaterial.cpp
--- a/Engine/Plugins/Experimental/Mutable/Source/CustomizableObjectEditor/Private/MuCOE/Nodes/CustomizableObjectNodeMaterial.cpp
+++ b/Engine/Plugins/Experimental/Mutable/Source/CustomizableObjectEditor/Private/MuCOE/Nodes/CustomizableObjectNodeMaterial.cpp
@@ -341,7 +341,7 @@ void UCustomizableObjectNodeMaterial::AllocateDefaultPins(UCustomizableObjectNod
 	}
 
 	{
-		const FString PinFriendlyName = TEXT("Material");
+		const FString PinFriendlyName = TEXT("Mesh Section");
 		const FString PinName = PinFriendlyName + FString(TEXT("_Output_Pin"));
 		UEdGraphPin* OutputPin = CustomCreatePin(EGPD_Output, Schema->PC_Material, FName(*PinName));
 		OutputPin->PinFriendlyName = FText::FromString(PinFriendlyName);
@@ -708,6 +708,19 @@ void UCustomizableObjectNodeMaterial::BackwardsCompatibleFixup(int32 Customizabl
 			}
 		}
 	}
+
+	if (CustomizableObjectCustomVersion == FCustomizableObjectCustomVersion::MaterialPinsRename)
+	{
+		UEdGraphPin* MaterialPin = FindPin(TEXT("Material_Output_Pin"), EEdGraphPinDirection::EGPD_Output);
+		if (MaterialPin)
+		{
+			const FString PinFriendlyName = TEXT("Mesh Section");
+			const FString PinName = PinFriendlyName + FString(TEXT("_Output_Pin"));
+			MaterialPin->PinName = FName(PinName);
+			MaterialPin->PinFriendlyName = FText::FromString(PinFriendlyName);
+		}
+
+	}
 }
 
 
@@ -753,14 +766,14 @@ FText UCustomizableObjectNodeMaterial::GetNodeTitle(ENodeTitleType::Type TitleTy
 {
 	if (TitleType == ENodeTitleType::ListView || !Material)
 	{
-		return LOCTEXT("Material", "Material");
+		return LOCTEXT("Mesh Section", "Mesh Section");
 	}
 	else
 	{
 		FFormatNamedArguments Args;
-		Args.Add(TEXT("MaterialName"), FText::FromString(Material->GetName()));
+		Args.Add(TEXT("MeshSectionName"), FText::FromString(Material->GetName()));
 
-		return FText::Format(LOCTEXT("Material_Title", "{MaterialName}\nMaterial"), Args);
+		return FText::Format(LOCTEXT("MeshSection_Title", "{MeshSectionName}\nMesh Section"), Args);
 	}
 }
 
@@ -774,19 +787,29 @@ FLinearColor UCustomizableObjectNodeMaterial::GetNodeTitleColor() const
 
 UEdGraphPin* UCustomizableObjectNodeMaterial::OutputPin() const
 {
-	FString PinFriendlyName = TEXT("Material");
+	FString PinFriendlyName = TEXT("Mesh Section");
 	FString PinName = PinFriendlyName + FString(TEXT("_Output_Pin"));
 
-	UEdGraphPin* Pin = FindPin(PinName);
-
-	if (Pin)
+	UEdGraphPin* Pin = FindPin(PinName, EEdGraphPinDirection::EGPD_Output);
+	if (!Pin)
 	{
-		return Pin;
+		Pin = FindPin(TEXT("Mesh Section"), EEdGraphPinDirection::EGPD_Output);
 	}
-	else
+
+	// Legacy name
+	if (!Pin)
 	{
-		return FindPin(TEXT("Material"));
+		PinFriendlyName = TEXT("Material");
+		PinName = PinFriendlyName + FString(TEXT("_Output_Pin"));
+
+		Pin = FindPin(PinName, EEdGraphPinDirection::EGPD_Output);
+		if (!Pin)
+		{
+			Pin = FindPin(TEXT("Material"), EEdGraphPinDirection::EGPD_Output);
+		}		
 	}
+
+	return Pin;
 }
 
 
@@ -1232,7 +1255,7 @@ bool UCustomizableObjectNodeMaterial::IsPinRelevant(const UEdGraphPin* Pin) cons
 
 FText UCustomizableObjectNodeMaterial::GetTooltipText() const
 {
-	return LOCTEXT("Material_Tooltip", "Defines a Customizable Object material.\nConsists of a mesh section, a material asset to paint it, and the runtime modifiable inputs to the material asset parameters.");
+	return LOCTEXT("MeshSection_Tooltip", "Defines a Customizable Object mesh section.\nIt has a mesh section, a material assigned to it and the runtime modifiable inputs to the material asset parameters.");
 }
 
 
diff --git a/Engine/Plugins/Experimental/Mutable/Source/CustomizableObjectEditor/Private/MuCOE/Nodes/CustomizableObjectNodeModifierClipWithMeshDetails.cpp b/Engine/Plugins/Experimental/Mutable/Source/CustomizableObjectEditor/Private/MuCOE/Nodes/CustomizableObjectNodeModifierClipWithMeshDetails.cpp
--- a/Engine/Plugins/Experimental/Mutable/Source/CustomizableObjectEditor/Private/MuCOE/Nodes/CustomizableObjectNodeModifierClipWithMeshDetails.cpp
+++ b/Engine/Plugins/Experimental/Mutable/Source/CustomizableObjectEditor/Private/MuCOE/Nodes/CustomizableObjectNodeModifierClipWithMeshDetails.cpp
@@ -91,9 +91,9 @@ void FCustomizableObjectNodeModifierClipWithMeshDetails::CustomizeDetails(IDetai
 
 		// Filling ComboBox options
 		ClippingMethods.Empty();
-		ClippingMethods.Add(MakeShareable(new FString("Material")));
+		ClippingMethods.Add(MakeShareable(new FString("Mesh Section")));
 		ClippingMethods.Add(MakeShareable(new FString("Tags")));
-		ClippingMethods.Add(MakeShareable(new FString("Material & Tags")));
+		ClippingMethods.Add(MakeShareable(new FString("Mesh Section & Tags")));
 
 		// Initializing the inital value of the ComboBox
 		SelectedMethod = SetInitialClippingMethod();
diff --git a/Engine/Plugins/Experimental/Mutable/Source/CustomizableObjectEditor/Public/MuCOE/Nodes/CustomizableObjectNodeCopyMaterial.h b/Engine/Plugins/Experimental/Mutable/Source/CustomizableObjectEditor/Public/MuCOE/Nodes/CustomizableObjectNodeCopyMaterial.h
--- a/Engine/Plugins/Experimental/Mutable/Source/CustomizableObjectEditor/Public/MuCOE/Nodes/CustomizableObjectNodeCopyMaterial.h
+++ b/Engine/Plugins/Experimental/Mutable/Source/CustomizableObjectEditor/Public/MuCOE/Nodes/CustomizableObjectNodeCopyMaterial.h
@@ -22,7 +22,6 @@ public:
 	// UEdGraphNode interface
 	virtual FText GetNodeTitle(ENodeTitleType::Type TitleType) const override;
 	virtual FText GetTooltipText() const override;
-	virtual FLinearColor GetNodeTitleColor() const override;
 
 	// UCustomizableObjectNode interface
 	virtual void AllocateDefaultPins(UCustomizableObjectNodeRemapPins* RemapPins) override;
@@ -31,6 +30,7 @@ public:
 	virtual bool IsNodeOutDatedAndNeedsRefresh() override;
 	virtual bool ProvidesCustomPinRelevancyTest() const override;
 	virtual bool IsPinRelevant(const UEdGraphPin* Pin) const override;
+	virtual void BackwardsCompatibleFixup(int32 CustomizableObjectCustomVersion) override;
 
 	// UCustomizableObjectNodeMaterialBase interface
 	virtual UMaterialInterface* GetMaterial() const override;
@@ -60,5 +60,8 @@ public:
 	// Own interface
 	UCustomizableObjectNodeSkeletalMesh* GetMeshNode() const;
 	
-	UEdGraphPin* GetMaterialPin() const;
+protected:
+
+	UEdGraphPin* GetMeshSectionPin() const;
+
 };
