[Mutable] Fix component switch not working properly.

#rnx

diff --git a/Engine/Plugins/Experimental/Mutable/Source/CustomizableObject/Private/MuCO/CustomizableObjectSystem.cpp b/Engine/Plugins/Experimental/Mutable/Source/CustomizableObject/Private/MuCO/CustomizableObjectSystem.cpp
--- a/Engine/Plugins/Experimental/Mutable/Source/CustomizableObject/Private/MuCO/CustomizableObjectSystem.cpp
+++ b/Engine/Plugins/Experimental/Mutable/Source/CustomizableObject/Private/MuCO/CustomizableObjectSystem.cpp
@@ -1699,7 +1699,7 @@ namespace impl
 
 		for (int32 ComponentIndex = 0; ComponentIndex < Operation->NumInstanceComponents; ++ComponentIndex)
 		{
-			Operation->NumLODsAvailablePerComponent[ComponentIndex] = Operation->MutableInstance->GetLODCount(ComponentIndex);
+			Operation->NumLODsAvailablePerComponent[Operation->MutableInstance->GetComponentId(ComponentIndex)] = Operation->MutableInstance->GetLODCount(ComponentIndex);
 		}
 		
 		int32 CurrentMinLOD = Operation->bStreamMeshLODs ? 0 : Operation->GetMinLOD();
