[Mutable] Fix possible crash in CustomizableObjectSystem::BeginDestroy when getting the CustomizableObjectEditor module.

#rnx

diff --git a/Engine/Plugins/Experimental/Mutable/Source/CustomizableObject/Private/MuCO/CustomizableObjectSystem.cpp b/Engine/Plugins/Experimental/Mutable/Source/CustomizableObject/Private/MuCO/CustomizableObjectSystem.cpp
--- a/Engine/Plugins/Experimental/Mutable/Source/CustomizableObject/Private/MuCO/CustomizableObjectSystem.cpp
+++ b/Engine/Plugins/Experimental/Mutable/Source/CustomizableObject/Private/MuCO/CustomizableObjectSystem.cpp
@@ -580,7 +580,7 @@ void UCustomizableObjectSystem::BeginDestroy()
 	if (Private)
 	{
 #if WITH_EDITOR
-		if (ICustomizableObjectEditorModule* EditorModule = ICustomizableObjectEditorModule::Get())
+		if (ICustomizableObjectEditorModule* EditorModule = FModuleManager::GetModulePtr<ICustomizableObjectEditorModule>("CustomizableObjectEditor"))
 		{
 			EditorModule->CancelCompileRequests();
 		}
