[Mutable] UX - Unify layout access for SkMesh and Table nodes.
Remove UCustomizableObjectNodeLayoutBlocks

#rnx

diff --git a/Engine/Plugins/Experimental/Mutable/Source/CustomizableObject/Internal/MuCO/CustomizableObjectCustomVersion.h b/Engine/Plugins/Experimental/Mutable/Source/CustomizableObject/Internal/MuCO/CustomizableObjectCustomVersion.h
--- a/Engine/Plugins/Experimental/Mutable/Source/CustomizableObject/Internal/MuCO/CustomizableObjectCustomVersion.h
+++ b/Engine/Plugins/Experimental/Mutable/Source/CustomizableObject/Internal/MuCO/CustomizableObjectCustomVersion.h
@@ -148,6 +148,10 @@ struct CUSTOMIZABLEOBJECT_API FCustomizableObjectCustomVersion
 
 		MaterialPinsRename, // Moved from "Material" to "Mesh Section"
 
+		MoveLayoutToNodeSkeletalMesh, // Move the layouts from the NodeLayout to the NodeSkeletalMesh
+
+		RemoveNodeLayout,
+
 		// -----<new versions can be added above this line>-------------------------------------------------
 		VersionPlusOne,
 		LatestVersion = VersionPlusOne - 1
diff --git a/Engine/Plugins/Experimental/Mutable/Source/CustomizableObjectEditor/Private/MuCOE/CustomizableObjectEditorModule.cpp b/Engine/Plugins/Experimental/Mutable/Source/CustomizableObjectEditor/Private/MuCOE/CustomizableObjectEditorModule.cpp
--- a/Engine/Plugins/Experimental/Mutable/Source/CustomizableObjectEditor/Private/MuCOE/CustomizableObjectEditorModule.cpp
+++ b/Engine/Plugins/Experimental/Mutable/Source/CustomizableObjectEditor/Private/MuCOE/CustomizableObjectEditorModule.cpp
@@ -33,7 +33,6 @@
 #include "MuCOE/Nodes/CustomizableObjectNodeExternalPin.h"
 #include "MuCOE/Nodes/CustomizableObjectNodeExternalPinDetails.h"
 #include "MuCOE/Nodes/CustomizableObjectNodeLayoutBlocks.h"
-#include "MuCOE/Nodes/CustomizableObjectNodeLayoutBlocksDetails.h"
 #include "MuCOE/Nodes/CustomizableObjectNodeModifierEditMeshSection.h"
 #include "MuCOE/Nodes/CustomizableObjectNodeModifierEditMeshSectionDetails.h"
 #include "MuCOE/Nodes/CustomizableObjectNodeModifierExtendMeshSection.h"
@@ -240,7 +239,6 @@ void FCustomizableObjectEditorModule::StartupModule()
 	// Property views
 	// Nodes
 	FPropertyEditorModule& PropertyModule = FModuleManager::LoadModuleChecked<FPropertyEditorModule>("PropertyEditor");
-	RegisterCustomDetails(PropertyModule, UCustomizableObjectNodeLayoutBlocks::StaticClass(), FOnGetDetailCustomizationInstance::CreateStatic(&FCustomizableObjectNodeLayoutBlocksDetails::MakeInstance));
 	RegisterCustomDetails(PropertyModule, UCustomizableObjectNodeModifierEditMeshSection::StaticClass(), FOnGetDetailCustomizationInstance::CreateStatic(&FCustomizableObjectNodeModifierEditMeshSectionDetails::MakeInstance));
 	RegisterCustomDetails(PropertyModule, UCustomizableObjectNodeModifierExtendMeshSection::StaticClass(), FOnGetDetailCustomizationInstance::CreateStatic(&FCustomizableObjectNodeModifierExtendMeshSectionDetails::MakeInstance));
 	RegisterCustomDetails(PropertyModule, UCustomizableObjectNodeModifierRemoveMesh::StaticClass(), FOnGetDetailCustomizationInstance::CreateStatic(&FCustomizableObjectNodeModifierRemoveMeshDetails::MakeInstance));
diff --git a/Engine/Plugins/Experimental/Mutable/Source/CustomizableObjectEditor/Private/MuCOE/CustomizableObjectGraph.cpp b/Engine/Plugins/Experimental/Mutable/Source/CustomizableObjectEditor/Private/MuCOE/CustomizableObjectGraph.cpp
--- a/Engine/Plugins/Experimental/Mutable/Source/CustomizableObjectEditor/Private/MuCOE/CustomizableObjectGraph.cpp
+++ b/Engine/Plugins/Experimental/Mutable/Source/CustomizableObjectEditor/Private/MuCOE/CustomizableObjectGraph.cpp
@@ -74,8 +74,10 @@ void UCustomizableObjectGraph::PostLoad()
 		}
 	}
 
+	NodesCopy = Nodes; // Copy to be able to remove nodes inside the PostBackwardsCompatibleFixup.
+
 	// Do any additional work which require nodes to be valid (i.e., have executed BackwardsCompatibleFixup).
-	for (UEdGraphNode* Node : Nodes)
+	for (UEdGraphNode* Node : NodesCopy)
 	{
 		if (UCustomizableObjectNode* CustomizableObjectNode = Cast<UCustomizableObjectNode>(Node))
 		{
diff --git a/Engine/Plugins/Experimental/Mutable/Source/CustomizableObjectEditor/Private/MuCOE/CustomizableObjectLayout.cpp b/Engine/Plugins/Experimental/Mutable/Source/CustomizableObjectEditor/Private/MuCOE/CustomizableObjectLayout.cpp
--- a/Engine/Plugins/Experimental/Mutable/Source/CustomizableObjectEditor/Private/MuCOE/CustomizableObjectLayout.cpp
+++ b/Engine/Plugins/Experimental/Mutable/Source/CustomizableObjectEditor/Private/MuCOE/CustomizableObjectLayout.cpp
@@ -8,8 +8,9 @@
 #include "MuCOE/CustomizableObjectCompiler.h"
 #include "MuCOE/GenerateMutableSource/GenerateMutableSourceLayout.h"
 #include "MuCOE/GraphTraversal.h"
+#include "MuCOE/MutableUtils.h"
 #include "MuCOE/ICustomizableObjectEditor.h"
-#include "MuCOE/Nodes/CustomizableObjectNodeLayoutBlocks.h"
+#include "MuCOE/Nodes/CustomizableObjectNodeSkeletalMesh.h"
 #include "MuCOE/Nodes/CustomizableObjectNodeMesh.h"
 #include "MuCOE/Nodes/CustomizableObjectNodeTable.h"
 
@@ -20,17 +21,13 @@ UCustomizableObjectLayout::UCustomizableObjectLayout()
 	GridSize = FIntPoint(4, 4);
 	MaxGridSize = FIntPoint(4, 4);
 
-	FCustomizableObjectLayoutBlock Block(FIntPoint(0, 0), FIntPoint(4, 4));
-	Blocks.Add(Block);
-
 	PackingStrategy = ECustomizableObjectTextureLayoutPackingStrategy::Resizable;
 	BlockReductionMethod = ECustomizableObjectLayoutBlockReductionMethod::Halve;
 }
 
 
-void UCustomizableObjectLayout::SetLayout(UObject* InMesh, int32 LODIndex, int32 MatIndex, int32 UVIndex)
+void UCustomizableObjectLayout::SetLayout(int32 LODIndex, int32 MatIndex, int32 UVIndex)
 {
-	Mesh = InMesh;
 	LOD = LODIndex;
 	Material = MatIndex;
 	UVChannel = UVIndex;
@@ -84,6 +81,7 @@ void UCustomizableObjectLayout::SetLayoutName(FString Name)
 void UCustomizableObjectLayout::GenerateAutomaticBlocksFromUVs()
 {
 	UCustomizableObjectNode* Node = Cast<UCustomizableObjectNode>(GetOuter());
+	const UObject* Mesh = GetMesh();
 
 	if (!Node || !Mesh)
 	{
@@ -171,23 +169,17 @@ void UCustomizableObjectLayout::ConsolidateAutomaticBlocks()
 }
 
 
-void UCustomizableObjectLayout::GetUVChannel(TArray<FVector2f>& UVs, int32 UVChannelIndex) const
+void UCustomizableObjectLayout::GetUVs(TArray<FVector2f>& UVs) const
 {
-	if (UObject* Node = GetOuter())
+	if(const UObject* Mesh = GetMesh())
 	{
-		if (const UCustomizableObjectNodeLayoutBlocks* TypedNodeLayout = Cast<UCustomizableObjectNodeLayoutBlocks>(Node))
+		if (const USkeletalMesh* SkeletalMesh = Cast<const USkeletalMesh>(Mesh))
 		{
-			if (const UEdGraphPin* ConnectedPin = FollowOutputPin(*TypedNodeLayout->OutputPin()))
-			{
-				if (const UCustomizableObjectNodeMesh* MeshNode = Cast<UCustomizableObjectNodeMesh>(ConnectedPin->GetOwningNode()))
-				{
-					MeshNode->GetUVChannelForPin(ConnectedPin, UVs, UVChannelIndex);
-				}
-			}
+			UVs = GetUV(*SkeletalMesh, LOD, Material, UVChannel);
 		}
-		else if (const UCustomizableObjectNodeTable* TypedNodeTable = Cast<UCustomizableObjectNodeTable>(Node))
+		else if (const UStaticMesh* StaticMesh = Cast<const UStaticMesh>(Mesh))
 		{
-			TypedNodeTable->GetUVChannel(this, UVs);
+			UVs = GetUV(*StaticMesh, LOD, Material, UVChannel);
 		}
 	}
 }
@@ -219,4 +211,22 @@ void UCustomizableObjectLayout::SetIgnoreWarningsLOD(int32 LODValue)
 }
 
 
+UObject* UCustomizableObjectLayout::GetMesh() const
+{
+	if (UObject* Node = GetOuter())
+	{
+		if (const UCustomizableObjectNodeSkeletalMesh* TypedNodeSkeletalMesh = Cast<UCustomizableObjectNodeSkeletalMesh>(Node))
+		{
+			return TypedNodeSkeletalMesh->GetMesh();
+		}
+		else if (const UCustomizableObjectNodeTable* TypedNodeTable = Cast<UCustomizableObjectNodeTable>(Node))
+		{
+			return TypedNodeTable->GetDefaultMeshForLayout(this);
+		}
+	}
+
+	return nullptr;
+}
+
+
 #undef LOCTEXT_NAMESPACE
diff --git a/Engine/Plugins/Experimental/Mutable/Source/CustomizableObjectEditor/Private/MuCOE/EdGraphSchema_CustomizableObject.cpp b/Engine/Plugins/Experimental/Mutable/Source/CustomizableObjectEditor/Private/MuCOE/EdGraphSchema_CustomizableObject.cpp
--- a/Engine/Plugins/Experimental/Mutable/Source/CustomizableObjectEditor/Private/MuCOE/EdGraphSchema_CustomizableObject.cpp
+++ b/Engine/Plugins/Experimental/Mutable/Source/CustomizableObjectEditor/Private/MuCOE/EdGraphSchema_CustomizableObject.cpp
@@ -428,7 +428,6 @@ void UEdGraphSchema_CustomizableObject::GetGraphContextActions(FGraphContextMenu
 		{
 			ContextMenuBuilder.CreateTemplateNode<UCustomizableObjectNodeSkeletalMesh>(),
 			ContextMenuBuilder.CreateTemplateNode<UCustomizableObjectNodeStaticMesh>(),
-			ContextMenuBuilder.CreateTemplateNode<UCustomizableObjectNodeLayoutBlocks>(),
 			ContextMenuBuilder.CreateTemplateNode<UCustomizableObjectNodeMeshMorph>(),
 			ContextMenuBuilder.CreateTemplateNode<UCustomizableObjectNodeMeshSwitch>(),
 			ContextMenuBuilder.CreateTemplateNode<UCustomizableObjectNodeMeshVariation>(),
diff --git a/Engine/Plugins/Experimental/Mutable/Source/CustomizableObjectEditor/Private/MuCOE/GenerateMutableSource/GenerateMutableSource.cpp b/Engine/Plugins/Experimental/Mutable/Source/CustomizableObjectEditor/Private/MuCOE/GenerateMutableSource/GenerateMutableSource.cpp
--- a/Engine/Plugins/Experimental/Mutable/Source/CustomizableObjectEditor/Private/MuCOE/GenerateMutableSource/GenerateMutableSource.cpp
+++ b/Engine/Plugins/Experimental/Mutable/Source/CustomizableObjectEditor/Private/MuCOE/GenerateMutableSource/GenerateMutableSource.cpp
@@ -127,6 +127,9 @@ FMutableGraphGenerationContext::FMutableGraphGenerationContext(UCustomizableObje
 {
 	// Default flags for mesh generation nodes.
 	MeshGenerationFlags.Push(EMutableMeshConversionFlags::None);
+
+	// Default flags for mesh layout generation.
+	LayoutGenerationFlags.Push(FLayoutGenerationFlags());
 }
 
 FMutableGraphGenerationContext::~FMutableGraphGenerationContext() = default;
@@ -468,6 +471,7 @@ FGeneratedKey::FGeneratedKey(void* InFunctionAddress, const UEdGraphPin& InPin,
 	if (UseMesh)
 	{
 		Flags = GenerationContext.MeshGenerationFlags.Last();
+		LayoutFlags = GenerationContext.LayoutGenerationFlags.Last();
 		MeshMorphStack = GenerationContext.MeshMorphStack;
 		bOnlyConnectedLOD = InbOnlyConnectedLOD;
 	}
@@ -480,6 +484,7 @@ bool FGeneratedKey::operator==(const FGeneratedKey& Other) const
 		Pin == Other.Pin &&
 		LOD == Other.LOD &&
 		Flags == Other.Flags &&
+		LayoutFlags == Other.LayoutFlags &&
 		MeshMorphStack == Other.MeshMorphStack &&
 		bOnlyConnectedLOD == Other.bOnlyConnectedLOD;
 }
@@ -491,6 +496,7 @@ uint32 GetTypeHash(const FGeneratedKey& Key)
 	Hash = HashCombine(Hash, GetTypeHash(Key.Pin));
 	Hash = HashCombine(Hash, GetTypeHash(Key.LOD));
 	Hash = HashCombine(Hash, GetTypeHash(Key.Flags));
+	//Hash = HashCombine(Hash, GetTypeHash(Key.LayoutFlags)); // Does not support array
 	//Hash = HashCombine(Hash, GetTypeHash(Key.MeshMorphStack)); // Does not support array
 	Hash = HashCombine(Hash, GetTypeHash(Key.bOnlyConnectedLOD));
 	
diff --git a/Engine/Plugins/Experimental/Mutable/Source/CustomizableObjectEditor/Private/MuCOE/GenerateMutableSource/GenerateMutableSource.h b/Engine/Plugins/Experimental/Mutable/Source/CustomizableObjectEditor/Private/MuCOE/GenerateMutableSource/GenerateMutableSource.h
--- a/Engine/Plugins/Experimental/Mutable/Source/CustomizableObjectEditor/Private/MuCOE/GenerateMutableSource/GenerateMutableSource.h
+++ b/Engine/Plugins/Experimental/Mutable/Source/CustomizableObjectEditor/Private/MuCOE/GenerateMutableSource/GenerateMutableSource.h
@@ -53,6 +53,7 @@ struct FMutableRefSkeletalMeshData;
 struct FMutableRefSocket;
 struct FMutableSkinWeightProfileInfo;
 struct FMorphTargetVertexData;
+enum class EPinMode;
 
 struct FGeneratedImageProperties
 {
@@ -105,6 +106,14 @@ enum class EMutableMeshConversionFlags : uint32
 
 ENUM_CLASS_FLAGS(EMutableMeshConversionFlags)
 
+struct FLayoutGenerationFlags
+{
+	bool operator==(const FLayoutGenerationFlags& Other) const = default;
+
+	// Texture pin mode per UV Channel
+	TArray<EPinMode> TexturePinModes;
+};
+
 /** 
 	Struct to store the necessary data to generate the morphs of a skeletal mesh 
 	This struct allows the stack morph nodes to use the same functions as the mesh morph nodes
@@ -155,6 +164,9 @@ private:
 	/** Active morphs at the time of mesh generation. */
 	TArray<FMorphNodeData> MeshMorphStack;
 
+	/** UV Layout modes */
+	FLayoutGenerationFlags LayoutFlags;
+
 	/** When caching a generated mesh, true if we force to generate the connected LOD when using Automatic LODs From Mesh. */
 	bool bOnlyConnectedLOD = false;
 };
@@ -700,6 +712,9 @@ struct FMutableGraphGenerationContext
 	// The value is a bit mask of EMutableMeshConversionFlags
 	TArray<EMutableMeshConversionFlags> MeshGenerationFlags;
 
+	// Stack of Layout generation flags. The last one is the currently valid.
+	TArray<FLayoutGenerationFlags> LayoutGenerationFlags;
+
 	/** Find a mesh if already generated for a given source and flags. */
 	mu::Ptr<mu::Mesh> FindGeneratedMesh(const FGeneratedMeshData::FKey& Key);
 
diff --git a/Engine/Plugins/Experimental/Mutable/Source/CustomizableObjectEditor/Private/MuCOE/GenerateMutableSource/GenerateMutableSourceLayout.cpp b/Engine/Plugins/Experimental/Mutable/Source/CustomizableObjectEditor/Private/MuCOE/GenerateMutableSource/GenerateMutableSourceLayout.cpp
--- a/Engine/Plugins/Experimental/Mutable/Source/CustomizableObjectEditor/Private/MuCOE/GenerateMutableSource/GenerateMutableSourceLayout.cpp
+++ b/Engine/Plugins/Experimental/Mutable/Source/CustomizableObjectEditor/Private/MuCOE/GenerateMutableSource/GenerateMutableSourceLayout.cpp
@@ -12,72 +12,23 @@
 
 #define LOCTEXT_NAMESPACE "CustomizableObjectEditor"
 
-
-mu::Ptr<mu::NodeLayout> GenerateMutableSourceLayout(const UEdGraphPin * Pin, FMutableGraphGenerationContext& GenerationContext, bool bIgnoreLayoutWarning)
+mu::Ptr<mu::NodeLayout> CreateDefaultLayout()
 {
-	check(Pin)
-	RETURN_ON_CYCLE(*Pin, GenerationContext)
-
-	CheckNumOutputs(*Pin, GenerationContext);
-
-	UCustomizableObjectNode* Node = CastChecked<UCustomizableObjectNode>(Pin->GetOwningNode());
-
-	const FGeneratedKey Key(reinterpret_cast<void*>(&GenerateMutableSourceLayout), *Pin, *Node, GenerationContext, true);
-	if (const FGeneratedData* Generated = GenerationContext.Generated.Find(Key))
-	{
-		return static_cast<mu::NodeLayout*>(Generated->Node.get());
-	}
+	constexpr int32 GridSize = 4;
+
+	mu::Ptr<mu::NodeLayout> LayoutNode = new mu::NodeLayout();
+	LayoutNode->Size = { GridSize, GridSize };
+	LayoutNode->MaxSize = { GridSize, GridSize };
+	LayoutNode->Strategy = mu::EPackStrategy::Resizeable;
+	LayoutNode->ReductionMethod = mu::EReductionMethod::Halve;
+	LayoutNode->Blocks.SetNum(1);
+	LayoutNode->Blocks[0].Min = { 0, 0 };
+	LayoutNode->Blocks[0].Size = { GridSize, GridSize };
+	LayoutNode->Blocks[0].Priority = 0;
+	LayoutNode->Blocks[0].bReduceBothAxes = false;
+	LayoutNode->Blocks[0].bReduceByTwo = false;
 
-	mu::Ptr<mu::NodeLayout> Result;
-	
-	if (const UCustomizableObjectNodeLayoutBlocks* TypedNodeBlocks = Cast<UCustomizableObjectNodeLayoutBlocks>(Node))
-	{
-		if (UCustomizableObjectNodeSkeletalMesh* SkeletalMeshNode = Cast<UCustomizableObjectNodeSkeletalMesh>(FollowOutputPin(*TypedNodeBlocks->OutputPin())->GetOwningNode()))
-		{
-			int32 LayoutIndex;
-			FString MaterialName;
-
-			if (!SkeletalMeshNode->CheckIsValidLayout(Pin, LayoutIndex, MaterialName))
-			{
-				FString msg = "Layouts ";
-				for (int32 i = 0; i < LayoutIndex; ++i)
-				{
-					msg += "UV" + FString::FromInt(i);
-					if (i < LayoutIndex - 1)
-					{
-						msg += ", ";
-					}
-				}
-				msg += " of " + MaterialName + " must be also connected to a Layout Blocks Node. ";
-				GenerationContext.Compiler->CompilerLog(FText::FromString(msg), Node, EMessageSeverity::Error);
-				return nullptr;
-			}
-		}
-
-		bool bWasEmpty = false;
-		Result = CreateMutableLayoutNode(GenerationContext, TypedNodeBlocks->Layout, bIgnoreLayoutWarning,bWasEmpty);
-		if (bWasEmpty)
-		{
-			FString msg = "Layout without any block found. A grid sized block will be used instead.";
-			GenerationContext.Compiler->CompilerLog(FText::FromString(msg), Node, EMessageSeverity::Warning);
-		}
-	}
-	
-	else
-	{
-		GenerationContext.Compiler->CompilerLog(LOCTEXT("UnimplementedNode", "Node type not implemented yet."), Node);
-	}
-
-	GenerationContext.Generated.Add(Key, FGeneratedData(Node, Result));
-	GenerationContext.GeneratedNodes.Add(Node);
-
-	
-	if (Result)
-	{
-		Result->SetMessageContext(Node);
-	}
-
-	return Result;
+	return LayoutNode;
 }
 
 
@@ -118,7 +69,7 @@ mu::Ptr<mu::NodeLayout> CreateMutableLayoutNode(FMutableGraphGenerationContext&
 		PackingStrategy == ECustomizableObjectTextureLayoutPackingStrategy::Overlay)
  	{
 		// Legacy behavior
-		if (!UnrealLayout->Blocks.Num())
+		if (UnrealLayout->Blocks.IsEmpty())
 		{
 			bWasEmpty = true;
 			LayoutNode->Blocks.SetNum(1);
@@ -129,7 +80,7 @@ mu::Ptr<mu::NodeLayout> CreateMutableLayoutNode(FMutableGraphGenerationContext&
 			LayoutNode->Blocks[0].bReduceByTwo = false;
 		}
 	}
-	else
+	else if (UnrealLayout->Blocks.IsEmpty())
 	{
 		// Convert the UE mesh in the layout into a Mutable mesh.
 		mu::Ptr<mu::Mesh> MutableMesh = nullptr;
diff --git a/Engine/Plugins/Experimental/Mutable/Source/CustomizableObjectEditor/Private/MuCOE/GenerateMutableSource/GenerateMutableSourceLayout.h b/Engine/Plugins/Experimental/Mutable/Source/CustomizableObjectEditor/Private/MuCOE/GenerateMutableSource/GenerateMutableSourceLayout.h
--- a/Engine/Plugins/Experimental/Mutable/Source/CustomizableObjectEditor/Private/MuCOE/GenerateMutableSource/GenerateMutableSourceLayout.h
+++ b/Engine/Plugins/Experimental/Mutable/Source/CustomizableObjectEditor/Private/MuCOE/GenerateMutableSource/GenerateMutableSourceLayout.h
@@ -9,7 +9,7 @@ class UEdGraphPin;
 struct FMutableGraphGenerationContext;
 
 
-extern mu::Ptr<mu::NodeLayout> GenerateMutableSourceLayout(const UEdGraphPin*, FMutableGraphGenerationContext&, bool bIgnoreLayoutWarning = false);
+extern mu::Ptr<mu::NodeLayout> CreateDefaultLayout();
 
 extern mu::Ptr<mu::NodeLayout> CreateMutableLayoutNode(FMutableGraphGenerationContext& GenerationContext, const UCustomizableObjectLayout*, bool bIgnoreLayoutWarnings, bool& bWasEmpty);
 
diff --git a/Engine/Plugins/Experimental/Mutable/Source/CustomizableObjectEditor/Private/MuCOE/GenerateMutableSource/GenerateMutableSourceMesh.cpp b/Engine/Plugins/Experimental/Mutable/Source/CustomizableObjectEditor/Private/MuCOE/GenerateMutableSource/GenerateMutableSourceMesh.cpp
--- a/Engine/Plugins/Experimental/Mutable/Source/CustomizableObjectEditor/Private/MuCOE/GenerateMutableSource/GenerateMutableSourceMesh.cpp
+++ b/Engine/Plugins/Experimental/Mutable/Source/CustomizableObjectEditor/Private/MuCOE/GenerateMutableSource/GenerateMutableSourceMesh.cpp
@@ -25,6 +25,7 @@
 #include "MuCOE/Nodes/CustomizableObjectNodeAnimationPose.h"
 #include "MuCOE/Nodes/CustomizableObjectNodeFloatConstant.h"
 #include "MuCOE/Nodes/CustomizableObjectNodeFloatParameter.h"
+#include "MuCOE/Nodes/CustomizableObjectNodeMaterial.h"
 #include "MuCOE/Nodes/CustomizableObjectNodeMeshGeometryOperation.h"
 #include "MuCOE/Nodes/CustomizableObjectNodeMeshMorph.h"
 #include "MuCOE/Nodes/CustomizableObjectNodeMeshMorphStackApplication.h"
@@ -2460,7 +2461,7 @@ mu::Ptr<mu::Mesh> GenerateMutableMesh(UObject * Mesh, const TSoftClassPtr<UAnimI
 	EMutableMeshConversionFlags CurrentFlags = GenerationContext.MeshGenerationFlags.Last();
 
 	FMutableGraphGenerationContext::FGeneratedMeshData::FKey Key = { Mesh, LODIndex, GenerationContext.CurrentLOD, SectionIndex, CurrentFlags, UniqueTags, CurrentNode };
-	mu::Ptr<mu::Mesh> MutableMesh = GenerationContext.FindGeneratedMesh(Key);	
+	mu::Ptr<mu::Mesh> MutableMesh = GenerationContext.FindGeneratedMesh(Key);
 	if (MutableMesh)
 	{
 		return MutableMesh;
@@ -3301,27 +3302,6 @@ mu::NodeMeshPtr GenerateMorphMesh(const UEdGraphPin* Pin,
 	return Result;
 }
 
-/** Create a default layout. Used when no layout is found. */
-mu::Ptr<mu::NodeLayout> CreateDefaultLayout()
-{
-	constexpr int32 GridSize = 4;
-	
-	mu::Ptr<mu::NodeLayout> LayoutNode = new mu::NodeLayout();
-	LayoutNode->Size = { GridSize, GridSize };
-	LayoutNode->MaxSize = { GridSize, GridSize };
-	LayoutNode->Strategy = mu::EPackStrategy::Resizeable;
-	LayoutNode->ReductionMethod = mu::EReductionMethod::Halve;
-	LayoutNode->Blocks.SetNum(1);
-	LayoutNode->Blocks[0].Min = { 0, 0 };
-	LayoutNode->Blocks[0].Size = { GridSize, GridSize };
-	LayoutNode->Blocks[0].Priority = 0;
-	LayoutNode->Blocks[0].bReduceBothAxes = false;
-	LayoutNode->Blocks[0].bReduceByTwo = false;
-
-	return LayoutNode;
-}
-
-
 mu::NodeMeshPtr GenerateMutableSourceMesh(const UEdGraphPin* Pin,
 	FMutableGraphGenerationContext& GenerationContext,
 	FMutableGraphMeshGenerationData& MeshData,
@@ -3558,28 +3538,28 @@ mu::NodeMeshPtr GenerateMutableSourceMesh(const UEdGraphPin* Pin,
 				const int32 NumLayouts = ImportedModel->LODModels[LODIndexLayout].NumTexCoords;
 				MeshNode->SetLayoutCount(NumLayouts);
 				
-				bool bAtLeastOneLayout = false;
+				const FLayoutGenerationFlags& LayoutFlags = GenerationContext.LayoutGenerationFlags.Last();
 
+				const TArray<UCustomizableObjectLayout*>& Layouts = TypedNodeSkel->GetLayouts(*Pin);
 				for (int32 LayoutIndex = 0; LayoutIndex < NumLayouts; ++LayoutIndex)
 				{
-					if (UEdGraphPin* LayoutPin = TypedNodeSkel->GetLayoutPin(LODIndexLayout, SectionIndexLayout, LayoutIndex))
+					if (!LayoutFlags.TexturePinModes.IsValidIndex(LayoutIndex) ||
+						LayoutFlags.TexturePinModes[LayoutIndex] != EPinMode::Mutable)
 					{
-						if (const UEdGraphPin* ConnectedPin = FollowInputPin(*LayoutPin))
-						{
-							mu::Ptr<mu::NodeLayout> LayoutNode = GenerateMutableSourceLayout(ConnectedPin, GenerationContext, bLinkedToExtendMaterial);
-							MeshNode->SetLayout(LayoutIndex, LayoutNode);
-							bAtLeastOneLayout = true;
-						}
+						MeshNode->SetLayout(LayoutIndex, CreateDefaultLayout());
+						// Ignore layout
+						continue;
 					}
-				}
 
-				if (!bAtLeastOneLayout)
-				{
-					MeshNode->SetLayoutCount(1);
+					const UCustomizableObjectLayout* Layout = Layouts.IsValidIndex(0) ? Layouts[0] : nullptr;
+					if (ensure(Layout))
+					{
+						bool bWasEmpty = false;
+						mu::Ptr<mu::NodeLayout> LayoutNode = CreateMutableLayoutNode(GenerationContext, Layout, bLinkedToExtendMaterial, bWasEmpty); // TODO PERE: Figure out
+						LayoutNode->SetMessageContext(Node);
 
-					mu::Ptr<mu::NodeLayout> LayoutNode = CreateDefaultLayout();
-					MeshNode->SetLayout(0, LayoutNode);
-					LayoutNode->SetMessageContext(Node); // We need it here because we create multiple nodes.
+						MeshNode->SetLayout(LayoutIndex, LayoutNode);
+					}
 				}
 			}
 
@@ -3649,17 +3629,29 @@ mu::NodeMeshPtr GenerateMutableSourceMesh(const UEdGraphPin* Pin,
 				// Layouts
 				MeshNode->SetLayoutCount(1);
 
-				if (const UEdGraphPin* ConnectedPin = FollowInputPin(*TypedNodeStatic->LODs[LODIndex].Materials[SectionIndex].LayoutPinRef.Get()))
+				mu::Ptr<mu::NodeLayout> LayoutNode;
+
+				const TArray<UCustomizableObjectLayout*>& Layouts = TypedNodeStatic->GetLayouts(*Pin);
+				const UCustomizableObjectLayout* Layout = Layouts.IsValidIndex(0) ? Layouts[0] : nullptr;
+				
+				if (Layout)
 				{
-					mu::Ptr<mu::NodeLayout> LayoutNode = GenerateMutableSourceLayout(ConnectedPin, GenerationContext);
-					MeshNode->SetLayout(0, LayoutNode);
+					bool bWasEmpty = false;
+					LayoutNode = CreateMutableLayoutNode(GenerationContext, Layout, false, bWasEmpty);
+					/*if (bWasEmpty)
+					{
+						FString msg = "Layout without any block found. A grid sized block will be used instead.";
+						GenerationContext.Compiler->CompilerLog(FText::FromString(msg), Node, EMessageSeverity::Warning);
+					}*/
 				}
-				else
+
+				if (!LayoutNode)
 				{
-					mu::Ptr<mu::NodeLayout> LayoutNode = CreateDefaultLayout();
-					MeshNode->SetLayout(0, LayoutNode);
-					LayoutNode->SetMessageContext(Node);  // We need it here because we create multiple nodes.
+					LayoutNode = CreateDefaultLayout();
 				}
+
+				MeshNode->SetLayout(0, LayoutNode);
+				LayoutNode->SetMessageContext(Node);  // We need it here because we create multiple nodes.
 			}
 			else
 			{
@@ -4300,8 +4292,18 @@ mu::NodeMeshPtr GenerateMutableSourceMesh(const UEdGraphPin* Pin,
 							if (Layouts.Num())
 							{
 								// Generating node Layouts
+								const FLayoutGenerationFlags& LayoutFlags = GenerationContext.LayoutGenerationFlags.Last();
+								
 								for (int32 i = 0; i < Layouts.Num(); ++i)
 								{
+									if (!LayoutFlags.TexturePinModes.IsValidIndex(i) ||
+										LayoutFlags.TexturePinModes[i] != EPinMode::Mutable)
+									{
+										MeshTableNode->SetLayout(i, CreateDefaultLayout());
+										// Ignore layouts
+										continue;
+									}
+
 									bool bWasEmpty = false;
 									// In tables, mimic the legacy behaviour and ignore all layout warnings beyond LOD 0.
 									bool bIgnoreLayoutWarnings = true;
diff --git a/Engine/Plugins/Experimental/Mutable/Source/CustomizableObjectEditor/Private/MuCOE/GenerateMutableSource/GenerateMutableSourceModifier.cpp b/Engine/Plugins/Experimental/Mutable/Source/CustomizableObjectEditor/Private/MuCOE/GenerateMutableSource/GenerateMutableSourceModifier.cpp
--- a/Engine/Plugins/Experimental/Mutable/Source/CustomizableObjectEditor/Private/MuCOE/GenerateMutableSource/GenerateMutableSourceModifier.cpp
+++ b/Engine/Plugins/Experimental/Mutable/Source/CustomizableObjectEditor/Private/MuCOE/GenerateMutableSource/GenerateMutableSourceModifier.cpp
@@ -290,7 +290,15 @@ mu::Ptr<mu::NodeModifier> GenerateMutableSourceModifier(const UEdGraphPin * Pin,
 			FMutableGraphMeshGenerationData MeshData;
 			if (const UEdGraphPin* ConnectedPin = FollowInputPin(*TypedNodeExt->AddMeshPin()))
 			{
+				// Flags to know which UV channels need layout
+				FLayoutGenerationFlags LayoutGenerationFlags;
+				LayoutGenerationFlags.TexturePinModes.Init(EPinMode::Mutable, TEXSTREAM_MAX_NUM_UVCHANNELS);
+
+				GenerationContext.LayoutGenerationFlags.Push(LayoutGenerationFlags);
+
 				AddMeshNode = GenerateMutableSourceMesh(ConnectedPin, GenerationContext, MeshData, true, false);
+
+				GenerationContext.LayoutGenerationFlags.Pop();
 			}
 
 			if (AddMeshNode)
diff --git a/Engine/Plugins/Experimental/Mutable/Source/CustomizableObjectEditor/Private/MuCOE/GenerateMutableSource/GenerateMutableSourceSurface.cpp b/Engine/Plugins/Experimental/Mutable/Source/CustomizableObjectEditor/Private/MuCOE/GenerateMutableSource/GenerateMutableSourceSurface.cpp
--- a/Engine/Plugins/Experimental/Mutable/Source/CustomizableObjectEditor/Private/MuCOE/GenerateMutableSource/GenerateMutableSourceSurface.cpp
+++ b/Engine/Plugins/Experimental/Mutable/Source/CustomizableObjectEditor/Private/MuCOE/GenerateMutableSource/GenerateMutableSourceSurface.cpp
@@ -340,9 +340,32 @@ mu::Ptr<mu::NodeSurface> GenerateMutableSourceSurface(const UEdGraphPin * Pin, F
 		{
 			if (const UEdGraphPin* ConnectedPin = FollowInputPin(*TypedNodeMat->GetMeshPin()))
 			{
+
+				// Flags to know which UV channels need layout
+				FLayoutGenerationFlags LayoutGenerationFlags;
+				
+				LayoutGenerationFlags.TexturePinModes.Init(EPinMode::Default, TEXSTREAM_MAX_NUM_UVCHANNELS);
+
+				const int32 NumImages = TypedNodeMat->GetNumParameters(EMaterialParameterType::Texture);
+				for (int32 ImageIndex = 0; ImageIndex < NumImages; ++ImageIndex)
+				{
+					if (TypedNodeMat->IsImageMutableMode(ImageIndex))
+					{
+						const int32 UVChannel = TypedNodeMat->GetImageUVLayout(ImageIndex);
+						if (LayoutGenerationFlags.TexturePinModes.IsValidIndex(UVChannel))
+						{
+							LayoutGenerationFlags.TexturePinModes[UVChannel] = EPinMode::Mutable;
+						}
+					}
+				}
+
+				GenerationContext.LayoutGenerationFlags.Push(LayoutGenerationFlags);
+
 				FMutableGraphMeshGenerationData MeshData;
 				MeshNode = GenerateMutableSourceMesh(ConnectedPin, GenerationContext, MeshData, false, false);
 
+				GenerationContext.LayoutGenerationFlags.Pop();
+
 				if (MeshNode)
 				{
 					mu::Ptr<mu::NodeMeshFormat> MeshFormatNode = new mu::NodeMeshFormat();
diff --git a/Engine/Plugins/Experimental/Mutable/Source/CustomizableObjectEditor/Private/MuCOE/Nodes/CustomizableObjectNodeLayoutBlocks.cpp b/Engine/Plugins/Experimental/Mutable/Source/CustomizableObjectEditor/Private/MuCOE/Nodes/CustomizableObjectNodeLayoutBlocks.cpp
--- a/Engine/Plugins/Experimental/Mutable/Source/CustomizableObjectEditor/Private/MuCOE/Nodes/CustomizableObjectNodeLayoutBlocks.cpp
+++ b/Engine/Plugins/Experimental/Mutable/Source/CustomizableObjectEditor/Private/MuCOE/Nodes/CustomizableObjectNodeLayoutBlocks.cpp
@@ -24,29 +24,13 @@ UCustomizableObjectNodeLayoutBlocks::UCustomizableObjectNodeLayoutBlocks()
 }
 
 
-void UCustomizableObjectNodeLayoutBlocks::PostBackwardsCompatibleFixup()
+void UCustomizableObjectNodeLayoutBlocks::BackwardsCompatibleFixup(int32 CustomizableObjectCustomVersion)
 {
-	Super::PostBackwardsCompatibleFixup();
+	Super::BackwardsCompatibleFixup(CustomizableObjectCustomVersion);
 
-	if (UEdGraphPin* Pin = OutputPin())
+	if (CustomizableObjectCustomVersion == FCustomizableObjectCustomVersion::RemoveNodeLayout)
 	{
-		LinkPostEditChangePropertyDelegate(Pin);
-	}
-
-	// Make sure we have IDs for all blocks. Generate fake ids for old blocks so that it is deterministic
-	// and new references to these blocks work even if this node is not saved.
-	for (int BlockIndex=0; BlockIndex< Layout->Blocks.Num(); ++BlockIndex)
-	{
-		FCustomizableObjectLayoutBlock& Block = Layout->Blocks[BlockIndex];
-		if (!Block.Id.IsValid())
-		{
-			FGuid FakeId;
-			FakeId[0] = BlockIndex + 1;
-			FakeId[1] = BlockIndex + 1;
-			FakeId[2] = BlockIndex + 1;
-			FakeId[3] = BlockIndex + 1;
-			Block.Id = FakeId;
-		}
+		DestroyNode();
 	}
 }
 
@@ -59,122 +43,6 @@ void UCustomizableObjectNodeLayoutBlocks::AllocateDefaultPins(UCustomizableObjec
 	OutputPin->bDefaultValueIsIgnored = true;
 }
 
-
-FText UCustomizableObjectNodeLayoutBlocks::GetNodeTitle(ENodeTitleType::Type TitleType) const
-{
-	return LOCTEXT("Mesh_Layout", "Mesh Layout");
-}
-
-
-FLinearColor UCustomizableObjectNodeLayoutBlocks::GetNodeTitleColor() const
-{
-	const UEdGraphSchema_CustomizableObject* Schema = GetDefault<UEdGraphSchema_CustomizableObject>();
-	return Schema->GetPinTypeColor(Schema->PC_Layout);
-}
-
-
-void UCustomizableObjectNodeLayoutBlocks::PinConnectionListChanged(UEdGraphPin* Pin)
-{
-	Super::PinConnectionListChanged(Pin);
-	
-	if (Layout && Pin == OutputPin())
-	{
-		LinkPostEditChangePropertyDelegate(Pin);
-	}
-}
-
-
-void UCustomizableObjectNodeLayoutBlocks::AddAttachedErrorData(const FAttachedErrorDataView& AttachedErrorData)
-{
-	if (Layout)
-	{
-		Layout->UnassignedUVs.Emplace();
-
-		const int32 UnassignedUVsCount = AttachedErrorData.UnassignedUVs.Num() / 2;
-		for (int32 Index = 0; Index < UnassignedUVsCount; ++Index)
-		{
-			Layout->UnassignedUVs.Last().Emplace(AttachedErrorData.UnassignedUVs[Index * 2 + 0],
-				AttachedErrorData.UnassignedUVs[Index * 2 + 1]);
-		}
-	}
-}
-
-
-void UCustomizableObjectNodeLayoutBlocks::ResetAttachedErrorData()
-{
-	if (Layout)
-	{
-		Layout->UnassignedUVs.Empty();
-	}
-}
-
-
-FText UCustomizableObjectNodeLayoutBlocks::GetTooltipText() const
-{
-	return LOCTEXT("Layout_Block_Tooltip", "Defines of how the material UV layout is cut into logical pieces.\nWhen destined to a Material Node, it also defines the resolution per layout block of that material for any extension that is made to that material.");
-
-}
-
-
-bool UCustomizableObjectNodeLayoutBlocks::IsSingleOutputNode() const
-{
-	return true;
-}
-
-
-void UCustomizableObjectNodeLayoutBlocks::MeshPostEditChangeProperty(UObject* Node, FPropertyChangedEvent& FPropertyChangedEvent)
-{
-	SetLayoutSkeletalMesh();
-}
-
-
-void UCustomizableObjectNodeLayoutBlocks::LinkPostEditChangePropertyDelegate(UEdGraphPin* Pin)
-{
-	if (LastMeshNodeConnected.IsValid())
-	{
-		LastMeshNodeConnected->PostEditChangePropertyRegularDelegate.AddUObject(this, &UCustomizableObjectNodeLayoutBlocks::MeshPostEditChangeProperty);
-	}
-
-	if (const UEdGraphPin* ConnectedPin = FollowOutputPin(*Pin))
-	{
-		UEdGraphNode* MeshNode = ConnectedPin->GetOwningNode();
-
-		if (MeshNode->IsA(UCustomizableObjectNodeStaticMesh::StaticClass()) || MeshNode->IsA(UCustomizableObjectNodeSkeletalMesh::StaticClass()))
-		{
-			LastMeshNodeConnected = Cast<UCustomizableObjectNode>(MeshNode);
-			LastMeshNodeConnected->PostEditChangePropertyRegularDelegate.AddUObject(this, &UCustomizableObjectNodeLayoutBlocks::MeshPostEditChangeProperty);
-
-			SetLayoutSkeletalMesh();
-		}
-	}
-}
-
-
-void UCustomizableObjectNodeLayoutBlocks::SetLayoutSkeletalMesh()
-{
-	if (const UEdGraphPin* ConnectedPin = FollowOutputPin(*OutputPin()))
-	{
-		if (const UCustomizableObjectNodeMesh* MeshNode = Cast<UCustomizableObjectNodeMesh>(LastMeshNodeConnected))
-		{
-			int32 LODIndex;
-			int32 SectionIndex;
-			int32 LayoutIndex;
-			MeshNode->GetPinSection(*ConnectedPin, LODIndex, SectionIndex, LayoutIndex);
-
-			Layout->SetLayout(MeshNode->GetMesh(), LODIndex, SectionIndex, LayoutIndex);
-		}
-	}
-}
-
-
-void UCustomizableObjectNodeLayoutBlocks::PostPasteNode()
-{
-	Super::PostPasteNode();
-	
-	Layout->SetLayout(nullptr, -1, -1, -1);
-}
-
-
 void UCustomizableObjectNodeLayoutBlocks::Serialize(FArchive& Ar)
 {
 	Super::Serialize(Ar);
diff --git a/Engine/Plugins/Experimental/Mutable/Source/CustomizableObjectEditor/Private/MuCOE/Nodes/CustomizableObjectNodeLayoutBlocksDetails.cpp b/Engine/Plugins/Experimental/Mutable/Source/CustomizableObjectEditor/Private/MuCOE/Nodes/CustomizableObjectNodeLayoutBlocksDetails.cpp
deleted file mode 100644
diff --git a/Engine/Plugins/Experimental/Mutable/Source/CustomizableObjectEditor/Private/MuCOE/Nodes/CustomizableObjectNodeLayoutBlocksDetails.h b/Engine/Plugins/Experimental/Mutable/Source/CustomizableObjectEditor/Private/MuCOE/Nodes/CustomizableObjectNodeLayoutBlocksDetails.h
deleted file mode 100644
diff --git a/Engine/Plugins/Experimental/Mutable/Source/CustomizableObjectEditor/Private/MuCOE/Nodes/CustomizableObjectNodeModifierEditMeshSectionDetails.cpp b/Engine/Plugins/Experimental/Mutable/Source/CustomizableObjectEditor/Private/MuCOE/Nodes/CustomizableObjectNodeModifierEditMeshSectionDetails.cpp
--- a/Engine/Plugins/Experimental/Mutable/Source/CustomizableObjectEditor/Private/MuCOE/Nodes/CustomizableObjectNodeModifierEditMeshSectionDetails.cpp
+++ b/Engine/Plugins/Experimental/Mutable/Source/CustomizableObjectEditor/Private/MuCOE/Nodes/CustomizableObjectNodeModifierEditMeshSectionDetails.cpp
@@ -9,7 +9,7 @@
 #include "MuCOE/Nodes/CustomizableObjectNodeModifierEditMeshSection.h"
 #include "MuCOE/Nodes/CustomizableObjectNodeModifierExtendMeshSection.h"
 #include "MuCOE/Nodes/CustomizableObjectNodeMaterialBase.h"
-#include "MuCOE/SCustomizableObjectNodeLayoutBlocksEditor.h"
+#include "MuCOE/SCustomizableObjectLayoutEditor.h"
 #include "MuCOE/PinViewer/SPinViewer.h"
 #include "MuR/MutableTrace.h"
 #include "Widgets/Layout/SBox.h"
@@ -39,7 +39,8 @@ void FCustomizableObjectNodeModifierEditMeshSectionDetails::CustomizeDetails( ID
 		Node = Cast<UCustomizableObjectNodeModifierEditMeshSection>(DetailsView->GetSelectedObjects()[0].Get());
 	}
 
-	IDetailCategoryBuilder& LayoutCategory = DetailBuilder.EditCategory("LayoutOptions");
+	
+	IDetailCategoryBuilder& LayoutCategory = DetailBuilder.EditCategory("Layout Editor");
 
 	if (!Node)
 	{
@@ -53,7 +54,6 @@ void FCustomizableObjectNodeModifierEditMeshSectionDetails::CustomizeDetails( ID
 
 	// UV Channel combo (for now hardcoded to a maximum of 4)
 	{
-		int32 MaxGridSize = 128;
 		UVChannelOptions.Empty();
 
 		TSharedPtr<FString> CurrentUVChannel;
@@ -67,12 +67,12 @@ void FCustomizableObjectNodeModifierEditMeshSectionDetails::CustomizeDetails( ID
 			}
 		}
 
-		IDetailGroup& LayoutOptionsGroup = LayoutCategory.AddGroup(TEXT("LayoutOptionsGroup"), LOCTEXT("LayoutGroup", "Layout Group"), false, true);
+		IDetailGroup& LayoutOptionsGroup = LayoutCategory.AddGroup(TEXT("EditLayoutGroup"), LOCTEXT("EditMeshSection", "Edit Layout"), false, true);
 		LayoutOptionsGroup.HeaderRow()
 			.NameContent()
 			[
 				SNew(STextBlock)
-					.Text(LOCTEXT("UVChannel", "UV Channel"))
+					.Text(LOCTEXT("UVChannel", "Edit UV Channel"))
 					.Font(DetailBuilder.GetDetailFont())
 			]
 			.ValueContent()
@@ -85,97 +85,56 @@ void FCustomizableObjectNodeModifierEditMeshSectionDetails::CustomizeDetails( ID
 			];
 	}
 
-	// Grid size combo
-	{
-		int32 MaxGridSize = 128;
-		LayoutGridSizes.Empty();
+	// Mesh sections with the editable layouts
+	TArray<FLayoutEditorMeshSection> MeshSectionsAndLayouts;
 
-		TSharedPtr<FString> CurrentSize;
-		for (int32 Size = 1; Size <= MaxGridSize; Size *= 2)
-		{
-			LayoutGridSizes.Add(MakeShareable(new FString(FString::Printf(TEXT("%d x %d"), Size, Size))));
+	// Add Dummy section with the layout to edit.
+	FLayoutEditorMeshSection& MeshSection = MeshSectionsAndLayouts.AddDefaulted_GetRef();
+	MeshSection.MeshName = MakeShareable(new FString("NameNone"));
+	MeshSection.Layouts.Add(Node->Layout);
 
-			if (Node->Layout->GetGridSize() == FIntPoint(Size))
-			{
-				CurrentSize = LayoutGridSizes.Last();
-			}
-		}
-
-		IDetailGroup& LayoutOptionsGroup = LayoutCategory.AddGroup(TEXT("LayoutOptionsGroup"), LOCTEXT("LayoutGroup", "Layout Group"), false, true);
-		LayoutOptionsGroup.HeaderRow()
-			.NameContent()
-			[
-				SNew(STextBlock)
-					.Text(LOCTEXT("LayoutGridSizeText", "Grid Size"))
-					.Font(DetailBuilder.GetDetailFont())
-			]
-			.ValueContent()
-			[
-				SNew(STextComboBox)
-					.InitiallySelectedItem(CurrentSize)
-					.OptionsSource(&LayoutGridSizes)
-					.OnSelectionChanged(this, &FCustomizableObjectNodeModifierEditMeshSectionDetails::OnGridSizeChanged)
-					.Font(DetailBuilder.GetDetailFont())
-			];
-	}
+	LayoutBlocksEditor = SNew(SCustomizableObjectLayoutEditor)
+		.Node(Node)
+		.MeshSections(MeshSectionsAndLayouts)
+		.OnPreUpdateLayoutDelegate(this, &FCustomizableObjectNodeModifierEditMeshSectionDetails::OnPreUpdateLayout);
 
-	// Block editor
-	LayoutBlocksEditor = SNew(SCustomizableObjectNodeLayoutBlocksEditor);
+	FCustomizableObjectLayoutEditorDetailsBuilder LayoutEditorBuilder;
+	LayoutEditorBuilder.LayoutEditor = LayoutBlocksEditor;
+	LayoutEditorBuilder.bShowGridSize = true;
 
-	LayoutCategory.AddCustomRow(LOCTEXT("BlocksDetails_BlockInstructions", "BlockInstructions"))
-	[
-		SNew(SBox)
-		.HeightOverride(700.0f)
-		.WidthOverride(700.0f)
-		[
-			LayoutBlocksEditor.ToSharedRef()
-		]
-	];
+	LayoutEditorBuilder.CustomizeDetails(DetailBuilder);
 
-	UpdateLayout();
+	LayoutBlocksEditor->UpdateLayout(Node->Layout);
 }
 
 
 void FCustomizableObjectNodeModifierEditMeshSectionDetails::OnRequiredTagsPropertyChanged()
 {
 	FCustomizableObjectNodeModifierBaseDetails::OnRequiredTagsPropertyChanged();
-	UpdateLayout();
-}
-
 
-void FCustomizableObjectNodeModifierEditMeshSectionDetails::UpdateLayout()
-{
-	MUTABLE_CPUPROFILER_SCOPE(FCustomizableObjectNodeModifierEditMeshSectionDetails_UpdateLayout);
-
-	// Try to find the parent layout, because we want to show its UVs in the widget
-	UCustomizableObjectLayout* ParentLayout = Node->GetPossibleParentLayout();
-
-	LayoutBlocksEditor->SetCurrentLayout(Node->Layout, ParentLayout);
+	if (ensure(LayoutBlocksEditor))
+	{
+		LayoutBlocksEditor->UpdateLayout(Node->Layout);
+	}
 }
 
 
-void FCustomizableObjectNodeModifierEditMeshSectionDetails::OnGridSizeChanged(TSharedPtr<FString> NewSelection, ESelectInfo::Type SelectInfo)
+void FCustomizableObjectNodeModifierEditMeshSectionDetails::OnPreUpdateLayout()
 {
-	if (!Node->Layout)
-	{
-		return;
-	}
-
-	int32 Size = 1 << LayoutGridSizes.Find(NewSelection);
+	MUTABLE_CPUPROFILER_SCOPE(FCustomizableObjectNodeModifierEditMeshSectionDetails_UpdateLayout);
 
-	if (Node->Layout->GetGridSize().X != Size || Node->Layout->GetGridSize().Y != Size)
+	if (ensure(LayoutBlocksEditor))
 	{
-		Node->Layout->SetGridSize(FIntPoint(Size));
-
-		Node->MarkPackageDirty();
-		UpdateLayout();
+		// Try to find the parent layout, because we want to show its UVs in the widget
+		UCustomizableObjectLayout* ParentLayout = Node->GetPossibleParentLayout();
+		LayoutBlocksEditor->SetUVsOverride(ParentLayout);
 	}
 }
 
 
 void FCustomizableObjectNodeModifierEditMeshSectionDetails::OnUVChannelChanged(TSharedPtr<FString> NewSelection, ESelectInfo::Type SelectInfo)
 {
-	if (!Node->Layout)
+	if (!Node->Layout || !LayoutBlocksEditor)
 	{
 		return;
 	}
@@ -185,8 +144,10 @@ void FCustomizableObjectNodeModifierEditMeshSectionDetails::OnUVChannelChanged(T
 	if (Node->ParentLayoutIndex != Index)
 	{
 		Node->ParentLayoutIndex = Index;
+
 		Node->Modify();
-		UpdateLayout();
+
+		LayoutBlocksEditor->UpdateLayout(Node->Layout);
 	}
 }
 
diff --git a/Engine/Plugins/Experimental/Mutable/Source/CustomizableObjectEditor/Private/MuCOE/Nodes/CustomizableObjectNodeModifierEditMeshSectionDetails.h b/Engine/Plugins/Experimental/Mutable/Source/CustomizableObjectEditor/Private/MuCOE/Nodes/CustomizableObjectNodeModifierEditMeshSectionDetails.h
--- a/Engine/Plugins/Experimental/Mutable/Source/CustomizableObjectEditor/Private/MuCOE/Nodes/CustomizableObjectNodeModifierEditMeshSectionDetails.h
+++ b/Engine/Plugins/Experimental/Mutable/Source/CustomizableObjectEditor/Private/MuCOE/Nodes/CustomizableObjectNodeModifierEditMeshSectionDetails.h
@@ -10,7 +10,7 @@
 class FString;
 class IDetailLayoutBuilder;
 class UCustomizableObjectNodeModifierEditMeshSection;
-class SCustomizableObjectNodeLayoutBlocksEditor;
+class SCustomizableObjectLayoutEditor;
 
 
 class FCustomizableObjectNodeModifierEditMeshSectionDetails : public FCustomizableObjectNodeModifierBaseDetails
@@ -29,16 +29,14 @@ private:
 	UCustomizableObjectNodeModifierEditMeshSection* Node;
 
 	// Layout block editor widget
-	TSharedPtr<SCustomizableObjectNodeLayoutBlocksEditor> LayoutBlocksEditor;
+	TSharedPtr<SCustomizableObjectLayoutEditor> LayoutBlocksEditor;
 
 	/** List of available layout grid sizes. */
-	TArray< TSharedPtr< FString > > LayoutGridSizes;
-	TArray< TSharedPtr< FString > > UVChannelOptions;
+	TArray<TSharedPtr<FString>> UVChannelOptions;
 
 	/** Layout Options Callbacks */
-	void OnGridSizeChanged(TSharedPtr<FString> NewSelection, ESelectInfo::Type SelectInfo);
 	void OnUVChannelChanged(TSharedPtr<FString> NewSelection, ESelectInfo::Type SelectInfo);
 
 	/** Reset the layout in the widget to force a refresh. */
-	void UpdateLayout();
+	void OnPreUpdateLayout();
 };
diff --git a/Engine/Plugins/Experimental/Mutable/Source/CustomizableObjectEditor/Private/MuCOE/Nodes/CustomizableObjectNodeModifierRemoveMeshBlocksDetails.cpp b/Engine/Plugins/Experimental/Mutable/Source/CustomizableObjectEditor/Private/MuCOE/Nodes/CustomizableObjectNodeModifierRemoveMeshBlocksDetails.cpp
--- a/Engine/Plugins/Experimental/Mutable/Source/CustomizableObjectEditor/Private/MuCOE/Nodes/CustomizableObjectNodeModifierRemoveMeshBlocksDetails.cpp
+++ b/Engine/Plugins/Experimental/Mutable/Source/CustomizableObjectEditor/Private/MuCOE/Nodes/CustomizableObjectNodeModifierRemoveMeshBlocksDetails.cpp
@@ -8,7 +8,7 @@
 #include "IDetailsView.h"
 #include "MuCOE/Nodes/CustomizableObjectNodeModifierRemoveMeshBlocks.h"
 #include "MuCOE/Nodes/CustomizableObjectNodeMaterialBase.h"
-#include "MuCOE/SCustomizableObjectNodeLayoutBlocksEditor.h"
+#include "MuCOE/SCustomizableObjectLayoutEditor.h"
 #include "Widgets/Layout/SBox.h"
 #include "Widgets/Text/STextBlock.h"
 #include "Widgets/Input/STextComboBox.h"
@@ -40,7 +40,7 @@ void FCustomizableObjectNodeModifierRemoveMeshBlocksDetails::CustomizeDetails(ID
 	// This property is not relevant for this node
 	DetailBuilder.HideProperty(GET_MEMBER_NAME_CHECKED(UCustomizableObjectNodeModifierWithMaterial, ReferenceMaterial), UCustomizableObjectNodeModifierWithMaterial::StaticClass());
 
-	IDetailCategoryBuilder& LayoutCategory = DetailBuilder.EditCategory("LayoutOptions");
+	IDetailCategoryBuilder& LayoutCategory = DetailBuilder.EditCategory("Layout Editor");
 	LayoutCategory.SetSortOrder(10000);
 
 	if (!Node)
@@ -69,12 +69,12 @@ void FCustomizableObjectNodeModifierRemoveMeshBlocksDetails::CustomizeDetails(ID
 			}
 		}
 
-		IDetailGroup& LayoutOptionsGroup = LayoutCategory.AddGroup(TEXT("LayoutOptionsGroup"), LOCTEXT("LayoutGroup", "Layout Group"), false, true);
+		IDetailGroup& LayoutOptionsGroup = LayoutCategory.AddGroup(TEXT("LayoutOptionsGroup"), LOCTEXT("LayoutGroup", "Edit Layout"), false, true);
 		LayoutOptionsGroup.HeaderRow()
 			.NameContent()
 			[
 				SNew(STextBlock)
-					.Text(LOCTEXT("UVChannel", "UV Channel"))
+					.Text(LOCTEXT("UVChannel", "Edit UV Channel"))
 					.Font(DetailBuilder.GetDetailFont())
 			]
 			.ValueContent()
@@ -87,88 +87,45 @@ void FCustomizableObjectNodeModifierRemoveMeshBlocksDetails::CustomizeDetails(ID
 			];
 	}
 
-	// Grid size combo
-	{
-		int32 MaxGridSize = 128;
-		LayoutGridSizes.Empty();
+	TArray<FLayoutEditorMeshSection> MeshLayoutsAndLayouts;
 
-		TSharedPtr<FString> CurrentSize;
-		for (int32 Size = 1; Size <= MaxGridSize; Size *= 2)
-		{
-			LayoutGridSizes.Add(MakeShareable(new FString(FString::Printf(TEXT("%d x %d"), Size, Size))));
+	FLayoutEditorMeshSection& MeshSection = MeshLayoutsAndLayouts.AddDefaulted_GetRef();
+	MeshSection.MeshName = MakeShareable(new FString("NameNone"));
+	MeshSection.Layouts.Add(Node->Layout);
 
-			if (Node->Layout->GetGridSize() == FIntPoint(Size))
-			{
-				CurrentSize = LayoutGridSizes.Last();
-			}
-		}
-
-		IDetailGroup& LayoutOptionsGroup = LayoutCategory.AddGroup(TEXT("LayoutOptionsGroup"), LOCTEXT("LayoutGroup", "Layout Group"), false, true);
-		LayoutOptionsGroup.HeaderRow()
-			.NameContent()
-			[
-				SNew(STextBlock)
-					.Text(LOCTEXT("LayoutGridSizeText", "Grid Size"))
-					.Font(DetailBuilder.GetDetailFont())
-			]
-			.ValueContent()
-			[
-				SNew(STextComboBox)
-					.InitiallySelectedItem(CurrentSize)
-					.OptionsSource(&LayoutGridSizes)
-					.OnSelectionChanged(this, &FCustomizableObjectNodeModifierRemoveMeshBlocksDetails::OnGridSizeChanged)
-					.Font(DetailBuilder.GetDetailFont())
-			];
-	}
+	LayoutBlocksEditor = SNew(SCustomizableObjectLayoutEditor)
+		.Node(Node)
+		.MeshSections(MeshLayoutsAndLayouts)
+		.OnPreUpdateLayoutDelegate(this, &FCustomizableObjectNodeModifierRemoveMeshBlocksDetails::OnPreUpdateLayout);
 
-	// Block editor
-	LayoutBlocksEditor = SNew(SCustomizableObjectNodeLayoutBlocksEditor);
+	FCustomizableObjectLayoutEditorDetailsBuilder LayoutEditorBuilder;
+	LayoutEditorBuilder.LayoutEditor = LayoutBlocksEditor;
+	LayoutEditorBuilder.bShowGridSize = true;
 
-	LayoutCategory.AddCustomRow(LOCTEXT("BlocksDetails_BlockInstructions", "BlockInstructions"))
-		[
-			SNew(SBox)
-				.HeightOverride(700.0f)
-				.WidthOverride(700.0f)
-				[
-					LayoutBlocksEditor.ToSharedRef()
-				]
-		];
+	LayoutEditorBuilder.CustomizeDetails(DetailBuilder);
 
-	UpdateLayout();
+	LayoutBlocksEditor->UpdateLayout(Node->Layout);
 }
 
 
 void FCustomizableObjectNodeModifierRemoveMeshBlocksDetails::OnRequiredTagsPropertyChanged()
 {
 	FCustomizableObjectNodeModifierBaseDetails::OnRequiredTagsPropertyChanged();
-	UpdateLayout();
-}
-
-
-void FCustomizableObjectNodeModifierRemoveMeshBlocksDetails::UpdateLayout()
-{
-	// Try to find the parent layout, because we want to show its UVs in the widget
-	UCustomizableObjectLayout* ParentLayout = Node->GetPossibleParentLayout();
-
-	LayoutBlocksEditor->SetCurrentLayout(Node->Layout, ParentLayout);
-
-}
 
-
-void FCustomizableObjectNodeModifierRemoveMeshBlocksDetails::OnGridSizeChanged(TSharedPtr<FString> NewSelection, ESelectInfo::Type SelectInfo)
-{
-	if (!Node->Layout)
+	if (ensure(LayoutBlocksEditor))
 	{
-		return;
+		LayoutBlocksEditor->UpdateLayout(Node->Layout);
 	}
+}
 
-	int32 Size = 1 << LayoutGridSizes.Find(NewSelection);
 
-	if (Node->Layout->GetGridSize().X != Size || Node->Layout->GetGridSize().Y != Size)
+void FCustomizableObjectNodeModifierRemoveMeshBlocksDetails::OnPreUpdateLayout()
+{
+	if (ensure(LayoutBlocksEditor))
 	{
-		Node->Layout->SetGridSize(FIntPoint(Size));
-		Node->MarkPackageDirty();
-		UpdateLayout();
+		// Try to find the parent layout, because we want to show its UVs in the widget
+		UCustomizableObjectLayout* ParentLayout = Node->GetPossibleParentLayout();
+		LayoutBlocksEditor->SetUVsOverride(ParentLayout);
 	}
 }
 
@@ -186,7 +143,8 @@ void FCustomizableObjectNodeModifierRemoveMeshBlocksDetails::OnUVChannelChanged(
 	{
 		Node->ParentLayoutIndex = Index;
 		Node->MarkPackageDirty();
-		UpdateLayout();
+
+		LayoutBlocksEditor->UpdateLayout(Node->Layout);
 	}
 }
 
diff --git a/Engine/Plugins/Experimental/Mutable/Source/CustomizableObjectEditor/Private/MuCOE/Nodes/CustomizableObjectNodeModifierRemoveMeshBlocksDetails.h b/Engine/Plugins/Experimental/Mutable/Source/CustomizableObjectEditor/Private/MuCOE/Nodes/CustomizableObjectNodeModifierRemoveMeshBlocksDetails.h
--- a/Engine/Plugins/Experimental/Mutable/Source/CustomizableObjectEditor/Private/MuCOE/Nodes/CustomizableObjectNodeModifierRemoveMeshBlocksDetails.h
+++ b/Engine/Plugins/Experimental/Mutable/Source/CustomizableObjectEditor/Private/MuCOE/Nodes/CustomizableObjectNodeModifierRemoveMeshBlocksDetails.h
@@ -10,7 +10,7 @@
 class FString;
 class IDetailLayoutBuilder;
 class UCustomizableObjectNodeModifierRemoveMeshBlocks;
-class SCustomizableObjectNodeLayoutBlocksEditor;
+class SCustomizableObjectLayoutEditor;
 
 
 class FCustomizableObjectNodeModifierRemoveMeshBlocksDetails : public FCustomizableObjectNodeModifierBaseDetails
@@ -29,17 +29,15 @@ private:
 	UCustomizableObjectNodeModifierRemoveMeshBlocks* Node;
 
 	// Layout block editor widget
-	TSharedPtr<SCustomizableObjectNodeLayoutBlocksEditor> LayoutBlocksEditor;
+	TSharedPtr<SCustomizableObjectLayoutEditor> LayoutBlocksEditor;
 
 	/** List of available layout grid sizes. */
-	TArray< TSharedPtr< FString > > LayoutGridSizes;
-	TArray< TSharedPtr< FString > > UVChannelOptions;
+	TArray<TSharedPtr< FString>> UVChannelOptions;
 
 	/** Layout Options Callbacks */
-	void OnGridSizeChanged(TSharedPtr<FString> NewSelection, ESelectInfo::Type SelectInfo);
 	void OnUVChannelChanged(TSharedPtr<FString> NewSelection, ESelectInfo::Type SelectInfo);
 
 	/** Reset the layout in the widget to force a refresh. */
-	void UpdateLayout();
+	void OnPreUpdateLayout();
 
 };
diff --git a/Engine/Plugins/Experimental/Mutable/Source/CustomizableObjectEditor/Private/MuCOE/Nodes/CustomizableObjectNodeSkeletalMesh.cpp b/Engine/Plugins/Experimental/Mutable/Source/CustomizableObjectEditor/Private/MuCOE/Nodes/CustomizableObjectNodeSkeletalMesh.cpp
--- a/Engine/Plugins/Experimental/Mutable/Source/CustomizableObjectEditor/Private/MuCOE/Nodes/CustomizableObjectNodeSkeletalMesh.cpp
+++ b/Engine/Plugins/Experimental/Mutable/Source/CustomizableObjectEditor/Private/MuCOE/Nodes/CustomizableObjectNodeSkeletalMesh.cpp
@@ -75,7 +75,7 @@ void UCustomizableObjectNodeSkeletalMesh::AllocateDefaultPins(UCustomizableObjec
 	if (!SkeletalMesh)
 	{
 		UCustomizableObjectNodeSkeletalMeshPinDataMesh* PinData = NewObject<UCustomizableObjectNodeSkeletalMeshPinDataMesh>(this);
-		PinData->Init(-1, -1);
+		PinData->Init(-1, -1, -1, nullptr);
 		
 		DefaultPin = CustomCreatePin(EGPD_Output, Schema->PC_Mesh, FName(SKELETAL_MESH_PIN_NAME), PinData);
 		return;
@@ -109,32 +109,14 @@ void UCustomizableObjectNodeSkeletalMesh::AllocateDefaultPins(UCustomizableObjec
 					FString MeshName = FString::Printf(TEXT("LOD %i - Section %i - Mesh"), LODIndex, SectionIndex);
 
 					UCustomizableObjectNodeSkeletalMeshPinDataMesh* PinData = NewObject<UCustomizableObjectNodeSkeletalMeshPinDataMesh>(this);
-					PinData->Init(LODIndex, SectionIndex);
+					PinData->Init(LODIndex, SectionIndex, ImportedModel->LODModels[LODIndex].NumTexCoords, SkeletalMesh);
 					
 					UEdGraphPin* Pin = CustomCreatePin(EGPD_Output, Schema->PC_Mesh, FName(*MeshName), PinData);
 					Pin->PinFriendlyName = FText::FromString(FString::Printf(TEXT("LOD %i - %s - Mesh"), LODIndex, *SectionFriendlyName));
 					Pin->PinToolTip = MeshName;
 				}
 				
-				// Layout
-				{
-					if (ImportedModel->LODModels.Num() > LODIndex)
-					{
-						const uint32 NumUVs = ImportedModel->LODModels[LODIndex].NumTexCoords;
-						for (uint32 UVIndex = 0; UVIndex < NumUVs; ++UVIndex)
-						{
-							FString PinName = FString::Printf(TEXT("LOD %i - Section %i - UV %i"), LODIndex, SectionIndex, UVIndex);
-							
-							UCustomizableObjectNodeSkeletalMeshPinDataLayout* PinData = NewObject<UCustomizableObjectNodeSkeletalMeshPinDataLayout>(this);
-							PinData->Init(LODIndex, SectionIndex, UVIndex);
-							
-							UEdGraphPin* Pin = CustomCreatePin(EGPD_Input, Schema->PC_Layout, FName(*PinName), PinData);
-							Pin->PinFriendlyName = FText::FromString(FString::Printf(TEXT("LOD %i - %s - UV %i"), LODIndex, *SectionFriendlyName, UVIndex));
-							Pin->PinToolTip = PinName;
-						}
-					}
-				}
-
+				
 				// Images
 				if (MaterialInterface)
 				{
@@ -238,51 +220,12 @@ UTexture2D* UCustomizableObjectNodeSkeletalMesh::FindTextureForPin(const UEdGrap
 }
 
 
-void UCustomizableObjectNodeSkeletalMesh::GetUVChannelForPin(const UEdGraphPin* Pin, TArray<FVector2f>& OutSegments, int32 UVIndex) const
-{
-	check(Pin);
-	
-	if (!SkeletalMesh)
-	{
-		return;
-	}
-
-	int32 LODIndex;
-	int32 SectionIndex;
-	int32 LayoutIndex;
-	GetPinSection(*Pin, LODIndex, SectionIndex, LayoutIndex);
-
-	OutSegments = GetUV(*SkeletalMesh, LODIndex, SectionIndex, UVIndex);
-}
-
-
 TArray<UCustomizableObjectLayout*> UCustomizableObjectNodeSkeletalMesh::GetLayouts(const UEdGraphPin& MeshPin) const
 {
-	TArray<UCustomizableObjectLayout*> Result;
-
 	const UCustomizableObjectNodeSkeletalMeshPinDataMesh* MeshPinData = Cast<UCustomizableObjectNodeSkeletalMeshPinDataMesh>(GetPinData(MeshPin));
 	check(MeshPinData); // Not a mesh pin
 
-	for (const UEdGraphPin* Pin : GetAllNonOrphanPins())
-	{
-		if (const UCustomizableObjectNodeSkeletalMeshPinDataLayout* PinData = Cast<UCustomizableObjectNodeSkeletalMeshPinDataLayout>(GetPinData(*Pin)))
-		{
-			if (PinData->GetLODIndex() == MeshPinData->GetLODIndex() &&
-				PinData->GetSectionIndex() == MeshPinData->GetSectionIndex())
-			{
-				if (const UEdGraphPin* ConnectedPin = FollowInputPin(*Pin))
-				{
-					const UCustomizableObjectNodeLayoutBlocks* LayoutNode = Cast<UCustomizableObjectNodeLayoutBlocks>(ConnectedPin->GetOwningNode());
-					if (LayoutNode && LayoutNode->Layout)
-					{
-						Result.Add(LayoutNode->Layout);
-					}
-				}
-			}
-		}
-	}
-
-	return Result;
+	return MeshPinData->Layouts;
 }
 
 
@@ -310,25 +253,6 @@ UEdGraphPin* UCustomizableObjectNodeSkeletalMesh::GetMeshPin(const int32 LODInde
 }
 
 
-UEdGraphPin* UCustomizableObjectNodeSkeletalMesh::GetLayoutPin(int32 LODIndex, int32 SectionIndex, int32 LayoutIndex) const
-{
-	for (UEdGraphPin* Pin : GetAllNonOrphanPins())
-	{
-		if (const UCustomizableObjectNodeSkeletalMeshPinDataLayout* PinData = Cast<UCustomizableObjectNodeSkeletalMeshPinDataLayout>(GetPinData(*Pin)))
-		{
-			if (PinData->GetLODIndex() == LODIndex &&
-				PinData->GetSectionIndex() == SectionIndex &&
-				PinData->GetUVIndex() == LayoutIndex)
-			{
-				return Pin;
-			}
-		}
-	}
-
-	return nullptr;
-}
-
-
 void UCustomizableObjectNodeSkeletalMesh::GetPinSection(const UEdGraphPin& Pin, int32& OutLODIndex, int32& OutSectionIndex, int32& OutLayoutIndex) const
 {
 	if (const UCustomizableObjectNodeSkeletalMeshPinDataSection* PinData = Cast<UCustomizableObjectNodeSkeletalMeshPinDataSection>(GetPinData(Pin)))
@@ -587,7 +511,7 @@ void UCustomizableObjectNodeSkeletalMesh::BackwardsCompatibleFixup(int32 Customi
 
 				{
 					UCustomizableObjectNodeSkeletalMeshPinDataMesh* PinData = NewObject<UCustomizableObjectNodeSkeletalMeshPinDataMesh>(this);
-					PinData->Init(LODIndex, SectionIndex);
+					PinData->Init(LODIndex, SectionIndex, -1, nullptr);
 				
 					AddPinData(*Section.MeshPinRef.Get(), *PinData);
 				}
@@ -636,7 +560,7 @@ void UCustomizableObjectNodeSkeletalMesh::BackwardsCompatibleFixup(int32 Customi
 				}
 			}
 		}
-			
+
 		ReconstructNode();
 	}
 
@@ -660,11 +584,83 @@ void UCustomizableObjectNodeSkeletalMesh::BackwardsCompatibleFixup(int32 Customi
 		if (const UEdGraphPin* Pin = DefaultPin.Get())
 		{
 			UCustomizableObjectNodeSkeletalMeshPinDataMesh* PinData = NewObject<UCustomizableObjectNodeSkeletalMeshPinDataMesh>(this);
-			PinData->Init(-1, -1);
+			PinData->Init(-1, -1, -1, nullptr);
 
 			AddPinData(*Pin, *PinData);
 		}
 	}
+
+	if (CustomizableObjectCustomVersion == FCustomizableObjectCustomVersion::MoveLayoutToNodeSkeletalMesh)
+	{
+		if (const FSkeletalMeshModel* ImportedModel = SkeletalMesh ? SkeletalMesh->GetImportedModel() : nullptr)
+		{
+			TArray<UEdGraphPin*> PinsToDelete;
+
+			TArray<UEdGraphPin*> NonOrphanPins = GetAllNonOrphanPins();
+			for (UEdGraphPin* Pin : NonOrphanPins)
+			{
+				UCustomizableObjectNodeSkeletalMeshPinDataMesh* MeshPinData = Cast<UCustomizableObjectNodeSkeletalMeshPinDataMesh>(GetPinData(*Pin));
+				if (!MeshPinData)
+				{
+					continue;
+				}
+
+				const int32 LODIndex = MeshPinData->GetLODIndex();
+				const int32 SectionIndex = MeshPinData->GetSectionIndex();
+				if (!ImportedModel->LODModels.IsValidIndex(LODIndex))
+				{
+					continue;
+				}
+
+				const int32 NumTexCoords = ImportedModel->LODModels[LODIndex].NumTexCoords;
+				MeshPinData->Layouts.SetNum(NumTexCoords);
+
+				for (int32 UVIndex = 0; UVIndex < NumTexCoords; ++UVIndex)
+				{
+					MeshPinData->Layouts[UVIndex] = NewObject<UCustomizableObjectLayout>(this);
+					MeshPinData->Layouts[UVIndex]->SetLayout(LODIndex, SectionIndex, UVIndex);
+					MeshPinData->Layouts[UVIndex]->SetIgnoreWarningsLOD(0);
+				}
+
+				for (UEdGraphPin* OtherPin : NonOrphanPins)
+				{
+					UCustomizableObjectNodeSkeletalMeshPinDataLayout* LayoutPinData = Cast<UCustomizableObjectNodeSkeletalMeshPinDataLayout>(GetPinData(*OtherPin));
+					if (!LayoutPinData ||
+						LayoutPinData->GetLODIndex() != LODIndex ||
+						LayoutPinData->GetSectionIndex() != SectionIndex ||
+						!MeshPinData->Layouts.IsValidIndex(LayoutPinData->GetUVIndex()))
+					{
+						continue;
+					}
+
+					if (const UEdGraphPin* ConnectedPin = FollowInputPin(*OtherPin))
+					{
+						const UCustomizableObjectNodeLayoutBlocks* LayoutNode = Cast<UCustomizableObjectNodeLayoutBlocks>(ConnectedPin->GetOwningNode());
+						if (LayoutNode && LayoutNode->Layout)
+						{
+							UCustomizableObjectLayout* Layout = MeshPinData->Layouts[LayoutPinData->GetUVIndex()];
+							Layout->Blocks = LayoutNode->Layout->Blocks;
+							Layout->SetGridSize(LayoutNode->Layout->GetGridSize());
+							Layout->SetMaxGridSize(LayoutNode->Layout->GetMaxGridSize());
+							Layout->SetIgnoreVertexLayoutWarnings(LayoutNode->Layout->GetIgnoreVertexLayoutWarnings());
+							Layout->SetIgnoreWarningsLOD(LayoutNode->Layout->GetFirstLODToIgnoreWarnings());
+							Layout->PackingStrategy = LayoutNode->Layout->PackingStrategy;
+							Layout->AutomaticBlocksStrategy = LayoutNode->Layout->AutomaticBlocksStrategy;
+							Layout->AutomaticBlocksMergeStrategy = LayoutNode->Layout->AutomaticBlocksMergeStrategy;
+							Layout->BlockReductionMethod = LayoutNode->Layout->BlockReductionMethod;
+						}
+					}
+
+					PinsToDelete.Add(OtherPin);
+				}
+			}
+
+			for (UEdGraphPin* Pin : PinsToDelete)
+			{
+				CustomRemovePin(*Pin);
+			}
+		}
+	}
 }
 
 
@@ -984,6 +980,45 @@ bool UCustomizableObjectNodeSkeletalMeshPinDataSection::Equals(const UCustomizab
 }
 
 
+void UCustomizableObjectNodeSkeletalMeshPinDataMesh::Copy(const UCustomizableObjectNodePinData& Other)
+{
+	if (const UCustomizableObjectNodeSkeletalMeshPinDataMesh* PinDataOldPin = Cast<UCustomizableObjectNodeSkeletalMeshPinDataMesh>(&Other))
+	{
+		for (UCustomizableObjectLayout* OldLayout : PinDataOldPin->Layouts)
+		{
+			if (!OldLayout)
+			{
+				continue;
+			}
+
+			const int32 UVChannel = OldLayout->GetUVChannel();
+			if (Layouts.IsValidIndex(UVChannel))
+			{
+				Layouts[UVChannel] = OldLayout;
+			}
+		}
+	}
+}
+
+void UCustomizableObjectNodeSkeletalMeshPinDataMesh::Init(int32 InLODIndex, int32 InSectionIndex, int32 NumTexCoords, USkeletalMesh* Mesh)
+{
+	Super::Init(InLODIndex, InSectionIndex);
+
+	if (NumTexCoords > 0)
+	{
+		UObject* Outer = GetOuter();
+
+		Layouts.SetNum(NumTexCoords);
+
+		for (int32 Index = 0; Index < NumTexCoords; ++Index)
+		{
+			Layouts[Index] = NewObject<UCustomizableObjectLayout>(Outer);
+			Layouts[Index]->SetLayout(InLODIndex, InSectionIndex, Index);
+		}
+	}
+}
+
+
 void UCustomizableObjectNodeSkeletalMeshPinDataImage::Init(int32 InLODIndex, int32 InSectionIndex, FGuid InTextureParameterId)
 {
 	Super::Init(InLODIndex, InSectionIndex);
diff --git a/Engine/Plugins/Experimental/Mutable/Source/CustomizableObjectEditor/Private/MuCOE/Nodes/CustomizableObjectNodeSkeletalMeshDetails.cpp b/Engine/Plugins/Experimental/Mutable/Source/CustomizableObjectEditor/Private/MuCOE/Nodes/CustomizableObjectNodeSkeletalMeshDetails.cpp
--- a/Engine/Plugins/Experimental/Mutable/Source/CustomizableObjectEditor/Private/MuCOE/Nodes/CustomizableObjectNodeSkeletalMeshDetails.cpp
+++ b/Engine/Plugins/Experimental/Mutable/Source/CustomizableObjectEditor/Private/MuCOE/Nodes/CustomizableObjectNodeSkeletalMeshDetails.cpp
@@ -8,6 +8,7 @@
 #include "MuCOE/Nodes/CustomizableObjectNodeSkeletalMesh.h"
 #include "MuCOE/PinViewer/SPinViewer.h"
 #include "MuCOE/SCustomizableObjectNodeSkeletalMeshRTMorphSelector.h"
+#include "MuCOE/SCustomizableObjectLayoutEditor.h"
 
 #define LOCTEXT_NAMESPACE "CustomizableObjectNodeMaterialDetails"
 
@@ -51,6 +52,52 @@ void FCustomizableObjectNodeSkeletalMeshDetails::CustomizeDetails(IDetailLayoutB
 
 	TSharedRef<IPropertyHandle> SkeletalMeshProperty = DetailBuilder.GetProperty(GET_MEMBER_NAME_CHECKED(UCustomizableObjectNodeSkeletalMesh, SkeletalMesh));
 	SkeletalMeshProperty->SetOnPropertyValueChanged(FSimpleDelegate::CreateSP(MorphSelector.Get(), &SCustomizableObjectNodeSkeletalMeshRTMorphSelector::UpdateWidget));
+
+	TArray<FLayoutEditorMeshSection> MeshSectionsAndLayouts;
+	GenerateMeshSectionOptions(MeshSectionsAndLayouts);
+
+	TSharedPtr<SCustomizableObjectLayoutEditor> LayoutBlocksEditor = SNew(SCustomizableObjectLayoutEditor)
+		.Node(Node)
+		.MeshSections(MeshSectionsAndLayouts);
+
+	FCustomizableObjectLayoutEditorDetailsBuilder LayoutEditorBuilder;
+	LayoutEditorBuilder.LayoutEditor = LayoutBlocksEditor;
+	LayoutEditorBuilder.bShowLayoutSelector = true;
+	LayoutEditorBuilder.bShowPackagingStrategy = true;
+	LayoutEditorBuilder.bShowAutomaticGenerationSettings = true;
+	LayoutEditorBuilder.bShowGridSize = true;
+	LayoutEditorBuilder.bShowMaxGridSize = true;
+	LayoutEditorBuilder.bShowReductionMethods = true;
+	LayoutEditorBuilder.bShowWarningSettings = true;
+
+	LayoutEditorBuilder.CustomizeDetails(DetailBuilder);
+
+	LayoutBlocksEditor->UpdateLayout(nullptr);
+}
+
+
+void FCustomizableObjectNodeSkeletalMeshDetails::GenerateMeshSectionOptions(TArray<FLayoutEditorMeshSection>& OutMeshSections)
+{
+	OutMeshSections.Empty();
+
+	if (!Node)
+	{
+		return;
+	}
+
+	for (const UEdGraphPin* Pin : Node->GetAllNonOrphanPins())
+	{
+		if (const UCustomizableObjectNodeSkeletalMeshPinDataMesh* PinData = Cast<UCustomizableObjectNodeSkeletalMeshPinDataMesh>(Node->GetPinData(*Pin)))
+		{
+			FLayoutEditorMeshSection& MeshSection = OutMeshSections.AddDefaulted_GetRef();
+			MeshSection.MeshName = MakeShareable(new FString(Pin->PinFriendlyName.ToString()));
+
+			for (UCustomizableObjectLayout* Layout : PinData->Layouts)
+			{
+				MeshSection.Layouts.Add(Layout);
+			}
+		}
+	}
 }
 
 
diff --git a/Engine/Plugins/Experimental/Mutable/Source/CustomizableObjectEditor/Private/MuCOE/Nodes/CustomizableObjectNodeSkeletalMeshDetails.h b/Engine/Plugins/Experimental/Mutable/Source/CustomizableObjectEditor/Private/MuCOE/Nodes/CustomizableObjectNodeSkeletalMeshDetails.h
--- a/Engine/Plugins/Experimental/Mutable/Source/CustomizableObjectEditor/Private/MuCOE/Nodes/CustomizableObjectNodeSkeletalMeshDetails.h
+++ b/Engine/Plugins/Experimental/Mutable/Source/CustomizableObjectEditor/Private/MuCOE/Nodes/CustomizableObjectNodeSkeletalMeshDetails.h
@@ -6,6 +6,7 @@
 #include "IDetailCustomization.h"
 
 class IDetailLayoutBuilder;
+struct FLayoutEditorMeshSection;
 
 
 class FCustomizableObjectNodeSkeletalMeshDetails : public FCustomizableObjectNodeDetails
@@ -19,6 +20,8 @@ public:
 
 private:
 
+	void GenerateMeshSectionOptions(TArray<FLayoutEditorMeshSection>& OutMeshSections);
+
 	// Pointer to the node represented in this details
 	class UCustomizableObjectNodeSkeletalMesh* Node;
 };
diff --git a/Engine/Plugins/Experimental/Mutable/Source/CustomizableObjectEditor/Private/MuCOE/Nodes/CustomizableObjectNodeStaticMesh.cpp b/Engine/Plugins/Experimental/Mutable/Source/CustomizableObjectEditor/Private/MuCOE/Nodes/CustomizableObjectNodeStaticMesh.cpp
--- a/Engine/Plugins/Experimental/Mutable/Source/CustomizableObjectEditor/Private/MuCOE/Nodes/CustomizableObjectNodeStaticMesh.cpp
+++ b/Engine/Plugins/Experimental/Mutable/Source/CustomizableObjectEditor/Private/MuCOE/Nodes/CustomizableObjectNodeStaticMesh.cpp
@@ -252,24 +252,6 @@ UTexture2D* UCustomizableObjectNodeStaticMesh::FindTextureForPin(const UEdGraphP
 }
 
 
-void UCustomizableObjectNodeStaticMesh::GetUVChannelForPin(const UEdGraphPin* Pin, TArray<FVector2f>& OutSegments, int32 UVIndex) const
-{
-	check(Pin);
-	
-	if (!StaticMesh)
-	{
-		return;
-	}
-
-	int32 LODIndex;
-	int32 SectionIndex;
-	int32 LayoutIndex;
-	GetPinSection(*Pin, LODIndex, SectionIndex, LayoutIndex);
-	
-	OutSegments = GetUV(*StaticMesh, LODIndex, SectionIndex, UVIndex);
-}
-
-
 TArray<class UCustomizableObjectLayout*> UCustomizableObjectNodeStaticMesh::GetLayouts(const UEdGraphPin& OutPin) const
 {
 	TArray<class UCustomizableObjectLayout*> Result;
@@ -320,23 +302,6 @@ UEdGraphPin* UCustomizableObjectNodeStaticMesh::GetMeshPin(const int32 LODIndex,
 }
 
 
-UEdGraphPin* UCustomizableObjectNodeStaticMesh::GetLayoutPin(int32 LODIndex, int32 SectionIndex, int32 LayoutIndex) const
-{
-	check(LayoutIndex == 1); // Multiple UVs not supported on Static Mesh Node.
-
-	if (LODIndex < LODs.Num())
-	{
-		if (const FCustomizableObjectNodeStaticMeshLOD& LOD = LODs[LODIndex];
-			SectionIndex < LOD.Materials.Num())
-		{
-			return LOD.Materials[SectionIndex].LayoutPinRef.Get();
-		}
-	}
-
-	return nullptr;
-}
-
-
 bool UCustomizableObjectNodeStaticMesh::IsPinRelevant(const UEdGraphPin* Pin) const
 {
 	const UEdGraphSchema_CustomizableObject* Schema = GetDefault<UEdGraphSchema_CustomizableObject>();
diff --git a/Engine/Plugins/Experimental/Mutable/Source/CustomizableObjectEditor/Private/MuCOE/Nodes/CustomizableObjectNodeTable.cpp b/Engine/Plugins/Experimental/Mutable/Source/CustomizableObjectEditor/Private/MuCOE/Nodes/CustomizableObjectNodeTable.cpp
--- a/Engine/Plugins/Experimental/Mutable/Source/CustomizableObjectEditor/Private/MuCOE/Nodes/CustomizableObjectNodeTable.cpp
+++ b/Engine/Plugins/Experimental/Mutable/Source/CustomizableObjectEditor/Private/MuCOE/Nodes/CustomizableObjectNodeTable.cpp
@@ -18,7 +18,6 @@
 #include "MuCOE/GraphTraversal.h"
 #include "MuCOE/MutableUtils.h"
 #include "MuCOE/Nodes/CustomizableObjectNodeMaterial.h"
-#include "MuCOE/SCustomizableObjectNodeLayoutBlocksEditor.h"
 #include "MuCOE/UnrealEditorPortabilityHelpers.h"
 #include "Rendering/SkeletalMeshLODModel.h"
 #include "Rendering/SkeletalMeshModel.h"
@@ -586,7 +585,7 @@ void UCustomizableObjectNodeTable::GenerateMeshPins(UObject* Mesh, const FString
 						{
 							PinData->Layouts.Add(Layout);
 
-							Layout->SetLayout(SkeletalMesh, LODIndex, MatIndex, LayoutIndex);
+							Layout->SetLayout(LODIndex, MatIndex, LayoutIndex);
 							Layout->SetLayoutName(LayoutName);
 						}
 					}
@@ -918,74 +917,23 @@ UTexture2D* UCustomizableObjectNodeTable::FindReferenceTextureParameter(const UE
 }
 
 
-void UCustomizableObjectNodeTable::GetUVChannel(const UCustomizableObjectLayout* CurrentLayout, TArray<FVector2f>& OutSegments) const
+UObject* UCustomizableObjectNodeTable::GetDefaultMeshForLayout(const UCustomizableObjectLayout* InLayout) const
 {
-	FString ColumnName;
-
 	for (const UEdGraphPin* Pin : Pins)
 	{
-		const UCustomizableObjectNodeTableMeshPinData* PinData = Cast<UCustomizableObjectNodeTableMeshPinData > (GetPinData(*Pin));
-		
-		if (!PinData)
-		{
-			continue;
-		}
-
-		bool bFound = false;
-
-		for (const UCustomizableObjectLayout* Layout : PinData->Layouts)
+		if (const UCustomizableObjectNodeTableMeshPinData* MeshPinData = Cast<UCustomizableObjectNodeTableMeshPinData >(GetPinData(*Pin)))
 		{
-			if (Layout == CurrentLayout)
+			for (const UCustomizableObjectLayout* Layout : MeshPinData->Layouts)
 			{
-				ColumnName = PinData->ColumnName;
-				bFound = true;
-				break;
+				if (Layout == InLayout)
+				{
+					return GetColumnDefaultAssetByType<UObject>(MeshPinData->ColumnName);
+				}
 			}
 		}
-
-		if (bFound)
-		{
-			break;
-		}
 	}
 
-	const int32 LODIndex = CurrentLayout->GetLOD();
-	const int32 MaterialIndex = CurrentLayout->GetMaterial();
-	const int32 LayoutIndex = CurrentLayout->GetUVChannel();
-
-	if (const USkeletalMesh* SkeletalMesh = GetColumnDefaultAssetByType<USkeletalMesh>(ColumnName))
-	{
-		OutSegments = GetUV(*SkeletalMesh, LODIndex, MaterialIndex, LayoutIndex);
-	}
-	else if (const UStaticMesh* StaticMesh = GetColumnDefaultAssetByType<UStaticMesh>(ColumnName))
-	{
-		OutSegments = GetUV(*StaticMesh, LODIndex, MaterialIndex, LayoutIndex);
-	}
-}
-
-
-void UCustomizableObjectNodeTable::GetUVChannelForPin(const UEdGraphPin* Pin, TArray<FVector2f>& OutSegments, int32 UVChannel) const
-{
-	const UCustomizableObjectNodeTableMeshPinData* PinData = Cast<UCustomizableObjectNodeTableMeshPinData >(GetPinData(*Pin));
-
-	for (int32 LayoutIndex = 0; LayoutIndex < PinData->Layouts.Num(); ++LayoutIndex)
-	{
-		if (PinData->Layouts[LayoutIndex]->GetUVChannel() == UVChannel)
-		{
-			// Getting LOD and Mat index from Layouts
-			const uint32 LODIndex = PinData->Layouts[LayoutIndex]->GetLOD();
-			const uint32 MaterialIndex = PinData->Layouts[LayoutIndex]->GetMaterial();
-
-			if (const USkeletalMesh* SkeletalMesh = GetColumnDefaultAssetByType<USkeletalMesh>(PinData->ColumnName))
-			{
-				OutSegments.Append(GetUV(*SkeletalMesh, LODIndex, MaterialIndex, LayoutIndex));
-			}
-			else if (const UStaticMesh* StaticMesh = GetColumnDefaultAssetByType<UStaticMesh>(PinData->ColumnName))
-			{
-				OutSegments.Append(GetUV(*StaticMesh, LODIndex, MaterialIndex, LayoutIndex));
-			}
-		}
-	}
+	return nullptr;
 }
 
 
diff --git a/Engine/Plugins/Experimental/Mutable/Source/CustomizableObjectEditor/Private/MuCOE/Nodes/CustomizableObjectNodeTableDetails.cpp b/Engine/Plugins/Experimental/Mutable/Source/CustomizableObjectEditor/Private/MuCOE/Nodes/CustomizableObjectNodeTableDetails.cpp
--- a/Engine/Plugins/Experimental/Mutable/Source/CustomizableObjectEditor/Private/MuCOE/Nodes/CustomizableObjectNodeTableDetails.cpp
+++ b/Engine/Plugins/Experimental/Mutable/Source/CustomizableObjectEditor/Private/MuCOE/Nodes/CustomizableObjectNodeTableDetails.cpp
@@ -13,7 +13,7 @@
 #include "Layout/Visibility.h"
 #include "MuCOE/CustomizableObjectLayout.h"
 #include "MuCOE/Nodes/CustomizableObjectNodeTable.h"
-#include "MuCOE/SCustomizableObjectNodeLayoutBlocksEditor.h"
+#include "MuCOE/SCustomizableObjectLayoutEditor.h"
 #include "MuCOE/UnrealEditorPortabilityHelpers.h"
 #include "SSearchableComboBox.h"
 #include "Styling/SlateColor.h"
@@ -37,12 +37,12 @@ TSharedRef<IDetailCustomization> FCustomizableObjectNodeTableDetails::MakeInstan
 void FCustomizableObjectNodeTableDetails::CustomizeDetails(const TSharedPtr<IDetailLayoutBuilder>& DetailBuilder)
 {
 	IDetailCustomization::CustomizeDetails(DetailBuilder);
-	
+
 	Node = 0;
 	DetailBuilderPtr = DetailBuilder;
 
 	const IDetailsView* DetailsView = DetailBuilder->GetDetailsView();
-	
+
 	if (DetailsView->GetSelectedObjects().Num())
 	{
 		Node = Cast<UCustomizableObjectNodeTable>(DetailsView->GetSelectedObjects()[0].Get());
@@ -57,7 +57,6 @@ void FCustomizableObjectNodeTableDetails::CustomizeDetails(const TSharedPtr<IDet
 		DetailBuilder->HideProperty("ParamUIMetadataColumn");
 		DetailBuilder->HideProperty("ThumbnailColumn");
 		IDetailCategoryBuilder& AnimationCategory = DetailBuilder->EditCategory("AnimationProperties");
-		IDetailCategoryBuilder& LayoutCategory = DetailBuilder->EditCategory("DefaultMeshLayoutEditor");
 
 		// Attaching the Posrecontruct delegate to force a refresh of the details
 		Node->PostReconstructNodeDelegate.AddSP(this, &FCustomizableObjectNodeTableDetails::OnNodePinValueChanged);
@@ -69,65 +68,65 @@ void FCustomizableObjectNodeTableDetails::CustomizeDetails(const TSharedPtr<IDet
 
 		CustomizableObjectCategory.AddProperty("ParameterName");
 
-		CompilationRestrictionsCategory.AddCustomRow(LOCTEXT("VersionColumn_Selector","VersionColumn"))
-		.NameContent()
-		[
-			SNew(STextBlock)
-			.Text(LOCTEXT("VersionColumn_SelectorText","Version Column"))
-			.ToolTipText(LOCTEXT("VersionColumn_SelectorTooltip","Select the column that contains the version of each row."))
-			.Font(IDetailLayoutBuilder::GetDetailFont())
-		]
-		.ValueContent()
-		[
-			SAssignNew(VersionColumnsComboBox,STextComboBox)
-			.InitiallySelectedItem(CurrentVersionColumn)
-			.OptionsSource(&VersionColumnsOptionNames)
-			.OnComboBoxOpening(this, &FCustomizableObjectNodeTableDetails::OnOpenVersionColumnComboBox)
-			.OnSelectionChanged(this, &FCustomizableObjectNodeTableDetails::OnVersionColumnComboBoxSelectionChanged)
-			.Font(IDetailLayoutBuilder::GetDetailFont())
-			.ColorAndOpacity(this, &FCustomizableObjectNodeTableDetails::GetVersionColumnComboBoxTextColor, &VersionColumnsOptionNames)
-		]
-		.OverrideResetToDefault(FResetToDefaultOverride::Create(FSimpleDelegate::CreateSP(this, &FCustomizableObjectNodeTableDetails::OnVersionColumnComboBoxSelectionReset)));
-
-		UICategory.AddCustomRow(LOCTEXT("MutableUIMetadataColumn_Selector","MutableUIMetadataColumn"))
-		.NameContent()
-		[
-			SNew(STextBlock)
-			.Text(LOCTEXT("MutableUIMetadataColumn_SelectorText","Options UI Metadata Column"))
-			.ToolTipText(LOCTEXT("MutableUIMetadataColumn_SelectorTooltip","Select a column that contains a Parameter UI Metadata for each Parameter Option (table row)."))
-			.Font(IDetailLayoutBuilder::GetDetailFont())
-		]
-		.ValueContent()
-		[
-			SAssignNew(MutableMetaDataComboBox,STextComboBox)
-			.InitiallySelectedItem(CurrentMutableMetadataColumn)
-			.OptionsSource(&MutableMetaDataColumnsOptionNames)
-			.OnComboBoxOpening(this, &FCustomizableObjectNodeTableDetails::OnOpenMutableMetadataComboBox)
-			.OnSelectionChanged(this, &FCustomizableObjectNodeTableDetails::OnMutableMetaDataColumnComboBoxSelectionChanged)
-			.Font(IDetailLayoutBuilder::GetDetailFont())
-			.ColorAndOpacity(this, &FCustomizableObjectNodeTableDetails::GetComboBoxTextColor, &MutableMetaDataColumnsOptionNames, Node->ParamUIMetadataColumn)
-		]
-		.OverrideResetToDefault(FResetToDefaultOverride::Create(FSimpleDelegate::CreateSP(this, &FCustomizableObjectNodeTableDetails::OnMutableMetaDataColumnComboBoxSelectionReset)));
-		
+		CompilationRestrictionsCategory.AddCustomRow(LOCTEXT("VersionColumn_Selector", "VersionColumn"))
+			.NameContent()
+			[
+				SNew(STextBlock)
+					.Text(LOCTEXT("VersionColumn_SelectorText", "Version Column"))
+					.ToolTipText(LOCTEXT("VersionColumn_SelectorTooltip", "Select the column that contains the version of each row."))
+					.Font(IDetailLayoutBuilder::GetDetailFont())
+			]
+			.ValueContent()
+			[
+				SAssignNew(VersionColumnsComboBox, STextComboBox)
+					.InitiallySelectedItem(CurrentVersionColumn)
+					.OptionsSource(&VersionColumnsOptionNames)
+					.OnComboBoxOpening(this, &FCustomizableObjectNodeTableDetails::OnOpenVersionColumnComboBox)
+					.OnSelectionChanged(this, &FCustomizableObjectNodeTableDetails::OnVersionColumnComboBoxSelectionChanged)
+					.Font(IDetailLayoutBuilder::GetDetailFont())
+					.ColorAndOpacity(this, &FCustomizableObjectNodeTableDetails::GetVersionColumnComboBoxTextColor, &VersionColumnsOptionNames)
+			]
+			.OverrideResetToDefault(FResetToDefaultOverride::Create(FSimpleDelegate::CreateSP(this, &FCustomizableObjectNodeTableDetails::OnVersionColumnComboBoxSelectionReset)));
+
+		UICategory.AddCustomRow(LOCTEXT("MutableUIMetadataColumn_Selector", "MutableUIMetadataColumn"))
+			.NameContent()
+			[
+				SNew(STextBlock)
+					.Text(LOCTEXT("MutableUIMetadataColumn_SelectorText", "Options UI Metadata Column"))
+					.ToolTipText(LOCTEXT("MutableUIMetadataColumn_SelectorTooltip", "Select a column that contains a Parameter UI Metadata for each Parameter Option (table row)."))
+					.Font(IDetailLayoutBuilder::GetDetailFont())
+			]
+			.ValueContent()
+			[
+				SAssignNew(MutableMetaDataComboBox, STextComboBox)
+					.InitiallySelectedItem(CurrentMutableMetadataColumn)
+					.OptionsSource(&MutableMetaDataColumnsOptionNames)
+					.OnComboBoxOpening(this, &FCustomizableObjectNodeTableDetails::OnOpenMutableMetadataComboBox)
+					.OnSelectionChanged(this, &FCustomizableObjectNodeTableDetails::OnMutableMetaDataColumnComboBoxSelectionChanged)
+					.Font(IDetailLayoutBuilder::GetDetailFont())
+					.ColorAndOpacity(this, &FCustomizableObjectNodeTableDetails::GetComboBoxTextColor, &MutableMetaDataColumnsOptionNames, Node->ParamUIMetadataColumn)
+			]
+			.OverrideResetToDefault(FResetToDefaultOverride::Create(FSimpleDelegate::CreateSP(this, &FCustomizableObjectNodeTableDetails::OnMutableMetaDataColumnComboBoxSelectionReset)));
+
 		UICategory.AddCustomRow(LOCTEXT("ThumbnailColumn_Selector", "ThumbnailColumn"))
-		.NameContent()
-		[
-			SNew(STextBlock)
-			.Text(LOCTEXT("ThumbnailColumn_SelectorText", "Options Thumbnail Column"))
-			.ToolTipText(LOCTEXT("ThumbnailColumn_SelectorTooltip", "Select a column that contains the assets to use its thumbnails as Option thumbnails."))
-			.Font(IDetailLayoutBuilder::GetDetailFont())
-		]
-		.ValueContent()
-		[
-			SAssignNew(ThumbnailComboBox,STextComboBox)
-			.InitiallySelectedItem(CurrentThumbnailColumn)
-			.OptionsSource(&ThumbnailColumnOptionNames)
-			.OnComboBoxOpening(this, &FCustomizableObjectNodeTableDetails::OnOpenThumbnailComboBox)
-			.OnSelectionChanged(this, &FCustomizableObjectNodeTableDetails::OnThumbnailColumnComboBoxSelectionChanged)
-			.Font(IDetailLayoutBuilder::GetDetailFont())
-			.ColorAndOpacity(this, &FCustomizableObjectNodeTableDetails::GetComboBoxTextColor, &MutableMetaDataColumnsOptionNames, Node->ThumbnailColumn)
-		]
-		.OverrideResetToDefault(FResetToDefaultOverride::Create(FSimpleDelegate::CreateSP(this, &FCustomizableObjectNodeTableDetails::OnThumbnailColumnComboBoxSelectionReset)));
+			.NameContent()
+			[
+				SNew(STextBlock)
+					.Text(LOCTEXT("ThumbnailColumn_SelectorText", "Options Thumbnail Column"))
+					.ToolTipText(LOCTEXT("ThumbnailColumn_SelectorTooltip", "Select a column that contains the assets to use its thumbnails as Option thumbnails."))
+					.Font(IDetailLayoutBuilder::GetDetailFont())
+			]
+			.ValueContent()
+			[
+				SAssignNew(ThumbnailComboBox, STextComboBox)
+					.InitiallySelectedItem(CurrentThumbnailColumn)
+					.OptionsSource(&ThumbnailColumnOptionNames)
+					.OnComboBoxOpening(this, &FCustomizableObjectNodeTableDetails::OnOpenThumbnailComboBox)
+					.OnSelectionChanged(this, &FCustomizableObjectNodeTableDetails::OnThumbnailColumnComboBoxSelectionChanged)
+					.Font(IDetailLayoutBuilder::GetDetailFont())
+					.ColorAndOpacity(this, &FCustomizableObjectNodeTableDetails::GetComboBoxTextColor, &MutableMetaDataColumnsOptionNames, Node->ThumbnailColumn)
+			]
+			.OverrideResetToDefault(FResetToDefaultOverride::Create(FSimpleDelegate::CreateSP(this, &FCustomizableObjectNodeTableDetails::OnThumbnailColumnComboBoxSelectionReset)));
 
 
 
@@ -135,273 +134,108 @@ void FCustomizableObjectNodeTableDetails::CustomizeDetails(const TSharedPtr<IDet
 
 		// Mesh Column Selector
 		AnimationCategory.AddCustomRow(LOCTEXT("AnimationProperties", "Animation Properties"))
-		.NameContent()
-		[
-			SNew(STextBlock)
-			.Text(LOCTEXT("AnimMeshColumnText", "Mesh Column: "))
-			.ToolTipText(LOCTEXT("AnimMeshColumnTooltip", "Select a mesh column from the Data Table to edit its animation options (Applied to all LODs)."))
-			.Font(IDetailLayoutBuilder::GetDetailFont())
-		]
-		.ValueContent()
-		.HAlign(EHorizontalAlignment::HAlign_Fill)
-		[
-			SAssignNew(AnimMeshColumnComboBox, STextComboBox)
-			.OptionsSource(&AnimMeshColumnOptionNames)
-			.InitiallySelectedItem(AnimMeshColumnOptionNames[0])
-			.OnSelectionChanged(this, &FCustomizableObjectNodeTableDetails::OnAnimMeshColumnComboBoxSelectionChanged)
-			.Font(IDetailLayoutBuilder::GetDetailFont())
-		]
-		.OverrideResetToDefault(FResetToDefaultOverride::Create(FSimpleDelegate::CreateSP(this, &FCustomizableObjectNodeTableDetails::OnAnimMeshCustomRowResetButtonClicked)));
+			.NameContent()
+			[
+				SNew(STextBlock)
+					.Text(LOCTEXT("AnimMeshColumnText", "Mesh Column: "))
+					.ToolTipText(LOCTEXT("AnimMeshColumnTooltip", "Select a mesh column from the Data Table to edit its animation options (Applied to all LODs)."))
+					.Font(IDetailLayoutBuilder::GetDetailFont())
+			]
+			.ValueContent()
+			.HAlign(EHorizontalAlignment::HAlign_Fill)
+			[
+				SAssignNew(AnimMeshColumnComboBox, STextComboBox)
+					.OptionsSource(&AnimMeshColumnOptionNames)
+					.InitiallySelectedItem(AnimMeshColumnOptionNames[0])
+					.OnSelectionChanged(this, &FCustomizableObjectNodeTableDetails::OnAnimMeshColumnComboBoxSelectionChanged)
+					.Font(IDetailLayoutBuilder::GetDetailFont())
+			]
+			.OverrideResetToDefault(FResetToDefaultOverride::Create(FSimpleDelegate::CreateSP(this, &FCustomizableObjectNodeTableDetails::OnAnimMeshCustomRowResetButtonClicked)));
 
 
 		// AnimBP Column Selector
 		AnimationCategory.AddCustomRow(LOCTEXT("AnimationProperties", "Animation Properties"))
-		.NameContent()
-		[
-			SNew(STextBlock)
-			.Text(LOCTEXT("AnimBPText", "Animation Blueprint Column: "))
-			.ToolTipText(LOCTEXT("AnimBlueprintColumnTooltip", "Select an animation blueprint column from the Data Table that will be applied to the mesh selected"))
-			.Font(IDetailLayoutBuilder::GetDetailFont())
-		]
-		.ValueContent()
-		.HAlign(EHorizontalAlignment::HAlign_Fill)
-		[
-			SAssignNew(AnimComboBox, STextComboBox)
-			.OptionsSource(&AnimOptionNames)
-			.OnSelectionChanged(this, &FCustomizableObjectNodeTableDetails::OnAnimInstanceComboBoxSelectionChanged)
-			.Font(IDetailLayoutBuilder::GetDetailFont())
-		]
-		.OverrideResetToDefault(FResetToDefaultOverride::Create(FSimpleDelegate::CreateSP(this, &FCustomizableObjectNodeTableDetails::OnAnimCustomRowResetButtonClicked, EAnimColumnType::EACT_BluePrintColumn)))
-		.Visibility(TAttribute<EVisibility>(this, &FCustomizableObjectNodeTableDetails::AnimWidgetsVisibility));
+			.NameContent()
+			[
+				SNew(STextBlock)
+					.Text(LOCTEXT("AnimBPText", "Animation Blueprint Column: "))
+					.ToolTipText(LOCTEXT("AnimBlueprintColumnTooltip", "Select an animation blueprint column from the Data Table that will be applied to the mesh selected"))
+					.Font(IDetailLayoutBuilder::GetDetailFont())
+			]
+			.ValueContent()
+			.HAlign(EHorizontalAlignment::HAlign_Fill)
+			[
+				SAssignNew(AnimComboBox, STextComboBox)
+					.OptionsSource(&AnimOptionNames)
+					.OnSelectionChanged(this, &FCustomizableObjectNodeTableDetails::OnAnimInstanceComboBoxSelectionChanged)
+					.Font(IDetailLayoutBuilder::GetDetailFont())
+			]
+			.OverrideResetToDefault(FResetToDefaultOverride::Create(FSimpleDelegate::CreateSP(this, &FCustomizableObjectNodeTableDetails::OnAnimCustomRowResetButtonClicked, EAnimColumnType::EACT_BluePrintColumn)))
+			.Visibility(TAttribute<EVisibility>(this, &FCustomizableObjectNodeTableDetails::AnimWidgetsVisibility));
 
 
 		// AnimSlot Column Selector
 		AnimationCategory.AddCustomRow(LOCTEXT("AnimationProperties", "Animation Properties"))
-		.NameContent()
-		[
-			SNew(STextBlock)
-			.Text(LOCTEXT("AnimSlotText", "Animation Slot Column: "))
-			.ToolTipText(LOCTEXT("AnimSlotColumnTooltip", "Select an animation slot column from the Data Table that will set to the slot value of the animation blueprint"))
-			.Font(IDetailLayoutBuilder::GetDetailFont())
-		]
-		.ValueContent()
-		.HAlign(EHorizontalAlignment::HAlign_Fill)
-		[
-			SAssignNew(AnimSlotComboBox, STextComboBox)
-			.OptionsSource(&AnimSlotOptionNames)
-			.OnSelectionChanged(this, &FCustomizableObjectNodeTableDetails::OnAnimSlotComboBoxSelectionChanged)
-			.Font(IDetailLayoutBuilder::GetDetailFont())
-		]
-		.OverrideResetToDefault(FResetToDefaultOverride::Create(FSimpleDelegate::CreateSP(this, &FCustomizableObjectNodeTableDetails::OnAnimCustomRowResetButtonClicked, EAnimColumnType::EACT_SlotColumn)))
-		.Visibility(TAttribute<EVisibility>(this, &FCustomizableObjectNodeTableDetails::AnimWidgetsVisibility));
-		
-
-		// AnimTags Column Selector
-		AnimationCategory.AddCustomRow(LOCTEXT("AnimationProperties", "Animation Properties"))
-		.NameContent()
-		[
-			SNew(STextBlock)
-			.Text(LOCTEXT("AnimTagsText", "Animation Tags Column: "))
-			.ToolTipText(LOCTEXT("AnimTagColumnTooltip", "Select an animation tag column from the Data Table that will set to the animation tags of the animation blueprint"))
-			.Font(IDetailLayoutBuilder::GetDetailFont())
-		]
-		.ValueContent()
-		.HAlign(EHorizontalAlignment::HAlign_Fill)
-		[
-			SAssignNew(AnimTagsComboBox, STextComboBox)
-			.OptionsSource(&AnimTagsOptionNames)
-			.OnSelectionChanged(this, &FCustomizableObjectNodeTableDetails::OnAnimTagsComboBoxSelectionChanged)
-			.Font(IDetailLayoutBuilder::GetDetailFont())
-		]
-		.OverrideResetToDefault(FResetToDefaultOverride::Create(FSimpleDelegate::CreateSP(this, &FCustomizableObjectNodeTableDetails::OnAnimCustomRowResetButtonClicked, EAnimColumnType::EACT_TagsColumn)))
-		.Visibility(TAttribute<EVisibility>(this, &FCustomizableObjectNodeTableDetails::AnimWidgetsVisibility));
-
-
-
-		// Layout Category -----------------------------------
-
-		//	Layout Selector
-		SelectedLayout = nullptr;
-
-		// Mesh selector of the layout editor
-		LayoutCategory.AddCustomRow(LOCTEXT("TableLayoutEditor_MeshSelector", "Mesh Selector"))
-		.NameContent()
-		[
-			SNew(STextBlock)
-			.Text(LOCTEXT("LayoutMeshColumnText", "Mesh Column: "))
-			.ToolTipText(LOCTEXT("LayoutMeshColumnTooltip", "Select a mesh from the Data Table to edit its layout blocks."))
-			.Font(IDetailLayoutBuilder::GetDetailFont())
-		]
-		.ValueContent()
-		.HAlign(EHorizontalAlignment::HAlign_Fill)
-		[
-			SNew(SHorizontalBox)
-			+ SHorizontalBox::Slot()
-			.FillWidth(0.9f)
+			.NameContent()
 			[
-				SAssignNew(LayoutMeshColumnComboBox, STextComboBox)
-				.OptionsSource(&LayoutMeshColumnOptionNames)
-				.InitiallySelectedItem(LayoutMeshColumnOptionNames[0])
-				.OnSelectionChanged(this, &FCustomizableObjectNodeTableDetails::OnLayoutMeshColumnComboBoxSelectionChanged)
-				.Font(IDetailLayoutBuilder::GetDetailFont())
+				SNew(STextBlock)
+					.Text(LOCTEXT("AnimSlotText", "Animation Slot Column: "))
+					.ToolTipText(LOCTEXT("AnimSlotColumnTooltip", "Select an animation slot column from the Data Table that will set to the slot value of the animation blueprint"))
+					.Font(IDetailLayoutBuilder::GetDetailFont())
 			]
-			
-			+ SHorizontalBox::Slot()
-			.FillWidth(0.1f)
+			.ValueContent()
+			.HAlign(EHorizontalAlignment::HAlign_Fill)
 			[
-				SNew(SBox)
-				.HAlign(EHorizontalAlignment::HAlign_Center)
-				.VAlign(EVerticalAlignment::VAlign_Center)
-				.Content()
-				[
-					SNew(SImage)
-					.Image(UE_MUTABLE_GET_BRUSH(TEXT("Icons.Info")))
-					.ToolTipText(FText(LOCTEXT("LaytoutMeshNoteTooltipText","Note:"
-						"\nAs all meshes of a Data Table column share the same layout, the UVs shown"
-						"\nin the editor are from the Default Skeletal Mesh of the Structure.")))
-				]
+				SAssignNew(AnimSlotComboBox, STextComboBox)
+					.OptionsSource(&AnimSlotOptionNames)
+					.OnSelectionChanged(this, &FCustomizableObjectNodeTableDetails::OnAnimSlotComboBoxSelectionChanged)
+					.Font(IDetailLayoutBuilder::GetDetailFont())
 			]
-			
-		]
-		.OverrideResetToDefault(FResetToDefaultOverride::Create(FSimpleDelegate::CreateSP(this, &FCustomizableObjectNodeTableDetails::OnLayoutMeshCustomRowResetButtonClicked)));
-
-
-		// Layout grid size options
-		int32 MaxGridSize = 32;
-		for (int32 Size = 1; Size <= MaxGridSize; Size *= 2)
-		{
-			LayoutGridSizes.Add(MakeShareable(new FString(FString::Printf(TEXT("%d x %d"), Size, Size))));
-		}
-
-		// Layout size selector widget
-		LayoutCategory.AddCustomRow(LOCTEXT("TableBlocksDetails_SizeSelector", "SizeSelector"))
-		.Visibility(TAttribute<EVisibility>(this, &FCustomizableObjectNodeTableDetails::LayoutOptionsVisibility))
-		.NameContent()
-		[
-			SNew(STextBlock)
-			.Text(LOCTEXT("TableLayoutGridSizeText", "Grid Size"))
-			.Font(IDetailLayoutBuilder::GetDetailFont())
-		]
-		.ValueContent()
-		[
-			SAssignNew(GridSizeComboBox, STextComboBox)
-			.OptionsSource(&LayoutGridSizes)
-			.OnSelectionChanged(this, &FCustomizableObjectNodeTableDetails::OnGridSizeChanged)
-			.Font(IDetailLayoutBuilder::GetDetailFont())
-		];
-
-
-		// Layout Strategy options. Hardcoded: we should get names and tooltips from the enum property
-		{
-			LayoutPackingStrategies.Empty();
-			LayoutPackingStrategiesTooltips.Empty();
-
-			LayoutPackingStrategies.Add(MakeShareable(new FString("Resizable")));
-			LayoutPackingStrategiesTooltips.Add(LOCTEXT("TableDetails_ResizableStrategyTooltip", "In a layout merge, Layout size will increase if blocks don't fit inside."));
-
-			LayoutPackingStrategies.Add(MakeShareable(new FString("Fixed")));
-			LayoutPackingStrategiesTooltips.Add(LOCTEXT("TableDetails_FixedStrategyTooltip", "In a layout merge, the layout will increase its size until the maximum layout grid size"
-				"\nBlock sizes will be reduced if they don't fit inside the layout."
-				"\nSet the reduction priority of each block to control which blocks are reduced first and how they are reduced."));
+			.OverrideResetToDefault(FResetToDefaultOverride::Create(FSimpleDelegate::CreateSP(this, &FCustomizableObjectNodeTableDetails::OnAnimCustomRowResetButtonClicked, EAnimColumnType::EACT_SlotColumn)))
+			.Visibility(TAttribute<EVisibility>(this, &FCustomizableObjectNodeTableDetails::AnimWidgetsVisibility));
 
-			LayoutPackingStrategies.Add(MakeShareable(new FString("Overlay")));
-			LayoutPackingStrategiesTooltips.Add(LOCTEXT("TableDetails_OverlayStrategyTooltip", "In a layout merge, the layout will not be modified and blocks will be ignored."
-				"\nExtend material nodes just add their layouts on top of the base one"));
-		}
 
-		// Layout strategy selector group widget
-		IDetailGroup* LayoutStrategyOptionsGroup = &LayoutCategory.AddGroup(TEXT("TableLayoutStrategyOptionsGroup"), LOCTEXT("TableLayoutStrategyGroup", "Table Layout Strategy Group"), false, true);
-		LayoutStrategyOptionsGroup->HeaderRow()
-		.Visibility(TAttribute<EVisibility>(this, &FCustomizableObjectNodeTableDetails::LayoutOptionsVisibility))
-		.NameContent()
-		[
-			SNew(STextBlock)
-			.Text(LOCTEXT("TableLayoutStrategy_Text", "Layout Strategy:"))
-			.ToolTipText(LOCTEXT("TableLayoutStrategyTooltip", "Selects the packing strategy in case of a layout merge."))
-			.Font(IDetailLayoutBuilder::GetDetailFont())
-		]
-		.ValueContent()
-		[
-			SAssignNew(StrategyComboBox, SSearchableComboBox)
-			.OptionsSource(&LayoutPackingStrategies)
-			.OnSelectionChanged(this, &FCustomizableObjectNodeTableDetails::OnLayoutPackingStrategyChanged)
-			.OnGenerateWidget(this, &FCustomizableObjectNodeTableDetails::OnGenerateStrategyComboBox)
-			.ToolTipText(this, &FCustomizableObjectNodeTableDetails::GetSelectedLayoutStrategyTooltip)
+		// AnimTags Column Selector
+		AnimationCategory.AddCustomRow(LOCTEXT("AnimationProperties", "Animation Properties"))
+			.NameContent()
 			[
 				SNew(STextBlock)
-				.Text(this, &FCustomizableObjectNodeTableDetails::GetSelectedLayoutStrategyName)
-				.Font(IDetailLayoutBuilder::GetDetailFont())
+					.Text(LOCTEXT("AnimTagsText", "Animation Tags Column: "))
+					.ToolTipText(LOCTEXT("AnimTagColumnTooltip", "Select an animation tag column from the Data Table that will set to the animation tags of the animation blueprint"))
+					.Font(IDetailLayoutBuilder::GetDetailFont())
 			]
-		];
-
-
-		// Max layout size selector widget
-		LayoutStrategyOptionsGroup->AddWidgetRow()
-		.Visibility(TAttribute<EVisibility>(this, &FCustomizableObjectNodeTableDetails::FixedStrategyOptionsVisibility))
-		.NameContent()
-		[
-			SNew(STextBlock)
-			.Text(LOCTEXT("TableMaxLayoutSize_Text", "Max Layout Size:"))
-			.Font(IDetailLayoutBuilder::GetDetailFont())
-		]
-		.ValueContent()
-		[
-			SAssignNew(MaxGridSizeComboBox, STextComboBox)
-			.OptionsSource(&LayoutGridSizes)
-			.OnSelectionChanged(this, &FCustomizableObjectNodeTableDetails::OnMaxGridSizeChanged)
-			.Font(IDetailLayoutBuilder::GetDetailFont())
-		];
-
-
-		// Block reduction methods options
-		{
-			BlockReductionMethods.Empty();
-			BlockReductionMethods.Add(MakeShareable(new FString("Halve")));
-			BlockReductionMethodsTooltips.Add(LOCTEXT("TableDetails_HalveRedMethodTooltip", "Blocks will be reduced by half each time."));
-
-			BlockReductionMethods.Add(MakeShareable(new FString("Unitary")));
-			BlockReductionMethodsTooltips.Add(LOCTEXT("TableDetails_UnitaryRedMethodTooltip", "Blocks will be reduced by one unit each time."));
-		}
-
-		// Reduction method selector widget
-		LayoutStrategyOptionsGroup->AddWidgetRow()
-		.Visibility(TAttribute<EVisibility>(this, &FCustomizableObjectNodeTableDetails::FixedStrategyOptionsVisibility))
-		.NameContent()
-		[
-			SNew(STextBlock)
-			.Text(LOCTEXT("TableReductionMethod_Text", "Reduction Method:"))
-			.ToolTipText(LOCTEXT("TableReduction_Method_Tooltip", "Select how blocks will be reduced in case that they do not fit in the layout:"))
-			.Font(IDetailLayoutBuilder::GetDetailFont())
-		]
-		.ValueContent()
-		[
-			SAssignNew(ReductionMethodComboBox, SSearchableComboBox)
-			.OptionsSource(&BlockReductionMethods)
-			.OnSelectionChanged(this, &FCustomizableObjectNodeTableDetails::OnReductionMethodChanged)
-			.OnGenerateWidget(this, &FCustomizableObjectNodeTableDetails::OnGenerateReductionMethodComboBox)
-			.ToolTipText(this, &FCustomizableObjectNodeTableDetails::GetSelectedLayoutReductionMethodTooltip)
+			.ValueContent()
+			.HAlign(EHorizontalAlignment::HAlign_Fill)
 			[
-				SNew(STextBlock)
-				.Text(this, &FCustomizableObjectNodeTableDetails::GetSelectedLayoutReductionMethodName)
-				.Font(IDetailLayoutBuilder::GetDetailFont())
+				SAssignNew(AnimTagsComboBox, STextComboBox)
+					.OptionsSource(&AnimTagsOptionNames)
+					.OnSelectionChanged(this, &FCustomizableObjectNodeTableDetails::OnAnimTagsComboBoxSelectionChanged)
+					.Font(IDetailLayoutBuilder::GetDetailFont())
 			]
-		];
+			.OverrideResetToDefault(FResetToDefaultOverride::Create(FSimpleDelegate::CreateSP(this, &FCustomizableObjectNodeTableDetails::OnAnimCustomRowResetButtonClicked, EAnimColumnType::EACT_TagsColumn)))
+			.Visibility(TAttribute<EVisibility>(this, &FCustomizableObjectNodeTableDetails::AnimWidgetsVisibility));
 
-		// Layout blocks editor widget
-		LayoutBlocksEditor = SNew(SCustomizableObjectNodeLayoutBlocksEditor);
+		// Array of MeshSections and their editable layouts
+		TArray<FLayoutEditorMeshSection> MeshSectionsAndLayouts;
+		GenerateMeshSectionOptions(MeshSectionsAndLayouts);
 
-		// Block editor Widget
-		LayoutCategory.AddCustomRow(LOCTEXT("TableLayoutEditor", "Layout Editor"))
-		[
-			SNew(SBox)
-			.HeightOverride(700.0f)
-			.WidthOverride(700.0f)
-			[
-				LayoutBlocksEditor.ToSharedRef()
-			]
-		];
+		TSharedPtr<SCustomizableObjectLayoutEditor> LayoutBlocksEditor = SNew(SCustomizableObjectLayoutEditor)
+			.Node(Node.Get())
+			.MeshSections(MeshSectionsAndLayouts);
+
+		FCustomizableObjectLayoutEditorDetailsBuilder LayoutEditorBuilder;
+		LayoutEditorBuilder.LayoutEditor = LayoutBlocksEditor;
+		LayoutEditorBuilder.bShowLayoutSelector = true;
+		LayoutEditorBuilder.bShowPackagingStrategy = true;
+		LayoutEditorBuilder.bShowAutomaticGenerationSettings = true;
+		LayoutEditorBuilder.bShowGridSize = true;
+		LayoutEditorBuilder.bShowMaxGridSize = true;
+		LayoutEditorBuilder.bShowReductionMethods = true;
+		LayoutEditorBuilder.bShowWarningSettings = true;
 
-		LayoutBlocksEditor->SetCurrentLayout(nullptr);
+		LayoutEditorBuilder.CustomizeDetails(*DetailBuilder.Get());
+
+		LayoutBlocksEditor->UpdateLayout(nullptr);
 	}
 }
 
@@ -409,11 +243,9 @@ void FCustomizableObjectNodeTableDetails::CustomizeDetails(const TSharedPtr<IDet
 void FCustomizableObjectNodeTableDetails::GenerateMeshColumnComboBoxOptions()
 {
 	AnimMeshColumnOptionNames.Empty();
-	LayoutMeshColumnOptionNames.Empty();
 
 	// Add first element to clear selection
 	AnimMeshColumnOptionNames.Add(MakeShareable(new FString("- Nothing Selected -")));
-	LayoutMeshColumnOptionNames.Add(MakeShareable(new FString("- Nothing Selected -")));
 
 	const UScriptStruct* TableStruct = Node->GetTableNodeStruct();
 
@@ -439,6 +271,40 @@ void FCustomizableObjectNodeTableDetails::GenerateMeshColumnComboBoxOptions()
 			{
 				FString MeshColumnName = DataTableUtils::GetPropertyExportName(ColumnProperty);
 				AnimMeshColumnOptionNames.Add(MakeShareable(new FString(MeshColumnName)));
+			}
+		}
+	}
+}
+
+
+void FCustomizableObjectNodeTableDetails::GenerateMeshSectionOptions(TArray<FLayoutEditorMeshSection>& OutMeshSections)
+{
+	// Add first element to clear selection
+	OutMeshSections.Empty();
+
+	const UScriptStruct* TableStruct = Node->GetTableNodeStruct();
+
+	if (!TableStruct)
+	{
+		return;
+	}
+
+	// Get mesh columns only
+	for (TFieldIterator<FProperty> It(TableStruct); It; ++It)
+	{
+		FProperty* ColumnProperty = *It;
+
+		if (!ColumnProperty)
+		{
+			continue;
+		}
+
+		if (const FSoftObjectProperty* SoftObjectProperty = CastField<FSoftObjectProperty>(ColumnProperty))
+		{
+			if (SoftObjectProperty->PropertyClass->IsChildOf(USkeletalMesh::StaticClass())
+				|| SoftObjectProperty->PropertyClass->IsChildOf(UStaticMesh::StaticClass()))
+			{
+				FString MeshColumnName = DataTableUtils::GetPropertyExportName(ColumnProperty);
 
 				for (const UEdGraphPin* Pin : Node->GetAllNonOrphanPins())
 				{
@@ -451,17 +317,12 @@ void FCustomizableObjectNodeTableDetails::GenerateMeshColumnComboBoxOptions()
 
 					if (PinData && PinData->ColumnName == MeshColumnName)
 					{
-						if (PinData->Layouts.Num() > 1)
-						{
-							for (int32 LayoutIndex = 0; LayoutIndex < PinData->Layouts.Num(); ++LayoutIndex)
-							{
-								FString LayoutName = Pin->PinFriendlyName.ToString() + FString::Printf(TEXT(" UV_%d"), LayoutIndex);
-								LayoutMeshColumnOptionNames.Add(MakeShareable(new FString(LayoutName)));
-							}
-						}
-						else
+						FLayoutEditorMeshSection& MeshSection = OutMeshSections.AddDefaulted_GetRef();
+						MeshSection.MeshName = MakeShareable(new FString(Pin->PinFriendlyName.ToString()));
+
+						for (UCustomizableObjectLayout* Layout : PinData->Layouts)
 						{
-							LayoutMeshColumnOptionNames.Add(MakeShareable(new FString(Pin->PinFriendlyName.ToString())));
+							MeshSection.Layouts.Add(Layout);
 						}
 					}
 				}
@@ -718,258 +579,6 @@ void FCustomizableObjectNodeTableDetails::OnAnimCustomRowResetButtonClicked(EAni
 }
 
 
-// Layout Category --------------------------------------------------------------------------------
-
-void FCustomizableObjectNodeTableDetails::OnLayoutMeshCustomRowResetButtonClicked()
-{
-	if (LayoutMeshColumnComboBox.IsValid() && LayoutMeshColumnOptionNames.Num())
-	{
-		LayoutMeshColumnComboBox->SetSelectedItem(LayoutMeshColumnOptionNames[0]);
-		SelectedLayout = nullptr;
-		LayoutBlocksEditor->SetCurrentLayout(nullptr);
-	}
-}
-
-
-void FCustomizableObjectNodeTableDetails::OnLayoutMeshColumnComboBoxSelectionChanged(TSharedPtr<FString> Selection, ESelectInfo::Type SelectInfo)
-{
-	if (Selection.IsValid() && Selection != LayoutMeshColumnOptionNames[0])
-	{
-		FString ColumnName;
-		(*Selection).Split(" LOD_", &ColumnName, NULL);
-
-		for (const UEdGraphPin* Pin : Node->GetAllNonOrphanPins())
-		{
-			const UCustomizableObjectNodeTableMeshPinData* PinData = Cast<UCustomizableObjectNodeTableMeshPinData >(Node->GetPinData(*Pin));
-
-			if (PinData && PinData->ColumnName == ColumnName)
-			{
-				for (int32 LayoutIndex = 0; LayoutIndex < PinData->Layouts.Num(); ++LayoutIndex)
-				{
-					if (PinData->Layouts[LayoutIndex]->GetLayoutName() == *Selection)
-					{
-						LayoutBlocksEditor->SetCurrentLayout(PinData->Layouts[LayoutIndex]);
-						SelectedLayout = PinData->Layouts[LayoutIndex];
-
-						FillLayoutComboBoxOptions();
-					}
-				}
-			}
-		}
-	}
-	else
-	{
-		SelectedLayout = nullptr;
-		LayoutBlocksEditor->SetCurrentLayout(nullptr);
-	}
-}
-
-
-EVisibility FCustomizableObjectNodeTableDetails::LayoutOptionsVisibility() const
-{
-	return SelectedLayout.IsValid() ? EVisibility::Visible : EVisibility::Collapsed;
-}
-
-
-EVisibility FCustomizableObjectNodeTableDetails::FixedStrategyOptionsVisibility() const
-{
-	return (SelectedLayout.IsValid() && SelectedLayout->PackingStrategy == ECustomizableObjectTextureLayoutPackingStrategy::Fixed) ? EVisibility::Visible : EVisibility::Collapsed;
-}
-
-
-void FCustomizableObjectNodeTableDetails::FillLayoutComboBoxOptions()
-{
-	if (SelectedLayout.IsValid() && GridSizeComboBox.IsValid() && StrategyComboBox.IsValid()
-		&& MaxGridSizeComboBox.IsValid() && ReductionMethodComboBox.IsValid())
-	{
-		for (int32 OptionIndex = 0; OptionIndex < LayoutGridSizes.Num(); ++OptionIndex)
-		{
-			int32 Size = 1 << OptionIndex;
-
-			if (SelectedLayout->GetGridSize() == FIntPoint(Size))
-			{
-				GridSizeComboBox->SetSelectedItem(LayoutGridSizes[OptionIndex]);
-			}
-
-			if (SelectedLayout->GetMaxGridSize() == FIntPoint(Size))
-			{
-				MaxGridSizeComboBox->SetSelectedItem(LayoutGridSizes[OptionIndex]);
-			}
-		}
-
-		StrategyComboBox->SetSelectedItem(LayoutPackingStrategies[(uint32)SelectedLayout->PackingStrategy]);
-		ReductionMethodComboBox->SetSelectedItem(BlockReductionMethods[(uint32)SelectedLayout->BlockReductionMethod]);
-	}
-}
-
-
-TSharedRef<SWidget> FCustomizableObjectNodeTableDetails::OnGenerateStrategyComboBox(TSharedPtr<FString> InItem) const
-{
-	FText Tooltip;
-
-	if (InItem.IsValid())
-	{
-		int32 TooltipIndex = LayoutPackingStrategies.IndexOfByKey(InItem);
-
-		if (LayoutPackingStrategiesTooltips.IsValidIndex(TooltipIndex))
-		{
-			//A list of tool tips should have been populated in a 1 to 1 correspondance
-			check(LayoutPackingStrategies.Num() == LayoutPackingStrategiesTooltips.Num());
-			Tooltip = LayoutPackingStrategiesTooltips[TooltipIndex];
-		}
-	}
-
-	return SNew(STextBlock)
-		.Text(FText::FromString(*InItem.Get()))
-		.Font(IDetailLayoutBuilder::GetDetailFont())
-		.ToolTipText(Tooltip);
-}
-
-
-TSharedRef<SWidget> FCustomizableObjectNodeTableDetails::OnGenerateReductionMethodComboBox(TSharedPtr<FString> InItem) const
-{
-	FText Tooltip;
-
-	if (InItem.IsValid())
-	{
-		int32 TooltipIndex = BlockReductionMethods.IndexOfByKey(InItem);
-
-		if (BlockReductionMethodsTooltips.IsValidIndex(TooltipIndex))
-		{
-			//A list of tool tips should have been populated in a 1 to 1 correspondance
-			check(BlockReductionMethods.Num() == BlockReductionMethodsTooltips.Num());
-			Tooltip = BlockReductionMethodsTooltips[TooltipIndex];
-		}
-	}
-
-	return SNew(STextBlock)
-		.Text(FText::FromString(*InItem.Get()))
-		.Font(IDetailLayoutBuilder::GetDetailFont())
-		.ToolTipText(Tooltip);
-}
-
-
-FText FCustomizableObjectNodeTableDetails::GetSelectedLayoutStrategyName() const
-{
-	if (SelectedLayout.IsValid())
-	{
-		return FText::FromString(*LayoutPackingStrategies[(uint32)SelectedLayout->PackingStrategy]);
-	}
-
-	return FText();
-}
-
-
-FText FCustomizableObjectNodeTableDetails::GetSelectedLayoutReductionMethodName() const
-{
-	if (SelectedLayout.IsValid())
-	{
-		return FText::FromString(*BlockReductionMethods[(uint32)SelectedLayout->BlockReductionMethod]);
-	}
-
-	return FText();
-}
-
-
-void FCustomizableObjectNodeTableDetails::OnGridSizeChanged(TSharedPtr<FString> NewSelection, ESelectInfo::Type SelectInfo)
-{
-	if (SelectedLayout.IsValid())
-	{
-		int32 Size = 1 << LayoutGridSizes.Find(NewSelection);
-
-		if (SelectedLayout->GetGridSize().X != Size || SelectedLayout->GetGridSize().Y != Size)
-		{
-			SelectedLayout->SetGridSize(FIntPoint(Size));
-
-			// Adjust all the blocks sizes
-			//for (int32 BlockIndex = 0; BlockIndex < SelectedLayout->Blocks.Num(); ++BlockIndex)
-			//{
-			//	FCustomizableObjectLayoutBlock& Block = SelectedLayout->Blocks[BlockIndex];
-			//	Block.Min.X = FMath::Min(Block.Min.X, Size - 1);
-			//	Block.Min.Y = FMath::Min(Block.Min.Y, Size - 1);
-			//	Block.Max.X = FMath::Min(Block.Max.X, Size);
-			//	Block.Max.Y = FMath::Min(Block.Max.Y, Size);
-			//}
-
-			Node->MarkPackageDirty();
-		}
-	}
-}
-
-
-void FCustomizableObjectNodeTableDetails::OnLayoutPackingStrategyChanged(TSharedPtr<FString> NewSelection, ESelectInfo::Type SelectInfo)
-{
-	if (SelectedLayout.IsValid())
-	{
-		uint32 Selection = LayoutPackingStrategies.IndexOfByKey(NewSelection);
-
-		if (SelectedLayout->PackingStrategy != (ECustomizableObjectTextureLayoutPackingStrategy)Selection)
-		{
-			SelectedLayout->PackingStrategy=(ECustomizableObjectTextureLayoutPackingStrategy)Selection;
-			Node->MarkPackageDirty();
-		}
-	}
-}
-
-
-void FCustomizableObjectNodeTableDetails::OnMaxGridSizeChanged(TSharedPtr<FString> NewSelection, ESelectInfo::Type SelectInfo)
-{
-	if (SelectedLayout.IsValid())
-	{
-		int32 Size = 1 << LayoutGridSizes.Find(NewSelection);
-
-		if (SelectedLayout->GetMaxGridSize().X != Size || SelectedLayout->GetMaxGridSize().Y != Size)
-		{
-			SelectedLayout->SetMaxGridSize(FIntPoint(Size));
-			SelectedLayout->MarkPackageDirty();
-		}
-	}
-}
-
-
-void FCustomizableObjectNodeTableDetails::OnReductionMethodChanged(TSharedPtr<FString> NewSelection, ESelectInfo::Type SelectInfo)
-{
-	if (SelectedLayout.IsValid())
-	{
-		uint32 Selection = BlockReductionMethods.IndexOfByKey(NewSelection);
-
-		if (SelectedLayout->BlockReductionMethod != (ECustomizableObjectLayoutBlockReductionMethod)Selection)
-		{
-			SelectedLayout->BlockReductionMethod=(ECustomizableObjectLayoutBlockReductionMethod)Selection;
-			Node->MarkPackageDirty();
-		}
-	}
-}
-
-
-FText FCustomizableObjectNodeTableDetails::GetSelectedLayoutStrategyTooltip() const
-{
-	if (SelectedLayout.IsValid())
-	{
-		//A list of tool tips should have been populated in a 1 to 1 correspondance
-		check(LayoutPackingStrategies.Num() == LayoutPackingStrategiesTooltips.Num());
-
-		return LayoutPackingStrategiesTooltips[(uint32)SelectedLayout->PackingStrategy];
-	}
-
-	return FText();
-}
-
-
-FText FCustomizableObjectNodeTableDetails::GetSelectedLayoutReductionMethodTooltip() const
-{
-	if (SelectedLayout.IsValid())
-	{
-		//A list of tool tips should have been populated in a 1 to 1 correspondance
-		check(BlockReductionMethods.Num() == BlockReductionMethodsTooltips.Num());
-
-		return BlockReductionMethodsTooltips[(uint32)SelectedLayout->BlockReductionMethod];
-	}
-
-	return FText();
-}
-
-
 // Metadata Category --------------------------------------------------------------------------------
 
 TSharedPtr<FString> FCustomizableObjectNodeTableDetails::GenerateMutableMetaDataColumnComboBoxOptions()
diff --git a/Engine/Plugins/Experimental/Mutable/Source/CustomizableObjectEditor/Private/MuCOE/Nodes/CustomizableObjectNodeTableDetails.h b/Engine/Plugins/Experimental/Mutable/Source/CustomizableObjectEditor/Private/MuCOE/Nodes/CustomizableObjectNodeTableDetails.h
--- a/Engine/Plugins/Experimental/Mutable/Source/CustomizableObjectEditor/Private/MuCOE/Nodes/CustomizableObjectNodeTableDetails.h
+++ b/Engine/Plugins/Experimental/Mutable/Source/CustomizableObjectEditor/Private/MuCOE/Nodes/CustomizableObjectNodeTableDetails.h
@@ -14,7 +14,6 @@ class IDetailLayoutBuilder;
 class ISinglePropertyView;
 class FReply;
 class SButton;
-class SCustomizableObjectNodeLayoutBlocksEditor;
 class SSearchableComboBox;
 class STextBlock;
 class STextComboBox;
@@ -23,6 +22,7 @@ class UCustomizableObjectLayout;
 class UCustomizableObjectNodeTable;
 struct EVisibility;
 struct FSlateColor;
+struct FLayoutEditorMeshSection;
 
 enum class EAnimColumnType 
 {
@@ -48,6 +48,7 @@ private:
 	// Details -------------
 	// Generates Mesh columns combobox options
 	void GenerateMeshColumnComboBoxOptions();
+	void GenerateMeshSectionOptions(TArray<FLayoutEditorMeshSection>& OutMeshSections);
 	
 	// Function called when the table node has been refreshed
 	void OnNodePinValueChanged();
@@ -79,47 +80,6 @@ private:
 	void OnAnimCustomRowResetButtonClicked(EAnimColumnType ColumnType);
 
 
-	// Layout Category -------------
-	// OnComboBoxSelectionChanged Callback for Layout ComboBox
-	void OnLayoutMeshColumnComboBoxSelectionChanged(TSharedPtr<FString> Selection, ESelectInfo::Type SelectInfo);
-
-	// Callback to clear the mesh column to edit its layout properties
-	void OnLayoutMeshCustomRowResetButtonClicked();
-
-	// Layout options visibility
-	EVisibility LayoutOptionsVisibility() const;
-
-	/** Returns the visibility of the Fixed layout widgets */
-	EVisibility FixedStrategyOptionsVisibility() const;
-
-	/** Generates the Layout Packing Strategy text widgets of the ComboBox */
-	TSharedRef<SWidget> OnGenerateStrategyComboBox(TSharedPtr<FString> InItem) const;
-	
-	/** Generates the Layout Redutction Methods text widgets of the ComboBox*/
-	TSharedRef<SWidget> OnGenerateReductionMethodComboBox(TSharedPtr<FString> InItem) const;
-
-	/** Returns the selected Packing Strategy Name as a FText */
-	FText GetSelectedLayoutStrategyName() const;
-
-	/** Returns the selected Reduction Method Name as a FText */
-	FText GetSelectedLayoutReductionMethodName() const;
-
-	/** Returns the selected Packing Strategy Tooltip as a FText */
-	FText GetSelectedLayoutStrategyTooltip() const;
-
-	/** Returns the selected Reduction Method Tooltip as a FText */
-	FText GetSelectedLayoutReductionMethodTooltip() const;
-
-	/** Fills the combo box arrays sources */
-	void FillLayoutComboBoxOptions();
-
-	/** Layout Options Callbacks */
-	void OnGridSizeChanged(TSharedPtr<FString> NewSelection, ESelectInfo::Type SelectInfo);
-	void OnMaxGridSizeChanged(TSharedPtr<FString> NewSelection, ESelectInfo::Type SelectInfo);
-	void OnReductionMethodChanged(TSharedPtr<FString> NewSelection, ESelectInfo::Type SelectInfo);
-	void OnLayoutPackingStrategyChanged(TSharedPtr<FString> NewSelection, ESelectInfo::Type SelectInfo);
-
-
 	// Mutable UI Metadata Category -------------
 
 	// Generates MutableMetadata columns combobox options
@@ -199,40 +159,6 @@ private:
 	// Array with the name of the Animation Tags columns as combobox options
 	TArray<TSharedPtr<FString>> AnimTagsOptionNames;
 
-	// Layout -------------
-	// ComboBox widget to select a column from the NodeTable
-	TSharedPtr<STextComboBox> LayoutMeshColumnComboBox;
-	
-	// Array with the name of the table columns as combobox options
-	TArray<TSharedPtr<FString>> LayoutMeshColumnOptionNames;
-	
-	// Layout block editor widget
-	TSharedPtr<SCustomizableObjectNodeLayoutBlocksEditor> LayoutBlocksEditor;
-
-	// Weakpointer to the selected layout objet
-	TWeakObjectPtr<UCustomizableObjectLayout> SelectedLayout;
-
-	/** List of available layout grid sizes. */
-	TArray< TSharedPtr< FString > > LayoutGridSizes;
-
-	/** List of available layout packing strategies. */
-	TArray< TSharedPtr< FString > > LayoutPackingStrategies;
-	TArray<FText> LayoutPackingStrategiesTooltips;
-
-	/** List of available block reduction methods. */
-	TArray< TSharedPtr< FString > > BlockReductionMethods;
-	TArray<FText> BlockReductionMethodsTooltips;
-
-	// ComboBox widget to select a Grid Size from the Selected Layout
-	TSharedPtr<STextComboBox> GridSizeComboBox;
-	// ComboBox widget to select a Strategy from the Selected Layout. SSearchableComboBox allows us to set a custom tooltip per option.
-	TSharedPtr<SSearchableComboBox> StrategyComboBox;
-	// ComboBox widget to select a Max Grid Size from the Selected Layout
-	TSharedPtr<STextComboBox> MaxGridSizeComboBox;
-	// ComboBox widget to select a Reduction Method from the Selected Layout
-	TSharedPtr<SSearchableComboBox> ReductionMethodComboBox;
-
-
 	// Mutable UI Metadata -------------
 	// Array with the name of the MutableMetaData columns
 	TArray<TSharedPtr<FString>> MutableMetaDataColumnsOptionNames;
diff --git a/Engine/Plugins/Experimental/Mutable/Source/CustomizableObjectEditor/Private/MuCOE/SCustomizableObjectLayoutEditor.cpp b/Engine/Plugins/Experimental/Mutable/Source/CustomizableObjectEditor/Private/MuCOE/SCustomizableObjectLayoutEditor.cpp
new file mode 100644
--- /dev/null
+++ b/Engine/Plugins/Experimental/Mutable/Source/CustomizableObjectEditor/Private/MuCOE/SCustomizableObjectLayoutEditor.cpp
@@ -0,0 +1,1504 @@
+// Copyright Epic Games, Inc. All Rights Reserved.
+
+#include "MuCOE/SCustomizableObjectLayoutEditor.h"
+
+#include "Framework/Commands/Commands.h"
+#include "Framework/Commands/UICommandList.h"
+#include "Framework/MultiBox/MultiBoxBuilder.h"
+#include "MuCOE/CustomizableObjectEditorStyle.h"
+#include "MuCOE/SCustomizableObjectLayoutGrid.h"
+#include "MuCOE/UnrealEditorPortabilityHelpers.h"
+#include "Widgets/Input/STextComboBox.h"
+#include "Widgets/Input/SCheckBox.h"
+#include "Widgets/Input/SSpinBox.h"
+#include "Widgets/Layout/SGridPanel.h"
+#include "Widgets/SToolTip.h"
+#include "SSearchableComboBox.h"
+#include "ScopedTransaction.h"
+
+#include "DetailLayoutBuilder.h"
+#include "DetailWidgetRow.h"
+#include "IDetailGroup.h"
+#include "IDetailsView.h"
+
+class ISlateStyle;
+class SWidget;
+
+
+#define LOCTEXT_NAMESPACE "CustomizableObjectEditor"
+
+/**
+ * 
+ */
+class FLayoutEditorCommands : public TCommands<FLayoutEditorCommands>
+{
+
+public:
+	FLayoutEditorCommands() : TCommands<FLayoutEditorCommands>
+	(
+		"LayoutEditorCommands", // Context name for fast lookup
+		NSLOCTEXT( "CustomizableObjectEditor", "LayoutEditorCommands", "Layout Editor" ), // Localized context name for displaying
+		NAME_None, // Parent
+		FCustomizableObjectEditorStyle::GetStyleSetName()
+	)
+	{
+	}	
+	
+	/**  */
+	TSharedPtr< FUICommandInfo > AddBlock;
+	TSharedPtr< FUICommandInfo > RemoveBlock;
+	TSharedPtr< FUICommandInfo > ConsolidateBlocks;
+
+	/**
+	 * Initialize commands
+	 */
+	virtual void RegisterCommands() override
+	{
+		UI_COMMAND( AddBlock, "Add Block", "Add a new block to the layout.", EUserInterfaceActionType::Button, FInputChord() );
+		UI_COMMAND( RemoveBlock, "Remove Block", "Remove a block from the layout.", EUserInterfaceActionType::Button, FInputChord() );
+		UI_COMMAND(ConsolidateBlocks, "Consolidate Blocks", "Convert automatic blocks into user blocks.", EUserInterfaceActionType::Button, FInputChord() );
+	}
+};
+
+
+void FCustomizableObjectLayoutEditorDetailsBuilder::CustomizeDetails(IDetailLayoutBuilder& DetailsBuilder)
+{
+	TWeakObjectPtr<UObject> Node = nullptr;
+
+	const IDetailsView* DetailsView = DetailsBuilder.GetDetailsView();
+	if (DetailsView->GetSelectedObjects().Num())
+	{
+		Node = DetailsView->GetSelectedObjects()[0].Get();
+	}
+
+	// Layout category
+	IDetailCategoryBuilder& LayoutCategory = DetailsBuilder.EditCategory("Layout Editor");
+
+	if (!Node.IsValid())
+	{
+		LayoutCategory.AddCustomRow(LOCTEXT("LayoutEditor_MissingNode", "NodeNotFound"))
+			[
+				SNew(STextBlock)
+					.Text(LOCTEXT("LayoutEditor_MissingNodeText", "Node not found"))
+			];
+		return;
+	}
+
+	if (!LayoutEditor)
+	{
+		check(false);
+		return;
+	}
+
+	if (bShowLayoutSelector)
+	{
+		LayoutCategory.AddCustomRow(LOCTEXT("LayoutEditor_MeshSectionRow", "Mesh Section Selector"))
+			.NameContent()
+			[
+				SNew(STextBlock)
+					.Text(LOCTEXT("LayoutEditor_MeshSectionText", "Mesh Section: "))
+					.ToolTipText(LOCTEXT("LayoutEditor_MeshSectionTooltip", "Select the mesh section to visualize."))
+					.Font(IDetailLayoutBuilder::GetDetailFont())
+			]
+			.ValueContent()
+			.HAlign(EHorizontalAlignment::HAlign_Fill)
+			[
+				SNew(SHorizontalBox)
+					+ SHorizontalBox::Slot()
+					.FillWidth(0.9f)
+					[
+						SAssignNew(LayoutEditor->MeshSectionComboBox, STextComboBox)
+							.OptionsSource(&LayoutEditor->MeshSectionNames)
+							.InitiallySelectedItem(LayoutEditor->MeshSectionNames[0])
+							.OnSelectionChanged(LayoutEditor.Get(), &SCustomizableObjectLayoutEditor::OnMeshSectionChanged)
+							.Font(IDetailLayoutBuilder::GetDetailFont())
+					]
+					// TODO UE-223305: Get the tooltip text from a function
+					//+ SHorizontalBox::Slot()
+					//.FillWidth(0.1f)
+					//[
+					//	SNew(SBox)
+					//		.HAlign(EHorizontalAlignment::HAlign_Center)
+					//		.VAlign(EVerticalAlignment::VAlign_Center)
+					//		.Content()
+					//		[
+					//			SNew(SImage)
+					//				.Image(UE_MUTABLE_GET_BRUSH(TEXT("Icons.Info")))
+					//				.ToolTipText(FText(LOCTEXT("LaytoutEditor_MeshNoteTooltip", "Note:"
+					//					"\nAs all meshes of a Data Table column share the same layout, the UVs shown"
+					//					"\nin the editor are from the Default Skeletal Mesh of the Structure.")))
+					//		]
+					//]
+
+			]
+			.OverrideResetToDefault(FResetToDefaultOverride::Create(FSimpleDelegate::CreateSP(LayoutEditor.Get(), &SCustomizableObjectLayoutEditor::OnResetSelection)));
+
+
+		LayoutCategory.AddCustomRow(LOCTEXT("LayoutEditor_UVChannelRow", "UV Channel Selector"))
+			.NameContent()
+			[
+				SNew(STextBlock)
+					.Text(LOCTEXT("LayoutEditor_UVChannelText", "UV Channel"))
+					.Font(IDetailLayoutBuilder::GetDetailFont())
+			]
+			.ValueContent()
+			[
+				SAssignNew(LayoutEditor->UVChannelComboBox, STextComboBox)
+					.InitiallySelectedItem(LayoutEditor->UVChannels[0])
+					.OptionsSource(&LayoutEditor->UVChannels)
+					.OnSelectionChanged(LayoutEditor.Get(), &SCustomizableObjectLayoutEditor::OnUVChannelChanged)
+					.Font(IDetailLayoutBuilder::GetDetailFont())
+			];
+	}
+
+	IDetailGroup* BaseLayoutOptionsGroup = &LayoutCategory.AddGroup(TEXT("LayoutEditor_OptionsGroupRow"), LOCTEXT("LayoutEditor_OptionsGroupRow", "Layout Options"), false, true);
+	IDetailGroup* LayoutOptionsGroup = BaseLayoutOptionsGroup;
+
+	BaseLayoutOptionsGroup->HeaderRow()
+		.Visibility(TAttribute<EVisibility>(LayoutEditor.Get(), &SCustomizableObjectLayoutEditor::LayoutOptionsVisibility))
+		.NameContent()
+		[
+			SNew(STextBlock)
+				.Text(LOCTEXT("LayoutEditor_OptionsGroupText", "Layout Options"))
+				.ToolTipText(LOCTEXT("LayoutEditor_OptionsGroupTooltip", "Selects the packing strategy in case of a layout merge."))
+				.Font(IDetailLayoutBuilder::GetDetailFont())
+		];
+
+	if (bShowPackagingStrategy)
+	{
+		// Layout strategy selector group widget
+		LayoutOptionsGroup = &LayoutOptionsGroup->AddGroup(TEXT("LayoutEditor_LayoutStrategyGroup"), LOCTEXT("LayoutEditor_LayoutStrategyGroup", "Layout Strategy Group"), true);
+		LayoutOptionsGroup->HeaderRow()
+			.Visibility(TAttribute<EVisibility>(LayoutEditor.Get(), &SCustomizableObjectLayoutEditor::LayoutOptionsVisibility))
+			.NameContent()
+			[
+				SNew(STextBlock)
+					.Text(LOCTEXT("LayoutEditor_LayoutStrategyText", "Layout Strategy:"))
+					.ToolTipText(LOCTEXT("LayoutEditor_LayoutStrategyTooltip", "Selects the packing strategy in case of a layout merge."))
+					.Font(IDetailLayoutBuilder::GetDetailFont())
+			]
+			.ValueContent()
+			[
+				SAssignNew(LayoutEditor->StrategyComboBox, SSearchableComboBox)
+					.OptionsSource(&LayoutEditor->LayoutPackingStrategies)
+					.OnSelectionChanged(LayoutEditor.Get(), &SCustomizableObjectLayoutEditor::OnPackagingStrategyChanged)
+					.OnGenerateWidget(LayoutEditor.Get(), &SCustomizableObjectLayoutEditor::GenerateLayoutPackagingStrategyComboBox)
+					.ToolTipText(LayoutEditor.Get(), &SCustomizableObjectLayoutEditor::GetLayoutPackagingStrategyToolTip)
+					[
+						SNew(STextBlock)
+							.Text(LayoutEditor.Get(), &SCustomizableObjectLayoutEditor::GetLayoutPackagingStrategyName)
+							.Font(IDetailLayoutBuilder::GetDetailFont())
+					]
+			];
+	}
+
+
+	if (bShowAutomaticGenerationSettings)
+	{
+		LayoutOptionsGroup->AddWidgetRow()
+			.Visibility(TAttribute<EVisibility>(LayoutEditor.Get(), &SCustomizableObjectLayoutEditor::AutoBlocksStrategyVisibility))
+			.NameContent()
+			[
+				SNew(STextBlock)
+					.Text(LOCTEXT("LayoutEditor_AutoBlockStrategyText", "Automatic Blocks Strategy:"))
+					.ToolTipText(LOCTEXT("AutoBlockStrategyTooltip", "Selects the strategy to create layout blocks from unassigned UVs."))
+					.Font(IDetailLayoutBuilder::GetDetailFont())
+			]
+			.ValueContent()
+			[
+				SAssignNew(LayoutEditor->AutoBlocksComboBox, SSearchableComboBox)
+					.InitiallySelectedItem(LayoutEditor->AutoBlocksStrategies[0])
+					.OptionsSource(&LayoutEditor->AutoBlocksStrategies)
+					.OnSelectionChanged(LayoutEditor.Get(), &SCustomizableObjectLayoutEditor::OnAutoBlocksChanged)
+					.OnGenerateWidget(LayoutEditor.Get(), &SCustomizableObjectLayoutEditor::GenerateAutoBlocksComboBox)
+					.ToolTipText(LayoutEditor.Get(), &SCustomizableObjectLayoutEditor::GetAutoBlocksTooltip)
+					[
+						SNew(STextBlock)
+							.Text(LayoutEditor.Get(), &SCustomizableObjectLayoutEditor::GetAutoBlocksName)
+							.Font(IDetailLayoutBuilder::GetDetailFont())
+					]
+			];
+
+		// Option to merge child automatic blocks
+		LayoutOptionsGroup->AddWidgetRow()
+			.Visibility(TAttribute<EVisibility>(LayoutEditor.Get(), &SCustomizableObjectLayoutEditor::AutoBlocksMergeStrategyVisibility))
+			.NameContent()
+			[
+				SNew(STextBlock)
+					.Text(LOCTEXT("LayoutEditor_AutoBlockMergeStrategyText", "Automatic Blocks Merge Strategy:"))
+					.ToolTipText(LOCTEXT("LayoutEditor_AutoBlockMergeStrategyTooltip", "Selects the strategy to merge blocks during automatic generation."))
+					.Font(IDetailLayoutBuilder::GetDetailFont())
+			]
+			.ValueContent()
+			[
+				SAssignNew(LayoutEditor->AutoBlocksMergeComboBox, SSearchableComboBox)
+					.InitiallySelectedItem(LayoutEditor->AutoBlocksMergeStrategies[0])
+					.OptionsSource(&LayoutEditor->AutoBlocksMergeStrategies)
+					.OnSelectionChanged(LayoutEditor.Get(), &SCustomizableObjectLayoutEditor::OnAutoBlocksMergeChanged)
+					.OnGenerateWidget(LayoutEditor.Get(), &SCustomizableObjectLayoutEditor::GenerateAutoBlocksMergeComboBox)
+					.ToolTipText(LayoutEditor.Get(), &SCustomizableObjectLayoutEditor::GetAutoBlocksMergeTooltip)
+					[
+						SNew(STextBlock)
+							.Text(LayoutEditor.Get(), &SCustomizableObjectLayoutEditor::GetAutoBlocksMergeName)
+							.Font(IDetailLayoutBuilder::GetDetailFont())
+					]
+			];
+	}
+
+	// Grid size combo
+	if (bShowGridSize)
+	{
+		LayoutOptionsGroup->AddWidgetRow()
+			.NameContent()
+			[
+				SNew(STextBlock)
+					.Text(LOCTEXT("LayoutEditor_GridSizeText", "Grid Size:"))
+					.Font(IDetailLayoutBuilder::GetDetailFont())
+			]
+			.ValueContent()
+			[
+				SAssignNew(LayoutEditor->GridSizeComboBox, STextComboBox)
+					.InitiallySelectedItem(LayoutEditor->LayoutGridSizes[0])
+					.OptionsSource(&LayoutEditor->LayoutGridSizes)
+					.OnSelectionChanged(LayoutEditor.Get(), &SCustomizableObjectLayoutEditor::OnGridSizeChanged)
+					.Font(IDetailLayoutBuilder::GetDetailFont())
+			];
+	}
+
+	if (bShowMaxGridSize)
+	{
+		LayoutOptionsGroup->AddWidgetRow()
+			.NameContent()
+			[
+				SNew(STextBlock)
+					.Text(LOCTEXT("LayoutEditor_MaxGridSizeText", "Max Grid Size:"))
+					.Font(IDetailLayoutBuilder::GetDetailFont())
+			]
+			.ValueContent()
+			[
+				SAssignNew(LayoutEditor->MaxGridSizeComboBox, STextComboBox)
+					.InitiallySelectedItem(LayoutEditor->LayoutGridSizes[0])
+					.OptionsSource(&LayoutEditor->LayoutGridSizes)
+					.OnSelectionChanged(LayoutEditor.Get(), &SCustomizableObjectLayoutEditor::OnMaxGridSizeChanged)
+					.Font(IDetailLayoutBuilder::GetDetailFont())
+			];
+	}
+
+	if (bShowReductionMethods)
+	{
+
+		// Reduction method selector widget
+		LayoutOptionsGroup->AddWidgetRow()
+			.Visibility(TAttribute<EVisibility>(LayoutEditor.Get(), &SCustomizableObjectLayoutEditor::FixedStrategyOptionsVisibility))
+			.NameContent()
+			[
+				SNew(STextBlock)
+					.Text(LOCTEXT("LayoutEditor_ReductionMethodText", "Reduction Method:"))
+					.ToolTipText(LOCTEXT("LayoutEditor_ReductionMethodTooltip", "Select how blocks will be reduced in case that they do not fit in the layout."))
+					.Font(IDetailLayoutBuilder::GetDetailFont())
+			]
+			.ValueContent()
+			[
+				SAssignNew(LayoutEditor->ReductionMethodComboBox, SSearchableComboBox)
+					.InitiallySelectedItem(LayoutEditor->BlockReductionMethods[0])
+					.OptionsSource(&LayoutEditor->BlockReductionMethods)
+					.OnSelectionChanged(LayoutEditor.Get(), &SCustomizableObjectLayoutEditor::OnReductionMethodChanged)
+					.OnGenerateWidget(LayoutEditor.Get(), &SCustomizableObjectLayoutEditor::GenerateReductionMethodComboBox)
+					.ToolTipText(LayoutEditor.Get(), &SCustomizableObjectLayoutEditor::GetLayoutReductionMethodTooltip)
+					[
+						SNew(STextBlock)
+							.Text(LayoutEditor.Get(), &SCustomizableObjectLayoutEditor::GetLayoutReductionMethodName)
+							.Font(IDetailLayoutBuilder::GetDetailFont())
+					]
+			];
+	}
+
+	if (bShowWarningSettings)
+	{
+		// Warning selector group widget
+		IDetailGroup* IgnoreWarningsGroup = &LayoutCategory.AddGroup(TEXT("LayoutEditor_IgnoreWarningsGroup"), LOCTEXT("LayoutEditor_IgnoreWarningsGroup", "Ignore Unassigned Vertives Warning group"), false, true);
+		IgnoreWarningsGroup->HeaderRow()
+			.Visibility(TAttribute<EVisibility>(LayoutEditor.Get(), &SCustomizableObjectLayoutEditor::LayoutOptionsVisibility))
+			.NameContent()
+			[
+				SNew(STextBlock)
+					.Text(LOCTEXT("LayoutEditor_IgnoreLodsCheckBoxText", "Ignore Unassigned Vertices Warning:"))
+					.ToolTipText(LOCTEXT("LayoutEditor_IgnoreLodsCheckBoxTooltip",
+						"If true, warning message \"Source mesh has vertices not assigned to any layout block\" will be ignored."
+						"\n Note:"
+						"\n This warning can appear when a CO has more than one LOD using the same Layout Block node and these LODs have been generated using the automatic LOD generation."
+						"\n (At high LODs, some vertices may have been displaced from their original position which means they could have been displaced outside their layout blocks.)"
+						"\n Ignoring these warnings can cause some visual artifacts that may or may not be visually important at higher LODs."))
+					.Font(IDetailLayoutBuilder::GetDetailFont())
+			]
+			.ValueContent()
+			[
+				SNew(SCheckBox)
+					.IsChecked(false)
+					.OnCheckStateChanged(LayoutEditor.Get(), &SCustomizableObjectLayoutEditor::OnIgnoreErrorsCheckStateChanged)
+			];
+
+		// LOD selector widget
+		IgnoreWarningsGroup->AddWidgetRow()
+			.Visibility(TAttribute<EVisibility>(LayoutEditor.Get(), &SCustomizableObjectLayoutEditor::LayoutOptionsVisibility))
+			.NameContent()
+			[
+				SAssignNew(LayoutEditor->LODSelectorTextWidget, STextBlock)
+					.Text(LOCTEXT("LayoutEditor_IgnoreLodText", "First LOD to ignore"))
+					.ToolTipText(LOCTEXT("LayoutEditor_IgnoreLodTooltip", "LOD from which vertex warning messages will be ignored."))
+					.Font(IDetailLayoutBuilder::GetDetailFont())
+					.IsEnabled(false)
+			]
+			.ValueContent()
+			[
+				SAssignNew(LayoutEditor->LODSelectorWidget, SSpinBox<int32>)
+					//.Value_Lambda(0)
+					.IsEnabled(false)
+					.OnValueChanged(LayoutEditor.Get(), &SCustomizableObjectLayoutEditor::OnIgnoreErrorsLODBoxValueChanged)
+					.MinValue(0)
+					.Delta(1)
+					.AlwaysUsesDeltaSnap(true)
+					.MinDesiredWidth(40.0f)
+					.Font(IDetailLayoutBuilder::GetDetailFont())
+			];
+	}
+
+	IDetailGroup& BlockEditorGroup = LayoutCategory.AddGroup(TEXT("LayoutEditor_BlockEditorGroup"), LOCTEXT("LayoutEditor_Layout", "Layout"), false, true);
+	BlockEditorGroup.HeaderRow()
+		.Visibility(TAttribute<EVisibility>(LayoutEditor.Get(), &SCustomizableObjectLayoutEditor::LayoutOptionsVisibility))
+		.NameContent()
+		[
+			SNew(STextBlock)
+				.Text(LOCTEXT("LayoutEditor_BlockEditorText", "Layout"))
+				.ToolTipText(LOCTEXT("LayoutEditor_BlockEditorTooltip", "Selected editable layout.")) 
+				.Font(IDetailLayoutBuilder::GetDetailFont())
+		];
+
+	BlockEditorGroup.AddWidgetRow()
+		.Visibility(TAttribute<EVisibility>(LayoutEditor.Get(), &SCustomizableObjectLayoutEditor::LayoutOptionsVisibility))
+		[
+			SNew(SBox)
+				.HeightOverride(700.0f)
+				.WidthOverride(700.0f)
+				[
+					LayoutEditor.ToSharedRef()
+				]
+		];
+}
+
+
+SCustomizableObjectLayoutEditor::SCustomizableObjectLayoutEditor() : UICommandList(new FUICommandList())
+{
+}
+
+
+void SCustomizableObjectLayoutEditor::Construct(const FArguments& InArgs)
+{
+	CurrentLayout = nullptr;
+	Node = InArgs._Node;
+	check(Node.Get());
+	
+
+	BindCommands();
+
+	NothingSelectedString = MakeShareable(new FString("- Nothing Selected -"));
+
+	MeshSections = InArgs._MeshSections.Get();
+	OnPreUpdateLayoutDelegate = InArgs._OnPreUpdateLayoutDelegate;
+
+	// Layout selector to select the mesh section and uv channel to edit
+	MeshSectionNames.Empty();
+	MeshSectionNames.Add(NothingSelectedString);
+
+	for (FLayoutEditorMeshSection& MeshSection : MeshSections)
+	{
+		MeshSectionNames.Add(MeshSection.MeshName);
+	}
+
+
+	// UV Channel combo (for now hardcoded to a maximum of 4)
+	UVChannels.Empty();
+
+	TSharedPtr<FString> CurrentUVChannel;
+	for (int32 Index = 0; Index < 4; ++Index)
+	{
+		UVChannels.Add(MakeShareable(new FString(FString::FromInt(Index))));
+	}
+
+	{
+		// Layout Strategy options. Hardcoded: we should get names and tooltips from the enum property
+		LayoutPackingStrategies.Empty();
+		LayoutPackingStrategiesOptions.Empty();
+
+		LayoutPackingStrategies.Add(MakeShareable(new FString("Resizable")));
+		LayoutPackingStrategiesOptions.Add({ ECustomizableObjectTextureLayoutPackingStrategy::Resizable, LOCTEXT("LayoutEditor_ResizableStrategyTooltip", "In a layout merge, Layout size will increase if blocks don't fit inside.") });
+
+		LayoutPackingStrategies.Add(MakeShareable(new FString("Fixed")));
+		LayoutPackingStrategiesOptions.Add({ ECustomizableObjectTextureLayoutPackingStrategy::Fixed,LOCTEXT("LayoutEditor_FixedStrategyTooltip", "In a layout merge, the layout will increase its size until the maximum layout grid size"
+			"\nBlock sizes will be reduced if they don't fit inside the layout."
+			"\nSet the reduction priority of each block to control which blocks are reduced first and how they are reduced.") });
+
+		LayoutPackingStrategies.Add(MakeShareable(new FString("Overlay")));
+		LayoutPackingStrategiesOptions.Add({ ECustomizableObjectTextureLayoutPackingStrategy::Overlay,LOCTEXT("LayoutEditor_OverlayStrategyTooltip", "In a layout merge, the layout will not be modified and blocks will be ignored."
+			"\nExtend material nodes just add their layouts on top of the base one") });
+	}
+
+	{
+		AutoBlocksStrategies.Empty();
+		AutoBlocksStrategiesOptions.Empty();
+
+		AutoBlocksStrategies.Add(MakeShareable(new FString("Rectangles")));
+		AutoBlocksStrategiesOptions.Add(
+			{
+				ECustomizableObjectLayoutAutomaticBlocksStrategy::Rectangles ,
+				LOCTEXT("AutoBlockDetails_RectanglesStrategyTooltip", "Try to build rectangles splitting the UVs.")
+			});
+
+		AutoBlocksStrategies.Add(MakeShareable(new FString("UV islands")));
+		AutoBlocksStrategiesOptions.Add(
+			{
+				ECustomizableObjectLayoutAutomaticBlocksStrategy::UVIslands,
+				LOCTEXT("AutoBlockDetails_UVIslandsStrategyTooltip", "Try to build rectangles around each UV island, with a mask.")
+			});
+
+		AutoBlocksStrategies.Add(MakeShareable(new FString("Ignore (legacy)")));
+		AutoBlocksStrategiesOptions.Add(
+			{
+				ECustomizableObjectLayoutAutomaticBlocksStrategy::Ignore,
+				LOCTEXT("AutoBlockDetails_IgnoreStrategyTooltip", "Legacy behavior: assign to first block, or ignore if none.")
+			});
+	}
+
+	{
+		AutoBlocksMergeStrategies.Empty();
+		AutoBlocksMergeStrategiesOptions.Empty();
+
+		AutoBlocksMergeStrategies.Add(MakeShareable(new FString("Don't merge")));
+		AutoBlocksMergeStrategiesOptions.Add(
+			{
+				ECustomizableObjectLayoutAutomaticBlocksMergeStrategy::DontMerge ,
+				LOCTEXT("AutoBlockMerge_DontMergeTooltip", "Don't merge and make each UV island a unique block.")
+			});
+
+		AutoBlocksMergeStrategies.Add(MakeShareable(new FString("Merge child blocks")));
+		AutoBlocksMergeStrategiesOptions.Add(
+			{
+				ECustomizableObjectLayoutAutomaticBlocksMergeStrategy::MergeChildBlocks,
+				LOCTEXT("AutoBlockMerge_ChildBlocksTooltip", "Merge the blocks that are already fully included in another block.")
+			});
+	}
+
+
+	// Array of available grid size options
+	const int32 MaxGridSize = 128;
+	LayoutGridSizes.Empty();
+
+	for (int32 Size = 1; Size <= MaxGridSize; Size *= 2)
+	{
+		LayoutGridSizes.Add(MakeShareable(new FString(FString::Printf(TEXT("%d x %d"), Size, Size))));
+	}
+
+	{
+		// Block reduction methods options
+		BlockReductionMethods.Empty();
+		BlockReductionMethods.Add(MakeShareable(new FString("Halve")));
+		BlockReductionMethodsTooltips.Add(LOCTEXT("LayoutEditor_HalveRedMethodTooltip", "Blocks will be reduced by half each time."));
+
+		BlockReductionMethods.Add(MakeShareable(new FString("Unitary")));
+		BlockReductionMethodsTooltips.Add(LOCTEXT("LayoutEditor_UnitaryRedMethodTooltip", "Blocks will be reduced by one unit each time."));
+	}
+}
+
+
+void SCustomizableObjectLayoutEditor::SetLayout( UCustomizableObjectLayout* InLayout )
+{
+	if (!Node.IsValid())
+	{
+		return;
+	}
+
+	CurrentLayout = InLayout;
+
+	FillLayoutComboBoxOptions();
+
+	UCustomizableObjectLayout* LayoutForUVs = UVOverrideLayout ? UVOverrideLayout : CurrentLayout;
+
+	TArray<FVector2f> UVs;
+	TArray<FVector2f> UnassignedUVs;
+	if (CurrentLayout)
+	{
+		LayoutForUVs->GetUVs(UVs);
+
+		UnassignedUVs = TArray<FVector2f>();
+		
+		if (CurrentLayout->UnassignedUVs.Num())
+		{
+			UnassignedUVs = LayoutForUVs->UnassignedUVs[0];
+		}
+	}
+
+	// Save some layout widget state to persist between updates
+	bool bHadWidget = false;
+	SCustomizableObjectLayoutGrid::FPointOfView OldView;
+	if (LayoutGridWidget)
+	{
+		bHadWidget = true;
+		OldView = LayoutGridWidget->PointOfView;
+	}
+
+	this->ChildSlot
+	[
+		SNew(SVerticalBox)
+
+		+ SVerticalBox::Slot()
+		.Padding(0.0f,2.0f,0.0f,0.0f )
+		.AutoHeight()
+		[
+			BuildLayoutToolBar()
+		]
+
+		+ SVerticalBox::Slot()
+		.FillHeight(1)
+		[
+			SAssignNew(LayoutGridWidget, SCustomizableObjectLayoutGrid)
+			.Mode(this, &SCustomizableObjectLayoutEditor::GetGridMode)
+			.GridSize(this, &SCustomizableObjectLayoutEditor::GetGridSize)
+			.Blocks(this, &SCustomizableObjectLayoutEditor::GetBlocks)
+			.UVLayout(UVs)
+			.UnassignedUVLayoutVertices(UnassignedUVs)
+			.SelectionColor(FColor(75, 106, 230, 155))
+			.OnBlockChanged(this, &SCustomizableObjectLayoutEditor::OnBlockChanged)
+			.OnDeleteBlocks(this, &SCustomizableObjectLayoutEditor::OnRemoveBlock)
+			.OnAddBlockAt(this, &SCustomizableObjectLayoutEditor::OnAddBlockAt)
+			.OnSetBlockPriority(this, &SCustomizableObjectLayoutEditor::OnSetBlockPriority)
+			.OnSetReduceBlockSymmetrically(this, &SCustomizableObjectLayoutEditor::OnSetBlockReductionSymmetry)
+			.OnSetReduceBlockByTwo(this, &SCustomizableObjectLayoutEditor::OnSetBlockReductionByTwo)
+			.OnSetBlockMask(this, &SCustomizableObjectLayoutEditor::OnSetBlockMask)
+		]
+	];	
+
+	if (CurrentLayout)
+	{
+		CurrentLayout->GenerateAutomaticBlocksFromUVs();
+	}
+
+	if (bHadWidget)
+	{
+		LayoutGridWidget->PointOfView = OldView;
+	}
+}
+
+
+void SCustomizableObjectLayoutEditor::SetUVsOverride(UCustomizableObjectLayout* InUVOverrideLayout)
+{
+	UVOverrideLayout = InUVOverrideLayout;
+}
+
+
+void SCustomizableObjectLayoutEditor::AddReferencedObjects( FReferenceCollector& Collector )
+{
+	Collector.AddReferencedObject( CurrentLayout );
+}
+
+
+TSharedRef<SWidget> SCustomizableObjectLayoutEditor::BuildLayoutToolBar()
+{
+	FSlimHorizontalToolBarBuilder LayoutToolbarBuilder(UICommandList, FMultiBoxCustomization::None, nullptr, true);
+	LayoutToolbarBuilder.SetLabelVisibility(EVisibility::Visible);
+
+	//Getting toolbar style
+	const ISlateStyle* const StyleSet = &FCoreStyle::Get();
+	const FName& StyleName = "ToolBar";
+
+	if (CurrentLayout && CurrentLayout->PackingStrategy!=ECustomizableObjectTextureLayoutPackingStrategy::Overlay)
+	{
+		LayoutToolbarBuilder.BeginSection("Blocks");
+		{
+			LayoutToolbarBuilder.AddToolBarButton(FLayoutEditorCommands::Get().AddBlock);
+			LayoutToolbarBuilder.AddToolBarButton(FLayoutEditorCommands::Get().RemoveBlock);
+			// Disable block consolidation if we are defininf blocks on top of another layout
+			bool bCanConsolidate = !UVOverrideLayout;
+
+			// Disable block consolidation if the cuyrrent automatic strategy doesn't generate blocks.
+			if (CurrentLayout->AutomaticBlocksStrategy != ECustomizableObjectLayoutAutomaticBlocksStrategy::Rectangles)
+			{
+				bCanConsolidate = false;
+			}
+
+			if (bCanConsolidate)
+			{
+				LayoutToolbarBuilder.AddToolBarButton(FLayoutEditorCommands::Get().ConsolidateBlocks);
+			}
+		}
+		LayoutToolbarBuilder.EndSection();
+	}
+
+	LayoutToolbarBuilder.BeginSection("Info");
+	{
+		LayoutToolbarBuilder.AddWidget
+		(
+			SNew(SBox)
+			.VAlign(EVerticalAlignment::VAlign_Center)
+			.HAlign(EHorizontalAlignment::HAlign_Left)
+			[
+				SNew(SImage)
+				.Image(UE_MUTABLE_GET_BRUSH(TEXT("Icons.Info")))
+				.ToolTip(GenerateInfoToolTip())
+			]
+		);
+	}
+	LayoutToolbarBuilder.EndSection();
+
+	return
+	SNew(SHorizontalBox)
+	+SHorizontalBox::Slot()
+	.Padding(4,0)
+	[
+		SNew(SBorder)
+		.Padding(2.0f)
+		.BorderImage(UE_MUTABLE_GET_BRUSH("NoBorder"))
+		.IsEnabled( FSlateApplication::Get().GetNormalExecutionAttribute() )
+		[
+			LayoutToolbarBuilder.MakeWidget()
+		]
+	];
+}
+
+
+void SCustomizableObjectLayoutEditor::OnAddBlock()
+{
+	if (CurrentLayout)
+	{
+		const FScopedTransaction Transaction(LOCTEXT("OnAddBlock", "Add Block"));
+		CurrentLayout->Modify();
+
+		FCustomizableObjectLayoutBlock Block;
+		CurrentLayout->Blocks.Add(Block);
+
+		CurrentLayout->GenerateAutomaticBlocksFromUVs();
+
+		if (LayoutGridWidget.IsValid())
+		{
+			LayoutGridWidget->SetSelectedBlock(Block.Id);
+		}
+	}
+}
+
+
+void SCustomizableObjectLayoutEditor::OnAddBlockAt(const FIntPoint Min, const FIntPoint Max)
+{
+	if (CurrentLayout)
+	{
+		const FScopedTransaction Transaction(LOCTEXT("OnAddBlockAt", "Add Block"));
+		CurrentLayout->Modify();
+		
+		FCustomizableObjectLayoutBlock block(Min,Max);
+		CurrentLayout->Blocks.Add(block);
+
+		CurrentLayout->GenerateAutomaticBlocksFromUVs();
+	}
+}
+
+
+void SCustomizableObjectLayoutEditor::OnRemoveBlock()
+{
+	if (CurrentLayout)
+	{
+		if (LayoutGridWidget.IsValid())
+		{
+			const FScopedTransaction Transaction(LOCTEXT("OnRemoveBlock", "Remove Block"));
+			CurrentLayout->Modify();
+
+			const TArray<FGuid>& Selected = LayoutGridWidget->GetSelectedBlocks();
+			for (TArray<FCustomizableObjectLayoutBlock>::TIterator It = CurrentLayout->Blocks.CreateIterator(); It; ++It)
+			{
+				if (Selected.Contains(It->Id))
+				{
+					It.RemoveCurrent();
+				}
+			}
+
+			CurrentLayout->GenerateAutomaticBlocksFromUVs();
+		}
+	}
+}
+
+
+void SCustomizableObjectLayoutEditor::OnMeshSectionChanged(TSharedPtr<FString> MeshSection, ESelectInfo::Type SelectInfo)
+{
+	if (!Node.IsValid())
+	{
+		return;
+	}
+
+	TSharedPtr<FString> UVChannel = UVChannelComboBox->GetSelectedItem();
+	UpdateLayout(FindSelectedLayout(MeshSection, UVChannel));
+}
+
+
+void SCustomizableObjectLayoutEditor::OnUVChannelChanged(TSharedPtr<FString> UVChannel, ESelectInfo::Type SelectInfo)
+{
+	if (!Node.IsValid())
+	{
+		return;
+	}
+
+	TSharedPtr<FString> MeshSection = MeshSectionComboBox->GetSelectedItem();
+	UpdateLayout(FindSelectedLayout(MeshSection, UVChannel));
+}
+
+
+void SCustomizableObjectLayoutEditor::OnGridSizeChanged(TSharedPtr<FString> NewSelection, ESelectInfo::Type SelectInfo)
+{
+	if (!Node.IsValid() || !CurrentLayout)
+	{
+		return;
+	}
+
+	const int32 Size = 1 << LayoutGridSizes.Find(NewSelection);
+
+	if (CurrentLayout->GetGridSize().X != Size || CurrentLayout->GetGridSize().Y != Size)
+	{
+		CurrentLayout->SetGridSize(FIntPoint(Size));
+
+		Node->Modify();
+
+		// Reset to update the UI.
+		UpdateLayout(CurrentLayout);
+	}
+}
+
+
+void SCustomizableObjectLayoutEditor::OnMaxGridSizeChanged(TSharedPtr<FString> NewSelection, ESelectInfo::Type SelectInfo)
+{
+
+	if (!Node.IsValid() || !CurrentLayout)
+	{
+		return;
+	}
+
+	const int32 Size = 1 << LayoutGridSizes.Find(NewSelection);
+
+	if (CurrentLayout->GetMaxGridSize().X != Size || CurrentLayout->GetMaxGridSize().Y != Size)
+	{
+		CurrentLayout->SetMaxGridSize(FIntPoint(Size));
+
+		Node->Modify();
+
+		// Reset to update the UI.
+		UpdateLayout(CurrentLayout);
+	}
+}
+
+
+void SCustomizableObjectLayoutEditor::OnAutoBlocksChanged(TSharedPtr<FString> NewSelection, ESelectInfo::Type SelectInfo)
+{
+	if (!Node.IsValid() || !CurrentLayout)
+	{
+		return;
+	}
+
+	uint32 Selection = AutoBlocksStrategies.IndexOfByKey(NewSelection);
+	
+	if (CurrentLayout->AutomaticBlocksStrategy != AutoBlocksStrategiesOptions[Selection].Value)
+	{
+		CurrentLayout->AutomaticBlocksStrategy = AutoBlocksStrategiesOptions[Selection].Value;
+
+		Node->Modify();
+
+		// Reset to update the UI.
+		UpdateLayout(CurrentLayout);
+	}
+}
+
+
+void SCustomizableObjectLayoutEditor::OnAutoBlocksMergeChanged(TSharedPtr<FString> NewSelection, ESelectInfo::Type SelectInfo)
+{
+	if (!Node.IsValid() || !CurrentLayout)
+	{
+		return;
+	}
+
+	uint32 Selection = AutoBlocksMergeStrategies.IndexOfByKey(NewSelection);
+	
+	if (CurrentLayout->AutomaticBlocksMergeStrategy != AutoBlocksMergeStrategiesOptions[Selection].Value)
+	{
+		CurrentLayout->AutomaticBlocksMergeStrategy = AutoBlocksMergeStrategiesOptions[Selection].Value;
+
+		Node->Modify();
+	
+		// Reset to update the UI.
+		UpdateLayout(CurrentLayout);
+	}
+}
+
+
+void SCustomizableObjectLayoutEditor::OnReductionMethodChanged(TSharedPtr<FString> NewSelection, ESelectInfo::Type SelectInfo)
+{
+	if (!Node.IsValid() || !CurrentLayout)
+	{
+		return;
+	}
+
+	uint32 ReductionMethod = BlockReductionMethods.IndexOfByKey(NewSelection);
+	
+	if (CurrentLayout->BlockReductionMethod != (ECustomizableObjectLayoutBlockReductionMethod)ReductionMethod)
+	{
+		CurrentLayout->BlockReductionMethod = (ECustomizableObjectLayoutBlockReductionMethod)ReductionMethod;
+
+		Node->Modify();
+
+		// Reset to update the UI.
+		UpdateLayout(CurrentLayout);
+	}
+}
+
+
+void SCustomizableObjectLayoutEditor::OnPackagingStrategyChanged(TSharedPtr<FString> NewSelection, ESelectInfo::Type SelectInfo)
+{
+	if (!Node.IsValid() || !CurrentLayout)
+	{
+		return;
+	}
+
+	const uint32 Index = LayoutPackingStrategies.IndexOfByKey(NewSelection);
+	if (CurrentLayout->PackingStrategy != LayoutPackingStrategiesOptions[Index].Value)
+	{
+		CurrentLayout->PackingStrategy = LayoutPackingStrategiesOptions[Index].Value;
+		
+		Node->Modify();
+	
+		// Reset to update the UI.
+		UpdateLayout(CurrentLayout);
+	}
+}
+
+
+void SCustomizableObjectLayoutEditor::OnIgnoreErrorsCheckStateChanged(ECheckBoxState CheckedBoxState)
+{
+	if (!Node.IsValid() || !CurrentLayout)
+	{
+		return;
+	}
+
+	const bool bIsChecked = CheckedBoxState == ECheckBoxState::Checked;
+	if (CurrentLayout->GetIgnoreVertexLayoutWarnings() != bIsChecked)
+	{
+		CurrentLayout->SetIgnoreVertexLayoutWarnings(bIsChecked);
+
+		Node->Modify();
+
+		// Reset to update the UI.
+		UpdateLayout(CurrentLayout);
+	}
+}
+
+
+void SCustomizableObjectLayoutEditor::OnIgnoreErrorsLODBoxValueChanged(int32 Value)
+{
+	if (!Node.IsValid() || !CurrentLayout)
+	{
+		return;
+	}
+
+	if (CurrentLayout->GetFirstLODToIgnoreWarnings() != Value)
+	{
+		CurrentLayout->SetIgnoreWarningsLOD(Value);
+
+		Node->Modify();
+
+		// Reset to update the UI.
+		UpdateLayout(CurrentLayout);
+	}
+}
+
+
+void SCustomizableObjectLayoutEditor::OnResetSelection()
+{
+	if (MeshSectionComboBox.IsValid())
+	{
+		MeshSectionComboBox->SetSelectedItem(NothingSelectedString);
+	}
+
+	UpdateLayout(nullptr);
+}
+
+
+void SCustomizableObjectLayoutEditor::OnConsolidateBlocks()
+{
+	if (CurrentLayout)
+	{
+		const FScopedTransaction Transaction(LOCTEXT("OnConsolidateBlocks", "Consolidate Blocks"));
+		CurrentLayout->Modify();
+				
+		CurrentLayout->GenerateAutomaticBlocksFromUVs();
+
+		CurrentLayout->ConsolidateAutomaticBlocks();
+	}
+}
+
+
+ELayoutGridMode SCustomizableObjectLayoutEditor::GetGridMode() const
+{
+	if (CurrentLayout && CurrentLayout->PackingStrategy != ECustomizableObjectTextureLayoutPackingStrategy::Overlay)
+	{
+		return ELayoutGridMode::ELGM_Edit;
+	}
+
+	return ELayoutGridMode::ELGM_ShowUVsOnly;
+}
+
+
+FIntPoint SCustomizableObjectLayoutEditor::GetGridSize() const
+{
+	if (CurrentLayout)
+	{
+		return CurrentLayout->GetGridSize();
+	}
+	return FIntPoint(1);
+}
+
+
+void SCustomizableObjectLayoutEditor::OnBlockChanged( FGuid BlockId, FIntRect Block )
+{
+	if (CurrentLayout)
+	{
+		const FScopedTransaction Transaction(LOCTEXT("OnBlockChanged", "Edit Block"));
+		CurrentLayout->Modify();
+		
+		for (FCustomizableObjectLayoutBlock& B : CurrentLayout->Blocks)
+		{
+			if (B.Id == BlockId)
+			{
+				B.Min = Block.Min;
+				B.Max = Block.Max;
+
+				break;
+			}
+		}
+
+		CurrentLayout->GenerateAutomaticBlocksFromUVs();
+	}
+}
+
+
+void SCustomizableObjectLayoutEditor::OnSetBlockPriority(int32 InValue)
+{
+	if (CurrentLayout)
+	{
+		const FScopedTransaction Transaction(LOCTEXT("OnSetBlockPriority", "Change Block Priority"));
+		CurrentLayout->Modify();
+		
+		if (LayoutGridWidget.IsValid())
+		{
+			const TArray<FGuid>& SelectedBlocks = LayoutGridWidget->GetSelectedBlocks();
+
+			for (FCustomizableObjectLayoutBlock& Block : CurrentLayout->Blocks)
+			{
+				if (SelectedBlocks.Contains(Block.Id))
+				{
+					Block.Priority = InValue;
+				}
+			}
+		}
+	}
+}
+
+
+void SCustomizableObjectLayoutEditor::OnSetBlockReductionSymmetry(bool bInValue)
+{
+	if (CurrentLayout)
+	{
+		if (LayoutGridWidget.IsValid())
+		{
+			const FScopedTransaction Transaction(LOCTEXT("OnSetBlockReductionSymmetry", "Change Block Symetry"));
+			CurrentLayout->Modify();
+			
+			const TArray<FGuid>& SelectedBlocks = LayoutGridWidget->GetSelectedBlocks();
+
+			for (FCustomizableObjectLayoutBlock& Block : CurrentLayout->Blocks)
+			{
+				if (SelectedBlocks.Contains(Block.Id))
+				{
+					Block.bReduceBothAxes = bInValue;
+				}
+			}
+		}
+	}
+}
+
+
+void SCustomizableObjectLayoutEditor::OnSetBlockReductionByTwo(bool bInValue)
+{
+	if (CurrentLayout)
+	{
+		if (LayoutGridWidget.IsValid())
+		{
+			const FScopedTransaction Transaction(LOCTEXT("OnSetBlockReductionByTwo", "Change Block Reduction By Two"));
+			CurrentLayout->Modify();
+			
+			const TArray<FGuid>& SelectedBlocks = LayoutGridWidget->GetSelectedBlocks();
+
+			for (FCustomizableObjectLayoutBlock& Block : CurrentLayout->Blocks)
+			{
+				if (SelectedBlocks.Contains(Block.Id))
+				{
+					Block.bReduceByTwo = bInValue;
+				}
+			}
+		}
+	}
+}
+
+
+void SCustomizableObjectLayoutEditor::OnSetBlockMask(UTexture2D* InValue)
+{
+	if (CurrentLayout)
+	{
+		const FScopedTransaction Transaction(LOCTEXT("OnSetBlockMask", "Change Block Mask"));
+		CurrentLayout->Modify();
+
+		if (LayoutGridWidget.IsValid())
+		{
+			const TArray<FGuid>& SelectedBlocks = LayoutGridWidget->GetSelectedBlocks();
+
+			for (FCustomizableObjectLayoutBlock& Block : CurrentLayout->Blocks)
+			{
+				if (SelectedBlocks.Contains(Block.Id))
+				{
+					Block.Mask = InValue;
+				}
+			}
+		}
+
+		CurrentLayout->GenerateAutomaticBlocksFromUVs();
+	}
+}
+
+
+void SCustomizableObjectLayoutEditor::UpdateLayout(UCustomizableObjectLayout* Layout)
+{
+	OnPreUpdateLayoutDelegate.ExecuteIfBound();
+
+	SetLayout(Layout);
+}
+
+
+TArray<FCustomizableObjectLayoutBlock> SCustomizableObjectLayoutEditor::GetBlocks() const
+{
+	TArray<FCustomizableObjectLayoutBlock> Blocks;
+
+	if (CurrentLayout)
+	{
+		Blocks = CurrentLayout->Blocks;
+		Blocks.Append(CurrentLayout->AutomaticBlocks);
+	}
+
+	return Blocks;
+}
+
+
+void SCustomizableObjectLayoutEditor::BindCommands()
+{
+	// Register our commands. This will only register them if not previously registered
+	FLayoutEditorCommands::Register();
+
+	const FLayoutEditorCommands& Commands = FLayoutEditorCommands::Get();
+
+	UICommandList->MapAction(
+		Commands.AddBlock,
+		FExecuteAction::CreateSP( this, &SCustomizableObjectLayoutEditor::OnAddBlock ),
+		FCanExecuteAction(),
+		FIsActionChecked() );
+
+	UICommandList->MapAction(
+		Commands.RemoveBlock,
+		FExecuteAction::CreateSP( this, &SCustomizableObjectLayoutEditor::OnRemoveBlock ),
+		FCanExecuteAction(),
+		FIsActionChecked() );
+
+	UICommandList->MapAction(
+		Commands.ConsolidateBlocks,
+		FExecuteAction::CreateSP(this, &SCustomizableObjectLayoutEditor::OnConsolidateBlocks),
+		FCanExecuteAction(),
+		FIsActionChecked());
+}
+
+
+TSharedPtr<IToolTip> SCustomizableObjectLayoutEditor::GenerateInfoToolTip() const
+{
+	TSharedPtr<SGridPanel> ToolTipWidget = SNew(SGridPanel);
+	int32 SlotCount = 0;
+
+	auto BuildShortcutAndTooltip = [ToolTipWidget, &SlotCount](const FText& Shortcut, const FText& Tooltip)
+	{
+		// Command Shortcut
+		ToolTipWidget->AddSlot(0, SlotCount)
+		[
+			SNew(STextBlock)
+			.Text(Shortcut)
+		];
+
+		// Command Explanation
+		ToolTipWidget->AddSlot(1, SlotCount)
+		.Padding(15.0f, 0.0f, 0.0f, 0.0f)
+		[
+			SNew(STextBlock)
+			.Text(Tooltip)
+		];
+
+		++SlotCount;
+	};
+
+	// Duplicate command
+	if (CurrentLayout && CurrentLayout->PackingStrategy != ECustomizableObjectTextureLayoutPackingStrategy::Overlay)
+	{
+		BuildShortcutAndTooltip(LOCTEXT("ShortCut_DuplicateBlocks", "CTRL + D"), LOCTEXT("Tooltip_DuplicateBlocks", "Duplicate selected block/s"));
+		BuildShortcutAndTooltip(LOCTEXT("ShortCut_CreateNewBlock", "CTRL + N"), LOCTEXT("Tooltip_CreateNewBlock", "Create new block"));
+		BuildShortcutAndTooltip(LOCTEXT("ShortCut_FillGridSize", "CTRL + F"), LOCTEXT("Tooltip_FillGridSize", "Resize selected block/s to grid size"));
+		BuildShortcutAndTooltip(LOCTEXT("ShortCut_DeleteSelectedBlock", "DEL"), LOCTEXT("Tooltip_DeleteSelectedBlock", "Delete selected block/s"));
+		BuildShortcutAndTooltip(LOCTEXT("ShortCut_SelectMultipleBlocksOneByOne", "SHIFT + L Click"), LOCTEXT("Tooltip_SelectMultipleBlocksOneByOne", "Select multiple blocks one by one"));
+		BuildShortcutAndTooltip(LOCTEXT("ShortCut_SelectMultipleBlocks", "L Click + Drag"), LOCTEXT("Tooltip_SelectMultipleBlocks", "Select blocks that intersect with the yellow rectangle"));
+	}
+	BuildShortcutAndTooltip(LOCTEXT("ShortCut_Pan", "M Click + Drag"), LOCTEXT("Tooltip_Pan", "Pan the UV view."));
+	BuildShortcutAndTooltip(LOCTEXT("ShortCut_Zoom", "M Wheel"), LOCTEXT("Tooltip_Zoom", "Zoom in and out the UV view."));
+
+	return SNew(SToolTip)
+	[
+		ToolTipWidget.ToSharedRef()
+	];
+}
+
+
+TSharedRef<SWidget> SCustomizableObjectLayoutEditor::GenerateAutoBlocksComboBox(TSharedPtr<FString> InItem) const
+{	
+	//A list of tool tips should have been populated in a 1 to 1 correspondance
+	check(AutoBlocksStrategies.Num() == AutoBlocksStrategiesOptions.Num());
+
+	FText AutoBlocksName;
+	FText AutoBlocksTooltip;
+
+	if (InItem.IsValid())
+	{
+		AutoBlocksName = FText::FromString(*InItem.Get());
+		int32 TooltipIndex = AutoBlocksStrategies.IndexOfByKey(InItem);
+
+		if (ensure(AutoBlocksStrategiesOptions.IsValidIndex(TooltipIndex)))
+		{
+			AutoBlocksTooltip = AutoBlocksStrategiesOptions[TooltipIndex].Tooltip;
+		}
+	}
+
+	return SNew(STextBlock)
+		.Text(AutoBlocksName)
+		.ToolTipText(AutoBlocksTooltip)
+		.Font(IDetailLayoutBuilder::GetDetailFont());
+}
+
+
+FText SCustomizableObjectLayoutEditor::GetAutoBlocksName() const
+{
+	if (CurrentLayout)
+	{
+		const uint32 ReductionIndex = (uint32)CurrentLayout->AutomaticBlocksStrategy;
+		if (ensure(AutoBlocksStrategies.IsValidIndex(ReductionIndex)))
+		{
+			return FText::FromString(*AutoBlocksStrategies[ReductionIndex]);
+		}
+	}
+
+	return FText();
+}
+
+
+FText SCustomizableObjectLayoutEditor::GetAutoBlocksTooltip() const
+{
+	if (CurrentLayout)
+	{
+		const uint32 ReductionIndex = (uint32)CurrentLayout->AutomaticBlocksStrategy;
+		if (ensure(AutoBlocksStrategiesOptions.IsValidIndex(ReductionIndex)))
+		{
+			return AutoBlocksStrategiesOptions[ReductionIndex].Tooltip;
+		}
+	}
+
+	return FText();
+}
+
+
+TSharedRef<SWidget> SCustomizableObjectLayoutEditor::GenerateAutoBlocksMergeComboBox(TSharedPtr<FString> InItem) const
+{	
+	//A list of tool tips should have been populated in a 1 to 1 correspondance
+	check(AutoBlocksMergeStrategies.Num() == AutoBlocksMergeStrategiesOptions.Num());
+
+	FText AutoBlocksMergeName;
+	FText AutoBlocksMergeTooltip;
+
+	if (InItem.IsValid())
+	{
+		AutoBlocksMergeName = FText::FromString(*InItem.Get());
+		int32 TooltipIndex = AutoBlocksMergeStrategies.IndexOfByKey(InItem);
+
+		if (ensure(AutoBlocksMergeStrategiesOptions.IsValidIndex(TooltipIndex)))
+		{
+			AutoBlocksMergeTooltip = AutoBlocksMergeStrategiesOptions[TooltipIndex].Tooltip;
+		}
+	}
+
+	return SNew(STextBlock)
+		.Text(AutoBlocksMergeName)
+		.ToolTipText(AutoBlocksMergeTooltip)
+		.Font(IDetailLayoutBuilder::GetDetailFont());
+}
+
+
+FText SCustomizableObjectLayoutEditor::GetAutoBlocksMergeName() const
+{
+	if (CurrentLayout)
+	{
+		const uint32 AutoMergeIndex = (uint32)CurrentLayout->AutomaticBlocksMergeStrategy;
+		if (ensure(AutoBlocksMergeStrategies.IsValidIndex(AutoMergeIndex)))
+		{
+			return FText::FromString(*AutoBlocksMergeStrategies[AutoMergeIndex]);
+		}
+	}
+
+	return FText();
+}
+
+
+FText SCustomizableObjectLayoutEditor::GetAutoBlocksMergeTooltip() const
+{
+	if (CurrentLayout)
+	{
+		const uint32 AutoMergeIndex = (uint32)CurrentLayout->AutomaticBlocksMergeStrategy;
+		if (ensure(AutoBlocksMergeStrategiesOptions.IsValidIndex(AutoMergeIndex)))
+		{
+			return AutoBlocksMergeStrategiesOptions[AutoMergeIndex].Tooltip;
+		}
+	}
+
+	return FText();
+}
+
+
+TSharedRef<SWidget> SCustomizableObjectLayoutEditor::GenerateReductionMethodComboBox(TSharedPtr<FString> InItem) const
+{
+	//A list of tool tips should have been populated in a 1 to 1 correspondance
+	check(BlockReductionMethods.Num() == BlockReductionMethodsTooltips.Num());
+
+	FText ReductionMethodName;
+	FText ReductionMethodTooltip;
+	
+	if (InItem.IsValid())
+	{
+		ReductionMethodName = FText::FromString(*InItem.Get());
+		int32 TooltipIndex = BlockReductionMethods.IndexOfByKey(InItem);
+	
+		if (ensure(BlockReductionMethodsTooltips.IsValidIndex(TooltipIndex)))
+		{
+			ReductionMethodTooltip = BlockReductionMethodsTooltips[TooltipIndex];
+		}
+	}
+	
+	return SNew(STextBlock)
+		.Text(ReductionMethodName)
+		.ToolTipText(ReductionMethodTooltip)
+		.Font(IDetailLayoutBuilder::GetDetailFont());
+}
+
+
+FText SCustomizableObjectLayoutEditor::GetLayoutReductionMethodName() const
+{
+	if (CurrentLayout)
+	{
+		const uint32 ReductionIndex = (uint32)CurrentLayout->BlockReductionMethod;
+		if (ensure(BlockReductionMethods.IsValidIndex(ReductionIndex)))
+		{
+			return FText::FromString(*BlockReductionMethods[ReductionIndex]);
+		}
+	}
+
+	return FText();
+}
+
+
+FText SCustomizableObjectLayoutEditor::GetLayoutReductionMethodTooltip() const
+{
+	if (CurrentLayout)
+	{
+		const uint32 ReductionIndex = (uint32)CurrentLayout->BlockReductionMethod;
+		if (ensure(BlockReductionMethods.IsValidIndex(ReductionIndex)))
+		{
+			return BlockReductionMethodsTooltips[ReductionIndex];
+		}
+	}
+	
+	return FText();
+}
+
+
+TSharedRef<SWidget> SCustomizableObjectLayoutEditor::GenerateLayoutPackagingStrategyComboBox(TSharedPtr<FString> InItem) const
+{
+	//A list of tool tips should have been populated in a 1 to 1 correspondance
+	check(LayoutPackingStrategies.Num() == LayoutPackingStrategiesOptions.Num());
+	
+	FText StrategyTooltip;
+	FText StrategyName;
+
+	if (InItem.IsValid())
+	{
+		StrategyName = FText::FromString(*InItem.Get());
+		int32 TooltipIndex = LayoutPackingStrategies.IndexOfByKey(InItem);
+
+		if (ensure(LayoutPackingStrategiesOptions.IsValidIndex(TooltipIndex)))
+		{
+			StrategyTooltip = LayoutPackingStrategiesOptions[TooltipIndex].Tooltip;
+		}
+	}
+
+	return SNew(STextBlock)
+		.Text(StrategyName)
+		.ToolTipText(StrategyTooltip)
+		.Font(IDetailLayoutBuilder::GetDetailFont());
+}
+
+
+FText SCustomizableObjectLayoutEditor::GetLayoutPackagingStrategyName() const
+{
+	if (CurrentLayout)
+	{
+		const uint32 StrategyIndex = (uint32)CurrentLayout->PackingStrategy;
+		if (ensure(LayoutPackingStrategies.IsValidIndex(StrategyIndex)))
+		{
+			return FText::FromString(*LayoutPackingStrategies[StrategyIndex]);
+		}
+	}
+
+	return FText();
+}
+
+
+FText SCustomizableObjectLayoutEditor::GetLayoutPackagingStrategyToolTip() const
+{
+	if (CurrentLayout)
+	{
+		const uint32 StrategyIndex = (uint32)CurrentLayout->PackingStrategy;
+		if (ensure(LayoutPackingStrategiesOptions.IsValidIndex(StrategyIndex)))
+		{
+			return LayoutPackingStrategiesOptions[StrategyIndex].Tooltip;
+		}
+	}
+
+	return FText();
+}
+
+
+UCustomizableObjectLayout* SCustomizableObjectLayoutEditor::FindSelectedLayout(TSharedPtr<FString> MeshSectionName, TSharedPtr<FString> UVChannel)
+{
+	const int32 Index = UVChannels.Find(UVChannel);
+
+	for (FLayoutEditorMeshSection& MeshSection : MeshSections)
+	{
+		if (MeshSection.MeshName != MeshSectionName)
+		{
+			continue;
+		}
+
+		if (MeshSection.Layouts.IsValidIndex(Index))
+		{
+			return MeshSection.Layouts[Index].Get();
+		}
+		else
+		{
+			return nullptr;
+		}
+	}
+
+	return nullptr;
+}
+
+
+void SCustomizableObjectLayoutEditor::FillLayoutComboBoxOptions()
+{
+	if (CurrentLayout)
+	{
+		if (StrategyComboBox)
+		{
+			const uint32 StrategyIndex = (uint32)CurrentLayout->PackingStrategy;
+			StrategyComboBox->SetSelectedItem(LayoutPackingStrategies[StrategyIndex]);
+		}
+
+		if (AutoBlocksComboBox)
+		{
+			const uint32 AutoBlockIndex = (uint32)CurrentLayout->AutomaticBlocksStrategy;
+			AutoBlocksComboBox->SetSelectedItem(AutoBlocksStrategies[AutoBlockIndex]);
+		}
+
+		if (AutoBlocksMergeComboBox)
+		{
+			const uint32 AutoBlockMergeIndex = (uint32)CurrentLayout->AutomaticBlocksMergeStrategy;
+			AutoBlocksMergeComboBox->SetSelectedItem(AutoBlocksMergeStrategies[AutoBlockMergeIndex]);
+		}
+
+		if (GridSizeComboBox)
+		{
+			const int32 Index = FMath::CountTrailingZeros(CurrentLayout->GetGridSize().X);
+			GridSizeComboBox->SetSelectedItem(LayoutGridSizes[Index]);
+		}
+
+		if (MaxGridSizeComboBox)
+		{
+			const int32 Index = FMath::CountTrailingZeros(CurrentLayout->GetMaxGridSize().X);
+			MaxGridSizeComboBox->SetSelectedItem(LayoutGridSizes[Index]);
+		}
+
+		if (ReductionMethodComboBox)
+		{
+			const uint32 ReductionMethodIndex = (uint32)CurrentLayout->BlockReductionMethod;
+			ReductionMethodComboBox->SetSelectedItem(BlockReductionMethods[ReductionMethodIndex]);
+		}
+
+		if (LODSelectorWidget)
+		{
+			LODSelectorWidget->SetEnabled(CurrentLayout->GetIgnoreVertexLayoutWarnings());
+			LODSelectorWidget->SetValue(CurrentLayout->GetFirstLODToIgnoreWarnings());
+		}
+
+		if (LODSelectorTextWidget)
+		{
+			LODSelectorTextWidget->SetEnabled(CurrentLayout->GetIgnoreVertexLayoutWarnings());
+		}
+	}
+}
+
+
+EVisibility SCustomizableObjectLayoutEditor::LayoutOptionsVisibility() const
+{
+	return CurrentLayout ? EVisibility::Visible : EVisibility::Collapsed;
+}
+
+
+EVisibility SCustomizableObjectLayoutEditor::AutoBlocksStrategyVisibility() const
+{
+	return CurrentLayout && CurrentLayout->PackingStrategy != ECustomizableObjectTextureLayoutPackingStrategy::Overlay ? EVisibility::Visible : EVisibility::Collapsed;
+}
+
+
+EVisibility SCustomizableObjectLayoutEditor::AutoBlocksMergeStrategyVisibility() const
+{
+	const bool bIsVisible = CurrentLayout && 
+		CurrentLayout->PackingStrategy != ECustomizableObjectTextureLayoutPackingStrategy::Overlay &&
+		CurrentLayout->AutomaticBlocksStrategy == ECustomizableObjectLayoutAutomaticBlocksStrategy::UVIslands;
+	
+	return bIsVisible ? EVisibility::Visible : EVisibility::Collapsed;
+}
+
+
+EVisibility SCustomizableObjectLayoutEditor::FixedStrategyOptionsVisibility() const
+{
+	return  CurrentLayout && CurrentLayout->PackingStrategy == ECustomizableObjectTextureLayoutPackingStrategy::Fixed ? EVisibility::Visible : EVisibility::Collapsed;
+}
+
+#undef LOCTEXT_NAMESPACE
+
diff --git a/Engine/Plugins/Experimental/Mutable/Source/CustomizableObjectEditor/Private/MuCOE/SCustomizableObjectLayoutEditor.h b/Engine/Plugins/Experimental/Mutable/Source/CustomizableObjectEditor/Private/MuCOE/SCustomizableObjectLayoutEditor.h
new file mode 100644
--- /dev/null
+++ b/Engine/Plugins/Experimental/Mutable/Source/CustomizableObjectEditor/Private/MuCOE/SCustomizableObjectLayoutEditor.h
@@ -0,0 +1,248 @@
+// Copyright Epic Games, Inc. All Rights Reserved.
+
+#pragma once
+
+#include "UObject/GCObject.h"
+#include "Widgets/SCompoundWidget.h"
+#include "Widgets/Input/SSpinBox.h"
+#include "MuCOE/SCustomizableObjectLayoutGrid.h"
+
+namespace ESelectInfo { enum Type : int; }
+
+class FReferenceCollector;
+class ICustomizableObjectInstanceEditor;
+class IDetailCategoryBuilder;
+class IDetailLayoutBuilder;
+class ISlateStyle;
+class SWidget;
+class STextComboBox;
+class SSearchableComboBox;
+class SVerticalBox;
+class SHorizontalBox;
+class FUICommandList;
+class SCustomizableObjectLayoutEditor;
+class SCustomizableObjectLayoutGrid;
+class UCustomizableObjectLayout;
+class UTexture2D;
+
+struct FCustomizableObjectLayoutBlock;
+struct FGuid;
+
+struct FLayoutEditorMeshSection
+{
+	TSharedPtr<FString> MeshName;
+	TArray<TWeakObjectPtr<UCustomizableObjectLayout>> Layouts;
+};
+
+struct FCustomizableObjectLayoutEditorDetailsBuilder
+{
+	void CustomizeDetails(IDetailLayoutBuilder& Details);
+
+	bool bShowLayoutSelector = false;
+	bool bShowPackagingStrategy = false;
+	bool bShowAutomaticGenerationSettings = false;
+	bool bShowGridSize = false;
+	bool bShowMaxGridSize = false;
+	bool bShowReductionMethods = false;
+	bool bShowWarningSettings = false;
+
+	TSharedPtr<SCustomizableObjectLayoutEditor> LayoutEditor;
+};
+
+/**
+ * CustomizableObject Editor Preview viewport widget
+ */
+class SCustomizableObjectLayoutEditor : public SCompoundWidget, public FGCObject
+{
+
+	friend FCustomizableObjectLayoutEditorDetailsBuilder;
+
+public:
+
+	DECLARE_DELEGATE(FOnPreUpdateLayout);
+
+	SLATE_BEGIN_ARGS( SCustomizableObjectLayoutEditor ){}
+		SLATE_ATTRIBUTE(TArray<FLayoutEditorMeshSection>, MeshSections)
+		SLATE_ARGUMENT_DEFAULT(UObject*, Node) { nullptr };
+		SLATE_EVENT(FOnPreUpdateLayout, OnPreUpdateLayoutDelegate)
+	SLATE_END_ARGS()
+
+	SCustomizableObjectLayoutEditor();
+
+	void Construct(const FArguments& InArgs);
+
+	// FSerializableObject interface
+	void AddReferencedObjects( FReferenceCollector& Collector ) override;
+	virtual FString GetReferencerName() const override
+	{
+		return TEXT("SCustomizableObjectLayoutEditor");
+	}
+	// End of FSerializableObject interface
+
+
+	/** Binds commands associated with the viewport client. */
+	void BindCommands();
+
+	void SetLayout(UCustomizableObjectLayout* Layout);
+
+	/**
+	* UVOverrideLayout parameter can be speicifed to show different UVs in the widget instead of the ones in Layout.
+	*/
+	void SetUVsOverride(UCustomizableObjectLayout* UVOverrideLayout);
+
+	void UpdateLayout(UCustomizableObjectLayout* Layout);
+
+private:
+
+	FOnPreUpdateLayout OnPreUpdateLayoutDelegate;
+
+	TWeakObjectPtr<UObject> Node;
+
+	TArray<FLayoutEditorMeshSection> MeshSections;
+	TArray<TSharedPtr<FString>> MeshSectionNames;
+
+	TArray<TSharedPtr<FString>> UVChannels;
+
+	/** Layout whose blocksa are being edited. */
+	TObjectPtr<UCustomizableObjectLayout> CurrentLayout;
+
+	/** If valid, layout use to show the UVs instead of CurrentLayout. */
+	TObjectPtr<UCustomizableObjectLayout> UVOverrideLayout;
+	
+	/** List of available layout grid sizes. */
+	TArray<TSharedPtr<FString>> LayoutGridSizes;
+
+	TSharedPtr<FString> NothingSelectedString;
+
+	struct FAutoBlocksStrategyOption
+	{
+		ECustomizableObjectLayoutAutomaticBlocksStrategy Value;
+		FText Tooltip;
+	};
+	TArray<TSharedPtr<FString>> AutoBlocksStrategies;
+	TArray<FAutoBlocksStrategyOption> AutoBlocksStrategiesOptions;
+
+	struct FAutoBlocksMergeStrategyOption
+	{
+		ECustomizableObjectLayoutAutomaticBlocksMergeStrategy Value;
+		FText Tooltip;
+	};
+	TArray<TSharedPtr<FString>> AutoBlocksMergeStrategies;
+	TArray<FAutoBlocksMergeStrategyOption> AutoBlocksMergeStrategiesOptions;
+
+	/** List of available block reduction methods. */
+	TArray<TSharedPtr<FString>> BlockReductionMethods;
+	TArray<FText> BlockReductionMethodsTooltips;
+
+	/** List of available layout packing strategies. */
+	struct FPackingStrategyOption
+	{
+		ECustomizableObjectTextureLayoutPackingStrategy Value;
+		FText Tooltip;
+	};
+	TArray<TSharedPtr<FString>> LayoutPackingStrategies;
+	TArray<FPackingStrategyOption> LayoutPackingStrategiesOptions;
+
+	/** The list of UI Commands executable */
+	TSharedRef<FUICommandList> UICommandList;
+
+	// Layout -------------
+	// ComboBox widget to select a column from the NodeTable
+	TSharedPtr<STextComboBox> MeshSectionComboBox;
+	TSharedPtr<STextComboBox> UVChannelComboBox;
+
+
+	// ComboBox widget to select a Strategy from the Selected Layout. SSearchableComboBox allows us to set a custom tooltip per option.
+	TSharedPtr<SSearchableComboBox> StrategyComboBox;
+
+	TSharedPtr<SSearchableComboBox> AutoBlocksComboBox;
+	TSharedPtr<SSearchableComboBox> AutoBlocksMergeComboBox;
+
+	// ComboBox widget to select a Grid Size from the Selected Layout
+	TSharedPtr<STextComboBox> GridSizeComboBox;
+
+	// ComboBox widget to select a Max Grid Size from the Selected Layout
+	TSharedPtr<STextComboBox> MaxGridSizeComboBox;
+
+	// ComboBox widget to select a Reduction Method from the Selected Layout
+	TSharedPtr<SSearchableComboBox> ReductionMethodComboBox;
+
+	// Widget to select at which LOD layout vertex warnings will start to be ignored
+	TSharedPtr<SSpinBox<int32>> LODSelectorWidget;
+	TSharedPtr<STextBlock> LODSelectorTextWidget;
+
+	TSharedPtr<SCustomizableObjectLayoutGrid> LayoutGridWidget;
+
+private:
+
+	/** */
+	ELayoutGridMode GetGridMode() const;
+	FIntPoint GetGridSize() const;
+	void OnBlockChanged(FGuid BlockId, FIntRect Block );
+	TArray<FCustomizableObjectLayoutBlock> GetBlocks() const;
+
+	/** Callbacks from the layout block editor. */
+	TSharedRef<SWidget> BuildLayoutToolBar();
+
+	void OnAddBlock();
+	void OnAddBlockAt(const FIntPoint Min, const FIntPoint Max);
+	void OnRemoveBlock();
+
+	/** */
+	void OnMeshSectionChanged(TSharedPtr<FString> NewSelection, ESelectInfo::Type SelectInfo);
+	void OnUVChannelChanged(TSharedPtr<FString> NewSelection, ESelectInfo::Type SelectInfo);
+	void OnGridSizeChanged(TSharedPtr<FString> NewSelection, ESelectInfo::Type SelectInfo);
+	void OnMaxGridSizeChanged(TSharedPtr<FString> NewSelection, ESelectInfo::Type SelectInfo);
+
+	void OnAutoBlocksChanged(TSharedPtr<FString> NewSelection, ESelectInfo::Type SelectInfo);
+	void OnAutoBlocksMergeChanged(TSharedPtr<FString> NewSelection, ESelectInfo::Type SelectInfo);
+	void OnReductionMethodChanged(TSharedPtr<FString> NewSelection, ESelectInfo::Type SelectInfo);
+	void OnPackagingStrategyChanged(TSharedPtr<FString> NewSelection, ESelectInfo::Type SelectInfo);
+
+	void OnIgnoreErrorsCheckStateChanged(ECheckBoxState NewSelection);
+	void OnIgnoreErrorsLODBoxValueChanged(int32 Value);
+
+	void OnResetSelection();
+
+	/** Turn the automatic layout blocks into user-created blocks. */
+	void OnConsolidateBlocks();
+
+	/** Sets the block priority from the input text. */
+	void OnSetBlockPriority(int32 InValue);
+
+	/** Sets the block reduction symmetry option. */
+	void OnSetBlockReductionSymmetry(bool bInValue);
+
+	/** Sets the block reduction ReduceByTwo option. */
+	void OnSetBlockReductionByTwo(bool bInValue);
+
+	/** Callback for block mask change. */
+	void OnSetBlockMask(UTexture2D* InValue);
+
+	UCustomizableObjectLayout* FindSelectedLayout(TSharedPtr<FString> MeshSection, TSharedPtr<FString> UVChannel);
+
+	void FillLayoutComboBoxOptions();
+
+	TSharedPtr<IToolTip> GenerateInfoToolTip() const;
+	
+	TSharedRef<SWidget> GenerateAutoBlocksComboBox(TSharedPtr<FString> InItem) const;
+	FText GetAutoBlocksName() const;
+	FText GetAutoBlocksTooltip() const;	
+
+	TSharedRef<SWidget> GenerateAutoBlocksMergeComboBox(TSharedPtr<FString> InItem) const;
+	FText GetAutoBlocksMergeName() const;
+	FText GetAutoBlocksMergeTooltip() const;
+
+	TSharedRef<SWidget> GenerateReductionMethodComboBox(TSharedPtr<FString> InItem) const;
+	FText GetLayoutReductionMethodName() const;
+	FText GetLayoutReductionMethodTooltip() const;
+
+	TSharedRef<SWidget> GenerateLayoutPackagingStrategyComboBox(TSharedPtr<FString> InItem) const;
+	FText GetLayoutPackagingStrategyName() const;
+	FText GetLayoutPackagingStrategyToolTip() const;
+
+	EVisibility LayoutOptionsVisibility() const;
+	EVisibility AutoBlocksStrategyVisibility() const;
+	EVisibility AutoBlocksMergeStrategyVisibility() const;
+	EVisibility FixedStrategyOptionsVisibility() const;
+};
diff --git a/Engine/Plugins/Experimental/Mutable/Source/CustomizableObjectEditor/Private/MuCOE/SCustomizableObjectNodeLayoutBlocksEditor.cpp b/Engine/Plugins/Experimental/Mutable/Source/CustomizableObjectEditor/Private/MuCOE/SCustomizableObjectNodeLayoutBlocksEditor.cpp
deleted file mode 100644
diff --git a/Engine/Plugins/Experimental/Mutable/Source/CustomizableObjectEditor/Private/MuCOE/SCustomizableObjectNodeLayoutBlocksEditor.h b/Engine/Plugins/Experimental/Mutable/Source/CustomizableObjectEditor/Private/MuCOE/SCustomizableObjectNodeLayoutBlocksEditor.h
deleted file mode 100644
diff --git a/Engine/Plugins/Experimental/Mutable/Source/CustomizableObjectEditor/Public/MuCOE/CustomizableObjectLayout.h b/Engine/Plugins/Experimental/Mutable/Source/CustomizableObjectEditor/Public/MuCOE/CustomizableObjectLayout.h
--- a/Engine/Plugins/Experimental/Mutable/Source/CustomizableObjectEditor/Public/MuCOE/CustomizableObjectLayout.h
+++ b/Engine/Plugins/Experimental/Mutable/Source/CustomizableObjectEditor/Public/MuCOE/CustomizableObjectLayout.h
@@ -111,7 +111,7 @@ public:
 	UCustomizableObjectLayout();
 
 	// Sets the layout parameters
-	void SetLayout(UObject* InMesh, int32 LODIndex, int32 MatIndex, int32 UVIndex);
+	void SetLayout(int32 LODIndex, int32 MatIndex, int32 UVIndex);
 	void SetGridSize(FIntPoint Size);
 	void SetMaxGridSize(FIntPoint Size);
 	void SetLayoutName(FString Name);
@@ -122,14 +122,13 @@ public:
 	int32 GetMaterial() const { return Material; }
 	int32 GetUVChannel() const { return UVChannel; }
 	FString GetLayoutName() const { return LayoutName; }
-	UObject* GetMesh() const { return Mesh; }
+	UObject* GetMesh() const;
 	FIntPoint GetGridSize() const { return GridSize; }
 	FIntPoint GetMaxGridSize() const { return MaxGridSize; }
 	bool GetIgnoreVertexLayoutWarnings() const { return bIgnoreUnassignedVertexWarning; };
 	int32 GetFirstLODToIgnoreWarnings() const { return FirstLODToIgnore; };
 
-	// TODO: Remove the default channel index. Anywhere using the default is probably wrong.
-	void GetUVChannel(TArray<FVector2f>& UVs, int32 UVChannelIndex = 0) const;
+	void GetUVs(TArray<FVector2f>& UVs) const;
 
 	/** Get a block index in the array from its id.Return - 1 if not found. */
 	int32 FindBlock(const FGuid& InId) const;
@@ -165,9 +164,6 @@ public:
 
 private:
 
-	UPROPERTY()
-	TObjectPtr<UObject> Mesh = nullptr;
-
 	UPROPERTY()
 	int32 LOD;
 
diff --git a/Engine/Plugins/Experimental/Mutable/Source/CustomizableObjectEditor/Public/MuCOE/Nodes/CustomizableObjectNodeLayoutBlocks.h b/Engine/Plugins/Experimental/Mutable/Source/CustomizableObjectEditor/Public/MuCOE/Nodes/CustomizableObjectNodeLayoutBlocks.h
--- a/Engine/Plugins/Experimental/Mutable/Source/CustomizableObjectEditor/Public/MuCOE/Nodes/CustomizableObjectNodeLayoutBlocks.h
+++ b/Engine/Plugins/Experimental/Mutable/Source/CustomizableObjectEditor/Public/MuCOE/Nodes/CustomizableObjectNodeLayoutBlocks.h
@@ -43,39 +43,16 @@ public:
 	TObjectPtr<UCustomizableObjectLayout> Layout = nullptr;
 
 	// EdGraphNode interface
-	FText GetNodeTitle(ENodeTitleType::Type TitleType) const override;
-	FLinearColor GetNodeTitleColor() const override;
-	FText GetTooltipText() const override;
 	void Serialize(FArchive& Ar) override;
-	void PinConnectionListChanged(UEdGraphPin* Pin) override;
-	void PostPasteNode() override;
-
 
 	// UCustomizableObjectNode interface
-	virtual void PostBackwardsCompatibleFixup() override;
+	virtual void BackwardsCompatibleFixup(int32 CustomizableObjectCustomVersion) override;
 	void AllocateDefaultPins(UCustomizableObjectNodeRemapPins* RemapPins) override;
 
 	UEdGraphPin* OutputPin() const
 	{
 		return FindPin(TEXT("Layout"));
 	}
-
-	void AddAttachedErrorData(const FAttachedErrorDataView& AttachedErrorData) override;
-	void ResetAttachedErrorData() override;
-
-	bool IsSingleOutputNode() const override;
-
-private:
-
-	/** Last static or skeletal mesh connected. Used to remove the callback once disconnected. */
-	TWeakObjectPtr<UCustomizableObjectNode> LastMeshNodeConnected;
-
-	/** Connected NodeStaticMesh or NodeSkeletalMesh Mesh UPROPERTY changed callback function. Sets the layout mesh. */
-	void MeshPostEditChangeProperty(UObject* Node, FPropertyChangedEvent& FPropertyChangedEvent);
-
-	void LinkPostEditChangePropertyDelegate(UEdGraphPin* Pin);
-
-	void SetLayoutSkeletalMesh();
 };
 
 
diff --git a/Engine/Plugins/Experimental/Mutable/Source/CustomizableObjectEditor/Public/MuCOE/Nodes/CustomizableObjectNodeMesh.h b/Engine/Plugins/Experimental/Mutable/Source/CustomizableObjectEditor/Public/MuCOE/Nodes/CustomizableObjectNodeMesh.h
--- a/Engine/Plugins/Experimental/Mutable/Source/CustomizableObjectEditor/Public/MuCOE/Nodes/CustomizableObjectNodeMesh.h
+++ b/Engine/Plugins/Experimental/Mutable/Source/CustomizableObjectEditor/Public/MuCOE/Nodes/CustomizableObjectNodeMesh.h
@@ -24,17 +24,12 @@ public:
 
 	virtual UTexture2D* FindTextureForPin(const UEdGraphPin* Pin) const PURE_VIRTUAL(UCustomizableObjectNodeMesh::FindTextureForPin, return {};);
 
-	virtual void GetUVChannelForPin(const UEdGraphPin* Pin, TArray<FVector2f>& OutSegments, int32 UVIndex) const PURE_VIRTUAL(UCustomizableObjectNodeMesh::GetUVChannelForPin, );
-
 	/** Returns the Unreal mesh (e.g., USkeletalMesh, UStaticMesh...). */
 	virtual UObject* GetMesh() const PURE_VIRTUAL(UCustomizableObjectNodeMesh::GetMesh, return {};);
 
 	/** Returns the output Mesh pin associated to the given LODIndex and SectionIndex. Override. */
 	virtual UEdGraphPin* GetMeshPin(int32 LOD, int32 SectionIndex) const PURE_VIRTUAL(UCustomizableObjectNodeMesh::GetMeshPin, return {};);
 
-	/** Returns the input Layout pin associated to the given LODIndex, SectionIndex and LayoutIndex. Override. */
-	virtual UEdGraphPin* GetLayoutPin(int32 LODIndex, int32 SectionIndex, int32 LayoutIndex) const PURE_VIRTUAL(UCustomizableObjectNodeMesh::GetLayoutPin, return {};);
-	
 	/** Given a pin, return the Section and Layout index.
 	 *
 	 * Will always return a valid result. The result can be valid but out of sync with respect the Unreal mesh asset.
diff --git a/Engine/Plugins/Experimental/Mutable/Source/CustomizableObjectEditor/Public/MuCOE/Nodes/CustomizableObjectNodeSkeletalMesh.h b/Engine/Plugins/Experimental/Mutable/Source/CustomizableObjectEditor/Public/MuCOE/Nodes/CustomizableObjectNodeSkeletalMesh.h
--- a/Engine/Plugins/Experimental/Mutable/Source/CustomizableObjectEditor/Public/MuCOE/Nodes/CustomizableObjectNodeSkeletalMesh.h
+++ b/Engine/Plugins/Experimental/Mutable/Source/CustomizableObjectEditor/Public/MuCOE/Nodes/CustomizableObjectNodeSkeletalMesh.h
@@ -120,6 +120,18 @@ UCLASS()
 class UCustomizableObjectNodeSkeletalMeshPinDataMesh : public UCustomizableObjectNodeSkeletalMeshPinDataSection
 {
 	GENERATED_BODY()
+
+public:
+	
+	// NodePinDataParameter interface
+	/** Virtual function used to copy pin data when remapping pins. */
+	virtual void Copy(const UCustomizableObjectNodePinData& Other) override;
+
+	void Init(int32 InLODIndex, int32 InSectionIndex, int32 NumTexCoords, USkeletalMesh* Mesh);
+
+	/** Layouts related to this Mesh pin */
+	UPROPERTY()
+	TArray<TObjectPtr<UCustomizableObjectLayout>> Layouts;
 };
 
 
@@ -218,11 +230,9 @@ public:
 	
 	// UCustomizableObjectNodeMesh interface
 	virtual UTexture2D* FindTextureForPin(const UEdGraphPin* Pin) const override;
-	virtual void GetUVChannelForPin(const UEdGraphPin* Pin, TArray<FVector2f>& OutSegments, int32 UVIndex) const override;
 	virtual TArray<UCustomizableObjectLayout*> GetLayouts(const UEdGraphPin& MeshPin) const override;
 	virtual UObject* GetMesh() const override;
 	virtual UEdGraphPin* GetMeshPin(int32 LOD, int32 SectionIndex) const override;
-	virtual UEdGraphPin* GetLayoutPin(int32 LODIndex, int32 SectionIndex, int32 LayoutIndex) const override;
 	virtual void GetPinSection(const UEdGraphPin& Pin, int32& OutLODIndex, int32& OutSectionIndex, int32& OutLayoutIndex) const override;
 
 	// Own interface
diff --git a/Engine/Plugins/Experimental/Mutable/Source/CustomizableObjectEditor/Public/MuCOE/Nodes/CustomizableObjectNodeStaticMesh.h b/Engine/Plugins/Experimental/Mutable/Source/CustomizableObjectEditor/Public/MuCOE/Nodes/CustomizableObjectNodeStaticMesh.h
--- a/Engine/Plugins/Experimental/Mutable/Source/CustomizableObjectEditor/Public/MuCOE/Nodes/CustomizableObjectNodeStaticMesh.h
+++ b/Engine/Plugins/Experimental/Mutable/Source/CustomizableObjectEditor/Public/MuCOE/Nodes/CustomizableObjectNodeStaticMesh.h
@@ -151,11 +151,9 @@ public:
 
 	// UCustomizableObjectNodeMesh interface
 	virtual UTexture2D* FindTextureForPin(const UEdGraphPin* Pin) const override;
-	virtual void GetUVChannelForPin(const UEdGraphPin* Pin, TArray<FVector2f>& OutSegments, int32 UVIndex) const override;
 	virtual TArray<UCustomizableObjectLayout*> GetLayouts(const UEdGraphPin& OutPin) const override;
 	virtual UObject* GetMesh() const override;
 	virtual UEdGraphPin* GetMeshPin(int32 LOD, int32 SectionIndex) const override;
-	virtual UEdGraphPin* GetLayoutPin(int32 LODIndex, int32 SectionIndex, int32 LayoutIndex) const override;
 	virtual void GetPinSection(const UEdGraphPin& Pin, int32& OutLODIndex, int32& OutSectionIndex, int32& OutLayoutIndex) const override;
 
 	/** Returns the material associated to the given output pin. */
diff --git a/Engine/Plugins/Experimental/Mutable/Source/CustomizableObjectEditor/Public/MuCOE/Nodes/CustomizableObjectNodeTable.h b/Engine/Plugins/Experimental/Mutable/Source/CustomizableObjectEditor/Public/MuCOE/Nodes/CustomizableObjectNodeTable.h
--- a/Engine/Plugins/Experimental/Mutable/Source/CustomizableObjectEditor/Public/MuCOE/Nodes/CustomizableObjectNodeTable.h
+++ b/Engine/Plugins/Experimental/Mutable/Source/CustomizableObjectEditor/Public/MuCOE/Nodes/CustomizableObjectNodeTable.h
@@ -266,8 +266,7 @@ public:
 	UTexture2D* FindReferenceTextureParameter(const UEdGraphPin* Pin, FString ParameterImageName) const;
 
 	// Methods to get the UVs of the skeletal mesh 
-	void GetUVChannel(const UCustomizableObjectLayout* CurrentLayout, TArray<FVector2f>& OutSegments) const;
-	void GetUVChannelForPin(const UEdGraphPin* Pin, TArray<FVector2f>& OutSegments, int32 UVChannel = 0) const;
+	UObject* GetDefaultMeshForLayout(const UCustomizableObjectLayout* Layout) const;
 
 	// Methods to provide the Layouts to the Layout block editors
 	TArray<UCustomizableObjectLayout*> GetLayouts(const UEdGraphPin* Pin) const;
diff --git a/Engine/Plugins/Experimental/Mutable/Source/MutableTools/Private/MuT/CodeGenerator.cpp b/Engine/Plugins/Experimental/Mutable/Source/MutableTools/Private/MuT/CodeGenerator.cpp
--- a/Engine/Plugins/Experimental/Mutable/Source/MutableTools/Private/MuT/CodeGenerator.cpp
+++ b/Engine/Plugins/Experimental/Mutable/Source/MutableTools/Private/MuT/CodeGenerator.cpp
@@ -973,7 +973,7 @@ namespace mu
 								continue;
 							}
 
-							if (Data.GeneratedLayouts.Num() != MeshResults.GeneratedLayouts.Num())
+							if (Data.GeneratedLayouts.Num() < MeshResults.GeneratedLayouts.Num())
 							{
 								ErrorLog->GetPrivate()->Add(TEXT("Merged layout has been ignored because the number of layouts is different."), ELMT_ERROR, surfaceNode->GetMessageContext());
 							}
