[mutable] Add specific UI widgets for tag list in mesh section nodes.
#rnx

diff --git a/Engine/Plugins/Experimental/Mutable/Source/CustomizableObjectEditor/Private/MuCOE/CustomizableObjectEditorModule.cpp b/Engine/Plugins/Experimental/Mutable/Source/CustomizableObjectEditor/Private/MuCOE/CustomizableObjectEditorModule.cpp
--- a/Engine/Plugins/Experimental/Mutable/Source/CustomizableObjectEditor/Private/MuCOE/CustomizableObjectEditorModule.cpp
+++ b/Engine/Plugins/Experimental/Mutable/Source/CustomizableObjectEditor/Private/MuCOE/CustomizableObjectEditorModule.cpp
@@ -33,6 +33,8 @@
 #include "MuCOE/Nodes/CustomizableObjectNodeExternalPin.h"
 #include "MuCOE/Nodes/CustomizableObjectNodeExternalPinDetails.h"
 #include "MuCOE/Nodes/CustomizableObjectNodeLayoutBlocks.h"
+#include "MuCOE/Nodes/CustomizableObjectNodeMaterial.h"
+#include "MuCOE/Nodes/CustomizableObjectNodeMeshSectionDetails.h"
 #include "MuCOE/Nodes/CustomizableObjectNodeModifierEditMeshSection.h"
 #include "MuCOE/Nodes/CustomizableObjectNodeModifierEditMeshSectionDetails.h"
 #include "MuCOE/Nodes/CustomizableObjectNodeModifierExtendMeshSection.h"
@@ -254,7 +256,7 @@ void FCustomizableObjectEditorModule::StartupModule()
 	RegisterCustomDetails(PropertyModule, UCustomizableObjectNodeProjectorConstant::StaticClass(), FOnGetDetailCustomizationInstance::CreateStatic(&FCustomizableObjectNodeProjectorParameterDetails::MakeInstance));
 	RegisterCustomDetails(PropertyModule, UCustomizableObjectNodeMeshMorph::StaticClass(), FOnGetDetailCustomizationInstance::CreateStatic(&FCustomizableObjectNodeMeshMorphDetails::MakeInstance));
 	RegisterCustomDetails(PropertyModule, UCustomizableObjectNodeExternalPin::StaticClass(), FOnGetDetailCustomizationInstance::CreateStatic(&FCustomizableObjectNodeExternalPinDetails::MakeInstance));
-	RegisterCustomDetails(PropertyModule, UCustomizableObjectNodeMaterial::StaticClass(), FOnGetDetailCustomizationInstance::CreateStatic(&FCustomizableObjectNodeDetails::MakeInstance));
+	RegisterCustomDetails(PropertyModule, UCustomizableObjectNodeMaterial::StaticClass(), FOnGetDetailCustomizationInstance::CreateStatic(&FCustomizableObjectNodeMeshSectionDetails::MakeInstance));
 	RegisterCustomDetails(PropertyModule, UCustomizableObjectNodeSkeletalMesh::StaticClass(), FOnGetDetailCustomizationInstance::CreateStatic(&FCustomizableObjectNodeSkeletalMeshDetails::MakeInstance));
 	RegisterCustomDetails(PropertyModule, UCustomizableObjectNodeStaticMesh::StaticClass(), FOnGetDetailCustomizationInstance::CreateStatic(&FCustomizableObjectNodeDetails::MakeInstance));
 	RegisterCustomDetails(PropertyModule, UCustomizableObjectNodeTable::StaticClass(), FOnGetDetailCustomizationInstance::CreateStatic(&FCustomizableObjectNodeTableDetails::MakeInstance));
diff --git a/Engine/Plugins/Experimental/Mutable/Source/CustomizableObjectEditor/Private/MuCOE/Nodes/CustomizableObjectNodeMeshSectionDetails.cpp b/Engine/Plugins/Experimental/Mutable/Source/CustomizableObjectEditor/Private/MuCOE/Nodes/CustomizableObjectNodeMeshSectionDetails.cpp
new file mode 100644
--- /dev/null
+++ b/Engine/Plugins/Experimental/Mutable/Source/CustomizableObjectEditor/Private/MuCOE/Nodes/CustomizableObjectNodeMeshSectionDetails.cpp
@@ -0,0 +1,93 @@
+// Copyright Epic Games, Inc. All Rights Reserved.
+
+#include "MuCOE/Nodes/CustomizableObjectNodeMeshSectionDetails.h"
+
+#include "MuCOE/CustomizableObjectEditorUtilities.h"
+#include "MuCOE/CustomizableObjectEditorStyle.h"
+#include "MuCOE/GraphTraversal.h"
+#include "MuCOE/Nodes/CustomizableObjectNodeMaterial.h"
+#include "MuCOE/Nodes/CustomizableObjectNodeSkeletalMesh.h"
+#include "MuCOE/UnrealEditorPortabilityHelpers.h"
+#include "PropertyCustomizationHelpers.h"
+#include "Widgets/Input/STextComboBox.h"
+#include "Widgets/Input/SSearchBox.h"
+#include "Widgets/Input/SButton.h"
+#include "DetailLayoutBuilder.h"
+#include "IDetailsView.h"
+#include "Engine/SkeletalMesh.h"
+#include "Misc/Attribute.h"
+
+
+#define LOCTEXT_NAMESPACE "CustomizableObjectDetails"
+
+
+TSharedRef<IDetailCustomization> FCustomizableObjectNodeMeshSectionDetails::MakeInstance()
+{
+	return MakeShareable( new FCustomizableObjectNodeMeshSectionDetails );
+}
+
+
+void FCustomizableObjectNodeMeshSectionDetails::CustomizeDetails( IDetailLayoutBuilder& DetailBuilder )
+{
+	FCustomizableObjectNodeDetails::CustomizeDetails(DetailBuilder);
+
+	const IDetailsView* DetailsView = DetailBuilder.GetDetailsView();
+	if ( DetailsView->GetSelectedObjects().Num() )
+	{
+		Node = Cast<UCustomizableObjectNodeMaterial>( DetailsView->GetSelectedObjects()[0].Get() );
+	}
+
+	if (!Node)
+	{
+		return;
+	}
+	
+	// Move tags to enable higher.
+	IDetailCategoryBuilder& TagsCategory = DetailBuilder.EditCategory("Tags");
+	TagsCategory.SetSortOrder(-5000);
+
+	// Add the required tags widget
+	{
+		TagsPropertyHandle = DetailBuilder.GetProperty(GET_MEMBER_NAME_CHECKED(UCustomizableObjectNodeMaterial, Tags), UCustomizableObjectNodeMaterial::StaticClass());
+		DetailBuilder.HideProperty(TagsPropertyHandle);
+
+		TagsPropertyHandle->SetOnPropertyValueChanged(FSimpleDelegate::CreateSP(this, &FCustomizableObjectNodeMeshSectionDetails::OnEnableTagsPropertyChanged));
+		TagsPropertyHandle->SetOnChildPropertyValueChanged(FSimpleDelegate::CreateSP(this, &FCustomizableObjectNodeMeshSectionDetails::OnEnableTagsPropertyChanged));
+
+		TagsCategory.AddCustomRow(FText::FromString(TEXT("Enable Tags")))
+			.PropertyHandleList({ TagsPropertyHandle })
+			.NameContent()
+				.VAlign(VAlign_Fill)
+			[
+				SNew(SVerticalBox)
+					+ SVerticalBox::Slot()
+					.VAlign(VAlign_Top)
+					.Padding(FMargin(0,4.0f,0,4.0f))
+					[
+						SNew(STextBlock)
+							.Text(LOCTEXT("MeshSectionDetails_Tags", "Enable Tags"))
+							.Font(IDetailLayoutBuilder::GetDetailFont())
+					]
+			]
+			.ValueContent()
+				.HAlign(HAlign_Fill)
+			[
+				SAssignNew(this->TagListWidget, SMutableTagListWidget)
+					.Node(Node)
+					.TagArray(&Node->Tags)
+					.EmptyListText(LOCTEXT("MeshSectionDetails_NoTags", "No tags enabled by this mesh section."))
+					.OnTagListChanged( this, &FCustomizableObjectNodeMeshSectionDetails::OnEnableTagsPropertyChanged)
+			];
+	}
+}
+
+
+void FCustomizableObjectNodeMeshSectionDetails::OnEnableTagsPropertyChanged()
+{
+	// This seems necessary to detect the "Reset to default" actions.
+	TagListWidget->RefreshOptions();
+	Node->Modify();
+}
+
+
+#undef LOCTEXT_NAMESPACE
diff --git a/Engine/Plugins/Experimental/Mutable/Source/CustomizableObjectEditor/Private/MuCOE/Nodes/CustomizableObjectNodeMeshSectionDetails.h b/Engine/Plugins/Experimental/Mutable/Source/CustomizableObjectEditor/Private/MuCOE/Nodes/CustomizableObjectNodeMeshSectionDetails.h
new file mode 100644
--- /dev/null
+++ b/Engine/Plugins/Experimental/Mutable/Source/CustomizableObjectEditor/Private/MuCOE/Nodes/CustomizableObjectNodeMeshSectionDetails.h
@@ -0,0 +1,48 @@
+// Copyright Epic Games, Inc. All Rights Reserved.
+
+#pragma once
+
+#include "CustomizableObjectNodeDetails.h"
+#include "MuCOE/SMutableTagListWidget.h"
+#include "MuCOE/SMutableSearchComboBox.h"
+#include "IDetailCustomization.h"
+#include "Types/SlateEnums.h"
+#include "Widgets/Input/SComboButton.h"
+#include "Widgets/Input/SEditableTextBox.h"
+#include "Widgets/Views/SListView.h"
+#include "Widgets/Views/ITableRow.h"
+
+namespace ESelectInfo { enum Type : int; }
+
+class FString;
+class IDetailLayoutBuilder;
+class IPropertyHandle;
+class UCustomizableObjectNode;
+class UCustomizableObjectNodeObject;
+
+/** */
+class FCustomizableObjectNodeMeshSectionDetails : public FCustomizableObjectNodeDetails
+{
+public:
+	// Makes a new instance of this detail layout class for a specific detail view requesting it 
+	static TSharedRef<IDetailCustomization> MakeInstance();
+
+	// ILayoutDetails interface
+	void CustomizeDetails( IDetailLayoutBuilder& DetailBuilder ) override;
+
+protected:
+
+	/** */
+	virtual void OnEnableTagsPropertyChanged();
+
+private:
+
+	class UCustomizableObjectNodeMaterial* Node = nullptr;
+
+	/** */
+	TSharedPtr<IPropertyHandle> TagsPropertyHandle;
+
+	/** */
+	TSharedPtr<SMutableTagListWidget> TagListWidget;
+
+};
diff --git a/Engine/Plugins/Experimental/Mutable/Source/CustomizableObjectEditor/Private/MuCOE/Nodes/CustomizableObjectNodeModifierBaseDetails.cpp b/Engine/Plugins/Experimental/Mutable/Source/CustomizableObjectEditor/Private/MuCOE/Nodes/CustomizableObjectNodeModifierBaseDetails.cpp
--- a/Engine/Plugins/Experimental/Mutable/Source/CustomizableObjectEditor/Private/MuCOE/Nodes/CustomizableObjectNodeModifierBaseDetails.cpp
+++ b/Engine/Plugins/Experimental/Mutable/Source/CustomizableObjectEditor/Private/MuCOE/Nodes/CustomizableObjectNodeModifierBaseDetails.cpp
@@ -49,8 +49,6 @@ void FCustomizableObjectNodeModifierBaseDetails::CustomizeDetails( IDetailLayout
 
 	// Add the required tags widget
 	{
-		RefreshTagOptions();
-
 		RequiredTagsPropertyHandle = DetailBuilder.GetProperty(GET_MEMBER_NAME_CHECKED(UCustomizableObjectNodeModifierBase, RequiredTags), UCustomizableObjectNodeModifierBase::StaticClass());
 		DetailBuilder.HideProperty(RequiredTagsPropertyHandle);
 
@@ -71,63 +69,19 @@ void FCustomizableObjectNodeModifierBaseDetails::CustomizeDetails( IDetailLayout
 					.Padding(FMargin(0,4.0f,0,4.0f))
 					[
 						SNew(STextBlock)
-							.Text(LOCTEXT("tableodifierDetails_RequiredTags", "Required Tags"))
+							.Text(LOCTEXT("ModifierDetails_RequiredTags", "Required Tags"))
 							.Font(IDetailLayoutBuilder::GetDetailFont())
 					]
 			]
 			.ValueContent()
 				.HAlign(HAlign_Fill)
 			[
-				SNew(SVerticalBox)
-
-					// Header with the "add tag" UI
-					+ SVerticalBox::Slot()
-					.AutoHeight()
-					.HAlign(HAlign_Fill)
-					[
-						SNew(SHorizontalBox)
-
-						+ SHorizontalBox::Slot()
-						.FillWidth(10.0f)
-						.HAlign(HAlign_Fill)
-						.VAlign(VAlign_Center)
-						[
-							SAssignNew(this->TagCombo, SMutableSearchComboBox)
-								.ButtonStyle(FAppStyle::Get(), "NoBorder")
-								.MenuButtonBrush(FAppStyle::GetBrush(TEXT("Icons.PlusCircle")))
-								.OptionsSource(&TagComboOptionsSource)
-								.AllowAddNewOptions(true)
-								.OnSelectionChanged(this, &FCustomizableObjectNodeModifierBaseDetails::OnTagComboBoxSelectionChanged)
-						]
-					]
-
-					// List of tags
-					+ SVerticalBox::Slot()
-					.VAlign(VAlign_Fill)
-					.HAlign(HAlign_Fill)
-					[
-						SAssignNew(this->TagList, SListView<TSharedPtr<FTagUIData>>)
-							.ListItemsSource(&CurrentTagsSource)
-							.OnGenerateRow(this, &FCustomizableObjectNodeModifierBaseDetails::GenerateTagMenuItemRow)
-							.SelectionMode(ESelectionMode::None)							
-					]
-
-					// No-tags warning: shown only if there are no required tags defined
-					+ SVerticalBox::Slot()
-					.AutoHeight()
-					.HAlign(HAlign_Fill)
-					[
-						SNew(STextBlock)
-							.Text(LOCTEXT("ModifierDetails_NoRequiredTagsWarning", "Warning: There are no required tags, so this modifier will not do anything."))
-							.AutoWrapText(true)
-							.Font(IDetailLayoutBuilder::GetDetailFont())
-							.Visibility_Lambda([&]() 
-								{
-									return (!Node || Node->RequiredTags.IsEmpty()) ? EVisibility::Visible : EVisibility::Collapsed;
-								})
-					]
+				SAssignNew(this->TagListWidget, SMutableTagListWidget)
+					.Node(Node)
+					.TagArray(&Node->RequiredTags)
+					.EmptyListText(LOCTEXT("ModifierDetails_NoRequiredTagsWarning", "Warning: There are no required tags, so this modifier will not do anything."))
+					.OnTagListChanged( this, &FCustomizableObjectNodeModifierBaseDetails::OnRequiredTagsPropertyChanged )
 			];
-
 	}
 }
 
@@ -135,134 +89,8 @@ void FCustomizableObjectNodeModifierBaseDetails::CustomizeDetails( IDetailLayout
 void FCustomizableObjectNodeModifierBaseDetails::OnRequiredTagsPropertyChanged()
 {
 	// This seems necessary to detect the "Reset to default" actions.
-	RefreshTagOptions();
-}
-
-
-void FCustomizableObjectNodeModifierBaseDetails::RefreshTagOptions()
-{
-	// Tag combo options
-	{
-		TagComboOptionsSource.SetNum(0,EAllowShrinking::No);
-
-		// Scan all potential receivers
-		UCustomizableObject* ThisNodeObject = GetRootObject(*Node);
-		UCustomizableObject* RootObject = GetRootObject(ThisNodeObject);
-
-		TSet<UCustomizableObject*> AllCustomizableObject;
-		GetAllObjectsInGraph(RootObject, AllCustomizableObject);
-
-		for (const UCustomizableObject* CustObject : AllCustomizableObject)
-		{
-			if (CustObject)
-			{
-				for (const TObjectPtr<UEdGraphNode>& CandidateNode : CustObject->GetPrivate()->GetSource()->Nodes)
-				{
-					UCustomizableObjectNode* Typed = Cast<UCustomizableObjectNode>(CandidateNode);
-					if (Typed)
-					{
-						TArray<FString>* EnableTags = Typed->GetEnableTags();
-						if (EnableTags)
-						{
-							for (const FString& Tag : *EnableTags)
-							{
-								if (!Tag.IsEmpty())
-								{
-									TagComboOptionsSource.AddUnique(MakeShared<FString>(Tag));
-								}
-							}
-						}
-					}
-				}
-			}
-		}
-
-		// TODO: add material nodes and options to create tags for them, show hierarchy, ...
-	}
-
-	// Current Tags
-	if (Node)
-	{
-		CurrentTagsSource.SetNum(0, EAllowShrinking::No);
-
-		for (const FString& Tag : Node->RequiredTags)
-		{
-			TSharedPtr<FTagUIData> Data = MakeShared<FTagUIData>();
-			Data->DisplayName = Tag;
-			Data->Tag = Tag;
-			CurrentTagsSource.Add(Data);
-		}
-	}
-
-	if (TagList)
-	{
-		TagList->RequestListRefresh();
-	}
-}
-
-
-TSharedRef<ITableRow> FCustomizableObjectNodeModifierBaseDetails::GenerateTagMenuItemRow(TSharedPtr<FTagUIData> InItem, const TSharedRef<STableViewBase>& OwnerTable)
-{
-	if (!InItem)
-	{
-		return SNew(SComboRow<TSharedPtr<FString>>, OwnerTable)
-			[
-				SNew(STextBlock)
-					.Text(FText::FromString("No item."))
-					.Font(IDetailLayoutBuilder::GetDetailFont())
-			];
-	}
-
-	return SNew(SComboRow<TSharedPtr<FString>>, OwnerTable)
-		[
-			SNew(SHorizontalBox)
-
-				+ SHorizontalBox::Slot()
-				.HAlign(HAlign_Fill)
-				.VAlign(VAlign_Center)
-				[
-					SNew(STextBlock)
-						.Text(FText::FromString(InItem->DisplayName))
-						.Font(IDetailLayoutBuilder::GetDetailFont())
-				]
-
-				+ SHorizontalBox::Slot()
-				.AutoWidth()
-				[
-					SNew(SButton)
-						.ButtonStyle(FAppStyle::Get(), "NoBorder")
-						.ToolTipText(LOCTEXT("RemoveModifierTag","Remove this tag from the modifier."))
-						.OnClicked_Lambda([this,InItem]() 
-							{
-								if (Node && InItem )
-								{
-									if (Node->RequiredTags.Contains(InItem->Tag))
-									{
-										Node->RequiredTags.Remove(InItem->Tag);
-										Node->Modify();
-										OnRequiredTagsPropertyChanged();
-									}
-								}
-								return FReply::Handled();
-							})
-						[
-							SNew(SImage)
-								.Image(FAppStyle::GetBrush(TEXT("Icons.MinusCircle")))
-						]
-				]
-		];
-}
-
-
-void FCustomizableObjectNodeModifierBaseDetails::OnTagComboBoxSelectionChanged(const FText& NewText)
-{
-	FString Tag = NewText.ToString();
-	if (Node && !Node->RequiredTags.Contains(Tag))
-	{
-		Node->RequiredTags.Add(Tag);
-		Node->Modify();
-		OnRequiredTagsPropertyChanged();
-	}
+	TagListWidget->RefreshOptions();
+	Node->Modify();
 }
 
 
diff --git a/Engine/Plugins/Experimental/Mutable/Source/CustomizableObjectEditor/Private/MuCOE/Nodes/CustomizableObjectNodeModifierBaseDetails.h b/Engine/Plugins/Experimental/Mutable/Source/CustomizableObjectEditor/Private/MuCOE/Nodes/CustomizableObjectNodeModifierBaseDetails.h
--- a/Engine/Plugins/Experimental/Mutable/Source/CustomizableObjectEditor/Private/MuCOE/Nodes/CustomizableObjectNodeModifierBaseDetails.h
+++ b/Engine/Plugins/Experimental/Mutable/Source/CustomizableObjectEditor/Private/MuCOE/Nodes/CustomizableObjectNodeModifierBaseDetails.h
@@ -3,6 +3,7 @@
 #pragma once
 
 #include "CustomizableObjectNodeDetails.h"
+#include "MuCOE/SMutableTagListWidget.h"
 #include "MuCOE/SMutableSearchComboBox.h"
 #include "IDetailCustomization.h"
 #include "Types/SlateEnums.h"
@@ -16,9 +17,10 @@ namespace ESelectInfo { enum Type : int; }
 class FString;
 class IDetailLayoutBuilder;
 class IPropertyHandle;
+class UCustomizableObjectNode;
 class UCustomizableObjectNodeObject;
 
-
+/** */
 class FCustomizableObjectNodeModifierBaseDetails : public FCustomizableObjectNodeDetails
 {
 public:
@@ -40,20 +42,8 @@ private:
 	/** */
 	TSharedPtr<IPropertyHandle> RequiredTagsPropertyHandle;
 	TSharedPtr<IPropertyHandle> TagsPolicyPropertyHandle;
-	TSharedPtr<SMutableSearchComboBox> TagCombo;
-	TArray< TSharedPtr<FString> > TagComboOptionsSource;
-	void OnTagComboBoxSelectionChanged(const FText& NewText);
-
-	struct FTagUIData
-	{
-		FString Tag;
-		FString DisplayName;
-	};
-	TSharedPtr<SListView<TSharedPtr<FTagUIData>>> TagList;
-	TArray< TSharedPtr<FTagUIData> > CurrentTagsSource;
 
 	/** */
-	void RefreshTagOptions();
-	TSharedRef<ITableRow> GenerateTagMenuItemRow(TSharedPtr<FTagUIData> InItem, const TSharedRef<STableViewBase>& OwnerTable);
+	TSharedPtr<SMutableTagListWidget> TagListWidget;
 
 };
diff --git a/Engine/Plugins/Experimental/Mutable/Source/CustomizableObjectEditor/Private/MuCOE/Nodes/CustomizableObjectNodeModifierExtendMeshSectionDetails.cpp b/Engine/Plugins/Experimental/Mutable/Source/CustomizableObjectEditor/Private/MuCOE/Nodes/CustomizableObjectNodeModifierExtendMeshSectionDetails.cpp
--- a/Engine/Plugins/Experimental/Mutable/Source/CustomizableObjectEditor/Private/MuCOE/Nodes/CustomizableObjectNodeModifierExtendMeshSectionDetails.cpp
+++ b/Engine/Plugins/Experimental/Mutable/Source/CustomizableObjectEditor/Private/MuCOE/Nodes/CustomizableObjectNodeModifierExtendMeshSectionDetails.cpp
@@ -6,7 +6,6 @@
 #include "MuCOE/CustomizableObjectEditorStyle.h"
 #include "MuCOE/GraphTraversal.h"
 #include "MuCOE/Nodes/CustomizableObjectNodeMaterial.h"
-#include "MuCOE/Nodes/CustomizableObjectNodeModifierMorphMeshSection.h"
 #include "MuCOE/Nodes/CustomizableObjectNodeModifierExtendMeshSection.h"
 #include "MuCOE/Nodes/CustomizableObjectNodeSkeletalMesh.h"
 #include "MuCOE/UnrealEditorPortabilityHelpers.h"
@@ -42,7 +41,53 @@ void FCustomizableObjectNodeModifierExtendMeshSectionDetails::CustomizeDetails(
 	{
 		return;
 	}
-	
+
+	// Move tags to enable higher.
+	IDetailCategoryBuilder& TagsCategory = DetailBuilder.EditCategory("EnableTags");
+	TagsCategory.SetSortOrder(-5000);
+
+	// Add the required tags widget
+	{
+		EnableTagsPropertyHandle = DetailBuilder.GetProperty(GET_MEMBER_NAME_CHECKED(UCustomizableObjectNodeModifierExtendMeshSection, Tags), UCustomizableObjectNodeModifierExtendMeshSection::StaticClass());
+		DetailBuilder.HideProperty(EnableTagsPropertyHandle);
+
+		EnableTagsPropertyHandle->SetOnPropertyValueChanged(FSimpleDelegate::CreateSP(this, &FCustomizableObjectNodeModifierExtendMeshSectionDetails::OnRequiredTagsPropertyChanged));
+		EnableTagsPropertyHandle->SetOnChildPropertyValueChanged(FSimpleDelegate::CreateSP(this, &FCustomizableObjectNodeModifierExtendMeshSectionDetails::OnRequiredTagsPropertyChanged));
+
+		TagsCategory.AddCustomRow(FText::FromString(TEXT("Enable Tags")))
+			.PropertyHandleList({ EnableTagsPropertyHandle })
+			.NameContent()
+			.VAlign(VAlign_Fill)
+			[
+				SNew(SVerticalBox)
+					+ SVerticalBox::Slot()
+					.VAlign(VAlign_Top)
+					.Padding(FMargin(0, 4.0f, 0, 4.0f))
+					[
+						SNew(STextBlock)
+							.Text(LOCTEXT("MeshSectionDetails_Tags", "Tags enabled for extended data"))
+							.Font(IDetailLayoutBuilder::GetDetailFont())
+					]
+			]
+			.ValueContent()
+			.HAlign(HAlign_Fill)
+			[
+				SAssignNew(this->EnableTagListWidget, SMutableTagListWidget)
+					.Node(Node)
+					.TagArray(&Node->Tags)
+					.EmptyListText(LOCTEXT("ExtendMeshSectionDetails_NoTags", "No tags enabled by this extended mesh section."))
+					.OnTagListChanged(this, &FCustomizableObjectNodeModifierExtendMeshSectionDetails::OnEnableTagsPropertyChanged)
+			];
+	}
+
+}
+
+
+void FCustomizableObjectNodeModifierExtendMeshSectionDetails::OnEnableTagsPropertyChanged()
+{
+	// This seems necessary to detect the "Reset to default" actions.
+	EnableTagListWidget->RefreshOptions();
+	Node->Modify();
 }
 
 
diff --git a/Engine/Plugins/Experimental/Mutable/Source/CustomizableObjectEditor/Private/MuCOE/Nodes/CustomizableObjectNodeModifierExtendMeshSectionDetails.h b/Engine/Plugins/Experimental/Mutable/Source/CustomizableObjectEditor/Private/MuCOE/Nodes/CustomizableObjectNodeModifierExtendMeshSectionDetails.h
--- a/Engine/Plugins/Experimental/Mutable/Source/CustomizableObjectEditor/Private/MuCOE/Nodes/CustomizableObjectNodeModifierExtendMeshSectionDetails.h
+++ b/Engine/Plugins/Experimental/Mutable/Source/CustomizableObjectEditor/Private/MuCOE/Nodes/CustomizableObjectNodeModifierExtendMeshSectionDetails.h
@@ -33,4 +33,14 @@ private:
 	class UCustomizableObjectNodeModifierExtendMeshSection* Node = nullptr;
 
 
+	/** */
+	TSharedPtr<IPropertyHandle> EnableTagsPropertyHandle;
+
+	/** */
+	TSharedPtr<SMutableTagListWidget> EnableTagListWidget;
+
+	/** */
+	virtual void OnEnableTagsPropertyChanged();
+
+
 };
diff --git a/Engine/Plugins/Experimental/Mutable/Source/CustomizableObjectEditor/Private/MuCOE/SMutableTagListWidget.cpp b/Engine/Plugins/Experimental/Mutable/Source/CustomizableObjectEditor/Private/MuCOE/SMutableTagListWidget.cpp
new file mode 100644
--- /dev/null
+++ b/Engine/Plugins/Experimental/Mutable/Source/CustomizableObjectEditor/Private/MuCOE/SMutableTagListWidget.cpp
@@ -0,0 +1,241 @@
+// Copyright Epic Games, Inc. All Rights Reserved.
+
+#include "MuCOE/SMutableTagListWidget.h"
+
+#include "MuCOE/CustomizableObjectEditorUtilities.h"
+#include "MuCOE/CustomizableObjectEditorStyle.h"
+#include "MuCOE/GraphTraversal.h"
+#include "MuCOE/Nodes/CustomizableObjectNode.h"
+#include "MuCOE/Nodes/CustomizableObjectNodeModifierBase.h"
+#include "MuCOE/UnrealEditorPortabilityHelpers.h"
+#include "PropertyCustomizationHelpers.h"
+#include "Widgets/Input/STextComboBox.h"
+#include "Widgets/Input/SSearchBox.h"
+#include "Widgets/Input/SButton.h"
+#include "DetailLayoutBuilder.h"
+#include "IDetailsView.h"
+#include "Engine/SkeletalMesh.h"
+#include "Misc/Attribute.h"
+
+
+#define LOCTEXT_NAMESPACE "CustomizableObjectDetails"
+
+
+
+void SMutableTagListWidget::Construct(const FArguments& InArgs)
+{
+	Node = InArgs._Node;
+	TagArray = InArgs._TagArray;
+	EmptyListText = InArgs._EmptyListText;
+
+	OnTagListChangedDelegate = InArgs._OnTagListChanged;
+
+	RefreshOptions();
+
+	ChildSlot
+	[
+		SNew(SVerticalBox)
+
+		// Header with the "add tag" UI
+		+ SVerticalBox::Slot()
+		.AutoHeight()
+		.HAlign(HAlign_Fill)
+		[
+			SNew(SHorizontalBox)
+
+				+ SHorizontalBox::Slot()
+				.FillWidth(10.0f)
+				.HAlign(HAlign_Fill)
+				.VAlign(VAlign_Center)
+				[
+					SAssignNew(this->TagCombo, SMutableSearchComboBox)
+						.ButtonStyle(FAppStyle::Get(), "NoBorder")
+						.MenuButtonBrush(FAppStyle::GetBrush(TEXT("Icons.PlusCircle")))
+						.OptionsSource(&TagComboOptionsSource)
+						.AllowAddNewOptions(true)
+						.OnSelectionChanged(this, &SMutableTagListWidget::OnTagComboBoxSelectionChanged)
+				]
+		]
+
+		// List of tags
+		+ SVerticalBox::Slot()
+		.VAlign(VAlign_Fill)
+		.HAlign(HAlign_Fill)
+		[
+			SAssignNew(this->TagListWidget, SListView<TSharedPtr<FTagUIData>>)
+				.ListItemsSource(&CurrentTagsSource)
+				.OnGenerateRow(this, &SMutableTagListWidget::GenerateTagMenuItemRow)
+				.SelectionMode(ESelectionMode::None)
+		]
+
+		// Shown only if there are no tags defined
+		+ SVerticalBox::Slot()
+		.AutoHeight()
+		.HAlign(HAlign_Fill)
+		[
+			SNew(STextBlock)
+				.Text(EmptyListText)
+				.AutoWrapText(true)
+				.Font(IDetailLayoutBuilder::GetDetailFont())
+				.Visibility_Lambda([&]()
+					{
+						return (!TagArray || TagArray->IsEmpty()) ? EVisibility::Visible : EVisibility::Collapsed;
+					})
+		]
+	];
+}
+
+
+void SMutableTagListWidget::RefreshOptions()
+{
+	// Tag combo options
+	{
+		TagComboOptionsSource.SetNum(0, EAllowShrinking::No);
+
+		// Scan all potential receivers
+		UCustomizableObject* ThisNodeObject = GetRootObject(*Node);
+		UCustomizableObject* RootObject = GetRootObject(ThisNodeObject);
+
+		TSet<UCustomizableObject*> AllCustomizableObject;
+		GetAllObjectsInGraph(RootObject, AllCustomizableObject);
+
+		for (const UCustomizableObject* CustObject : AllCustomizableObject)
+		{
+			if (CustObject)
+			{
+				for (const TObjectPtr<UEdGraphNode>& CandidateNode : CustObject->GetPrivate()->GetSource()->Nodes)
+				{
+					UCustomizableObjectNode* Typed = Cast<UCustomizableObjectNode>(CandidateNode);
+					if (Typed)
+					{
+						TArray<FString>* EnableTags = Typed->GetEnableTags();
+						if (EnableTags)
+						{
+							for (const FString& OneTag : *EnableTags)
+							{
+								if (!OneTag.IsEmpty())
+								{
+									bool bContained = TagComboOptionsSource.ContainsByPredicate([&](const TSharedPtr<FString>& Candidate) 
+										{
+											return Candidate.IsValid() && *Candidate == OneTag;
+										});
+									if (!bContained)
+									{
+										TagComboOptionsSource.Add(MakeShared<FString>(OneTag));
+									}
+								}
+							}
+						}
+					}
+
+					UCustomizableObjectNodeModifierBase* TypedModifier = Cast<UCustomizableObjectNodeModifierBase>(CandidateNode);
+					if (TypedModifier)
+					{
+						for (const FString& OneTag : TypedModifier->RequiredTags)
+						{
+							if (!OneTag.IsEmpty())
+							{
+								bool bContained = TagComboOptionsSource.ContainsByPredicate([&](const TSharedPtr<FString>& Candidate)
+									{
+										return Candidate.IsValid() && *Candidate == OneTag;
+									});
+								if (!bContained)
+								{
+									TagComboOptionsSource.Add(MakeShared<FString>(OneTag));
+								}
+							}
+						}
+					}
+				}
+			}
+		}
+
+		// TODO: add material nodes and options to create tags for them, show hierarchy, ...
+	}
+
+	// Current Tags
+	if (TagArray)
+	{
+		CurrentTagsSource.SetNum(0, EAllowShrinking::No);
+
+		for (const FString& OneTag : *TagArray)
+		{
+			TSharedPtr<FTagUIData> Data = MakeShared<FTagUIData>();
+			Data->DisplayName = OneTag;
+			Data->Tag = OneTag;
+			CurrentTagsSource.Add(Data);
+		}
+	}
+
+	if (TagListWidget)
+	{
+		TagListWidget->RequestListRefresh();
+	}
+}
+
+
+void SMutableTagListWidget::OnTagComboBoxSelectionChanged(const FText& NewText)
+{
+	FString OneTag = NewText.ToString();
+	if (TagArray && !TagArray->Contains(OneTag))
+	{
+		TagArray->Add(OneTag);
+		OnTagListChangedDelegate.ExecuteIfBound();
+	}
+}
+
+
+TSharedRef<ITableRow> SMutableTagListWidget::GenerateTagMenuItemRow(TSharedPtr<FTagUIData> InItem, const TSharedRef<STableViewBase>& OwnerTable)
+{
+	if (!InItem)
+	{
+		return SNew(SComboRow<TSharedPtr<FString>>, OwnerTable)
+			[
+				SNew(STextBlock)
+					.Text(FText::FromString("No item."))
+					.Font(IDetailLayoutBuilder::GetDetailFont())
+			];
+	}
+
+	return SNew(SComboRow<TSharedPtr<FString>>, OwnerTable)
+		[
+			SNew(SHorizontalBox)
+
+				+ SHorizontalBox::Slot()
+				.HAlign(HAlign_Fill)
+				.VAlign(VAlign_Center)
+				[
+					SNew(STextBlock)
+						.Text(FText::FromString(InItem->DisplayName))
+						.Font(IDetailLayoutBuilder::GetDetailFont())
+				]
+
+				+ SHorizontalBox::Slot()
+				.AutoWidth()
+				[
+					SNew(SButton)
+						.ButtonStyle(FAppStyle::Get(), "NoBorder")
+						.ToolTipText(LOCTEXT("RemoveModifierTag", "Remove this tag from the modifier."))
+						.OnClicked_Lambda([this, InItem]()
+							{
+								if (Node && InItem)
+								{
+									if (TagArray->Contains(InItem->Tag))
+									{
+										TagArray->Remove(InItem->Tag);
+										OnTagListChangedDelegate.ExecuteIfBound();
+									}
+								}
+								return FReply::Handled();
+							})
+						[
+							SNew(SImage)
+								.Image(FAppStyle::GetBrush(TEXT("Icons.MinusCircle")))
+						]
+				]
+		];
+}
+
+
+
+#undef LOCTEXT_NAMESPACE
diff --git a/Engine/Plugins/Experimental/Mutable/Source/CustomizableObjectEditor/Private/MuCOE/SMutableTagListWidget.h b/Engine/Plugins/Experimental/Mutable/Source/CustomizableObjectEditor/Private/MuCOE/SMutableTagListWidget.h
new file mode 100644
--- /dev/null
+++ b/Engine/Plugins/Experimental/Mutable/Source/CustomizableObjectEditor/Private/MuCOE/SMutableTagListWidget.h
@@ -0,0 +1,71 @@
+// Copyright Epic Games, Inc. All Rights Reserved.
+
+#pragma once
+
+#include "Widgets/SCompoundWidget.h"
+#include "MuCOE/SMutableSearchComboBox.h"
+#include "IDetailCustomization.h"
+#include "Types/SlateEnums.h"
+#include "Widgets/Input/SComboButton.h"
+#include "Widgets/Input/SEditableTextBox.h"
+#include "Widgets/Views/SListView.h"
+#include "Widgets/Views/ITableRow.h"
+
+namespace ESelectInfo { enum Type : int; }
+
+class FString;
+class IDetailLayoutBuilder;
+class IPropertyHandle;
+class UCustomizableObjectNode;
+class UCustomizableObjectNodeObject;
+
+/** */
+class SMutableTagListWidget : public SCompoundWidget
+{
+public:
+
+	SLATE_BEGIN_ARGS(SMutableTagListWidget)
+		: _Node(nullptr)
+		, _TagArray(nullptr)
+		{}
+
+		SLATE_ARGUMENT(UCustomizableObjectNode*, Node)
+		SLATE_ARGUMENT(TArray<FString>*, TagArray)
+		SLATE_ARGUMENT(FText, EmptyListText)
+
+		SLATE_EVENT(FSimpleDelegate, OnTagListChanged)
+
+	SLATE_END_ARGS()
+
+	/** */
+	void Construct(const FArguments& InArgs);
+
+	/** */
+	FSimpleDelegate OnTagListChangedDelegate;
+
+	/** */
+	void RefreshOptions();
+
+private:
+
+	class UCustomizableObjectNode* Node = nullptr;
+	TArray<FString>* TagArray = nullptr;
+	FText EmptyListText;
+
+	TSharedPtr<SMutableSearchComboBox> TagCombo;
+	TArray< TSharedPtr<FString> > TagComboOptionsSource;
+	void OnTagComboBoxSelectionChanged(const FText& NewText);
+
+	struct FTagUIData
+	{
+		FString Tag;
+		FString DisplayName;
+	};
+	TSharedPtr<SListView<TSharedPtr<FTagUIData>>> TagListWidget;
+	TArray< TSharedPtr<FTagUIData> > CurrentTagsSource;
+
+	/** */
+	TSharedRef<ITableRow> GenerateTagMenuItemRow(TSharedPtr<FTagUIData> InItem, const TSharedRef<STableViewBase>& OwnerTable);
+
+};
+
diff --git a/Engine/Plugins/Experimental/Mutable/Source/CustomizableObjectEditor/Public/MuCOE/Nodes/CustomizableObjectNodeMaterial.h b/Engine/Plugins/Experimental/Mutable/Source/CustomizableObjectEditor/Public/MuCOE/Nodes/CustomizableObjectNodeMaterial.h
--- a/Engine/Plugins/Experimental/Mutable/Source/CustomizableObjectEditor/Public/MuCOE/Nodes/CustomizableObjectNodeMaterial.h
+++ b/Engine/Plugins/Experimental/Mutable/Source/CustomizableObjectEditor/Public/MuCOE/Nodes/CustomizableObjectNodeMaterial.h
@@ -166,7 +166,7 @@ public:
 	static bool HasParameter(const UMaterialInterface* InMaterial, const FNodeMaterialParameterId& ParameterId);
 	static int32 GetParameterLayerIndex(const UMaterialInterface* InMaterial, EMaterialParameterType Type, int32 ParameterIndex);
 
-	UPROPERTY(EditAnywhere, Category = CustomizableObject)
+	UPROPERTY(EditAnywhere, Category = Tags)
 	TArray<FString> Tags;
 
 private:
diff --git a/Engine/Plugins/Experimental/Mutable/Source/CustomizableObjectEditor/Public/MuCOE/Nodes/CustomizableObjectNodeModifierExtendMeshSection.h b/Engine/Plugins/Experimental/Mutable/Source/CustomizableObjectEditor/Public/MuCOE/Nodes/CustomizableObjectNodeModifierExtendMeshSection.h
--- a/Engine/Plugins/Experimental/Mutable/Source/CustomizableObjectEditor/Public/MuCOE/Nodes/CustomizableObjectNodeModifierExtendMeshSection.h
+++ b/Engine/Plugins/Experimental/Mutable/Source/CustomizableObjectEditor/Public/MuCOE/Nodes/CustomizableObjectNodeModifierExtendMeshSection.h
@@ -27,7 +27,7 @@ class CUSTOMIZABLEOBJECTEDITOR_API UCustomizableObjectNodeModifierExtendMeshSect
 public:
 
 	/** Tags enabled when the modifier is applied. Also used to decide what modifiers get applied to the data added by this node.  */
-	UPROPERTY(EditAnywhere, BlueprintReadWrite, Category = AddedData)
+	UPROPERTY(EditAnywhere, BlueprintReadWrite, Category = EnableTags)
 	TArray<FString> Tags;
 
 public:
