[Mutable] Removed components array.

- Moved Reference Skeletal Mesh to Mesh Component nodes.
- Added ComponentName to UCustomizableSkeletalComponent and UCustomizableObjectInstanceUsage.
    - Only used if ComponentIndex is INDEX_NONE.
- New public API:
    - UCustomizableObject::GetComponentName
    - UCustomizableObject::GetComponentMeshReferenceSkeletalMesh
    - UCustomizableObject::GetComponentMeshSkeletalMesh
- Fixed not being able to remove nodes in fixups.
- Removed COFactory component related options.

#rnx

diff --git a/Engine/Plugins/Experimental/Mutable/Source/CustomizableObject/Internal/MuCO/CustomizableObjectCustomVersion.h b/Engine/Plugins/Experimental/Mutable/Source/CustomizableObject/Internal/MuCO/CustomizableObjectCustomVersion.h
--- a/Engine/Plugins/Experimental/Mutable/Source/CustomizableObject/Internal/MuCO/CustomizableObjectCustomVersion.h
+++ b/Engine/Plugins/Experimental/Mutable/Source/CustomizableObject/Internal/MuCO/CustomizableObjectCustomVersion.h
@@ -152,6 +152,8 @@ struct CUSTOMIZABLEOBJECT_API FCustomizableObjectCustomVersion
 
 		RemoveNodeLayout,
 
+		ComponentsArray,
+
 		// -----<new versions can be added above this line>-------------------------------------------------
 		VersionPlusOne,
 		LatestVersion = VersionPlusOne - 1
@@ -162,4 +164,4 @@ struct CUSTOMIZABLEOBJECT_API FCustomizableObjectCustomVersion
 
 	private:
 		FCustomizableObjectCustomVersion() {}
-};
\ No newline at end of file
+};
diff --git a/Engine/Plugins/Experimental/Mutable/Source/CustomizableObject/Internal/MuCO/CustomizableObjectInstancePrivate.h b/Engine/Plugins/Experimental/Mutable/Source/CustomizableObject/Internal/MuCO/CustomizableObjectInstancePrivate.h
--- a/Engine/Plugins/Experimental/Mutable/Source/CustomizableObject/Internal/MuCO/CustomizableObjectInstancePrivate.h
+++ b/Engine/Plugins/Experimental/Mutable/Source/CustomizableObject/Internal/MuCO/CustomizableObjectInstancePrivate.h
@@ -190,8 +190,8 @@ public:
 
 	/** The generated skeletal meshes for this Instance, one for each component in the Object component list which means that some may be null if the component hasn't been generated. */
 	UPROPERTY(Transient, VisibleAnywhere, Category = NoCategory, meta=(DisplayName = "Meshes"))
-	TArray<TObjectPtr<USkeletalMesh>> SkeletalMeshes;
-
+	TMap<FName, TObjectPtr<USkeletalMesh>> SkeletalMeshes;
+	
 	UPROPERTY(Transient, VisibleAnywhere, Category = NoCategory, meta=(DisplayName = "Materials"))
 	TArray<FGeneratedMaterial> GeneratedMaterials;
 
@@ -229,7 +229,10 @@ public:
 	/** Invalidates the previously generated data and retrieves information from the CObject after specific actions.
 	 *  It'll be called in the PostLoad, after Compiling the CO, and after changing the CO of the Instance. */
 	void InitCustomizableObjectData(const UCustomizableObject* InCustomizableObject);
-	
+
+	FCustomizableInstanceComponentData* GetComponentData(const FName& ComponentName);
+
+	// DEPRECATED
 	FCustomizableInstanceComponentData* GetComponentData(int32 ComponentIndex);
 	const FCustomizableInstanceComponentData* GetComponentData(int32 ComponentIndex) const;
 
diff --git a/Engine/Plugins/Experimental/Mutable/Source/CustomizableObject/Internal/MuCO/CustomizableObjectPrivate.h b/Engine/Plugins/Experimental/Mutable/Source/CustomizableObject/Internal/MuCO/CustomizableObjectPrivate.h
--- a/Engine/Plugins/Experimental/Mutable/Source/CustomizableObject/Internal/MuCO/CustomizableObjectPrivate.h
+++ b/Engine/Plugins/Experimental/Mutable/Source/CustomizableObject/Internal/MuCO/CustomizableObjectPrivate.h
@@ -1097,6 +1097,10 @@ struct FModelResources
 	UPROPERTY()
 	uint8 FirstLODAvailable = 0;
 
+	/** Name of all possible components. Index is the ObjectComponentIndex. */
+	UPROPERTY()
+	TArray<FName> ComponentNames;
+	
 #if WITH_EDITORONLY_DATA
 	void CUSTOMIZABLEOBJECT_API Serialize(FObjectAndNameAsStringProxyArchive& Ar, bool bIsCooking);
 	bool CUSTOMIZABLEOBJECT_API Unserialize(FObjectAndNameAsStringProxyArchive& Ar, UCustomizableObject& Outer, const ITargetPlatform* InTargetPlatform, bool bIsCooking);
@@ -1202,6 +1206,8 @@ public:
 	void SetModelStreamableBulkData(const TSharedPtr<FModelStreamableBulkData>& StreamableData, bool bIsCooking);
 #endif
 
+	USkeletalMesh* GetRefSkeletalMesh(const FName& ComponentName) const;
+	
 	TSharedPtr<FModelStreamableBulkData> GetModelStreamableBulkData(bool bIsCooking = false);
 
 	// See UCustomizableObjectSystem::LockObject()
@@ -1332,9 +1338,8 @@ public:
 
 #if WITH_EDITORONLY_DATA
 
-	/** Necessary information to generate each component of the Customizable Object. */
-	UPROPERTY(EditAnywhere, Category = CustomizableObject)
-	TArray<FMutableMeshComponentData> MutableMeshComponents;
+	UPROPERTY()
+	TArray<FMutableMeshComponentData> MutableMeshComponents_DEPRECATED;
 
 	/** Unique Identifier - Deterministic. Used to locate Model and Streamable data on disk. Should not be modified. */
 	FGuid Identifier;
@@ -1465,6 +1470,8 @@ public:
 
 		DerivedDataCache,
 
+		ComponentsArray,
+		
 		// -----<new versions can be added above this line>--------
 		LastCustomizableObjectVersion
 	};
diff --git a/Engine/Plugins/Experimental/Mutable/Source/CustomizableObject/Internal/MuCO/ICustomizableObjectEditorModule.h b/Engine/Plugins/Experimental/Mutable/Source/CustomizableObject/Internal/MuCO/ICustomizableObjectEditorModule.h
--- a/Engine/Plugins/Experimental/Mutable/Source/CustomizableObject/Internal/MuCO/ICustomizableObjectEditorModule.h
+++ b/Engine/Plugins/Experimental/Mutable/Source/CustomizableObject/Internal/MuCO/ICustomizableObjectEditorModule.h
@@ -13,6 +13,7 @@ class IToolkitHost;
 class UCustomizableObject;
 class UCustomizableObjectPrivate;
 class UCustomizableObjectInstance;
+class USkeletalMesh;
 class FExtensibilityManager;
 class FBakeOperationCompletedDelegate;
 struct FBakingConfiguration;
@@ -80,5 +81,6 @@ public:
 	/** Return the number of pending compilation requests. Ongoing requests included. */
 	virtual int32 GetNumCompileRequests() = 0;
 
+	virtual USkeletalMesh* GetReferenceSkeletalMesh(const UCustomizableObject& Object, const FName& Component) const = 0;
 };
  
diff --git a/Engine/Plugins/Experimental/Mutable/Source/CustomizableObject/Private/MuCO/CustomizableInstanceLODManagement.cpp b/Engine/Plugins/Experimental/Mutable/Source/CustomizableObject/Private/MuCO/CustomizableInstanceLODManagement.cpp
--- a/Engine/Plugins/Experimental/Mutable/Source/CustomizableObject/Private/MuCO/CustomizableInstanceLODManagement.cpp
+++ b/Engine/Plugins/Experimental/Mutable/Source/CustomizableObject/Private/MuCO/CustomizableInstanceLODManagement.cpp
@@ -359,12 +359,13 @@ void UCustomizableInstanceLODManagement::UpdateInstanceDistsAndLODs(FMutableInst
 				{
 					COI->SetIsBeingUsedByComponentInPlay(true);
 
-					int32 ObjectComponentIndex = CustomizableObjectInstanceUsage->GetComponentIndex();
+					const FName& ComponentName = CustomizableObjectInstanceUsage->GetComponentName();
+					const int32 ObjectComponentIndex = COI->GetCustomizableObject()->GetPrivate()->GetModelResources().ComponentNames.IndexOfByKey(ComponentName);
 
 #if WITH_EDITOR
 					// If the instance is generated but the component doesn't have a mesh, set it.
 					// Can happen when duplicating instances in the editor.
-					const USkeletalMesh* SkeletalMesh = COI->GetSkeletalMesh(ObjectComponentIndex);
+					const USkeletalMesh* SkeletalMesh = COI->GetComponentMeshSkeletalMesh(ComponentName);
 					if (SkeletalMesh && !Parent->GetSkeletalMeshAsset())
 					{
 						// As the instance is already generated, this will be very fast and just set the mesh and call the delegates
diff --git a/Engine/Plugins/Experimental/Mutable/Source/CustomizableObject/Private/MuCO/CustomizableObject.cpp b/Engine/Plugins/Experimental/Mutable/Source/CustomizableObject/Private/MuCO/CustomizableObject.cpp
--- a/Engine/Plugins/Experimental/Mutable/Source/CustomizableObject/Private/MuCO/CustomizableObject.cpp
+++ b/Engine/Plugins/Experimental/Mutable/Source/CustomizableObject/Private/MuCO/CustomizableObject.cpp
@@ -336,7 +336,7 @@ void UCustomizableObject::PostLoad()
 
 	if (CustomizableObjectCustomVersion < FCustomizableObjectCustomVersion::NewComponentOptions)
 	{
-		if (GetPrivate()->MutableMeshComponents.IsEmpty())
+		if (GetPrivate()->MutableMeshComponents_DEPRECATED.IsEmpty())
 		{
 			for (int32 SkeletalMeshIndex = 0; SkeletalMeshIndex < ReferenceSkeletalMeshes_DEPRECATED.Num(); ++SkeletalMeshIndex)
 			{
@@ -344,7 +344,7 @@ void UCustomizableObject::PostLoad()
 				NewComponent.Name = FName(FString::FromInt(SkeletalMeshIndex));
 				NewComponent.ReferenceSkeletalMesh = ReferenceSkeletalMeshes_DEPRECATED[SkeletalMeshIndex];
 
-				GetPrivate()->MutableMeshComponents.Add(NewComponent);
+				GetPrivate()->MutableMeshComponents_DEPRECATED.Add(NewComponent);
 			}
 
 			ReferenceSkeletalMeshes_DEPRECATED.Empty();
@@ -743,6 +743,8 @@ void FModelResources::Serialize(FObjectAndNameAsStringProxyArchive& MemoryWriter
 	MemoryWriter << NumLODsToStream;
 	MemoryWriter << FirstLODAvailable;
 
+	MemoryWriter << ComponentNames;
+
 	// Editor Only data
 	if (!bIsCooking)
 	{
@@ -893,6 +895,8 @@ bool FModelResources::Unserialize(FObjectAndNameAsStringProxyArchive& MemoryRead
 	MemoryReader << NumLODsToStream;
 	MemoryReader << FirstLODAvailable;
 
+	MemoryReader << ComponentNames;
+
 	// Editor Only data
 	if (!bIsCooking)
 	{
@@ -1223,15 +1227,16 @@ void UCustomizableObjectPrivate::AddUncompiledCOWarning(const FString& Additiona
 }
 
 
-USkeletalMesh* UCustomizableObject::GetRefSkeletalMesh(int32 ObjectComponentIndex) const
+USkeletalMesh* UCustomizableObject::GetComponentMeshReferenceSkeletalMesh(const FName& ComponentName) const
 {
 #if WITH_EDITORONLY_DATA
-	if (GetPrivate()->MutableMeshComponents.IsValidIndex(ObjectComponentIndex))
+	if (const ICustomizableObjectEditorModule* Module = ICustomizableObjectEditorModule::Get())
 	{
-		return GetPrivate()->MutableMeshComponents[ObjectComponentIndex].ReferenceSkeletalMesh;
+		return Module->GetReferenceSkeletalMesh(*this, ComponentName);
 	}
 #else
 	const FModelResources& ModelResources = Private->GetModelResources();
+	int32 ObjectComponentIndex = ModelResources.ComponentNames.IndexOfByKey(ComponentName);
 	if (ModelResources.ReferenceSkeletalMeshesData.IsValidIndex(ObjectComponentIndex))
 	{
 		// Can be nullptr if RefSkeletalMeshes are not loaded yet.
@@ -1406,6 +1411,22 @@ int32 UCustomizableObject::GetComponentCount() const
 	return 0;
 }
 
+
+FName UCustomizableObject::GetComponentName(int32 ObjectComponentIndex) const
+{
+	if (IsCompiled())
+	{
+		const TArray<FName>& ComponentNames = GetPrivate()->GetModelResources().ComponentNames;
+		if (ComponentNames.IsValidIndex(ObjectComponentIndex))
+		{
+			return ComponentNames[ObjectComponentIndex];
+		}
+	}
+
+	return NAME_None;
+}
+
+
 int32 UCustomizableObject::GetParameterCount() const
 {
 	return GetPrivate()->ParameterProperties.Num();
@@ -1691,6 +1712,12 @@ FString UCustomizableObject::FindIntParameterValueName(int32 ParamIndex, int32 P
 }
 
 
+USkeletalMesh* UCustomizableObject::GetRefSkeletalMesh(int32 ObjectComponentIndex) const
+{
+	return GetComponentMeshReferenceSkeletalMesh(FName(FString::FromInt(ObjectComponentIndex)));
+}
+
+
 FMutableParamUIMetadata UCustomizableObject::GetParameterUIMetadata(const FString& ParamName) const
 {
 	const FMutableParameterData* ParameterData = Private->GetModelResources().ParameterUIDataMap.Find(ParamName);
diff --git a/Engine/Plugins/Experimental/Mutable/Source/CustomizableObject/Private/MuCO/CustomizableObjectInstance.cpp b/Engine/Plugins/Experimental/Mutable/Source/CustomizableObject/Private/MuCO/CustomizableObjectInstance.cpp
--- a/Engine/Plugins/Experimental/Mutable/Source/CustomizableObject/Private/MuCO/CustomizableObjectInstance.cpp
+++ b/Engine/Plugins/Experimental/Mutable/Source/CustomizableObject/Private/MuCO/CustomizableObjectInstance.cpp
@@ -241,6 +241,29 @@ void UCustomizableInstancePrivate::InitCustomizableObjectData(const UCustomizabl
 }
 
 
+FCustomizableInstanceComponentData* UCustomizableInstancePrivate::GetComponentData(const FName& ComponentName)
+{
+	UCustomizableObject* Object = GetPublic()->GetCustomizableObject();
+	if (!Object)
+	{
+		return nullptr;
+	}
+
+	const int32 ComponentIndex = Object->GetPrivate()->GetModelResources().ComponentNames.IndexOfByKey(ComponentName);
+	if (ComponentIndex != INDEX_NONE)
+	{
+		return nullptr;
+	}
+
+	if (!ComponentsData.IsValidIndex(ComponentIndex))
+	{
+		return nullptr;	
+	}
+
+	return &ComponentsData[ComponentIndex];
+}
+
+
 FCustomizableInstanceComponentData* UCustomizableInstancePrivate::GetComponentData(int32 ComponentIndex)
 {
 	return ComponentsData.IsValidIndex(ComponentIndex) ? &ComponentsData[ComponentIndex] : nullptr;
@@ -936,8 +959,10 @@ void UCustomizableInstancePrivate::PostEditChangePropertyWithoutEditor()
 {
 	MUTABLE_CPUPROFILER_SCOPE(UCustomizableInstancePrivate::PostEditChangePropertyWithoutEditor);
 
-	for (USkeletalMesh* SkeletalMesh : SkeletalMeshes)
+	for (TTuple<FName, TObjectPtr<USkeletalMesh>>& Tuple : SkeletalMeshes)
 	{
+		USkeletalMesh* SkeletalMesh = Tuple.Get<1>();
+		
 		if (SkeletalMesh && SkeletalMesh->GetResourceForRendering())
 		{
 			MUTABLE_CPUPROFILER_SCOPE(InitResources);
@@ -1925,14 +1950,16 @@ bool UCustomizableInstancePrivate::DoComponentsNeedUpdate(UCustomizableObjectIns
 	for (int32 InstanceComponentIndex = 0; InstanceComponentIndex < NumInstanceComponents; ++InstanceComponentIndex)
 	{
 		const FInstanceUpdateData::FComponent& Component = OperationData->InstanceUpdateData.Components[InstanceComponentIndex];
-		int32 ObjectComponentIndex = Component.Id;
+		const int32 ObjectComponentIndex = Component.Id;
+		const FName ComponentName = OperationData->Instance->GetCustomizableObject()->GetPrivate()->GetModelResources().ComponentNames[ObjectComponentIndex];
 
 		if (OperationData->bUseMeshCache)
 		{
 			USkeletalMesh* CachedMesh = CustomizableObject->GetPrivate()->MeshCache.Get(OperationData->MeshDescriptors[ObjectComponentIndex]);
 			if (CachedMesh)
 			{
-				const bool bMeshNeedsUpdate = !SkeletalMeshes.IsValidIndex(ObjectComponentIndex) || (SkeletalMeshes[ObjectComponentIndex] != CachedMesh);
+				TObjectPtr<USkeletalMesh>* SkeletalMesh = SkeletalMeshes.Find(ComponentName);
+				const bool bMeshNeedsUpdate = !SkeletalMesh || (*SkeletalMesh != CachedMesh);
 				OperationData->MeshChanged[InstanceComponentIndex] = bMeshNeedsUpdate;
 				ComponentWithMesh[InstanceComponentIndex] = true;
 				continue;
@@ -1941,7 +1968,8 @@ bool UCustomizableInstancePrivate::DoComponentsNeedUpdate(UCustomizableObjectIns
 
 		// Components with mesh must have valid geometry at CurrentMaxLOD
 
-		const bool bHasSkeletalMesh = SkeletalMeshes.IsValidIndex(ObjectComponentIndex) && SkeletalMeshes[ObjectComponentIndex];
+		TObjectPtr<USkeletalMesh>* SkeletalMeshPtr = SkeletalMeshes.Find(ComponentName);
+		const bool bHasSkeletalMesh = SkeletalMeshPtr && *SkeletalMeshPtr;
 
 		if (Component.LODCount == 0)
 		{
@@ -2023,8 +2051,9 @@ bool UCustomizableInstancePrivate::UpdateSkeletalMesh_PostBeginUpdate0(UCustomiz
 	
 	TextureReuseCache.Empty(); // Sections may have changed, so invalidate the texture reuse cache because it's indexed by section
 
-	TArray<TObjectPtr<USkeletalMesh>> OldSkeletalMeshes = SkeletalMeshes;
-	
+	TMap<FName, TObjectPtr<USkeletalMesh>> OldSkeletalMeshes = SkeletalMeshes;
+	SkeletalMeshes.Reset();
+
 	const FModelResources& ModelResources = CustomizableObject->GetPrivate()->GetModelResources();
 
 	// Collate the Extension Data on the instance into groups based on the extension that produced
@@ -2147,14 +2176,12 @@ bool UCustomizableInstancePrivate::UpdateSkeletalMesh_PostBeginUpdate0(UCustomiz
 		return true;
 	}
 	
-	// Initialize the maximum number of SkeletalMeshes we could possibly have. 
-	SkeletalMeshes.Init(nullptr, CustomizableObject->GetComponentCount());
-
 	const int32 NumInstanceComponents = OperationData->InstanceUpdateData.Components.Num();
 	for (int32 InstanceComponentIndex = 0; InstanceComponentIndex < NumInstanceComponents; ++InstanceComponentIndex)
 	{
 		const FInstanceUpdateData::FComponent& Component = OperationData->InstanceUpdateData.Components[InstanceComponentIndex];
 		const int32 ObjectComponentIndex = Component.Id;
+		const FName ComponentName = OperationData->Instance->GetCustomizableObject()->GetPrivate()->GetModelResources().ComponentNames[ObjectComponentIndex];
 
 		if (!ComponentsData.IsValidIndex(ObjectComponentIndex))
 		{
@@ -2167,9 +2194,9 @@ bool UCustomizableInstancePrivate::UpdateSkeletalMesh_PostBeginUpdate0(UCustomiz
 		// If the component doesn't need an update copy the previously generated mesh.
 		if (!OperationData->MeshChanged[InstanceComponentIndex])
 		{
-			if (OldSkeletalMeshes.IsValidIndex(ObjectComponentIndex))
+			if (TObjectPtr<USkeletalMesh>* Result = OldSkeletalMeshes.Find(ComponentName))
 			{
-				SkeletalMeshes[ObjectComponentIndex] = OldSkeletalMeshes[ObjectComponentIndex];
+				SkeletalMeshes.Add(ComponentName, *Result);
 			}
 
 			continue;
@@ -2181,7 +2208,7 @@ bool UCustomizableInstancePrivate::UpdateSkeletalMesh_PostBeginUpdate0(UCustomiz
 			{
 				check(OperationData->MeshDescriptors[ObjectComponentIndex].Num() == MAX_MESH_LOD_COUNT);
 				ComponentsData[ObjectComponentIndex].LastMeshIdPerLOD = OperationData->MeshDescriptors[ObjectComponentIndex];
-				SkeletalMeshes[ObjectComponentIndex] = CachedMesh;
+				SkeletalMeshes.Add(ComponentName, CachedMesh);
 				continue;
 			}
 		}
@@ -2213,7 +2240,7 @@ bool UCustomizableInstancePrivate::UpdateSkeletalMesh_PostBeginUpdate0(UCustomiz
         		UE_LOG(LogMutable, Error, TEXT("Referenced skeletal mesh [%s] was not pre-loaded. It will be sync-loaded probably causing a hitch. CO [%s]"), *Ref.ToString(), *GetNameSafe(CustomizableObject));
         	}
 
-        	SkeletalMeshes[ObjectComponentIndex] = Ref.LoadSynchronous();
+        	SkeletalMeshes.Add(ComponentName, Ref.LoadSynchronous());
         	continue;
         }
 
@@ -2226,17 +2253,20 @@ bool UCustomizableInstancePrivate::UpdateSkeletalMesh_PostBeginUpdate0(UCustomiz
 		// Create and initialize the SkeletalMesh for this component
 		MUTABLE_CPUPROFILER_SCOPE(ConstructMesh);
 
+		USkeletalMesh* SkeletalMesh = nullptr;
+
 		if (OperationData->bStreamMeshLODs)
 		{
-			SkeletalMeshes[ObjectComponentIndex] = UCustomizableObjectSkeletalMesh::CreateSkeletalMesh(OperationData, *Public, *CustomizableObject, InstanceComponentIndex);
+			SkeletalMesh =UCustomizableObjectSkeletalMesh::CreateSkeletalMesh(OperationData, *Public, *CustomizableObject, InstanceComponentIndex);
 		}
 		else
 		{
 			FName SkeletalMeshName = GenerateUniqueNameFromCOInstance(*Public);
-			SkeletalMeshes[ObjectComponentIndex] = NewObject<USkeletalMesh>(GetTransientPackage(), SkeletalMeshName, RF_Transient);
+			SkeletalMesh = NewObject<USkeletalMesh>(GetTransientPackage(), SkeletalMeshName, RF_Transient);
 		}
-		USkeletalMesh* SkeletalMesh = SkeletalMeshes[ObjectComponentIndex];
-		check(SkeletalMesh);
+		
+		check(SkeletalMesh);		
+		SkeletalMeshes.Add(ComponentName, SkeletalMesh);
 		
 		const FMutableRefSkeletalMeshData& RefSkeletalMeshData = ModelResources.ReferenceSkeletalMeshesData[ObjectComponentIndex];
 
@@ -2317,8 +2347,9 @@ bool UCustomizableInstancePrivate::UpdateSkeletalMesh_PostBeginUpdate0(UCustomiz
 			Pair.Key->OnSkeletalMeshCreated(Pair.Value, ObjectComponentIndex, SkeletalMesh);
 		}
 		
-		// Mesh to copy data from if possible. 
-		USkeletalMesh* OldSkeletalMesh = OldSkeletalMeshes.IsValidIndex(ObjectComponentIndex) ? OldSkeletalMeshes[ObjectComponentIndex] : nullptr;
+		// Mesh to copy data from if possible.
+		TObjectPtr<USkeletalMesh>* OldSkeletalMeshPtr = OldSkeletalMeshes.Find(ComponentName);
+		USkeletalMesh* OldSkeletalMesh = OldSkeletalMeshPtr ? *OldSkeletalMeshPtr : nullptr;
 
 		BuildOrCopyElementData(OperationData, SkeletalMesh, Public, InstanceComponentIndex);
 		bool const bCopyRenderDataSuccess = BuildOrCopyRenderData(OperationData, SkeletalMesh, OldSkeletalMesh, Public, InstanceComponentIndex);
@@ -2951,14 +2982,17 @@ void UCustomizableInstancePrivate::DiscardResources()
 
 	if (SkeletalMeshStatus == ESkeletalMeshStatus::Success)
 	{
-		for (int32 ObjectComponentIndex = 0; ObjectComponentIndex < SkeletalMeshes.Num(); ++ObjectComponentIndex)
+		for (const TTuple<FName, TObjectPtr<USkeletalMesh>>& Tuple : SkeletalMeshes)
 		{
-			if (SkeletalMeshes[ObjectComponentIndex] && SkeletalMeshes[ObjectComponentIndex]->IsValidLowLevel())
+			USkeletalMesh* SkeletalMesh = Tuple.Get<1>();
+			
+			if (SkeletalMesh->IsValidLowLevel())
 			{
-				SkeletalMeshes[ObjectComponentIndex]->ReleaseResources();
-				SkeletalMeshes[ObjectComponentIndex] = nullptr;
+				SkeletalMesh->ReleaseResources();
 			}
 		}
+		
+		SkeletalMeshes.Empty();
 
 		ReleaseMutableResources(false, *Instance);
 	}
@@ -2978,9 +3012,7 @@ void UCustomizableInstancePrivate::SetDefaultSkeletalMesh(bool bSetEmptyMesh) co
 
 	UCustomizableObject* CustomizableObject = Instance->GetCustomizableObject();
 	const FModelResources& ModelResources = CustomizableObject->GetPrivate()->GetModelResources();
-
-	const int32 NumComponents = CustomizableObject->GetComponentCount();
-
+	
 	for (TObjectIterator<UCustomizableObjectInstanceUsage> It; It; ++It)
 	{
 		UCustomizableObjectInstanceUsage* CustomizableObjectInstanceUsage = *It;
@@ -2995,7 +3027,9 @@ void UCustomizableInstancePrivate::SetDefaultSkeletalMesh(bool bSetEmptyMesh) co
 			continue;
 		}
 #endif
-		const int32 ObjectComponentIndex = CustomizableObjectInstanceUsage->GetComponentIndex();
+
+		const FName& ComponentName = CustomizableObjectInstanceUsage->GetComponentName();
+		const int32 ObjectComponentIndex = Instance->GetCustomizableObject()->GetPrivate()->GetModelResources().ComponentNames.IndexOfByKey(ComponentName);
 		if (!ModelResources.ReferenceSkeletalMeshesData.IsValidIndex(ObjectComponentIndex))
 		{
 			continue;
@@ -6119,16 +6153,11 @@ void UCustomizableInstancePrivate::BuildMaterials(const TSharedRef<FUpdateContex
 	{
 		const FInstanceUpdateData::FComponent& Component = OperationData->InstanceUpdateData.Components[InstanceComponentIndex];
 
-		int32 ObjectComponentIndex = Component.Id;
-
-		if (!SkeletalMeshes.IsValidIndex(ObjectComponentIndex))
-		{
-			// This could happen if there are no meshes at all.
-			continue;
-		}
-
-		USkeletalMesh* SkeletalMesh = SkeletalMeshes[ObjectComponentIndex];
+		const int32 ObjectComponentIndex = Component.Id;
+		const FName& ComponentName = GetPublic()->GetCustomizableObject()->GetPrivate()->GetModelResources().ComponentNames[ObjectComponentIndex];
 
+		TObjectPtr<USkeletalMesh>* Result = SkeletalMeshes.Find(ComponentName);
+		TObjectPtr<USkeletalMesh> SkeletalMesh = Result ? *Result : nullptr;
 		if (!SkeletalMesh)
 		{
 			continue;
@@ -6756,8 +6785,9 @@ void UCustomizableInstancePrivate::BuildMaterials(const TSharedRef<FUpdateContex
 			}
 #endif
 
-			// TODO: What is the proper way to refer to dynamic components?
-			const int32 ObjectComponentIndex = CustomizableObjectInstanceUsage->GetComponentIndex();
+			const FName& ComponentName = CustomizableObjectInstanceUsage->GetComponentName();
+			const int32 ObjectComponentIndex = OperationData->Instance->GetCustomizableObject()->GetPrivate()->GetModelResources().ComponentNames.IndexOfByKey(ComponentName);
+			
 			int32 InstanceComponentIndex = -1;
 			for (int32 CurrentInstanceIndex=0; CurrentInstanceIndex<OperationData->InstanceUpdateData.Components.Num(); ++CurrentInstanceIndex)
 			{
@@ -6774,7 +6804,8 @@ void UCustomizableInstancePrivate::BuildMaterials(const TSharedRef<FUpdateContex
 			}
 
 			USkeletalMeshComponent* AttachedParent = CustomizableObjectInstanceUsage->GetAttachParent();
-			if (!AttachedParent || AttachedParent->GetSkeletalMeshAsset() != SkeletalMeshes[ObjectComponentIndex])
+			TObjectPtr<USkeletalMesh>* SkeletalMesh = SkeletalMeshes.Find(ComponentName);
+			if (!AttachedParent || (SkeletalMesh && AttachedParent->GetSkeletalMeshAsset() != *SkeletalMesh))
 			{
 				continue;
 			}
@@ -7036,23 +7067,20 @@ void UCustomizableObjectInstance::Bake(const FBakingConfiguration& InBakingConfi
 
 USkeletalMesh* UCustomizableObjectInstance::GetSkeletalMesh(int32 ObjectComponentIndex) const
 {
-	check(PrivateData);
-	return PrivateData->SkeletalMeshes.IsValidIndex(ObjectComponentIndex) ? PrivateData->SkeletalMeshes[ObjectComponentIndex] : nullptr;
+	return GetComponentMeshSkeletalMesh(FName(FString::FromInt(ObjectComponentIndex)));
 }
 
 
-bool UCustomizableObjectInstance::HasAnySkeletalMesh() const
+USkeletalMesh* UCustomizableObjectInstance::GetComponentMeshSkeletalMesh(const FName& ComponentName) const
 {
-	check(PrivateData);
-	for (int32 ObjectComponentIndex = 0; ObjectComponentIndex < PrivateData->SkeletalMeshes.Num(); ++ObjectComponentIndex)
-	{
-		if (PrivateData->SkeletalMeshes[ObjectComponentIndex])
-		{
-			return true;
-		}
-	}
+	TObjectPtr<USkeletalMesh>* Result = GetPrivate()->SkeletalMeshes.Find(ComponentName);
+	return Result ? *Result : nullptr;
+}
 
-	return false;
+
+bool UCustomizableObjectInstance::HasAnySkeletalMesh() const
+{
+	return !GetPrivate()->SkeletalMeshes.IsEmpty();
 }
 
 
@@ -7366,8 +7394,10 @@ void UCustomizableInstancePrivate::RegenerateImportedModels()
 {
 	MUTABLE_CPUPROFILER_SCOPE(RegenerateImportedModels);
 
-	for (USkeletalMesh* SkeletalMesh : SkeletalMeshes)
+	for (const TTuple<FName, TObjectPtr<USkeletalMesh>>& Tuple : SkeletalMeshes)
 	{
+		USkeletalMesh* SkeletalMesh = Tuple.Get<1>();
+		
 		if (!SkeletalMesh)
 		{
 			continue;
diff --git a/Engine/Plugins/Experimental/Mutable/Source/CustomizableObject/Private/MuCO/CustomizableObjectInstanceUsage.cpp b/Engine/Plugins/Experimental/Mutable/Source/CustomizableObject/Private/MuCO/CustomizableObjectInstanceUsage.cpp
--- a/Engine/Plugins/Experimental/Mutable/Source/CustomizableObject/Private/MuCO/CustomizableObjectInstanceUsage.cpp
+++ b/Engine/Plugins/Experimental/Mutable/Source/CustomizableObject/Private/MuCO/CustomizableObjectInstanceUsage.cpp
@@ -102,6 +102,40 @@ int32 UCustomizableObjectInstanceUsage::GetComponentIndex() const
 }
 
 
+void UCustomizableObjectInstanceUsage::SetComponentName(const FName& Name)
+{
+	if (CustomizableSkeletalComponent)
+	{
+		CustomizableSkeletalComponent->SetComponentName(Name);
+	}
+	else
+	{
+		UsedComponentIndex = INDEX_NONE;
+		UsedComponentName = Name;
+	}
+}
+
+
+FName UCustomizableObjectInstanceUsage::GetComponentName() const
+{
+	if (CustomizableSkeletalComponent)
+	{
+		return CustomizableSkeletalComponent->GetComponentName();
+	}
+	else
+	{
+		if (UsedComponentIndex == INDEX_NONE)
+		{
+			return UsedComponentName;	
+		}
+		else
+		{
+			return FName(FString::FromInt(UsedComponentIndex));
+		}
+	}
+}
+
+
 void UCustomizableObjectInstanceUsage::SetPendingSetSkeletalMesh(bool bIsActive)
 {
 	if (CustomizableSkeletalComponent)
@@ -197,7 +231,7 @@ USkeletalMesh* UCustomizableObjectInstanceUsage::GetSkeletalMesh() const
 {
 	UCustomizableObjectInstance* CustomizableObjectInstance = GetCustomizableObjectInstance();
 
-	return CustomizableObjectInstance ? CustomizableObjectInstance->GetSkeletalMesh(GetComponentIndex()) : nullptr;
+	return CustomizableObjectInstance ? CustomizableObjectInstance->GetComponentMeshSkeletalMesh(GetComponentName()) : nullptr;
 }
 
 
@@ -279,7 +313,7 @@ void UCustomizableObjectInstanceUsage::SetSkeletalMeshAndOverrideMaterials(USkel
 
 	if (bUseOverrideMaterials)
 	{
-		if (FCustomizableInstanceComponentData* ComponentData = CustomizableObjectInstance.GetPrivate()->GetComponentData(GetComponentIndex()))
+		if (FCustomizableInstanceComponentData* ComponentData = CustomizableObjectInstance.GetPrivate()->GetComponentData(GetComponentName()))
 		{
 			for (int32 Index = 0; Index < ComponentData->OverrideMaterials.Num(); ++Index)
 			{
@@ -391,11 +425,11 @@ void UCustomizableObjectInstanceUsage::UpdateDistFromComponentToLevelEditorCamer
 		}
 
 		USkeletalMesh* AttachedSkeletalMesh = GetAttachedSkeletalMesh();
-		const int32 ObjectComponentIndex = GetComponentIndex();
+		const FName& ComponentName = GetComponentName();
 
 		const bool bInstanceGenerated = CustomizableObjectInstance->GetPrivate()->SkeletalMeshStatus != ESkeletalMeshStatus::NotGenerated;
-		USkeletalMesh* GeneratedSkeletalMesh = bInstanceGenerated ? CustomizableObjectInstance->GetSkeletalMesh(ObjectComponentIndex) :
-			CustomizableObjectInstance->GetCustomizableObject()->GetRefSkeletalMesh(ObjectComponentIndex);
+		USkeletalMesh* GeneratedSkeletalMesh = bInstanceGenerated ? CustomizableObjectInstance->GetComponentMeshSkeletalMesh(ComponentName) :
+			CustomizableObjectInstance->GetCustomizableObject()->GetComponentMeshReferenceSkeletalMesh(ComponentName);
 
 		if (AttachedSkeletalMesh != GeneratedSkeletalMesh)
 		{
@@ -419,14 +453,14 @@ void UCustomizableObjectInstanceUsage::EditorUpdateComponent()
 		if (ParentActor)
 		{
 			USkeletalMesh* AttachedSkeletalMesh = GetAttachedSkeletalMesh();
-			const int32 ComponentIndex = GetComponentIndex();
+			const FName& ComponentName = GetComponentName();
 
 			const bool bInstanceGenerated = CustomizableObjectInstance->GetPrivate()->SkeletalMeshStatus != ESkeletalMeshStatus::NotGenerated;
 			USkeletalMesh* GeneratedSkeletalMesh = bInstanceGenerated ?
-				CustomizableObjectInstance->GetSkeletalMesh(ComponentIndex) :
+				CustomizableObjectInstance->GetComponentMeshSkeletalMesh(ComponentName) :
 				GetSkipSetReferenceSkeletalMesh() ?
 					nullptr :
-					CustomizableObjectInstance->GetCustomizableObject()->GetRefSkeletalMesh(ComponentIndex);
+					CustomizableObjectInstance->GetCustomizableObject()->GetComponentMeshReferenceSkeletalMesh(ComponentName);
 
 			if (AttachedSkeletalMesh != GeneratedSkeletalMesh)
 			{
@@ -441,9 +475,9 @@ void UCustomizableObjectInstanceUsage::EditorUpdateComponent()
 void UCustomizableObjectInstanceUsage::SetVisibilityOfSkeletalMeshSectionWithMaterialName(bool bInVisible, const FString& MaterialName, int32 LOD)
 {
 	UCustomizableObjectInstance* CustomizableObjectInstance = GetCustomizableObjectInstance();
-	int32 ObjectComponentIndex = GetComponentIndex();
+	const FName& ComponentName = GetComponentName();
 
-	USkeletalMesh* SkeletalMesh = CustomizableObjectInstance->GetSkeletalMesh(ObjectComponentIndex);
+	USkeletalMesh* SkeletalMesh = CustomizableObjectInstance->GetComponentMeshSkeletalMesh(ComponentName);
 
 	if (!SkeletalMesh)
 	{
@@ -532,11 +566,11 @@ void UCustomizableObjectInstanceUsage::UpdateDistFromComponentToPlayer(const AAc
 			}
 		}
 
-		int32 ObjectComponentIndex = GetComponentIndex();
+		const FName& ComponentName = GetComponentName();
 
-		if (ParentActor && GetAttachedSkeletalMesh() == nullptr && CustomizableObjectInstance->GetSkeletalMesh(ObjectComponentIndex))
+		if (ParentActor && GetAttachedSkeletalMesh() == nullptr && CustomizableObjectInstance->GetComponentMeshSkeletalMesh(ComponentName))
 		{
-			SetSkeletalMesh(CustomizableObjectInstance->GetSkeletalMesh(ObjectComponentIndex));
+			SetSkeletalMesh(CustomizableObjectInstance->GetComponentMeshSkeletalMesh(ComponentName));
 		}
 	}
 }
@@ -564,7 +598,7 @@ void UCustomizableObjectInstanceUsage::Tick(float DeltaTime)
 
 	if (USkeletalMeshComponent* Parent = Cast<USkeletalMeshComponent>(GetAttachParent()))
 	{
-		const int32 ObjectComponentIndex = GetComponentIndex();
+		const FName& ComponentName = GetComponentName();
 
 		USkeletalMesh* SkeletalMesh = nullptr;
 
@@ -572,7 +606,7 @@ void UCustomizableObjectInstanceUsage::Tick(float DeltaTime)
 		if (bInstanceGenerated)
 		{
 			// Generated SkeletalMesh to set, can be null if the component is empty
-			SkeletalMesh = CustomizableObjectInstance->GetSkeletalMesh(ObjectComponentIndex);
+			SkeletalMesh = CustomizableObjectInstance->GetComponentMeshSkeletalMesh(ComponentName);
 		}
 		else
 		{
@@ -581,7 +615,7 @@ void UCustomizableObjectInstanceUsage::Tick(float DeltaTime)
 				CustomizableObject->bEnableUseRefSkeletalMeshAsPlaceholder)
 			{
 				// Can be nullptr
-				SkeletalMesh = CustomizableObject->GetRefSkeletalMesh(ObjectComponentIndex);
+				SkeletalMesh = CustomizableObject->GetComponentMeshReferenceSkeletalMesh(ComponentName);
 			}
 		}
 
diff --git a/Engine/Plugins/Experimental/Mutable/Source/CustomizableObject/Private/MuCO/CustomizableObjectSystem.cpp b/Engine/Plugins/Experimental/Mutable/Source/CustomizableObject/Private/MuCO/CustomizableObjectSystem.cpp
--- a/Engine/Plugins/Experimental/Mutable/Source/CustomizableObject/Private/MuCO/CustomizableObjectSystem.cpp
+++ b/Engine/Plugins/Experimental/Mutable/Source/CustomizableObject/Private/MuCO/CustomizableObjectSystem.cpp
@@ -954,19 +954,17 @@ void UpdateSkeletalMesh(const TSharedRef<FUpdateContextPrivate>& Context)
 		bool bPhysicsAssetUpdated = false;	
 		
 		if (IsValid(CustomizableObjectInstanceUsage) &&
-			(CustomizableObjectInstanceUsage->GetCustomizableObjectInstance() == CustomizableObjectInstance) &&
-			CustomizableObjectInstancePrivateData->SkeletalMeshes.IsValidIndex(CustomizableObjectInstanceUsage->GetComponentIndex())
-		   )
+			CustomizableObjectInstanceUsage->GetCustomizableObjectInstance() == CustomizableObjectInstance)
 		{
 			MUTABLE_CPUPROFILER_SCOPE(UpdateSkeletalMesh_SetSkeletalMesh);
 
-			CustomizableObjectInstanceUsage->SetSkeletalMesh(CustomizableObjectInstancePrivateData->SkeletalMeshes[CustomizableObjectInstanceUsage->GetComponentIndex()], &bSkeletalMeshUpdated, &bMaterialsUpdated);
+			USkeletalMesh* SkeletalMesh = CustomizableObjectInstance->GetComponentMeshSkeletalMesh(CustomizableObjectInstanceUsage->GetComponentName());
+			CustomizableObjectInstanceUsage->SetSkeletalMesh(SkeletalMesh, &bSkeletalMeshUpdated, &bMaterialsUpdated);
 
-			if (CustomizableObjectInstancePrivateData->HasCOInstanceFlags(ReplacePhysicsAssets))
+			if (CustomizableObjectInstancePrivateData->HasCOInstanceFlags(ReplacePhysicsAssets) &&
+				SkeletalMesh)
 			{
-				CustomizableObjectInstanceUsage->SetPhysicsAsset(
-					CustomizableObjectInstancePrivateData->SkeletalMeshes[CustomizableObjectInstanceUsage->GetComponentIndex()] ?
-					CustomizableObjectInstancePrivateData->SkeletalMeshes[CustomizableObjectInstanceUsage->GetComponentIndex()]->GetPhysicsAsset() : nullptr, &bPhysicsAssetUpdated);
+				CustomizableObjectInstanceUsage->SetPhysicsAsset(SkeletalMesh->GetPhysicsAsset(), &bPhysicsAssetUpdated);	
 			}
 		}
 
diff --git a/Engine/Plugins/Experimental/Mutable/Source/CustomizableObject/Private/MuCO/CustomizableSkeletalComponent.cpp b/Engine/Plugins/Experimental/Mutable/Source/CustomizableObject/Private/MuCO/CustomizableSkeletalComponent.cpp
--- a/Engine/Plugins/Experimental/Mutable/Source/CustomizableObject/Private/MuCO/CustomizableSkeletalComponent.cpp
+++ b/Engine/Plugins/Experimental/Mutable/Source/CustomizableObject/Private/MuCO/CustomizableSkeletalComponent.cpp
@@ -81,6 +81,26 @@ USkeletalMesh* UCustomizableSkeletalComponent::GetAttachedSkeletalMesh() const
 }
 
 
+void UCustomizableSkeletalComponent::SetComponentName(const FName& Name)
+{
+	ComponentIndex = INDEX_NONE;
+	ComponentName = Name;
+}
+
+
+FName UCustomizableSkeletalComponent::GetComponentName() const
+{
+	if (ComponentIndex == INDEX_NONE)
+	{
+		return ComponentName;	
+	}
+	else
+	{
+		return FName(FString::FromInt(ComponentIndex));
+	}
+}
+
+
 void UCustomizableSkeletalComponent::UpdateSkeletalMeshAsync(bool bNeverSkipUpdate)
 {
 	if (CustomizableObjectInstanceUsage)
diff --git a/Engine/Plugins/Experimental/Mutable/Source/CustomizableObject/Private/MuCO/LogBenchmarkUtil.cpp b/Engine/Plugins/Experimental/Mutable/Source/CustomizableObject/Private/MuCO/LogBenchmarkUtil.cpp
--- a/Engine/Plugins/Experimental/Mutable/Source/CustomizableObject/Private/MuCO/LogBenchmarkUtil.cpp
+++ b/Engine/Plugins/Experimental/Mutable/Source/CustomizableObject/Private/MuCO/LogBenchmarkUtil.cpp
@@ -236,7 +236,7 @@ void FLogBenchmarkUtil::UpdateStats()
 			GEngine->AddOnScreenDebugMessage(MsgIndex++, .0f, FColor::Green, TEXT("Animation tags: ") + TagString);
 
 			check(Instance->GetPrivate() != nullptr);
-			FCustomizableInstanceComponentData* ComponentData = Instance->GetPrivate()->GetComponentData(CustomizableObjectInstanceUsage->GetComponentIndex());
+			FCustomizableInstanceComponentData* ComponentData = Instance->GetPrivate()->GetComponentData(CustomizableObjectInstanceUsage->GetComponentName());
 
 			if (ComponentData)
 			{
diff --git a/Engine/Plugins/Experimental/Mutable/Source/CustomizableObject/Public/MuCO/CustomizableObject.h b/Engine/Plugins/Experimental/Mutable/Source/CustomizableObject/Public/MuCO/CustomizableObject.h
--- a/Engine/Plugins/Experimental/Mutable/Source/CustomizableObject/Public/MuCO/CustomizableObject.h
+++ b/Engine/Plugins/Experimental/Mutable/Source/CustomizableObject/Public/MuCO/CustomizableObject.h
@@ -363,6 +363,11 @@ public:
 	/** Get the number of components this Customizable Object has. */
 	UFUNCTION(BlueprintCallable, Category = CustomizableObject)
 	int32 GetComponentCount() const;
+
+	/** Return the name of the component.
+	 *  @return NAME_None if the component does not exist. */
+	UFUNCTION(BlueprintCallable, Category = CustomizableObject)
+	FName GetComponentName(int32 ObjectComponentIndex) const;
 	
 	/** Get the number of parameters available in instances of this object. */
 	UFUNCTION(BlueprintCallable, Category = CustomizableObject)
@@ -400,8 +405,13 @@ public:
 	int32 FindIntParameterValue( int32 ParamIndex, const FString& Value ) const;
 	FString FindIntParameterValueName(int32 ParamIndex, int32 ParamValue) const;
 
+	// DEPRECATED
 	USkeletalMesh* GetRefSkeletalMesh(int32 ObjectComponentIndex = 0) const;
 
+	/** Given a Mesh Component name, return its reference Skeletal Mesh. */
+	UFUNCTION(BlueprintCallable, Category = CustomizableObject)
+	USkeletalMesh* GetComponentMeshReferenceSkeletalMesh(const FName& Name) const;
+	
 	/** Get the default value of a parameter of type Float.
 	  * @param InParameterName The name of the Float parameter to get the default value of.
 	  * @return The default value of the provided parameter name. */
diff --git a/Engine/Plugins/Experimental/Mutable/Source/CustomizableObject/Public/MuCO/CustomizableObjectInstance.h b/Engine/Plugins/Experimental/Mutable/Source/CustomizableObject/Public/MuCO/CustomizableObjectInstance.h
--- a/Engine/Plugins/Experimental/Mutable/Source/CustomizableObject/Public/MuCO/CustomizableObjectInstance.h
+++ b/Engine/Plugins/Experimental/Mutable/Source/CustomizableObject/Public/MuCO/CustomizableObjectInstance.h
@@ -275,10 +275,13 @@ public:
 	UFUNCTION(BlueprintCallable, Category = CustomizableObjectInstance)
 	void SetCurrentState(const FString& StateName);
 
-	/** Get the skeletal mesh generated for this instance. 
-	  * If the object has multiple components, an index of the component can be specified. */
+	// DEPRECATED
 	UFUNCTION(BlueprintCallable, Category = CustomizableObjectInstance)
 	USkeletalMesh* GetSkeletalMesh(int32 ComponentIndex = 0) const;
+
+	/** Given a Mesh Component name, return its generated Skeletal Mesh. */
+	UFUNCTION(BlueprintCallable, Category = CustomizableObjectInstance)
+	USkeletalMesh* GetComponentMeshSkeletalMesh(const FName& ComponentName) const;
 	
 	/** Return true if a skeletal mesh has been generated for any component of this instance. */
 	UFUNCTION(BlueprintCallable, Category = CustomizableObjectInstance)
diff --git a/Engine/Plugins/Experimental/Mutable/Source/CustomizableObject/Public/MuCO/CustomizableObjectInstanceUsage.h b/Engine/Plugins/Experimental/Mutable/Source/CustomizableObject/Public/MuCO/CustomizableObjectInstanceUsage.h
--- a/Engine/Plugins/Experimental/Mutable/Source/CustomizableObject/Public/MuCO/CustomizableObjectInstanceUsage.h
+++ b/Engine/Plugins/Experimental/Mutable/Source/CustomizableObject/Public/MuCO/CustomizableObjectInstanceUsage.h
@@ -41,7 +41,13 @@ public:
 	// For now assume this is an *Object* component index (to index the array of all possible components in the CO) and not an *Instance* component index.
 	void SetComponentIndex(int32 ObjectComponentIndex);
 	int32 GetComponentIndex() const;
+	
+	UFUNCTION(BlueprintCallable, Category = CustomizableObjectInstanceUsage)
+	void SetComponentName(const FName& Name);
 
+	UFUNCTION(BlueprintCallable, Category = CustomizableObjectInstanceUsage)
+	FName GetComponentName() const;
+	
 	void AttachTo(USkeletalMeshComponent* SkeletalMeshComponent);
 	USkeletalMeshComponent* GetAttachParent() const;
 
@@ -102,6 +108,10 @@ private:
 	UPROPERTY()
 	int32 UsedComponentIndex;
 
+	/** Only used if the ComponentIndex is INDEX_NONE. */
+	UPROPERTY()
+	FName UsedComponentName;
+
 	// Used to replace the SkeletalMesh of the parent component by the ReferenceSkeletalMesh or the generated SkeletalMesh 
 	bool bUsedPendingSetSkeletalMesh = false;
 
diff --git a/Engine/Plugins/Experimental/Mutable/Source/CustomizableObject/Public/MuCO/CustomizableSkeletalComponent.h b/Engine/Plugins/Experimental/Mutable/Source/CustomizableObject/Public/MuCO/CustomizableSkeletalComponent.h
--- a/Engine/Plugins/Experimental/Mutable/Source/CustomizableObject/Public/MuCO/CustomizableSkeletalComponent.h
+++ b/Engine/Plugins/Experimental/Mutable/Source/CustomizableObject/Public/MuCO/CustomizableSkeletalComponent.h
@@ -24,6 +24,12 @@ public:
 	UPROPERTY(BlueprintReadWrite, EditAnywhere, Category = CustomizableSkeletalMesh)
 	int32 ComponentIndex;
 
+private:
+	/** Only used if the ComponentIndex is INDEX_NONE. */
+	UPROPERTY()
+	FName ComponentName;
+
+public:
 	// Used to replace the SkeletalMesh of the parent component by the ReferenceSkeletalMesh or the generated SkeletalMesh 
 	bool bPendingSetSkeletalMesh = false;
 
@@ -38,6 +44,12 @@ public:
 	USkeletalMesh* GetSkeletalMesh() const;
 	USkeletalMesh* GetAttachedSkeletalMesh() const;
 
+	UFUNCTION(BlueprintCallable, Category = CustomizableSkeletalMesh)
+	void SetComponentName(const FName& Name);
+
+	UFUNCTION(BlueprintCallable, Category = CustomizableSkeletalMesh)
+	FName GetComponentName() const;
+	
 	/** Update Skeletal Mesh asynchronously. */
 	UFUNCTION(BlueprintCallable, Category = CustomizableObject)
 	void UpdateSkeletalMeshAsync(bool bNeverSkipUpdate = false);
diff --git a/Engine/Plugins/Experimental/Mutable/Source/CustomizableObjectEditor/Private/MuCOE/CustomizableObjectCompiler.cpp b/Engine/Plugins/Experimental/Mutable/Source/CustomizableObjectEditor/Private/MuCOE/CustomizableObjectCompiler.cpp
--- a/Engine/Plugins/Experimental/Mutable/Source/CustomizableObjectEditor/Private/MuCOE/CustomizableObjectCompiler.cpp
+++ b/Engine/Plugins/Experimental/Mutable/Source/CustomizableObjectEditor/Private/MuCOE/CustomizableObjectCompiler.cpp
@@ -566,7 +566,8 @@ mu::NodeObjectPtr FCustomizableObjectCompiler::GenerateMutableRoot(
 	GenerationContext.RootVersionBridge = ActualRootObject->VersionBridge;
 
 	UCustomizableObjectNodeObject* ActualRoot = GetRootNode(ActualRootObject, bMultipleBaseObjectsFound);
-
+	GenerationContext.Root = ActualRoot;
+	
 	if (bMultipleBaseObjectsFound)
 	{
 		ErrorMsg = LOCTEXT("MultipleBaseActualRoot", "Multiple base object nodes found.");
@@ -658,60 +659,6 @@ mu::NodeObjectPtr FCustomizableObjectCompiler::GenerateMutableRoot(
 		}
 	}
 	
-	if (ActualRootObject->GetPrivate()->MutableMeshComponents.IsEmpty())
-	{
-		CompilerLog(LOCTEXT("NoComponentsError", "Error! There are no components defined in the Object Properties Tab."), ActualRoot, EMessageSeverity::Error);
-		return nullptr;
-	}
-
-	TArray<const FName> ComponentNames;
-
-	// Make sure we have a valid Name, Reference SkeletalMesh and Skeleton for each component
-	for (int32 ComponentIndex = 0; ComponentIndex < ActualRootObject->GetPrivate()->MutableMeshComponents.Num(); ++ComponentIndex)
-	{
-		FName ComponentName = ActualRootObject->GetPrivate()->MutableMeshComponents[ComponentIndex].Name;
-
-		if (ComponentName.IsNone())
-		{
-			CompilerLog(LOCTEXT("EmptyComponentNameError", "Error! Missing name in a component of the Customizable Object."), ActualRoot, EMessageSeverity::Error);
-			return nullptr;
-		}
-		else if (ComponentNames.Contains(ComponentName))
-		{
-			CompilerLog(FText::Format(LOCTEXT("RepeatedComponentName", "Error! Repeated name [{0}] used in more than one Component"),
-				FText::FromName(ComponentName)), ActualRoot, EMessageSeverity::Error);
-			return nullptr;
-		}
-
-		ComponentNames.Add(ComponentName);
-
-		USkeletalMesh* RefSkeletalMesh = ActualRootObject->GetRefSkeletalMesh(ComponentIndex);
-		if (!RefSkeletalMesh)
-		{
-			CompilerLog(LOCTEXT("NoReferenceMeshObjectTab", "Error! Missing reference mesh in the Object Properties Tab"), ActualRoot, EMessageSeverity::Error);
-			return nullptr;
-		}
-
-		USkeleton* RefSkeleton = RefSkeletalMesh->GetSkeleton();
-		if(!RefSkeleton)
-		{
-			FText Msg = FText::Format(LOCTEXT("NoReferenceSkeleton", "Error! Missing skeleton in the reference mesh [{0}]"),
-				FText::FromString(GenerationContext.CustomizableObjectWithCycle->GetPathName()));
-
-			CompilerLog(Msg, ActualRoot, EMessageSeverity::Error);
-			return nullptr;
-		}
-
-		// Add a new entry to the list of Component Infos
-		GenerationContext.ComponentInfos.Add(FMutableComponentInfo(ComponentName, RefSkeletalMesh));
-
-		// Make sure the Skeleton from the reference mesh is added to the list of referenced Skeletons.
-		GenerationContext.ReferencedSkeletons.Add(RefSkeleton);
-
-		// Add reference meshes to the participating objects
-		GenerationContext.AddParticipatingObject(*RefSkeletalMesh);
-	}
-
 	// Find all component nodes
 	TSet<UCustomizableObject*> ParticipatingObjects; // All COs participating in the hierarchy
 	{
@@ -748,29 +695,26 @@ mu::NodeObjectPtr FCustomizableObjectCompiler::GenerateMutableRoot(
 	{
 		for (UEdGraphNode* Node : ParticipatingObject->GetPrivate()->GetSource()->Nodes)
 		{
-			if (UCustomizableObjectNodeComponentMesh* NodeComponentMesh = Cast<UCustomizableObjectNodeComponentMesh>(Node))
+			UCustomizableObjectNodeComponentMesh* NodeComponentMesh = Cast<UCustomizableObjectNodeComponentMesh>(Node);
+			if (!NodeComponentMesh)
 			{
-				if (FMutableComponentInfo* ComponentInfo = GenerationContext.ComponentInfos.FindByPredicate([&](const FMutableComponentInfo& Element)
-				{
-					return Element.ComponentName == NodeComponentMesh->ComponentName;
-				}))
-				{
-					if (ComponentInfo->NodeComponentMesh)
-					{
-						FText Msg = FText::Format(LOCTEXT("ComponentNodeWithSameNameExists", "Already exist a Component node with the same name in Customizable Object [{0}]"), FText::FromString(GetRootObject(*ComponentInfo->NodeComponentMesh)->GetName()));
-						CompilerLog(Msg, Node, EMessageSeverity::Error);
-						return nullptr;
-					}
+				continue;
+			}
 
-					ComponentInfo->NodeComponentMesh = NodeComponentMesh;
-				}
+			if (UCustomizableObjectNodeComponentMesh** Result = GenerationContext.MeshComponents.Find(NodeComponentMesh->ComponentName))
+			{
+				UCustomizableObjectNodeComponentMesh* NodeMeshComponent = *Result;
+				check(NodeMeshComponent);
+				
+				FText Msg = FText::Format(LOCTEXT("ComponentNodeWithSameNameExists", "Already exists a Mesh Component node with the same name in Customizable Object [{0}]"), FText::FromString(GetRootObject(*NodeMeshComponent)->GetName()));
+				CompilerLog(Msg, Node, EMessageSeverity::Error);
+				return nullptr;
 			}
+
+			GenerationContext.MeshComponents.Add(NodeComponentMesh->ComponentName, NodeComponentMesh);
 		}
 	}
-
-	// Copy component data to the object being compiled
-	Object->GetPrivate()->MutableMeshComponents = ActualRootObject->GetPrivate()->MutableMeshComponents;
-
+	
     GenerationContext.RealTimeMorphTargetsOverrides = ActualRoot->RealTimeMorphSelectionOverrides;
     GenerationContext.RealTimeMorphTargetsOverrides.Reset();
 
@@ -1194,7 +1138,13 @@ void FCustomizableObjectCompiler::CompileInternal(bool bAsync)
 		ModelResources.CustomizableObjectPathMap = MoveTemp(GenerationContext.CustomizableObjectPathMap);
 #endif
 
-		ModelResources.NumComponents = GenerationContext.NumMeshComponentsInRoot + GenerationContext.NumPassthroughMeshComponents;
+		ModelResources.NumComponents = GenerationContext.NumComponents;
+
+		for (int32 Index = 0; Index <  GenerationContext.ComponentInfos.Num(); ++Index)
+		{
+			ModelResources.ComponentNames.Add(GenerationContext.ComponentInfos[Index].ComponentName);
+		}
+		
 		ModelResources.NumLODs = GenerationContext.NumLODsInRoot;
 		ModelResources.NumLODsToStream = GenerationContext.bEnableLODStreaming ? GenerationContext.NumMaxLODsToStream : 0;
 		ModelResources.FirstLODAvailable = GenerationContext.FirstLODAvailable;
diff --git a/Engine/Plugins/Experimental/Mutable/Source/CustomizableObjectEditor/Private/MuCOE/CustomizableObjectDetails.cpp b/Engine/Plugins/Experimental/Mutable/Source/CustomizableObjectEditor/Private/MuCOE/CustomizableObjectDetails.cpp
--- a/Engine/Plugins/Experimental/Mutable/Source/CustomizableObjectEditor/Private/MuCOE/CustomizableObjectDetails.cpp
+++ b/Engine/Plugins/Experimental/Mutable/Source/CustomizableObjectEditor/Private/MuCOE/CustomizableObjectDetails.cpp
@@ -159,15 +159,6 @@ void FCustomizableObjectDetails::CustomizeDetails( IDetailLayoutBuilder& DetailB
 			VersionStructProperty->MarkHiddenByCustomization();
 		}
 	}
-
-	// Add Private's components property.
-	if (CustomizableObject && !CustomizableObject->IsChildObject())
-	{
-		TArray<UObject*> Private;
-		Private.Add(CustomizableObject->GetPrivate());
-
-		CustomizableObjectCategory.AddExternalObjectProperty(Private, "MutableMeshComponents");
-	}
 }
 
 
diff --git a/Engine/Plugins/Experimental/Mutable/Source/CustomizableObjectEditor/Private/MuCOE/CustomizableObjectEditor.cpp b/Engine/Plugins/Experimental/Mutable/Source/CustomizableObjectEditor/Private/MuCOE/CustomizableObjectEditor.cpp
--- a/Engine/Plugins/Experimental/Mutable/Source/CustomizableObjectEditor/Private/MuCOE/CustomizableObjectEditor.cpp
+++ b/Engine/Plugins/Experimental/Mutable/Source/CustomizableObjectEditor/Private/MuCOE/CustomizableObjectEditor.cpp
@@ -2394,7 +2394,7 @@ void FCustomizableObjectEditor::CreatePreviewActor()
 
 	TArray<TWeakObjectPtr<UDebugSkelMeshComponent>> PreviewSkeletalMeshComponents;
 	
-	for (int32 ComponentIndex = 0; ComponentIndex < CustomizableObject->GetComponentCount(); ++ComponentIndex)
+	for (int32 ObjectComponentIndex = 0; ObjectComponentIndex < CustomizableObject->GetComponentCount(); ++ObjectComponentIndex)
 	{
 		UDebugSkelMeshComponent* DebugComponent = NewObject<UDebugSkelMeshComponent>(Actor.Get(), NAME_None, RF_Transient);
 		DebugComponent->bCastInsetShadow = true; // For better quality shadows in the editor previews, more similar to the in-game ones
@@ -2409,7 +2409,7 @@ void FCustomizableObjectEditor::CreatePreviewActor()
 		UCustomizableSkeletalComponent* CustomizableComponent = NewObject<UCustomizableSkeletalComponent>(DebugComponent);
 		CustomizableComponent->bSkipSetReferenceSkeletalMesh = true;
 		CustomizableComponent->CustomizableObjectInstance = PreviewInstance;
-		CustomizableComponent->ComponentIndex = ComponentIndex;
+		CustomizableComponent->SetComponentName(CustomizableObject->GetComponentName(ObjectComponentIndex));
 		CustomizableComponent->AttachToComponent(DebugComponent, FAttachmentTransformRules::KeepRelativeTransform);
 		CustomizableComponent->RegisterComponent();
 	}
diff --git a/Engine/Plugins/Experimental/Mutable/Source/CustomizableObjectEditor/Private/MuCOE/CustomizableObjectEditorModule.cpp b/Engine/Plugins/Experimental/Mutable/Source/CustomizableObjectEditor/Private/MuCOE/CustomizableObjectEditorModule.cpp
--- a/Engine/Plugins/Experimental/Mutable/Source/CustomizableObjectEditor/Private/MuCOE/CustomizableObjectEditorModule.cpp
+++ b/Engine/Plugins/Experimental/Mutable/Source/CustomizableObjectEditor/Private/MuCOE/CustomizableObjectEditorModule.cpp
@@ -77,6 +77,7 @@
 #include "MuCOE/GraphTraversal.h"
 #include "MuCOE/CustomizableObjectInstanceBaker.h"
 #include "Editor.h"
+#include "Nodes/CustomizableObjectNodeComponentMeshDetails.h"
 
 class AActor;
 class FString;
@@ -693,6 +694,31 @@ void FCustomizableObjectEditorModule::BakeCustomizableObjectInstance(UCustomizab
 }
 
 
+USkeletalMesh* FCustomizableObjectEditorModule::GetReferenceSkeletalMesh(const UCustomizableObject& Object, const FName& ComponentName) const
+{
+	UCustomizableObject* RootObject = GetRootObject(const_cast<UCustomizableObject*>(&Object));
+
+	TSet<UCustomizableObject*> Objects;
+	GetAllObjectsInGraph(RootObject, Objects);
+	
+	for (const UCustomizableObject* CurrentObject : Objects)
+	{
+		for (UEdGraphNode* Node : CurrentObject->GetPrivate()->GetSource()->Nodes)
+		{
+			if (UCustomizableObjectNodeComponentMesh* NodeComponentMesh = Cast<UCustomizableObjectNodeComponentMesh>(Node))
+			{
+				if (NodeComponentMesh->ComponentName == ComponentName)
+				{
+					return NodeComponentMesh->ReferenceSkeletalMesh;
+				}
+			}
+		}	
+	}
+
+	return {};
+}
+
+
 void FCustomizableObjectEditorModule::CompileCustomizableObject(const TSharedRef<FCompilationRequest>& InCompilationRequest, bool bForceRequest)
 {
 	if (IsRunningGame())
diff --git a/Engine/Plugins/Experimental/Mutable/Source/CustomizableObjectEditor/Private/MuCOE/CustomizableObjectEditorModule.h b/Engine/Plugins/Experimental/Mutable/Source/CustomizableObjectEditor/Private/MuCOE/CustomizableObjectEditorModule.h
--- a/Engine/Plugins/Experimental/Mutable/Source/CustomizableObjectEditor/Private/MuCOE/CustomizableObjectEditorModule.h
+++ b/Engine/Plugins/Experimental/Mutable/Source/CustomizableObjectEditor/Private/MuCOE/CustomizableObjectEditorModule.h
@@ -42,7 +42,8 @@ public:
 	virtual bool IsRootObject(const UCustomizableObject& Object) const override;
 	virtual FString GetCurrentContentVersionForObject(const UCustomizableObject& Object) const override;
 	virtual void BakeCustomizableObjectInstance(UCustomizableObjectInstance* InTargetInstance, const FBakingConfiguration& InBakingConfig) override;
-
+	virtual USkeletalMesh* GetReferenceSkeletalMesh(const UCustomizableObject& Object, const FName& ComponentName) const override;
+	
 	/** Request for a given customizable object to be compiled. Async compile requests will be queued and processed sequentially. 
 	 * @param InCompilationRequest - Request to compile an object. 
 	 * @param bForceRequest - Queue request even if already in the pending list. */
diff --git a/Engine/Plugins/Experimental/Mutable/Source/CustomizableObjectEditor/Private/MuCOE/CustomizableObjectFactory.cpp b/Engine/Plugins/Experimental/Mutable/Source/CustomizableObjectEditor/Private/MuCOE/CustomizableObjectFactory.cpp
--- a/Engine/Plugins/Experimental/Mutable/Source/CustomizableObjectEditor/Private/MuCOE/CustomizableObjectFactory.cpp
+++ b/Engine/Plugins/Experimental/Mutable/Source/CustomizableObjectEditor/Private/MuCOE/CustomizableObjectFactory.cpp
@@ -33,8 +33,7 @@
 
 #define LOCTEXT_NAMESPACE "CustomizableObjectFactory"
 
-static const FVector2D NewCOWindowsSize = FVector2D(300,340);
-static const FVector2D NewChildWindowsSize = FVector2D(300, 220);
+static const FVector2D NewCOWindowsSize = FVector2D(300, 230);
 
 UCustomizableObjectFactory::UCustomizableObjectFactory()
 	: Super()
@@ -104,9 +103,7 @@ UObject* UCustomizableObjectFactory::FactoryCreateNew(UClass* Class, UObject* In
 	{
 		return NewObj;
 	}
-
-	ObjectPrivate->MutableMeshComponents.SetNum(CreationSettings.NumMeshComponents);
-
+	
 	if (!CreationSettings.bEmptyObject)
 	{
 		if (CreationSettings.bIsChildObject && CreationSettings.ParentObject.IsValid())
@@ -128,16 +125,6 @@ UObject* UCustomizableObjectFactory::FactoryCreateNew(UClass* Class, UObject* In
 
 			ObjectPrivate->SetIsChildObject(true);
 		}
-		else
-		{
-			check(CreationSettings.ComponentsInfo.Num() == CreationSettings.NumMeshComponents);
-
-			for (int32 MeshIndex = 0; MeshIndex < CreationSettings.ComponentsInfo.Num(); ++MeshIndex)
-			{
-				ObjectPrivate->MutableMeshComponents[MeshIndex].Name = CreationSettings.ComponentsInfo[MeshIndex].ComponentName;
-				ObjectPrivate->MutableMeshComponents[MeshIndex].ReferenceSkeletalMesh = CreationSettings.ComponentsInfo[MeshIndex].ReferenceSkeletalMesh.LoadSynchronous();
-			}
-		}
 	}
 
 	return NewObj;
@@ -154,12 +141,7 @@ const FCustomizableObjectOptions FCustomizableObjectFactoryUI::ConstructFactoryU
 	
 	TSharedRef<IDetailsView> DetailsView = PropertyEditorModule.CreateDetailView(DetailsViewArgs);
 	DetailsView->SetObject(nullptr);
-
-	// Default number of components for non-child objects
-	Options.NumMeshComponents = 1;
-	Options.ComponentsInfo.SetNum(1);
-	GenerateComponentOptions();
-
+	
 	// Settings window
 	COSettingsWindow = SNew(SWindow)
 	.Title(LOCTEXT("CustomizableObjectFactoryptions", "New Costumizable Object"))
@@ -207,6 +189,7 @@ const FCustomizableObjectOptions FCustomizableObjectFactoryUI::ConstructFactoryU
 			[
 				SNew(SBorder)
 				.Padding(5.0f, 10.0f, 0.0f, 10.0f)
+				.Visibility(this, &FCustomizableObjectFactoryUI::GetParentWidgetsVisibility)
 				.BorderImage(FAppStyle::Get().GetBrush("Brushes.Panel"))
 				.HAlign(HAlign_Fill)
 				[
@@ -216,7 +199,6 @@ const FCustomizableObjectOptions FCustomizableObjectFactoryUI::ConstructFactoryU
 					.HAlign(HAlign_Left)
 					[
 						SNew(STextBlock)
-						.Visibility(this, &FCustomizableObjectFactoryUI::GetParentWidgetsVisibility)
 						.Text(LOCTEXT("SelectedParent_Text", "Parent Object: "))
 						.Font(IDetailLayoutBuilder::GetDetailFont())
 					]
@@ -226,7 +208,6 @@ const FCustomizableObjectOptions FCustomizableObjectFactoryUI::ConstructFactoryU
 					.AutoHeight()
 					[
 						SNew(SObjectPropertyEntryBox)
-						.Visibility(this, &FCustomizableObjectFactoryUI::GetParentWidgetsVisibility)
 						.AllowedClass(UCustomizableObject::StaticClass())
 						.OnObjectChanged(this, &FCustomizableObjectFactoryUI::OnPickedCustomizableObjectParent)
 						.ObjectPath(this, &FCustomizableObjectFactoryUI::GetSelectedCustomizableObjectPath)
@@ -244,7 +225,6 @@ const FCustomizableObjectOptions FCustomizableObjectFactoryUI::ConstructFactoryU
 					.Padding(0.0f, 15.0f, 0.0f, 0.0f)
 					[
 						SNew(STextBlock)
-						.Visibility(this, &FCustomizableObjectFactoryUI::GetParentWidgetsVisibility)
 						.Text(LOCTEXT("SelecteGroup_Text", "Group Node: "))
 						.Font(IDetailLayoutBuilder::GetDetailFont())
 					]
@@ -254,111 +234,11 @@ const FCustomizableObjectOptions FCustomizableObjectFactoryUI::ConstructFactoryU
 					.AutoHeight()
 					[
 						SAssignNew(GroupSelector, STextComboBox)
-						.Visibility(this, &FCustomizableObjectFactoryUI::GetParentWidgetsVisibility)
 						.OptionsSource(&GroupOptions)
 						.IsEnabled(this, &FCustomizableObjectFactoryUI::IsNodeGroupSelectorEnabled)
 						.OnSelectionChanged(this, &FCustomizableObjectFactoryUI::OnSelectGroupComboBox)
 						.Font(IDetailLayoutBuilder::GetDetailFont())
 					]
-
-					// Component Selector Widgets
-					+ SVerticalBox::Slot()
-					.AutoHeight()
-					.HAlign(HAlign_Left)
-					[
-						SNew(STextBlock)
-						.Visibility(this, &FCustomizableObjectFactoryUI::GetComponentWidgetsVisibility)
-						.Text(LOCTEXT("NumberComponents_Text", "Num Mesh Components: "))
-						.Font(IDetailLayoutBuilder::GetDetailFont())
-					]
-
-					+ SVerticalBox::Slot()
-					.AutoHeight()
-					.Padding(5.0f, 7.0f, 10.0f, 0.0f)
-					[
-						SNew(SSpinBox<int32>)
-						.Visibility(this, &FCustomizableObjectFactoryUI::GetComponentWidgetsVisibility)
-						.OnValueChanged(this, &FCustomizableObjectFactoryUI::OnNumComponentsChanged, ETextCommit::Type::Default)
-						.OnValueCommitted(this, &FCustomizableObjectFactoryUI::OnNumComponentsChanged)
-						.Value(this, &FCustomizableObjectFactoryUI::GetNumComponents)
-						.MinValue(1)
-						.MaxValue(255)
-						.MaxSliderValue(6)
-						.Delta(1)
-						.AlwaysUsesDeltaSnap(true)
-						.Font(IDetailLayoutBuilder::GetDetailFont())
-					]
-
-					+SVerticalBox::Slot()
-					.Padding(0.0f, 15.0f, 0.0f, 0.0f)
-					[
-						SNew(STextBlock)
-						.Visibility(this, &FCustomizableObjectFactoryUI::GetComponentWidgetsVisibility)
-						.Text(LOCTEXT("SelectedComponent_Text", "Component: "))
-						.Font(IDetailLayoutBuilder::GetDetailFont())
-					]
-
-					+SVerticalBox::Slot()
-					.AutoHeight()
-					.Padding(5.0f, 5.0f, 10.0f, 0.0f)
-					[
-						SAssignNew(ComponentSelector, STextComboBox)
-						.Visibility(this, &FCustomizableObjectFactoryUI::GetComponentWidgetsVisibility)
-						.ToolTipText(LOCTEXT("SelectedComponent_Tooltip", "Select a component to set its Reference Skeletal Mesh."))
-						.OptionsSource(&ComponentsOptions)
-						.InitiallySelectedItem(ComponentsOptions[0])
-						.IsEnabled(this, &FCustomizableObjectFactoryUI::IsComponentSelectorEnabled)
-						.OnSelectionChanged(this, &FCustomizableObjectFactoryUI::OnSelectComponentComboBox)
-						.Font(IDetailLayoutBuilder::GetDetailFont())
-					]
-
-					+ SVerticalBox::Slot()
-					.Padding(0.0f, 15.0f, 0.0f, 0.0f)
-					.AutoHeight()
-					[
-						SNew(STextBlock)
-						.Visibility(this, &FCustomizableObjectFactoryUI::GetComponentWidgetsVisibility)
-						.Text(this, &FCustomizableObjectFactoryUI::GetSelectorWidgetText, true)
-						.Font(IDetailLayoutBuilder::GetDetailFont())
-					]
-
-					+ SVerticalBox::Slot()
-					.Padding(5.0f, 10.0f, 10.0f, 0.0f)
-					.AutoHeight()
-					[
-						SNew(SEditableTextBox)
-						.Visibility(this, &FCustomizableObjectFactoryUI::GetComponentWidgetsVisibility)
-						.Text(this, &FCustomizableObjectFactoryUI::GetComponentName)
-						.OnTextCommitted(this, &FCustomizableObjectFactoryUI::OnTextCommited)
-						.Font(IDetailLayoutBuilder::GetDetailFont())
-					]
-
-					+ SVerticalBox::Slot()
-					.Padding(0.0f, 15.0f, 0.0f, 0.0f)
-					[
-						SNew(STextBlock)
-						.Visibility(this, &FCustomizableObjectFactoryUI::GetComponentWidgetsVisibility)
-						.Text(this, &FCustomizableObjectFactoryUI::GetSelectorWidgetText, false)
-						.Font(IDetailLayoutBuilder::GetDetailFont())
-					]
-
-					+SVerticalBox::Slot()
-					.AutoHeight()
-					.Padding(5.0f, 5.0f, 10.0f, 0.0f)
-					[
-						SNew(SObjectPropertyEntryBox)
-						.Visibility(this, &FCustomizableObjectFactoryUI::GetComponentWidgetsVisibility)
-						.IsEnabled(this, &FCustomizableObjectFactoryUI::IsComponentSelectorEnabled)
-						.AllowedClass(USkeletalMesh::StaticClass())
-						.OnObjectChanged(this, &FCustomizableObjectFactoryUI::OnPickedComponentSkeletalMesh)
-						.ObjectPath(this, &FCustomizableObjectFactoryUI::GetSelectedComponentSkeletalMeshPath)
-						.DisplayThumbnail(true)
-						.ThumbnailPool(UThumbnailManager::Get().GetSharedThumbnailPool())
-						.AllowClear(false)
-						.AllowCreate(false)
-						.DisplayBrowse(false)
-						.DisplayUseSelected(false)
-					]
 				]
 			]
 
@@ -427,13 +307,6 @@ void FCustomizableObjectFactoryUI::OnCheckBoxChanged(ECheckBoxState State)
 	// Reset settings
 	Options.ParentObject = nullptr;
 	Options.GroupNodeName.Empty();
-	Options.NumMeshComponents = Options.bIsChildObject ? 0 : 1;
-	Options.bIsChildObject ? Options.ComponentsInfo.Empty() : Options.ComponentsInfo.SetNum(1);
-
-	FVector2D ClientSize = Options.bIsChildObject ? NewChildWindowsSize : NewCOWindowsSize;
-	FVector2D WindowsSize = ClientSize * COSettingsWindow->GetDPIScaleFactor();
-	
-	COSettingsWindow->Resize(WindowsSize);
 }
 
 
@@ -496,19 +369,6 @@ bool FCustomizableObjectFactoryUI::IsConfigurationValid() const
 			return true;
 		}
 	}
-	else
-	{
-		// Check if all components have a valid Skeletal Mesh and name assigned
-		for (int32 CompIndex = 0; CompIndex < Options.NumMeshComponents; ++CompIndex)
-		{
-			if (!Options.ComponentsInfo.IsValidIndex(CompIndex) || Options.ComponentsInfo[CompIndex].ReferenceSkeletalMesh.IsNull() || Options.ComponentsInfo[CompIndex].ComponentName.IsNone())
-			{
-				return false;
-			}
-		}
-
-		return true;
-	}
 
 	return false;
 }
@@ -591,50 +451,6 @@ EVisibility FCustomizableObjectFactoryUI::GetComponentWidgetsVisibility() const
 }
 
 
-int32 FCustomizableObjectFactoryUI::GetNumComponents() const
-{
-	return Options.NumMeshComponents;
-}
-
-
-void FCustomizableObjectFactoryUI::OnNumComponentsChanged(int32 Value, ETextCommit::Type CommitInfo)
-{
-	if (Options.NumMeshComponents != Value)
-	{
-		Options.NumMeshComponents = Value;
-		Options.ComponentsInfo.SetNum(Value);
-		GenerateComponentOptions();
-
-		if (Value > 0)
-		{
-			ComponentSelector->SetSelectedItem(ComponentsOptions[0]);
-		}
-		else
-		{
-			ComponentSelector->ClearSelection();
-		}
-	}
-}
-
-
-void FCustomizableObjectFactoryUI::GenerateComponentOptions()
-{
-	ComponentsOptions.Reset();
-
-	for (int32 ComponentIndex = 0; ComponentIndex < Options.NumMeshComponents; ++ComponentIndex)
-	{
-		FString ComponentName = "Component " + FString::FromInt(ComponentIndex);
-		ComponentsOptions.Add(MakeShareable(new FString(ComponentName)));
-	}
-}
-
-
-bool FCustomizableObjectFactoryUI::IsComponentSelectorEnabled() const
-{
-	return Options.NumMeshComponents > 0;
-}
-
-
 void FCustomizableObjectFactoryUI::OnSelectComponentComboBox(TSharedPtr<FString> Selection, ESelectInfo::Type SelectInfo)
 {
 	if (Selection)
@@ -644,85 +460,4 @@ void FCustomizableObjectFactoryUI::OnSelectComponentComboBox(TSharedPtr<FString>
 }
 
 
-FText FCustomizableObjectFactoryUI::GetSelectorWidgetText(bool bIsName) const
-{
-	FString ComponentName;
-
-	if (ComponentSelector->GetSelectedItem().IsValid() && Options.NumMeshComponents)
-	{
-		ComponentName = *ComponentSelector->GetSelectedItem();
-	}
-
-	FString VariableName = bIsName ? " Name:" : " Skeletal Mesh:";
-	FString Message = "Select " + ComponentName + VariableName;
-
-	return FText::FromString(Message);
-}
-
-
-int32 FCustomizableObjectFactoryUI::GetSelectedComponentIndex() const
-{
-	for (int32 ComponentIndex = 0; ComponentIndex < ComponentsOptions.Num(); ++ComponentIndex)
-	{
-		if (ComponentSelector.IsValid() && ComponentSelector->GetSelectedItem().IsValid()
-			&& *ComponentsOptions[ComponentIndex] == *ComponentSelector->GetSelectedItem()
-			&& Options.ComponentsInfo.IsValidIndex(ComponentIndex))
-		{
-			return ComponentIndex;
-		}
-	}
-
-	return -1;
-}
-
-
-void FCustomizableObjectFactoryUI::OnPickedComponentSkeletalMesh(const FAssetData& SelectedAsset)
-{
-	int32 ComponentIndex = GetSelectedComponentIndex();
-	if (ComponentIndex != INDEX_NONE && SelectedAsset.IsValid())
-	{
-		if (USkeletalMesh* SkeletalMesh = Cast<USkeletalMesh>(SelectedAsset.GetAsset()))
-		{
-			// we have to load the asset otherwise the asset thumbnail won't be visible
-			Options.ComponentsInfo[ComponentIndex].ReferenceSkeletalMesh = SkeletalMesh;
-		}
-	}
-}
-
-
-FString FCustomizableObjectFactoryUI::GetSelectedComponentSkeletalMeshPath() const
-{
-	int32 ComponentIndex = GetSelectedComponentIndex();
-	FString SkeletalMeshPath;
-
-	if (ComponentIndex != INDEX_NONE && !Options.ComponentsInfo[ComponentIndex].ReferenceSkeletalMesh.IsNull())
-	{
-		SkeletalMeshPath = Options.ComponentsInfo[ComponentIndex].ReferenceSkeletalMesh.ToSoftObjectPath().ToString();
-	}
-
-	return SkeletalMeshPath;
-}
-
-
-FText FCustomizableObjectFactoryUI::GetComponentName() const
-{
-	int32 ComponentIndex = GetSelectedComponentIndex();
-	if (ComponentIndex != INDEX_NONE)
-	{
-		return FText::FromName(Options.ComponentsInfo[ComponentIndex].ComponentName);
-	}
-
-	return FText();
-}
-
-
-void FCustomizableObjectFactoryUI::OnTextCommited(const FText& NewName, ETextCommit::Type CommitInfo)
-{
-	int32 ComponentIndex = GetSelectedComponentIndex();
-	if (ComponentIndex != INDEX_NONE)
-	{
-		Options.ComponentsInfo[ComponentIndex].ComponentName = FName(*NewName.ToString());
-	}
-}
-
 #undef LOCTEXT_NAMESPACE
diff --git a/Engine/Plugins/Experimental/Mutable/Source/CustomizableObjectEditor/Private/MuCOE/CustomizableObjectGraph.cpp b/Engine/Plugins/Experimental/Mutable/Source/CustomizableObjectEditor/Private/MuCOE/CustomizableObjectGraph.cpp
--- a/Engine/Plugins/Experimental/Mutable/Source/CustomizableObjectEditor/Private/MuCOE/CustomizableObjectGraph.cpp
+++ b/Engine/Plugins/Experimental/Mutable/Source/CustomizableObjectEditor/Private/MuCOE/CustomizableObjectGraph.cpp
@@ -43,11 +43,9 @@ void UCustomizableObjectGraph::PostLoad()
 	}
 
 	// Execute backwards compatible code for all nodes. It requires all nodes to be loaded.
-
-	TArray<TObjectPtr<UEdGraphNode>> NodesCopy = Nodes; // Copy to be able to remove nodes inside the BackwardsCompatibleFixup.
-
 	for (int32 Version = GetLinkerCustomVersion(FCustomizableObjectCustomVersion::GUID) + 1; Version <= FCustomizableObjectCustomVersion::LatestVersion; ++Version)
 	{
+		TArray<TObjectPtr<UEdGraphNode>> NodesCopy = Nodes; // Copy to be able to remove nodes inside the BackwardsCompatibleFixup.
 		for (UEdGraphNode* Node : NodesCopy)
 		{
 			if (UCustomizableObjectNode* CustomizableObjectNode = Cast<UCustomizableObjectNode>(Node))
@@ -74,9 +72,9 @@ void UCustomizableObjectGraph::PostLoad()
 		}
 	}
 
-	NodesCopy = Nodes; // Copy to be able to remove nodes inside the PostBackwardsCompatibleFixup.
 
 	// Do any additional work which require nodes to be valid (i.e., have executed BackwardsCompatibleFixup).
+	TArray<TObjectPtr<UEdGraphNode>> NodesCopy = Nodes; // Copy to be able to remove nodes inside the PostBackwardsCompatibleFixup.
 	for (UEdGraphNode* Node : NodesCopy)
 	{
 		if (UCustomizableObjectNode* CustomizableObjectNode = Cast<UCustomizableObjectNode>(Node))
diff --git a/Engine/Plugins/Experimental/Mutable/Source/CustomizableObjectEditor/Private/MuCOE/CustomizableObjectInstanceEditor.cpp b/Engine/Plugins/Experimental/Mutable/Source/CustomizableObjectEditor/Private/MuCOE/CustomizableObjectInstanceEditor.cpp
--- a/Engine/Plugins/Experimental/Mutable/Source/CustomizableObjectEditor/Private/MuCOE/CustomizableObjectInstanceEditor.cpp
+++ b/Engine/Plugins/Experimental/Mutable/Source/CustomizableObjectEditor/Private/MuCOE/CustomizableObjectInstanceEditor.cpp
@@ -974,7 +974,7 @@ void FCustomizableObjectInstanceEditor::CreatePreviewActor()
 
 	PreviewSkeletalMeshComponents.Reset();
 	
-	for (int32 ComponentIndex = 0; ComponentIndex < Object->GetComponentCount(); ++ComponentIndex)
+	for (int32 ObjectComponentIndex = 0; ObjectComponentIndex < Object->GetComponentCount(); ++ObjectComponentIndex)
 	{
 		UDebugSkelMeshComponent* DebugComponent = NewObject<UDebugSkelMeshComponent>(Actor.Get(), NAME_None, RF_Transient);
 		DebugComponent->bCastInsetShadow = true; // For better quality shadows in the editor previews, more similar to the in-game ones
@@ -989,7 +989,7 @@ void FCustomizableObjectInstanceEditor::CreatePreviewActor()
 		UCustomizableSkeletalComponent* CustomizableComponent = NewObject<UCustomizableSkeletalComponent>(DebugComponent);
 		CustomizableComponent->bSkipSetReferenceSkeletalMesh = true;
 		CustomizableComponent->CustomizableObjectInstance = CustomizableObjectInstance;
-		CustomizableComponent->ComponentIndex = ComponentIndex;
+		CustomizableComponent->SetComponentName(Object->GetComponentName(ObjectComponentIndex));
 		CustomizableComponent->AttachToComponent(DebugComponent, FAttachmentTransformRules::KeepRelativeTransform);
 		CustomizableComponent->RegisterComponent();
 	}
diff --git a/Engine/Plugins/Experimental/Mutable/Source/CustomizableObjectEditor/Private/MuCOE/CustomizableObjectInstanceFactory.cpp b/Engine/Plugins/Experimental/Mutable/Source/CustomizableObjectEditor/Private/MuCOE/CustomizableObjectInstanceFactory.cpp
--- a/Engine/Plugins/Experimental/Mutable/Source/CustomizableObjectEditor/Private/MuCOE/CustomizableObjectInstanceFactory.cpp
+++ b/Engine/Plugins/Experimental/Mutable/Source/CustomizableObjectEditor/Private/MuCOE/CustomizableObjectInstanceFactory.cpp
@@ -32,7 +32,7 @@ void UCustomizableObjectInstanceFactory::PostSpawnActor(UObject* Asset, AActor*
 	UCustomizableObjectInstance* COInstance = Cast<UCustomizableObjectInstance>(Asset);
 	ACustomizableSkeletalMeshActor* NewCSMActor = CastChecked<ACustomizableSkeletalMeshActor>(NewActor);
 	if (NewCSMActor && COInstance)
-	{
+	{	
 		int32 NumComponents = COInstance->GetNumComponents();
 
 		for (int32 ComponentIndex = 0; ComponentIndex < NumComponents; ++ComponentIndex)
@@ -139,23 +139,26 @@ bool UCustomizableObjectInstanceFactory::CanCreateActorFrom(const FAssetData& As
 
 USkeletalMesh* UCustomizableObjectInstanceFactory::GetSkeletalMeshFromAsset(UObject* Asset, int32 ComponentIndex) const
 {
-    USkeletalMesh* SkeletalMesh = Cast<USkeletalMesh>(Asset);
     UCustomizableObjectInstance* CustomizableObjectInstance = Cast<UCustomizableObjectInstance>(Asset);
-
-    if (SkeletalMesh == nullptr && CustomizableObjectInstance != nullptr)
+    if (!CustomizableObjectInstance)
     {
-        SkeletalMesh = CustomizableObjectInstance->GetSkeletalMesh(ComponentIndex);
-        if (SkeletalMesh == nullptr)
-        {
-            if (UCustomizableObject* CustomizableObject = CustomizableObjectInstance->GetCustomizableObject())
-            {
-                SkeletalMesh = CustomizableObject->GetRefSkeletalMesh(ComponentIndex);
-            }
-        }
+    	return nullptr;
     }
 
-    //check(SkeletalMesh != NULL);
-    return SkeletalMesh;
+	UCustomizableObject* CustomizableObject = CustomizableObjectInstance->GetCustomizableObject();
+	if (!CustomizableObject)
+	{
+		return nullptr;
+	}
+	
+	const FName& ComponentName = CustomizableObject->GetComponentName(ComponentIndex);
+
+	if (USkeletalMesh* SkeletalMesh = CustomizableObjectInstance->GetComponentMeshSkeletalMesh(ComponentName))
+	{
+		return SkeletalMesh;
+	}
+    		
+	return CustomizableObject->GetComponentMeshReferenceSkeletalMesh(ComponentName);
 }
 
 
diff --git a/Engine/Plugins/Experimental/Mutable/Source/CustomizableObjectEditor/Private/MuCOE/GenerateMutableSource/GenerateMutableSource.cpp b/Engine/Plugins/Experimental/Mutable/Source/CustomizableObjectEditor/Private/MuCOE/GenerateMutableSource/GenerateMutableSource.cpp
--- a/Engine/Plugins/Experimental/Mutable/Source/CustomizableObjectEditor/Private/MuCOE/GenerateMutableSource/GenerateMutableSource.cpp
+++ b/Engine/Plugins/Experimental/Mutable/Source/CustomizableObjectEditor/Private/MuCOE/GenerateMutableSource/GenerateMutableSource.cpp
@@ -443,22 +443,9 @@ mu::FBoneName FMutableGraphGenerationContext::GetBoneUnique(const FName& InBoneN
 //}
 
 
-FMutableComponentInfo& FMutableGraphGenerationContext::GetCurrentComponentInfo()
+FMutableComponentInfo* FMutableGraphGenerationContext::GetCurrentComponentInfo()
 {
-	FMutableComponentInfo* CurrentComponentInfo = ComponentInfos.FindByPredicate(
-		[this](const FMutableComponentInfo& Component) { return Component.ComponentName == CurrentMeshComponent; });
-	
-	// Temp workaround to the problem of modifiers that generate meshes (like ExtendMeshSection) not having a known component at generation time.
-	// TODO: Actually detect all possible components and change this query in GenerationContext state to accomodate a set of "current components" instead of one.
-	//check(CurrentComponentInfo);
-
-	if (!CurrentComponentInfo)
-	{
-		check(!ComponentInfos.IsEmpty());
-		CurrentComponentInfo = &ComponentInfos[0];
-	}
-
-	return *CurrentComponentInfo;
+	return ComponentInfos.FindByPredicate([this](const FMutableComponentInfo& Component) { return Component.ComponentName == CurrentMeshComponent; });
 }
 
 
@@ -1298,7 +1285,7 @@ void PopulateReferenceSkeletalMeshesData(FMutableGraphGenerationContext& Generat
 	const FString PlatformName = GenerationContext.Options.TargetPlatform->IniPlatformName();
 	
 	const uint32 LODCount = GenerationContext.NumLODsInRoot;
-	const uint32 ComponentCount = GenerationContext.NumMeshComponentsInRoot;
+	const uint32 ComponentCount = GenerationContext.NumComponents;
 
 	GenerationContext.ReferenceSkeletalMeshesData.AddDefaulted(ComponentCount);
 	for(uint32 ComponentIndex = 0; ComponentIndex < ComponentCount; ++ComponentIndex)
diff --git a/Engine/Plugins/Experimental/Mutable/Source/CustomizableObjectEditor/Private/MuCOE/GenerateMutableSource/GenerateMutableSource.h b/Engine/Plugins/Experimental/Mutable/Source/CustomizableObjectEditor/Private/MuCOE/GenerateMutableSource/GenerateMutableSource.h
--- a/Engine/Plugins/Experimental/Mutable/Source/CustomizableObjectEditor/Private/MuCOE/GenerateMutableSource/GenerateMutableSource.h
+++ b/Engine/Plugins/Experimental/Mutable/Source/CustomizableObjectEditor/Private/MuCOE/GenerateMutableSource/GenerateMutableSource.h
@@ -563,6 +563,9 @@ struct FMutableGraphGenerationContext
 	
 	UCustomizableObject* Object = nullptr;
 
+	/** Full hierarchy root. */
+	UCustomizableObjectNodeObject* Root = nullptr;
+	
 	// Non-owned reference to the compiler object
 	FCustomizableObjectCompiler* Compiler = nullptr;
 
@@ -753,8 +756,9 @@ struct FMutableGraphGenerationContext
 	//void CheckPhysicsAssetInSkeletalMesh(const USkeletalMesh* InSkeletalMesh);
 
 	/** Get the reference skeletal mesh associated to the current mesh component being generated */
-	FMutableComponentInfo& GetCurrentComponentInfo();
+	FMutableComponentInfo* GetCurrentComponentInfo();
 
+	TMap<FName, UCustomizableObjectNodeComponentMesh*> MeshComponents;
 
 	TArray<FMutableComponentInfo> ComponentInfos;
 	TArray<FMutableRefSkeletalMeshData> ReferenceSkeletalMeshesData;
@@ -818,9 +822,7 @@ struct FMutableGraphGenerationContext
 	mu::Ptr<mu::NodeMesh> ComponentMeshOverride;
 
 	uint8 NumLODsInRoot = 0;
-	uint8 NumMeshComponentsInRoot = 0;
-
-	uint8 NumPassthroughMeshComponents = 0;
+	uint8 NumComponents = 0;
 
 	uint8 FirstLODAvailable = MAX_MESH_LOD_COUNT;
 	uint8 NumMaxLODsToStream = MAX_MESH_LOD_COUNT;
diff --git a/Engine/Plugins/Experimental/Mutable/Source/CustomizableObjectEditor/Private/MuCOE/GenerateMutableSource/GenerateMutableSourceComponent.cpp b/Engine/Plugins/Experimental/Mutable/Source/CustomizableObjectEditor/Private/MuCOE/GenerateMutableSource/GenerateMutableSourceComponent.cpp
--- a/Engine/Plugins/Experimental/Mutable/Source/CustomizableObjectEditor/Private/MuCOE/GenerateMutableSource/GenerateMutableSourceComponent.cpp
+++ b/Engine/Plugins/Experimental/Mutable/Source/CustomizableObjectEditor/Private/MuCOE/GenerateMutableSource/GenerateMutableSourceComponent.cpp
@@ -154,17 +154,53 @@ mu::Ptr<mu::NodeComponent> GenerateMutableSourceComponent(const UEdGraphPin* Pin
 	
 	if (const UCustomizableObjectNodeComponentMesh* TypedComponentMesh = Cast<UCustomizableObjectNodeComponentMesh>(Node))
 	{
-		int32 CurrentComponent = GenerationContext.Object->GetPrivate()->MutableMeshComponents.IndexOfByPredicate([&](const FMutableMeshComponentData& ComponentData)
-		{
-			return ComponentData.Name == TypedComponentMesh->ComponentName;
-		});
+		UCustomizableObjectNodeObject* ActualRoot = GenerationContext.Root;
+			
+		FName ComponentName = TypedComponentMesh->ComponentName;
 		
-		if (CurrentComponent == INDEX_NONE)
+		if (TypedComponentMesh->ComponentName.IsNone())
+		{
+			GenerationContext.Compiler->CompilerLog(LOCTEXT("EmptyComponentNameError", "Error! Missing name in a component of the Customizable Object."), ActualRoot, EMessageSeverity::Error);
+			return nullptr;
+		}
+
+		if (GenerationContext.ComponentInfos.ContainsByPredicate([&](const FMutableComponentInfo& ComponentInfo)
+		{
+			return ComponentInfo.ComponentName == ComponentName;
+		}))
+		{
+			GenerationContext.Compiler->CompilerLog(FText::Format(LOCTEXT("RepeatedComponentName", "Error! Repeated name [{0}] used in more than one Component"),
+				FText::FromName(ComponentName)), ActualRoot, EMessageSeverity::Error);
+			return nullptr;
+		}
+			
+		USkeletalMesh* RefSkeletalMesh = TypedComponentMesh->ReferenceSkeletalMesh;
+		if (!RefSkeletalMesh)
 		{
-			GenerationContext.Compiler->CompilerLog(FText::Format(LOCTEXT("ComponentNotDeclared", "Component {0} not declared in Object Properties tab."), FText::FromName(TypedComponentMesh->ComponentName)), Node, EMessageSeverity::Error);
-			return {};
+			GenerationContext.Compiler->CompilerLog(LOCTEXT("NoReferenceMeshObjectTab", "Error! Missing reference Skeletal Mesh"), ActualRoot, EMessageSeverity::Error);
+			return nullptr;
+		}
+
+		USkeleton* RefSkeleton = RefSkeletalMesh->GetSkeleton();
+		if(!RefSkeleton)
+		{
+			FText Msg = FText::Format(LOCTEXT("NoReferenceSkeleton", "Error! Missing skeleton in the reference mesh [{0}]"), FText::FromString(GenerationContext.CustomizableObjectWithCycle->GetPathName()));
+
+			GenerationContext.Compiler->CompilerLog(Msg, ActualRoot, EMessageSeverity::Error);
+			return nullptr;
 		}
 
+		// Add a new entry to the list of Component Infos
+		FMutableComponentInfo& ComponentInfo = GenerationContext.ComponentInfos.Add_GetRef(FMutableComponentInfo(ComponentName, RefSkeletalMesh));
+
+		ComponentInfo.AccumulateBonesToRemovePerLOD(TypedComponentMesh->LODReductionSettings, TypedComponentMesh->NumLODs);
+
+		// Make sure the Skeleton from the reference mesh is added to the list of referenced Skeletons.
+		GenerationContext.ReferencedSkeletons.Add(RefSkeleton);
+
+		// Add reference meshes to the participating objects
+		GenerationContext.AddParticipatingObject(*RefSkeletalMesh);
+		
 		// Ensure that the CO has a valid AutoLODStrategy on the ActualRoot.
 		if (TypedComponentMesh->AutoLODStrategy == ECustomizableObjectAutomaticLODStrategy::Inherited)
 		{
@@ -196,18 +232,7 @@ mu::Ptr<mu::NodeComponent> GenerateMutableSourceComponent(const UEdGraphPin* Pin
 			check(!GenerationContext.ComponentInfos.IsEmpty());
 
 			// NumLODsInRoot
-			int32 MaxRefMeshLODs = 1;
-			for (int32 MeshIndex = 0; MeshIndex < GenerationContext.ComponentInfos.Num(); ++MeshIndex)
-			{
-				const USkeletalMesh* RefSkeletalMesh = GenerationContext.ComponentInfos[MeshIndex].RefSkeletalMesh;
-				if (RefSkeletalMesh && RefSkeletalMesh->GetLODNum() > MaxRefMeshLODs)
-				{
-					MaxRefMeshLODs = RefSkeletalMesh->GetLODNum();
-				}
-	
-				GenerationContext.ComponentInfos[MeshIndex].AccumulateBonesToRemovePerLOD(TypedComponentMesh->LODReductionSettings, TypedComponentMesh->NumLODs);
-			}
-
+			int32 MaxRefMeshLODs = TypedComponentMesh->ReferenceSkeletalMesh->GetLODNum();
 			if (MaxRefMeshLODs < NumLODs)
 			{
 				FString Msg = FString::Printf(TEXT("The object has %d LODs but the reference mesh only %d. Resulting objects will have %d LODs."),
@@ -217,12 +242,9 @@ mu::Ptr<mu::NodeComponent> GenerateMutableSourceComponent(const UEdGraphPin* Pin
 			}
 			else
 			{
-				GenerationContext.NumLODsInRoot = NumLODs;
+				GenerationContext.NumLODsInRoot = FMath::Max(GenerationContext.NumLODsInRoot, static_cast<uint8>(NumLODs));
 			}
 		
-			USkeletalMesh* RefSkeletalMesh = GenerationContext.ComponentInfos[0].RefSkeletalMesh;
-			check(RefSkeletalMesh);
-
 			const FMutableLODSettings& LODSettings = GenerationContext.Object->LODSettings;
 
 			// Find the MinLOD available for the target platform
@@ -278,14 +300,8 @@ mu::Ptr<mu::NodeComponent> GenerateMutableSourceComponent(const UEdGraphPin* Pin
 			GenerationContext.NumMaxLODsToStream = FMath::Clamp(GenerationContext.NumMaxLODsToStream, 0, GenerationContext.NumLODsInRoot - 1);
 		}
 		
-		// Mesh components per instance
-		if (!GenerationContext.NumMeshComponentsInRoot)
-		{
-			GenerationContext.NumMeshComponentsInRoot = GenerationContext.Object->GetPrivate()->MutableMeshComponents.Num();
-		}
-		
 		mu::Ptr<mu::NodeComponentNew> NodeComponentNew = new mu::NodeComponentNew();
-		NodeComponentNew->Id = CurrentComponent;
+		NodeComponentNew->Id = GenerationContext.NumComponents++;
 		
 		NodeComponentNew->SetMessageContext(Node);
 		ObjectNode->Components.Add(NodeComponentNew);
@@ -298,13 +314,9 @@ mu::Ptr<mu::NodeComponent> GenerateMutableSourceComponent(const UEdGraphPin* Pin
 
 	else if (const UCustomizableObjectNodeComponentMeshAddTo* TypedComponentMeshExtend = Cast<UCustomizableObjectNodeComponentMeshAddTo>(Node))
 	{
-		if (const FMutableComponentInfo* ComponentInfo = GenerationContext.ComponentInfos.FindByPredicate([&TypedComponentMeshExtend](const FMutableComponentInfo& Element)
-		{
-			return Element.ComponentName == TypedComponentMeshExtend->ParentComponentName &&
-				Element.NodeComponentMesh;
-		}))
+		if (UCustomizableObjectNodeComponentMesh** FindResult = GenerationContext.MeshComponents.Find(TypedComponentMeshExtend->ParentComponentName))
 		{
-			UCustomizableObjectNodeComponentMesh* TypedParentComponentMesh = ComponentInfo->NodeComponentMesh;
+			UCustomizableObjectNodeComponentMesh* TypedParentComponentMesh = *FindResult;
 
 			if (TypedComponentMeshExtend->NumLODs > TypedParentComponentMesh->NumLODs)
 			{
@@ -337,8 +349,6 @@ mu::Ptr<mu::NodeComponent> GenerateMutableSourceComponent(const UEdGraphPin* Pin
 	
 	else if (const UCustomizableObjectNodeComponentPassthroughMesh* TypeComponentPassthroughMesh = Cast<UCustomizableObjectNodeComponentPassthroughMesh>(Node))
 	{
-		++GenerationContext.NumPassthroughMeshComponents;
-
 		GenerationContext.CurrentMeshComponent = TypeComponentPassthroughMesh->ComponentName;
 
 		if (TypeComponentPassthroughMesh->ComponentName.IsNone())
@@ -379,7 +389,7 @@ mu::Ptr<mu::NodeComponent> GenerateMutableSourceComponent(const UEdGraphPin* Pin
 
 		// Create the component node
 		mu::Ptr<mu::NodeComponentNew> ComponentNode = new mu::NodeComponentNew;
-		ComponentNode->Id = GenerationContext.NumMeshComponentsInRoot - 1 + GenerationContext.NumPassthroughMeshComponents; // Last root component id + Num explicit components
+		ComponentNode->Id = GenerationContext.NumComponents++;
 
 		// Create a LOD for each pass-through mesh LOD.
 		const FSkeletalMeshModel* Model = SkeletalMesh->GetImportedModel();
diff --git a/Engine/Plugins/Experimental/Mutable/Source/CustomizableObjectEditor/Private/MuCOE/GenerateMutableSource/GenerateMutableSourceGroupProjector.cpp b/Engine/Plugins/Experimental/Mutable/Source/CustomizableObjectEditor/Private/MuCOE/GenerateMutableSource/GenerateMutableSourceGroupProjector.cpp
--- a/Engine/Plugins/Experimental/Mutable/Source/CustomizableObjectEditor/Private/MuCOE/GenerateMutableSource/GenerateMutableSourceGroupProjector.cpp
+++ b/Engine/Plugins/Experimental/Mutable/Source/CustomizableObjectEditor/Private/MuCOE/GenerateMutableSource/GenerateMutableSourceGroupProjector.cpp
@@ -413,9 +413,9 @@ bool GenerateMutableSourceGroupProjector(const UEdGraphPin* Pin, FMutableGraphGe
 				FString msg = FString::Printf(TEXT("The group projection node must have at least one option image connected to a texture or at least one valid element in Option Images Data Table."));
 				GenerationContext.Compiler->CompilerLog(FText::FromString(msg), ProjParamNode, EMessageSeverity::Error, true);
 				return false;
-			}			
-
-			if (GenerationContext.ComponentInfos.IsValidIndex(0))
+			}
+			
+			if (FMutableComponentInfo* CurrentComponentInfo = GenerationContext.GetCurrentComponentInfo())
 			{
 				// Poses will only affect component 0 of the CO,
 				// TODO UE-206803
@@ -441,7 +441,7 @@ bool GenerateMutableSourceGroupProjector(const UEdGraphPin* Pin, FMutableGraphGe
 
 					check(GroupProjectorTempData.PoseBoneDataArray.Num() == PoseIndex);
 					GroupProjectorTempData.PoseBoneDataArray.AddDefaulted(1);
-					UCustomizableObjectNodeAnimationPose::StaticRetrievePoseInformation(PoseAsset, GenerationContext.GetCurrentComponentInfo().RefSkeletalMesh,
+					UCustomizableObjectNodeAnimationPose::StaticRetrievePoseInformation(PoseAsset, CurrentComponentInfo->RefSkeletalMesh,
 						GroupProjectorTempData.PoseBoneDataArray[PoseIndex].ArrayBoneName, GroupProjectorTempData.PoseBoneDataArray[PoseIndex].ArrayTransform);
 				}
 
diff --git a/Engine/Plugins/Experimental/Mutable/Source/CustomizableObjectEditor/Private/MuCOE/GenerateMutableSource/GenerateMutableSourceMesh.cpp b/Engine/Plugins/Experimental/Mutable/Source/CustomizableObjectEditor/Private/MuCOE/GenerateMutableSource/GenerateMutableSourceMesh.cpp
--- a/Engine/Plugins/Experimental/Mutable/Source/CustomizableObjectEditor/Private/MuCOE/GenerateMutableSource/GenerateMutableSourceMesh.cpp
+++ b/Engine/Plugins/Experimental/Mutable/Source/CustomizableObjectEditor/Private/MuCOE/GenerateMutableSource/GenerateMutableSourceMesh.cpp
@@ -775,17 +775,22 @@ mu::MeshPtr ConvertSkeletalMeshToMutable(const USkeletalMesh* InSkeletalMesh, co
 
 		GenerationContext.AddParticipatingObject(*InSkeleton);
 
-		FMutableComponentInfo& MutComponentInfo = GenerationContext.GetCurrentComponentInfo();
-		USkeletalMesh* ComponentRefSkeletalMesh = MutComponentInfo.RefSkeletalMesh;
-		USkeleton* ComponentRefSkeleton = MutComponentInfo.RefSkeleton;
+		FMutableComponentInfo* MutComponentInfo = GenerationContext.GetCurrentComponentInfo();
+		if (!MutComponentInfo)
+		{
+			return {};
+		}
+
+		USkeletalMesh* ComponentRefSkeletalMesh = MutComponentInfo->RefSkeletalMesh;
+		USkeleton* ComponentRefSkeleton = MutComponentInfo->RefSkeleton;
 		check(ComponentRefSkeletalMesh);
 		check(ComponentRefSkeleton);
 
 		// Compatibility check
 		{
 			FString ErrorMessage;
-			const bool bCompatible = IsSkeletalMeshCompatibleWithRefSkeleton(MutComponentInfo, InSkeletalMesh, ErrorMessage);
-			MutComponentInfo.SkeletonCompatibility.Add(InSkeleton, bCompatible);
+			const bool bCompatible = IsSkeletalMeshCompatibleWithRefSkeleton(*MutComponentInfo, InSkeletalMesh, ErrorMessage);
+			MutComponentInfo->SkeletonCompatibility.Add(InSkeleton, bCompatible);
 
 			if (!bCompatible)
 			{
@@ -843,7 +848,7 @@ mu::MeshPtr ConvertSkeletalMeshToMutable(const USkeletalMesh* InSkeletalMesh, co
 
 		// Remove bones and build an array to remap indices of the BoneMap
 		TArray<FBoneIndexType> RemappedBones;
-		BuildRemappedBonesArray(MutComponentInfo, InSkeletalMesh, LODIndex, SourceRequiredBones, RemappedBones);
+		BuildRemappedBonesArray(*MutComponentInfo, InSkeletalMesh, LODIndex, SourceRequiredBones, RemappedBones);
 		
 		// Build RequiredBones array
 		TArray<FBoneIndexType> RequiredBones;
@@ -2855,7 +2860,7 @@ bool GetAndValidateReshapeBonesToDeform(
 	else if (SelectionMethod == EBoneDeformSelectionMethod::DEFORM_REF_SKELETON)
 	{
 		// Getting reference skeleton from the reference skeletal mesh of the current component
-		const FReferenceSkeleton RefSkeleton = GenerationContext.GetCurrentComponentInfo().RefSkeletalMesh->GetRefSkeleton(); //GenerationContext.ComponentInfos[GenerationContext.CurrentMeshComponent].RefSkeletalMesh->GetRefSkeleton();
+		const FReferenceSkeleton RefSkeleton = GenerationContext.GetCurrentComponentInfo()->RefSkeletalMesh->GetRefSkeleton(); //GenerationContext.ComponentInfos[GenerationContext.CurrentMeshComponent].RefSkeletalMesh->GetRefSkeleton();
 		int32 NumBones = RefSkeleton.GetRawBoneNum();
 
 		for (int32 BoneIndex = 0; BoneIndex < NumBones; ++BoneIndex)
@@ -2870,7 +2875,7 @@ bool GetAndValidateReshapeBonesToDeform(
 	else if (SelectionMethod == EBoneDeformSelectionMethod::DEFORM_NONE_REF_SKELETON)
 	{
 		// Getting reference skeleton from the reference skeletal mesh of the current component
-		const FReferenceSkeleton RefSkeleton = GenerationContext.GetCurrentComponentInfo().RefSkeletalMesh->GetRefSkeleton();
+		const FReferenceSkeleton RefSkeleton = GenerationContext.GetCurrentComponentInfo()->RefSkeletalMesh->GetRefSkeleton();
 
 		for (const MeshInfoType& Mesh : SkeletalMeshesInfo)
 		{
@@ -2968,7 +2973,7 @@ bool GetAndValidateReshapePhysicsToDeform(
 	if (bIsReferenceSkeletalMeshMethod)
 	{
 		const FReferenceSkeleton& RefSkeleton =
-			GenerationContext.GetCurrentComponentInfo().RefSkeletalMesh->GetRefSkeleton();
+			GenerationContext.GetCurrentComponentInfo()->RefSkeletalMesh->GetRefSkeleton();
 
 		const int32 RefSkeletonNumBones = RefSkeleton.GetRawBoneNum();
 			BoneNames.SetNum(RefSkeletonNumBones);
@@ -2999,10 +3004,10 @@ bool GetAndValidateReshapePhysicsToDeform(
 
 	for (const PhysicsInfoType& PhysicsInfo : ContributingPhysicsAssetsInfo)
 	{
-		check(GenerationContext.GetCurrentComponentInfo().RefSkeletalMesh);
+		check(GenerationContext.GetCurrentComponentInfo()->RefSkeletalMesh);
 
 		const FReferenceSkeleton& RefSkeleton = bIsReferenceSkeletalMeshMethod
-			? GenerationContext.GetCurrentComponentInfo().RefSkeletalMesh->GetRefSkeleton()
+			? GenerationContext.GetCurrentComponentInfo()->RefSkeletalMesh->GetRefSkeleton()
 			: PhysicsInfo.Get<const FReferenceSkeleton&>();
 
 		UPhysicsAsset* PhysicsAsset = PhysicsInfo.Get<UPhysicsAsset*>();
@@ -4118,13 +4123,13 @@ mu::NodeMeshPtr GenerateMutableSourceMesh(const UEdGraphPin* Pin,
 		{
 			mu::Ptr<mu::NodeMesh> InputMeshNode = GenerateMutableSourceMesh(InputMeshPin, GenerationContext, MeshData, false, bOnlyConnectedLOD);
 
-			if (GenerationContext.GetCurrentComponentInfo().RefSkeletalMesh)
+			if (GenerationContext.GetCurrentComponentInfo()->RefSkeletalMesh)
 			{
 				if (TypedNode->PoseAsset)
 				{
 					TArray<FName> ArrayBoneName;
 					TArray<FTransform> ArrayTransform;
-					UCustomizableObjectNodeAnimationPose::StaticRetrievePoseInformation(TypedNode->PoseAsset, GenerationContext.GetCurrentComponentInfo().RefSkeletalMesh, ArrayBoneName, ArrayTransform);
+					UCustomizableObjectNodeAnimationPose::StaticRetrievePoseInformation(TypedNode->PoseAsset, GenerationContext.GetCurrentComponentInfo()->RefSkeletalMesh, ArrayBoneName, ArrayTransform);
 					mu::NodeMeshApplyPosePtr NodeMeshApplyPose = CreateNodeMeshApplyPose(GenerationContext, InputMeshNode, ArrayBoneName, ArrayTransform);
 
 					if (NodeMeshApplyPose)
diff --git a/Engine/Plugins/Experimental/Mutable/Source/CustomizableObjectEditor/Private/MuCOE/GenerateMutableSource/GenerateMutableSourceTable.cpp b/Engine/Plugins/Experimental/Mutable/Source/CustomizableObjectEditor/Private/MuCOE/GenerateMutableSource/GenerateMutableSourceTable.cpp
--- a/Engine/Plugins/Experimental/Mutable/Source/CustomizableObjectEditor/Private/MuCOE/GenerateMutableSource/GenerateMutableSourceTable.cpp
+++ b/Engine/Plugins/Experimental/Mutable/Source/CustomizableObjectEditor/Private/MuCOE/GenerateMutableSource/GenerateMutableSourceTable.cpp
@@ -589,7 +589,7 @@ bool FillTableColumn(const UCustomizableObjectNodeTable* TableNode, mu::TablePtr
 
 				TArray<FName> ArrayBoneName;
 				TArray<FTransform> ArrayTransform;
-				UCustomizableObjectNodeAnimationPose::StaticRetrievePoseInformation(PoseAsset, GenerationContext.GetCurrentComponentInfo().RefSkeletalMesh, ArrayBoneName, ArrayTransform);
+				UCustomizableObjectNodeAnimationPose::StaticRetrievePoseInformation(PoseAsset, GenerationContext.GetCurrentComponentInfo()->RefSkeletalMesh, ArrayBoneName, ArrayTransform);
 
 				mu::Ptr<mu::Mesh> MutableMesh = new mu::Mesh();
 				mu::Ptr<mu::Skeleton> MutableSkeleton = new mu::Skeleton;
diff --git a/Engine/Plugins/Experimental/Mutable/Source/CustomizableObjectEditor/Private/MuCOE/Nodes/CustomizableObjectNodeComponentMesh.cpp b/Engine/Plugins/Experimental/Mutable/Source/CustomizableObjectEditor/Private/MuCOE/Nodes/CustomizableObjectNodeComponentMesh.cpp
--- a/Engine/Plugins/Experimental/Mutable/Source/CustomizableObjectEditor/Private/MuCOE/Nodes/CustomizableObjectNodeComponentMesh.cpp
+++ b/Engine/Plugins/Experimental/Mutable/Source/CustomizableObjectEditor/Private/MuCOE/Nodes/CustomizableObjectNodeComponentMesh.cpp
@@ -2,7 +2,9 @@
 
 #include "MuCOE/Nodes/CustomizableObjectNodeComponentMesh.h"
 
+#include "MuCO/CustomizableObjectCustomVersion.h"
 #include "MuCOE/EdGraphSchema_CustomizableObject.h"
+#include "MuCOE/GraphTraversal.h"
 
 
 #define LOCTEXT_NAMESPACE "CustomizableObjectEditor"
@@ -27,6 +29,25 @@ void UCustomizableObjectNodeComponentMesh::PostEditChangeProperty(FPropertyChang
 }
 
 
+void UCustomizableObjectNodeComponentMesh::BackwardsCompatibleFixup(int32 CustomizableObjectCustomVersion)
+{
+	Super::BackwardsCompatibleFixup(CustomizableObjectCustomVersion);
+	
+	if (CustomizableObjectCustomVersion == FCustomizableObjectCustomVersion::ComponentsArray)
+	{
+		UCustomizableObject* Object = GetRootObject(*this);
+
+		if (FMutableMeshComponentData* Result = Object->GetPrivate()->MutableMeshComponents_DEPRECATED.FindByPredicate([&](const FMutableMeshComponentData& ComponentData)
+		{
+			return ComponentData.Name == ComponentName;
+		}))
+		{
+			ReferenceSkeletalMesh = Result->ReferenceSkeletalMesh;
+		}
+	}
+}
+
+
 FText UCustomizableObjectNodeComponentMesh::GetNodeTitle(ENodeTitleType::Type TitleType) const
 {
 	if (TitleType == ENodeTitleType::ListView)
diff --git a/Engine/Plugins/Experimental/Mutable/Source/CustomizableObjectEditor/Private/MuCOE/Nodes/CustomizableObjectNodeObject.cpp b/Engine/Plugins/Experimental/Mutable/Source/CustomizableObjectEditor/Private/MuCOE/Nodes/CustomizableObjectNodeObject.cpp
--- a/Engine/Plugins/Experimental/Mutable/Source/CustomizableObjectEditor/Private/MuCOE/Nodes/CustomizableObjectNodeObject.cpp
+++ b/Engine/Plugins/Experimental/Mutable/Source/CustomizableObjectEditor/Private/MuCOE/Nodes/CustomizableObjectNodeObject.cpp
@@ -118,7 +118,7 @@ void UCustomizableObjectNodeObject::BackwardsCompatibleFixup(int32 CustomizableO
 
 			if (UCustomizableObject* CurrentObject = Cast<UCustomizableObject>(GetOutermostObject()))
 			{
-				ComponentSettings_DEPRECATED.Init(ComponentSettingsTemplate, CurrentObject->GetPrivate()->MutableMeshComponents.Num());
+				ComponentSettings_DEPRECATED.Init(ComponentSettingsTemplate, CurrentObject->GetPrivate()->MutableMeshComponents_DEPRECATED.Num());
 			}
 		}
 	}
@@ -254,7 +254,7 @@ void UCustomizableObjectNodeObject::BackwardsCompatibleFixup(int32 CustomizableO
 				NodePosX += 400; // Move it a bit to make space for the new component nodes.
 			}
 
-			for (const FMutableMeshComponentData& MeshComponent : GetRootObject(*this)->GetPrivate()->MutableMeshComponents)
+			for (const FMutableMeshComponentData& MeshComponent : GetRootObject(*this)->GetPrivate()->MutableMeshComponents_DEPRECATED)
 			{
 				UCustomizableObjectNodeComponentMesh* NewNode = NewObject<UCustomizableObjectNodeComponentMesh>(this);
 				UEdGraphNode* Node = CreateNewNode(GetGraph(), ComponentsPin(), FVector2D(NodePosX - 300.0,  NodePosY + 200.0 * NodesCreated), NewNode);
diff --git a/Engine/Plugins/Experimental/Mutable/Source/CustomizableObjectEditor/Private/MuCOE/StatePerformanceTesting.cpp b/Engine/Plugins/Experimental/Mutable/Source/CustomizableObjectEditor/Private/MuCOE/StatePerformanceTesting.cpp
--- a/Engine/Plugins/Experimental/Mutable/Source/CustomizableObjectEditor/Private/MuCOE/StatePerformanceTesting.cpp
+++ b/Engine/Plugins/Experimental/Mutable/Source/CustomizableObjectEditor/Private/MuCOE/StatePerformanceTesting.cpp
@@ -88,7 +88,7 @@ void FTestingCustomizableObject::StartTest(UCustomizableObject* CustomizableObje
 
 	Instance = CustomizableObject->CreateInstance();
 
-	const int32 NumComponents = CustomizableObject->GetPrivate()->MutableMeshComponents.Num();
+	const int32 NumComponents = CustomizableObject->GetComponentCount();
 	SkeletalMeshComponents.Reset(NumComponents);
 	CustomizableSkeletalComponents.Reset(NumComponents);
 
diff --git a/Engine/Plugins/Experimental/Mutable/Source/CustomizableObjectEditor/Private/MuCOE/Widgets/CustomizableObjectLODReductionSettings.cpp b/Engine/Plugins/Experimental/Mutable/Source/CustomizableObjectEditor/Private/MuCOE/Widgets/CustomizableObjectLODReductionSettings.cpp
--- a/Engine/Plugins/Experimental/Mutable/Source/CustomizableObjectEditor/Private/MuCOE/Widgets/CustomizableObjectLODReductionSettings.cpp
+++ b/Engine/Plugins/Experimental/Mutable/Source/CustomizableObjectEditor/Private/MuCOE/Widgets/CustomizableObjectLODReductionSettings.cpp
@@ -141,27 +141,21 @@ FName FCustomizableObjectLODReductionSettings::GetSelectedBone(bool& bMultipleVa
 }
 
 
-const struct FReferenceSkeleton& FCustomizableObjectLODReductionSettings::GetReferenceSkeleton() const
+const FReferenceSkeleton& FCustomizableObjectLODReductionSettings::GetReferenceSkeleton() const
 {
 	static FReferenceSkeleton DummySkeleton;
-	USkeletalMesh* SkeletalMesh = nullptr;
 
-	if (NodeComponentMesh)
+	if (!NodeComponentMesh)
 	{
-		if (const UCustomizableObject* CustomizableObject = Cast<UCustomizableObject>(NodeComponentMesh->GetCustomizableObjectGraph()->GetOuter()))
-		{
-			if (const FMutableMeshComponentData* Component = CustomizableObject->GetPrivate()->MutableMeshComponents.FindByPredicate(
-				[&](const FMutableMeshComponentData& Component) 
-				{ 
-					return Component.Name == NodeComponentMesh->ComponentName; 
-				}))
-			{
-				SkeletalMesh = Component->ReferenceSkeletalMesh;
-			}
-		}
+		return DummySkeleton;
+	}
+
+	if (!NodeComponentMesh->ReferenceSkeletalMesh)
+	{
+		return DummySkeleton;
 	}
 
-	return SkeletalMesh ? SkeletalMesh->GetRefSkeleton() : DummySkeleton;
+	return NodeComponentMesh->ReferenceSkeletalMesh->GetRefSkeleton();
 }
 
 
diff --git a/Engine/Plugins/Experimental/Mutable/Source/CustomizableObjectEditor/Public/MuCOE/CustomizableObjectFactory.h b/Engine/Plugins/Experimental/Mutable/Source/CustomizableObjectEditor/Public/MuCOE/CustomizableObjectFactory.h
--- a/Engine/Plugins/Experimental/Mutable/Source/CustomizableObjectEditor/Public/MuCOE/CustomizableObjectFactory.h
+++ b/Engine/Plugins/Experimental/Mutable/Source/CustomizableObjectEditor/Public/MuCOE/CustomizableObjectFactory.h
@@ -37,18 +37,6 @@ struct FCustomizableObjectOptions
 
 	// Name of the group node of the parent where this CO asset will be linked
 	FString GroupNodeName;
-
-	// Non-Child Object Settings
-	// Number of components
-	int32 NumMeshComponents = 1;
-
-	struct FComponentInfo 
-	{
-		TSoftObjectPtr<USkeletalMesh> ReferenceSkeletalMesh;
-		FName ComponentName;
-	};
-
-	TArray<FComponentInfo> ComponentsInfo;
 };
 
 
@@ -105,19 +93,7 @@ public:
 
 	// Component Selectior Widget methods
 	EVisibility GetComponentWidgetsVisibility() const;
-	FText GetSelectorWidgetText(bool bIsName) const;
-	FString GetSelectedComponentSkeletalMeshPath() const;
-	int32 GetNumComponents() const;
-	void OnNumComponentsChanged(int32 Value, ETextCommit::Type CommitInfo);
 	void OnSelectComponentComboBox(TSharedPtr<FString> Selection, ESelectInfo::Type SelectInfo);
-	void OnPickedComponentSkeletalMesh(const FAssetData& SelectedAsset);
-	void GenerateComponentOptions();
-	bool IsComponentSelectorEnabled() const;
-
-	FText GetComponentName() const;
-	void OnTextCommited(const FText& NewName, ETextCommit::Type CommitInfo);
-
-	int32 GetSelectedComponentIndex() const;
 
 private:
 
diff --git a/Engine/Plugins/Experimental/Mutable/Source/CustomizableObjectEditor/Public/MuCOE/Nodes/CustomizableObjectNodeComponentMesh.h b/Engine/Plugins/Experimental/Mutable/Source/CustomizableObjectEditor/Public/MuCOE/Nodes/CustomizableObjectNodeComponentMesh.h
--- a/Engine/Plugins/Experimental/Mutable/Source/CustomizableObjectEditor/Public/MuCOE/Nodes/CustomizableObjectNodeComponentMesh.h
+++ b/Engine/Plugins/Experimental/Mutable/Source/CustomizableObjectEditor/Public/MuCOE/Nodes/CustomizableObjectNodeComponentMesh.h
@@ -52,6 +52,7 @@ class CUSTOMIZABLEOBJECTEDITOR_API UCustomizableObjectNodeComponentMesh : public
 public:
 	// UObject interface
 	virtual void PostEditChangeProperty(FPropertyChangedEvent& PropertyChangedEvent) override;
+	virtual void BackwardsCompatibleFixup(int32 CustomizableObjectCustomVersion) override;
 
 	// UEdGraphNode interface
 	virtual FText GetNodeTitle(ENodeTitleType::Type TitleType) const override;
@@ -59,6 +60,9 @@ public:
 	UPROPERTY(EditAnywhere, Category = ComponentMesh)
 	FName ComponentName;
 	
+	UPROPERTY(EditAnywhere, Category = ComponentMesh)
+	TObjectPtr<USkeletalMesh> ReferenceSkeletalMesh;
+
 	UPROPERTY(EditAnywhere, Category = ComponentMesh)
 	TArray<FLODReductionSettings> LODReductionSettings;
 
