[Mutable] Now, UDGGUI uses a config file to search the path of the DGGUI class.

#rnx

diff --git a/Engine/Plugins/Experimental/Mutable/Config/DefaultMutable.ini b/Engine/Plugins/Experimental/Mutable/Config/DefaultMutable.ini
--- a/Engine/Plugins/Experimental/Mutable/Config/DefaultMutable.ini
+++ b/Engine/Plugins/Experimental/Mutable/Config/DefaultMutable.ini
@@ -16,6 +16,7 @@ bExtraBoneInfluencesEnabled=False
 [EditorDefaults]
 ClipMorphMaterialName=/Engine/EditorMaterials/LevelGridMaterial.LevelGridMaterial
 ClipMeshMaterialName=/Mutable/Materials/ClipMeshMaterial.ClipMeshMaterial
+DynamicallyGenerated_DGGUI_Path=/Game/UI/DynamicallyGeneratedGUI_DGGUI/DynamicallyGeneratedGUI_DGGUI.DynamicallyGeneratedGUI_DGGUI_C
 
 ;--------------------------------------------------------------------------------------------------
 ; These are scalability settings that can be defined in any xxxScalability.ini file.
diff --git a/Engine/Plugins/Experimental/Mutable/Source/CustomizableObject/Private/MuCO/CustomizableObjectDGGUI.cpp b/Engine/Plugins/Experimental/Mutable/Source/CustomizableObject/Private/MuCO/CustomizableObjectDGGUI.cpp
--- a/Engine/Plugins/Experimental/Mutable/Source/CustomizableObject/Private/MuCO/CustomizableObjectDGGUI.cpp
+++ b/Engine/Plugins/Experimental/Mutable/Source/CustomizableObject/Private/MuCO/CustomizableObjectDGGUI.cpp
@@ -3,6 +3,7 @@
 #include "MuCO/CustomizableObjectDGGUI.h"
 
 #include "Kismet/GameplayStatics.h"
+#include "Misc/ConfigCacheIni.h"
 #include "UObject/UObjectIterator.h"
 
 #include UE_INLINE_GENERATED_CPP_BY_NAME(CustomizableObjectDGGUI)
@@ -12,7 +13,18 @@ void UDGGUI::OpenDGGUI(const int32 SlotID, UCustomizableObjectInstanceUsage* Sel
 #if !UE_BUILD_SHIPPING
 	if (APlayerController* Player = UGameplayStatics::GetPlayerController(CurrentWorld, PlayerIndex))
 	{
-		FSoftClassPath DGUIPath(TEXT("/Mutable/UI/DynamicallyGeneratedGUI_DGGUI/DynamicallyGeneratedGUI_DGGUI.DynamicallyGeneratedGUI_DGGUI_C"));
+		static FString DGGUIAssetPath;
+
+		if (DGGUIAssetPath.IsEmpty())
+		{
+			FConfigFile* PluginConfig = GConfig->FindConfigFileWithBaseName("Mutable");
+			if (PluginConfig)
+			{
+				PluginConfig->GetString(TEXT("EditorDefaults"), TEXT("DynamicallyGenerated_DGGUI_Path"), DGGUIAssetPath);
+			}
+		}
+
+		FSoftClassPath DGUIPath(DGGUIAssetPath);
 		if (UClass* DGUI = DGUIPath.TryLoadClass<UDGGUI>())
 		{
 			UDGGUI* WDGUI = CreateWidget<UDGGUI>(Player, DGUI);
@@ -23,6 +35,10 @@ void UDGGUI::OpenDGGUI(const int32 SlotID, UCustomizableObjectInstanceUsage* Sel
 				Player->SetShowMouseCursor(true);
 			}
 		}
+		else
+		{
+			UE_LOG(LogTemp, Warning, TEXT("Could not find the DynamicallyGenerated_DGGUI class inside the specified path. Check the DefaultMutable.ini file."));
+		}
 	}
 #endif // !UE_BUILD_SHIPPING
 }
