[mutable] Moved the Mutable plugin out of Experimental status into Beta.

#tests Editor
#rnx

#virtualized

diff --git a/Engine/Plugins/Mutable/Config/BaseMutable.ini b/Engine/Plugins/Mutable/Config/BaseMutable.ini
new file mode 100644
--- /dev/null
+++ b/Engine/Plugins/Mutable/Config/BaseMutable.ini
@@ -0,0 +1,20 @@
+[CoreRedirects]
+
++StructRedirects=(OldName="/Script/CustomizableObject.ParameterUIData",NewName="/Script/CustomizableObject.MutableParameterData")
++StructRedirects=(OldName="/Script/CustomizableObject.CompilationOptions",NewName="/Script/CustomizableObject.CompilationOptions_DEPRECATED")
+
++PropertyRedirects=(OldName="GroupProjectorParameterImage.OptionImage", NewName="OptionTexture")
++PropertyRedirects=(OldName="CustomizableObjectNodeGroupProjectorParameter.OptionImages", NewName="OptionTextures")
++PropertyRedirects=(OldName="CustomizableObjectNodeGroupProjectorParameter.OptionImagesDataTable", NewName="OptionTexturesDataTable")
+
+; Node renames
++ClassRedirects=(OldName="/Script/CustomizableObjectEditor.CustomizableObjectNodeMeshClipDeform", NewName="/Script/CustomizableObjectEditor.CustomizableObjectNodeModifierClipDeform")
++ClassRedirects=(OldName="/Script/CustomizableObjectEditor.CustomizableObjectNodeMeshClipWithMesh", NewName="/Script/CustomizableObjectEditor.CustomizableObjectNodeModifierClipWithMesh")
++ClassRedirects=(OldName="/Script/CustomizableObjectEditor.CustomizableObjectNodeMeshClipMorph", NewName="/Script/CustomizableObjectEditor.CustomizableObjectNodeModifierClipMorph")
++ClassRedirects=(OldName="/Script/CustomizableObjectEditor.CustomizableObjectNodeExtendMaterial", NewName="/Script/CustomizableObjectEditor.CustomizableObjectNodeModifierExtendMeshSection")
++ClassRedirects=(OldName="/Script/CustomizableObjectEditor.CustomizableObjectNodeEditMaterial", NewName="/Script/CustomizableObjectEditor.CustomizableObjectNodeModifierEditMeshSection")
++ClassRedirects=(OldName="/Script/CustomizableObjectEditor.CustomizableObjectNodeMorphMaterial", NewName="/Script/CustomizableObjectEditor.CustomizableObjectNodeModifierMorphMeshSection")
++ClassRedirects=(OldName="/Script/CustomizableObjectEditor.CustomizableObjectNodeRemoveMesh", NewName="/Script/CustomizableObjectEditor.CustomizableObjectNodeModifierRemoveMesh")
++ClassRedirects=(OldName="/Script/CustomizableObjectEditor.CustomizableObjectNodeRemoveMeshBlocks", NewName="/Script/CustomizableObjectEditor.CustomizableObjectNodeModifierRemoveMeshBlocks")
++ClassRedirects=(OldName="/Script/CustomizableObjectEditor.CustomizableObjectNodeEditMaterialBase", NewName="/Script/CustomizableObjectEditor.CustomizableObjectNodeModifierEditMeshSectionBase")
++ClassRedirects=(OldName="/Script/CustomizableObjectEditor.CustomizableObjectNodeEditLayoutBlocks", NewName="/Script/CustomizableObjectEditor.CustomizableObjectNodeModifierEditLayoutBlocks")
diff --git a/Engine/Plugins/Mutable/Config/DefaultMutable.ini b/Engine/Plugins/Mutable/Config/DefaultMutable.ini
new file mode 100644
--- /dev/null
+++ b/Engine/Plugins/Mutable/Config/DefaultMutable.ini
@@ -0,0 +1,54 @@
+; Settings for the Mutable plugin
+[Features]
+
+; If CustomizableObjectNumBoneInfluences is set to N, CustomizableObjects will be allowed to use a max of N bone influences 
+; per vertex. The only valid values are 4, 8, 12, Four, Eight, Twelve. 
+; If r.GPUSkin.UnlimitedBoneInfluencesThreshold is set to a lower number, then only that number of influences
+; will always be used and allocated thus saving memory per vertex when the real number is higher for a vertex.
+; If CustomizableObjectNumBoneInfluences is not defined, only the default 4 will be used. This is used at object compile time.
+;CustomizableObjectNumBoneInfluences=Eight
+
+; For backwards compatibility bExtraBoneInfluencesEnabled can be used instead, but will be overridden by 
+; CustomizableObjectNumBoneInfluences if both are present. If enabled CustomizableObjects will be allowed to 
+; use 8 bone influences per vertex instead of the default 4. This is used at object compile time.
+bExtraBoneInfluencesEnabled=False
+
+[EditorDefaults]
+ClipMorphMaterialName=/Engine/EditorMaterials/LevelGridMaterial.LevelGridMaterial
+ClipMeshMaterialName=/Mutable/Materials/ClipMeshMaterial.ClipMeshMaterial
+DynamicallyGenerated_DGGUI_Path=/Game/UI/DynamicallyGeneratedGUI_DGGUI/DynamicallyGeneratedGUI_DGGUI.DynamicallyGeneratedGUI_DGGUI_C
+
+;--------------------------------------------------------------------------------------------------
+; These are scalability settings that can be defined in any xxxScalability.ini file.
+[ScalabilitySettings]
+
+; If different than 0, limit the amount of memory (in KB) to use to cache streaming data when 
+; building characters. 0 means no limit, -1 means use default (10Mb in PC and 3Mb in consoles).
++CVars=b.MutableStreamingMemory=10000
+
+; Override in platform scalability .ini files
+; b.MutableStreamingMemory=3000
+
+[ViewDistanceQuality@0]
+; If different than 0, limit the number of mutable instances with full LODs to have at a given time.
++CVars=b.NumGeneratedInstancesLimit=7
+
+; If different than 0, limit the number of mutable instances with LOD 1 to have at a given time.
++CVars=b.NumGeneratedInstancesLimitLOD1=0
+
+; If different than 0, limit the number of mutable instances with LOD 2 to have at a given time.
++CVars=b.NumGeneratedInstancesLimitLOD2=0
+
+; If NumGeneratedInstancesLimit is different than 0, sets the distance at which the system will fix 
+; the LOD of an instance to the lowest res one (LOD2) to prevent unnecessary LOD changes and memory 
+; consumption
++CVars=b.DistanceForFixedLOD2=0
+
+[ViewDistanceQuality@1]
++CVars=b.NumGeneratedInstancesLimit=0
+[ViewDistanceQuality@2]
++CVars=b.NumGeneratedInstancesLimit=0
+[ViewDistanceQuality@3]
++CVars=b.NumGeneratedInstancesLimit=0
+[ViewDistanceQuality@Cine]
++CVars=b.NumGeneratedInstancesLimit=0
diff --git a/Engine/Plugins/Mutable/Source/CustomizableObject/Private/MuCO/CustomizableInstanceLODManagement.cpp b/Engine/Plugins/Mutable/Source/CustomizableObject/Private/MuCO/CustomizableInstanceLODManagement.cpp
new file mode 100644
--- /dev/null
+++ b/Engine/Plugins/Mutable/Source/CustomizableObject/Private/MuCO/CustomizableInstanceLODManagement.cpp
@@ -0,0 +1,469 @@
+// Copyright Epic Games, Inc. All Rights Reserved.
+
+#include "MuCO/CustomizableInstanceLODManagement.h"
+
+#include "Camera/CameraActor.h"
+#include "Engine/SkeletalMesh.h"
+#include "Engine/SkinnedAssetCommon.h"
+#include "GameFramework/Pawn.h"
+#include "GameFramework/PlayerController.h"
+#include "Kismet/GameplayStatics.h"
+#include "MuCO/CustomizableObject.h"
+#include "MuCO/CustomizableObjectInstance.h"
+#include "MuCO/CustomizableObjectSystem.h"
+#include "MuCO/UnrealPortabilityHelpers.h"
+#include "MuCO/CustomizableObjectInstanceUsage.h"
+#include "UObject/UObjectIterator.h"
+#include "Components/SkeletalMeshComponent.h"
+#include "MuCO/CustomizableObjectInstancePrivate.h"
+#include "MuCO/CustomizableObjectPrivate.h"
+
+#include UE_INLINE_GENERATED_CPP_BY_NAME(CustomizableInstanceLODManagement)
+
+#if WITH_EDITOR
+#include "LevelEditorViewport.h"
+#else
+#include "Camera/CameraComponent.h"
+#include "SceneManagement.h"
+#endif
+
+static TAutoConsoleVariable<int32> CVarNumGeneratedInstancesLimit(
+	TEXT("b.NumGeneratedInstancesLimit"),
+	0,
+	TEXT("If different than 0, limit the number of mutable instances with full LODs to have at a given time."),
+	ECVF_Scalability);
+
+static TAutoConsoleVariable<int32> CVarNumGeneratedInstancesLimitLOD1(
+	TEXT("b.NumGeneratedInstancesLimitLOD1"),
+	0,
+	TEXT("If different than 0, limit the number of mutable instances with LOD 1 to have at a given time."),
+	ECVF_Scalability);
+
+static TAutoConsoleVariable<int32> CVarNumGeneratedInstancesLimitLOD2(
+	TEXT("b.NumGeneratedInstancesLimitLOD2"),
+	0,
+	TEXT("If different than 0, limit the number of mutable instances with LOD 2 to have at a given time."),
+	ECVF_Scalability);
+
+static TAutoConsoleVariable<int32> CVarDistanceForFixedLOD2(
+	TEXT("b.DistanceForFixedLOD2"),
+	50000,
+	TEXT("If NumGeneratedInstancesLimit is different than 0, sets the distance at which the system will fix the LOD of an instance to the lowest res one (LOD2) to prevent unnecessary LOD changes and memory consumption"),
+	ECVF_Scalability);
+
+static TAutoConsoleVariable<bool> CVarOnlyUpdateCloseCustomizableObjects(
+	TEXT("b.OnlyUpdateCloseCustomizableObjects"),
+	false,
+	TEXT("If true, only CustomizableObjects within a predefined distance to the view centers will be generated"),
+	ECVF_Scalability);
+
+UCustomizableInstanceLODManagement::UCustomizableInstanceLODManagement() : UCustomizableInstanceLODManagementBase()
+{
+	CloseCustomizableObjectsDist = 2000.f;
+}
+
+
+UCustomizableInstanceLODManagement::~UCustomizableInstanceLODManagement()
+{
+
+}
+
+// Used to manually update distances used in "OnlyUpdateCloseCustomizableObjects" system. If OnlyForInstance is null, all instances have their distance updated
+// ViewCenter is the origin where the distances will be measured from.
+void UpdatePawnToInstancesDistances(const class UCustomizableObjectInstance* OnlyForInstance, const TWeakObjectPtr<const AActor> ViewCenter)
+{
+	for (TObjectIterator<UCustomizableObjectInstanceUsage> CustomizableObjectInstanceUsage; CustomizableObjectInstanceUsage; ++CustomizableObjectInstanceUsage)
+	{
+#if WITH_EDITOR
+		if (IsValid(*CustomizableObjectInstanceUsage) && CustomizableObjectInstanceUsage->IsNetMode(NM_DedicatedServer))
+		{
+			continue;
+		}
+#endif
+
+		if (IsValid(*CustomizableObjectInstanceUsage) && (OnlyForInstance == nullptr || CustomizableObjectInstanceUsage->GetCustomizableObjectInstance() == OnlyForInstance))
+		{
+			CustomizableObjectInstanceUsage->UpdateDistFromComponentToPlayer(ViewCenter.IsValid() ? ViewCenter.Get() : nullptr, OnlyForInstance != nullptr);
+		}
+	}
+}
+
+#if WITH_EDITOR
+// Used to manually update instances in the level editor
+void UpdateCameraToInstancesDistance(const FVector CameraPosition)
+{
+	for (TObjectIterator<UCustomizableObjectInstanceUsage> CustomizableObjectInstanceUsage; CustomizableObjectInstanceUsage; ++CustomizableObjectInstanceUsage)
+	{
+#if WITH_EDITOR
+		if (IsValid(*CustomizableObjectInstanceUsage) && CustomizableObjectInstanceUsage->IsNetMode(NM_DedicatedServer))
+		{
+			continue;
+		}
+#endif
+		if (IsValid(*CustomizableObjectInstanceUsage) && !CustomizableObjectInstanceUsage->IsTemplate())
+		{
+			CustomizableObjectInstanceUsage->UpdateDistFromComponentToLevelEditorCamera(CameraPosition);
+		}
+	}
+}
+#endif
+
+
+void UCustomizableInstanceLODManagement::UpdateInstanceDistsAndLODs(FMutableInstanceUpdateMap& InOutRequestedUpdates)
+{
+	int32 NumGeneratedInstancesLimitLODs = GetNumGeneratedInstancesLimitFullLODs();
+	if (NumGeneratedInstancesLimitLODs > 0 || (IsOnlyUpdateCloseCustomizableObjectsEnabled() && IsOnlyGenerateRequestedLODLevelsEnabled()))
+	{
+		UCustomizableObjectSystem* CustomizableObjectSystem = UCustomizableObjectSystem::GetInstance();
+
+		if (ViewCenters.Num() == 0) // Just use the first pawn
+		{
+			UCustomizableObjectInstanceUsage* FirstCustomizableObjectInstanceUsage = nullptr;
+
+	#if WITH_EDITOR
+			bool bLevelEditorInstancesUpdated = false;
+	#endif
+
+			for (TObjectIterator<UCustomizableObjectInstanceUsage> CustomizableObjectInstanceUsage; CustomizableObjectInstanceUsage; ++CustomizableObjectInstanceUsage)
+			{
+#if WITH_EDITOR
+				if (IsValid(*CustomizableObjectInstanceUsage) && CustomizableObjectInstanceUsage->IsNetMode(NM_DedicatedServer))
+				{
+					continue;
+				}
+#endif
+				if (!IsValid(*CustomizableObjectInstanceUsage))
+				{
+					continue;
+				}
+
+				USkeletalMeshComponent* Parent = Cast<USkeletalMeshComponent>(CustomizableObjectInstanceUsage->GetAttachParent());
+
+				if (!CustomizableObjectInstanceUsage->IsTemplate())
+				{
+					UWorld* LocalWorld = Parent ? Parent->GetWorld() : nullptr;
+					APlayerController* Controller = LocalWorld ? LocalWorld->GetFirstPlayerController() : nullptr;
+					TWeakObjectPtr<const AActor> ViewCenter = Controller ? TWeakObjectPtr<const AActor>(Controller->GetPawn()) : nullptr;
+
+					if (ViewCenter.IsValid())
+					{
+						UpdatePawnToInstancesDistances(nullptr, ViewCenter);
+						break;
+					}
+
+	#if WITH_EDITOR
+					else if (Parent && Parent->GetWorld())
+					{
+						EWorldType::Type WorldType = Parent->GetWorld()->WorldType;
+					
+						// Level Editor Instances (non PIE)
+						if (!bLevelEditorInstancesUpdated && WorldType == EWorldType::Editor)
+						{
+							for (FLevelEditorViewportClient* LevelVC : GEditor->GetLevelViewportClients())
+							{
+								if (LevelVC && LevelVC->IsPerspective())
+								{
+									UpdateCameraToInstancesDistance(LevelVC->GetViewLocation());
+									bLevelEditorInstancesUpdated = true;
+									break;
+								}
+							}
+						}
+
+						// Blueprint instances
+						else if (WorldType == EWorldType::EditorPreview)
+						{
+							CustomizableObjectInstanceUsage->EditorUpdateComponent();
+						}
+					}
+	#endif // WITH_EDITOR
+
+				}
+			}
+		}
+		else
+		{
+			for (const TWeakObjectPtr<const AActor> ViewCenter : ViewCenters)
+			{
+				if (ViewCenter.IsValid())
+				{
+					UpdatePawnToInstancesDistances(nullptr, ViewCenter);
+				}
+			}
+		}
+	}
+
+	if (NumGeneratedInstancesLimitLODs > 0)
+	{
+		int32 NumGeneratedInstancesLimitLOD1 = GetNumGeneratedInstancesLimitLOD1();
+		int32 NumGeneratedInstancesLimitLOD2 = GetNumGeneratedInstancesLimitLOD2();
+
+		TArray<UCustomizableObjectInstance*> SortedInstances;
+
+		for (TObjectIterator<UCustomizableObjectInstance> CustomizableObjectInstance; CustomizableObjectInstance; ++CustomizableObjectInstance)
+		{
+			if (IsValid(*CustomizableObjectInstance) &&
+				CustomizableObjectInstance->GetPrivate() &&
+				CustomizableObjectInstance->GetIsBeingUsedByComponentInPlay())
+			{
+				if (const UCustomizableObject* CustomizableObject = CustomizableObjectInstance->GetCustomizableObject();
+					CustomizableObject &&
+					!CustomizableObject->GetPrivate()->IsLocked())
+				{
+					UCustomizableObjectInstance* Ptr = *CustomizableObjectInstance;
+					Ptr->SetIsDiscardedBecauseOfTooManyInstances(false);
+					SortedInstances.Add(Ptr);
+				}
+			}
+		}
+
+		for (int32 i = 0; i < GetNumberOfPriorityUpdateInstances() && i < SortedInstances.Num(); ++i)
+		{
+			SortedInstances[i]->SetIsPlayerOrNearIt(true);
+		}
+
+		TArray<uint16> RequestedLODs;
+
+		if (SortedInstances.Num() > NumGeneratedInstancesLimitLODs)
+		{
+			SortedInstances.Sort([](UCustomizableObjectInstance& A, UCustomizableObjectInstance& B)
+			{
+				return A.GetMinSquareDistToPlayer() < B.GetMinSquareDistToPlayer();
+			});
+
+			bool bAlreadyReachedFixedLOD = false;
+			const float DistanceForFixedLOD = float(CVarDistanceForFixedLOD2.GetValueOnGameThread());
+			const float DistanceForFixedLODSquared = FMath::Square(DistanceForFixedLOD);
+
+
+			for (int32 i = 0; i < NumGeneratedInstancesLimitLODs && i < SortedInstances.Num(); ++i)
+			{
+				if (!bAlreadyReachedFixedLOD && SortedInstances[i]->GetMinSquareDistToPlayer() < DistanceForFixedLODSquared)
+				{
+					RequestedLODs.Init(MAX_uint16, SortedInstances[i]->GetNumComponents());
+					SortedInstances[i]->SetRequestedLODs(0, 0, RequestedLODs, InOutRequestedUpdates);
+				}
+				else
+				{
+					RequestedLODs.Init(MAX_uint16, SortedInstances[i]->GetNumComponents());
+					SortedInstances[i]->SetRequestedLODs(2, 0, RequestedLODs, InOutRequestedUpdates);
+					bAlreadyReachedFixedLOD = true;
+				}
+			}
+
+			for (int32 i = NumGeneratedInstancesLimitLODs; i < NumGeneratedInstancesLimitLODs + NumGeneratedInstancesLimitLOD1 && i < SortedInstances.Num(); ++i)
+			{
+				if (!bAlreadyReachedFixedLOD && SortedInstances[i]->GetMinSquareDistToPlayer() < DistanceForFixedLODSquared)
+				{
+					RequestedLODs.Init(MAX_uint16, SortedInstances[i]->GetNumComponents());
+					SortedInstances[i]->SetRequestedLODs(1, 0, RequestedLODs, InOutRequestedUpdates);
+				}
+				else
+				{
+					RequestedLODs.Init(MAX_uint16, SortedInstances[i]->GetNumComponents());
+					SortedInstances[i]->SetRequestedLODs(2, 0, RequestedLODs, InOutRequestedUpdates);
+					bAlreadyReachedFixedLOD = true;
+				}
+			}
+
+			for (int32 i = NumGeneratedInstancesLimitLODs + NumGeneratedInstancesLimitLOD1; i < NumGeneratedInstancesLimitLODs + NumGeneratedInstancesLimitLOD1 + NumGeneratedInstancesLimitLOD2 && i < SortedInstances.Num(); ++i)
+			{
+				RequestedLODs.Init(MAX_uint16, SortedInstances[i]->GetNumComponents());
+				SortedInstances[i]->SetRequestedLODs(2, 0, RequestedLODs, InOutRequestedUpdates);
+			}
+
+			for (int32 i = NumGeneratedInstancesLimitLODs + NumGeneratedInstancesLimitLOD1 + NumGeneratedInstancesLimitLOD2; i < SortedInstances.Num(); ++i)
+			{
+				SortedInstances[i]->SetIsDiscardedBecauseOfTooManyInstances(true);
+			}
+		}
+		else
+		{
+			// No limit surpassed, set all instances to have all LODs, there will be an UpdateSkeletalMesh only if there's a change in LOD state
+			for (int32 i = 0; i < SortedInstances.Num(); ++i)
+			{
+				RequestedLODs.Init(MAX_uint16, SortedInstances[i]->GetNumComponents());
+				SortedInstances[i]->SetRequestedLODs(2, 0, RequestedLODs, InOutRequestedUpdates);
+			}
+		}
+	}
+	else if (IsOnlyGenerateRequestedLODLevelsEnabled())
+	{
+		struct FLODTracker
+		{
+			int32 MinLOD = MAX_int32;
+
+			bool bInitialized = false;
+			TArray<uint16> RequestedLODPerComponent;
+		};
+
+		TMap<TObjectPtr<UCustomizableObjectInstance>, FLODTracker> InstancesMinLOD;
+		InstancesMinLOD.Reserve(100);
+
+		for (TObjectIterator<UCustomizableObjectInstanceUsage> CustomizableObjectInstanceUsage; CustomizableObjectInstanceUsage; ++CustomizableObjectInstanceUsage)
+		{
+#if WITH_EDITOR
+			if (IsValid(*CustomizableObjectInstanceUsage) && CustomizableObjectInstanceUsage->IsNetMode(NM_DedicatedServer))
+			{
+				continue;
+			}
+#endif
+
+			if (IsValid(*CustomizableObjectInstanceUsage) && !CustomizableObjectInstanceUsage->IsTemplate())
+			{
+				UCustomizableObjectInstance* COI = CustomizableObjectInstanceUsage->GetCustomizableObjectInstance();
+				if (!COI || !COI->GetCustomizableObject())
+				{
+					continue;
+				}
+
+				USkeletalMeshComponent* Parent = Cast<USkeletalMeshComponent>(CustomizableObjectInstanceUsage->GetAttachParent());
+
+#if WITH_EDITOR
+
+				EWorldType::Type WorldType = EWorldType::Type::None;
+
+				UWorld* World = Parent ? Parent->GetWorld() : nullptr;
+				if (World)
+				{
+					WorldType = World->WorldType;
+				}
+
+				// Blueprint instances and CO Editors
+				
+				const USceneComponent* const AttachParentComponent = CustomizableObjectInstanceUsage->GetAttachParent();
+				bool bAttachParentActor = AttachParentComponent ? AttachParentComponent->GetOwner()!=nullptr : false;
+
+				if (WorldType == EWorldType::EditorPreview || (!World && !bAttachParentActor))
+				{
+					CustomizableObjectInstanceUsage->EditorUpdateComponent();
+					continue;
+				}
+
+				// Skip if world type is EditorPreview, GamePreview, etc...
+				if (WorldType == EWorldType::GamePreview || WorldType == EWorldType::GameRPC || WorldType == EWorldType::Inactive)
+				{
+					continue;
+				}
+#endif // WITH_EDITOR
+
+				FLODTracker& LODTracker = InstancesMinLOD.FindOrAdd(COI);
+
+				if (!LODTracker.bInitialized)
+				{
+					LODTracker.RequestedLODPerComponent.Init(MAX_uint16, COI->GetCustomizableObject()->GetComponentCount());
+					LODTracker.bInitialized = true;
+				}
+
+				if (Parent)
+				{
+					COI->SetIsBeingUsedByComponentInPlay(true);
+
+					const FName& ComponentName = CustomizableObjectInstanceUsage->GetComponentName();
+					const int32 ObjectComponentIndex = COI->GetCustomizableObject()->GetPrivate()->GetModelResources().ComponentNames.IndexOfByKey(ComponentName);
+
+#if WITH_EDITOR
+					// If the instance is generated but the component doesn't have a mesh, set it.
+					// Can happen when duplicating instances in the editor.
+					const USkeletalMesh* SkeletalMesh = COI->GetComponentMeshSkeletalMesh(ComponentName);
+					if (SkeletalMesh && !Parent->GetSkeletalMeshAsset())
+					{
+						// As the instance is already generated, this will be very fast and just set the mesh and call the delegates
+						COI->UpdateSkeletalMeshAsync();
+					}
+#endif
+
+					// If it's the local player set max priority
+					const USceneComponent* const AttachParentParentComponent = Parent->GetAttachParent();
+					AActor* ParentParentActor = AttachParentParentComponent ? AttachParentParentComponent->GetOwner() : nullptr;
+
+					const APawn* Pawn = Cast<APawn>(ParentParentActor);
+					if (Pawn && Pawn->IsPlayerControlled())
+					{
+						COI->SetMinSquareDistToPlayer(-1.f);
+					}
+
+					// Use the component minLOD to set the minimum LOD to generate
+					LODTracker.MinLOD = FMath::Min(LODTracker.MinLOD, Parent->bOverrideMinLod ? Parent->MinLodModel : 0);
+
+					// If the parent component have a SkeletalMesh use the RequestedLODLevel of the component as reference to know which LODs mutable should generate.
+					if (UE_MUTABLE_GETSKELETALMESHASSET(Parent) && LODTracker.RequestedLODPerComponent.IsValidIndex(ObjectComponentIndex))
+					{
+						uint16& RequestedLOD = LODTracker.RequestedLODPerComponent[ObjectComponentIndex];
+						RequestedLOD = FMath::Min((int32)RequestedLOD, Parent->GetPredictedLODLevel());
+					}
+				}
+			}
+		}
+
+		for (TPair<TObjectPtr<UCustomizableObjectInstance>, FLODTracker>& It : InstancesMinLOD)
+		{
+			if (IsValidChecked(It.Key) && It.Key->GetPrivate())
+			{
+				const UCustomizableObject* CustomizableObject = It.Key->GetCustomizableObject();
+				if (!CustomizableObject || CustomizableObject->GetPrivate()->IsLocked())
+				{
+					continue;
+				}
+
+				if (IsOnlyUpdateCloseCustomizableObjectsEnabled() && (It.Key->GetMinSquareDistToPlayer() > FMath::Square(CloseCustomizableObjectsDist) || !It.Key->GetPrivate()->NearestToActor.IsValid()))
+				{
+					continue;
+				}
+
+				// Limit MinLOD
+				It.Value.MinLOD = FMath::Min(It.Value.MinLOD, CustomizableObject->GetNumLODs() - 1);
+
+				It.Key->SetRequestedLODs(It.Value.MinLOD, 0, It.Value.RequestedLODPerComponent, InOutRequestedUpdates);
+			}
+		}
+	}
+}
+
+
+int32 UCustomizableInstanceLODManagement::GetNumGeneratedInstancesLimitFullLODs() const
+{
+	return CVarNumGeneratedInstancesLimit.GetValueOnGameThread();
+}
+
+
+int32 UCustomizableInstanceLODManagement::GetNumGeneratedInstancesLimitLOD1() const
+{
+	return CVarNumGeneratedInstancesLimitLOD1.GetValueOnGameThread();
+}
+
+
+int32 UCustomizableInstanceLODManagement::GetNumGeneratedInstancesLimitLOD2() const
+{
+	return CVarNumGeneratedInstancesLimitLOD2.GetValueOnGameThread();
+}
+
+
+void UCustomizableInstanceLODManagement::SetNumberOfPriorityUpdateInstances(int32 InNumPriorityUpdateInstances)
+{
+	NumPriorityUpdateInstances = InNumPriorityUpdateInstances;
+}
+
+
+int32 UCustomizableInstanceLODManagement::GetNumberOfPriorityUpdateInstances() const
+{
+	return NumPriorityUpdateInstances;
+}
+
+
+void UCustomizableInstanceLODManagement::SetCustomizableObjectsUpdateDistance(float Distance)
+{
+	CloseCustomizableObjectsDist = Distance;
+}
+
+
+float UCustomizableInstanceLODManagement::GetOnlyUpdateCloseCustomizableObjectsDist() const
+{
+	return CloseCustomizableObjectsDist;
+}
+
+
+bool UCustomizableInstanceLODManagement::IsOnlyUpdateCloseCustomizableObjectsEnabled() const
+{
+	return CVarOnlyUpdateCloseCustomizableObjects.GetValueOnGameThread();
+}
diff --git a/Engine/Plugins/Mutable/Source/CustomizableObject/Private/MuCO/CustomizableObjectMipDataProvider.cpp b/Engine/Plugins/Mutable/Source/CustomizableObject/Private/MuCO/CustomizableObjectMipDataProvider.cpp
new file mode 100644
--- /dev/null
+++ b/Engine/Plugins/Mutable/Source/CustomizableObject/Private/MuCO/CustomizableObjectMipDataProvider.cpp
@@ -0,0 +1,560 @@
+// Copyright Epic Games, Inc. All Rights Reserved.
+
+#include "MuCO/CustomizableObjectMipDataProvider.h"
+
+#include "MuCO/CustomizableObjectInstance.h"
+#include "MuCO/CustomizableObjectSystemPrivate.h"
+#include "MuR/Model.h"
+#include "TextureResource.h"
+#include "UnrealMutableImageProvider.h"
+
+#include "Containers/Ticker.h"
+
+#include UE_INLINE_GENERATED_CPP_BY_NAME(CustomizableObjectMipDataProvider)
+
+
+#define UE_MUTABLE_MIPDATA_PROVIDER_UPDATE_IMAGE_REGION		TEXT("MipDataProvider_Mutable_UpdateImage")
+
+UMutableTextureMipDataProviderFactory::UMutableTextureMipDataProviderFactory(const FObjectInitializer& ObjectInitializer)
+	: Super(ObjectInitializer)
+{
+
+}
+
+
+FMutableUpdateContext::FMutableUpdateContext(
+	const FString& InCustomizableObjectPathName,
+	const FString& InInstancePathName,
+	mu::Ptr<mu::System> InSystem,
+	TSharedPtr<mu::Model, ESPMode::ThreadSafe> InModel,
+	mu::Ptr<const mu::Parameters> InParameters,
+	int32 InState) :
+	CustomizableObjectPathName(InCustomizableObjectPathName),
+	InstancePathName(InInstancePathName),
+	System(InSystem),
+	Model(InModel),
+	Parameters(InParameters),
+	State(InState)
+{
+	if (Parameters)
+	{
+		const UCustomizableObjectSystemPrivate* Private = UCustomizableObjectSystem::GetInstance()->GetPrivate();
+		Private->GetResourceProviderChecked()->CacheImages(*Parameters);
+	}
+}
+
+
+FMutableUpdateContext::~FMutableUpdateContext()
+{
+	if (Parameters &&
+		UCustomizableObjectSystem::IsCreated())
+	{
+		const UCustomizableObjectSystemPrivate* Private = UCustomizableObjectSystem::GetInstance()->GetPrivate();
+		Private->GetResourceProviderChecked()->UnCacheImages(*Parameters);
+	}
+}
+
+
+const FString& FMutableUpdateContext::GetCustomizableObjectPathName() const
+{
+	return CustomizableObjectPathName;
+}
+
+
+const FString& FMutableUpdateContext::GetInstancePathName() const
+{
+	return InstancePathName;
+}
+
+
+mu::Ptr<mu::System> FMutableUpdateContext::GetSystem() const
+{
+	return System;
+}
+
+
+TSharedPtr<mu::Model, ESPMode::ThreadSafe> FMutableUpdateContext::GetModel() const
+{
+	return Model;
+}
+
+
+mu::Ptr<const mu::Parameters> FMutableUpdateContext::GetParameters() const
+{
+	return Parameters;
+}
+
+
+int32 FMutableUpdateContext::GetState() const
+{
+	return State;
+}
+
+
+const TArray<mu::Ptr<const mu::Image>>& FMutableUpdateContext::GetImageParameterValues() const
+{
+	return ImageParameterValues;
+}
+
+
+FMutableTextureMipDataProvider::FMutableTextureMipDataProvider(const UTexture* Texture, UCustomizableObjectInstance* InCustomizableObjectInstance, const FMutableImageReference& InImageRef)
+	: FTextureMipDataProvider(Texture, ETickState::Init, ETickThread::Async),
+	CustomizableObjectInstance(InCustomizableObjectInstance), ImageRef(InImageRef)
+{
+	check(ImageRef.ImageID > 0);
+}
+
+
+void FMutableTextureMipDataProvider::PrintWarningAndAdvanceToCleanup()
+{
+	UE_LOG(LogMutable, Warning, TEXT("Tried to update a mip from a Customizable Object being compiled, cancelling mip update."));
+	AdvanceTo(ETickState::CleanUp, ETickThread::Async);
+}
+
+
+void FMutableTextureMipDataProvider::Init(const FTextureUpdateContext& Context, const FTextureUpdateSyncOptions& SyncOptions)
+{
+#if WITH_EDITOR
+	check(Context.Texture->HasPendingInitOrStreaming());
+	check(CustomizableObjectInstance->GetCustomizableObject());
+	if (CustomizableObjectInstance->GetCustomizableObject()->GetPrivate()->IsLocked())
+	{
+		PrintWarningAndAdvanceToCleanup();
+
+		return;
+	}
+#endif
+
+	AdvanceTo(ETickState::GetMips, ETickThread::Async);
+}
+
+
+namespace Impl
+{
+	void Task_Mutable_UpdateImage(TSharedPtr<FMutableImageOperationData> OperationData)
+	{
+		const double StartTime = FPlatformTime::Seconds();
+		
+		if (FLogBenchmarkUtil::IsBenchmarkingReportingEnabled())
+		{
+			// Cache memory used when starting the update of the image
+			OperationData->ImageUpdateStartBytes = mu::FGlobalMemoryCounter::GetAbsoluteCounter();
+			mu::FGlobalMemoryCounter::Zero();
+		}
+		
+		// Any external texture that may be needed for this update will be requested from Mutable Core's GetImage
+		// which will safely access the GlobalExternalImages map, and then just get the cached image or issue a disk read
+
+		// This runs in a worker thread.
+		check(OperationData.IsValid());
+		check(OperationData->UpdateContext->GetSystem().get());
+		check(OperationData->UpdateContext->GetModel());
+		check(OperationData->UpdateContext->GetParameters().get());
+
+		if (!OperationData.IsValid())
+		{
+			return;
+		}
+
+		TRACE_BEGIN_REGION(UE_MUTABLE_MIPDATA_PROVIDER_UPDATE_IMAGE_REGION);
+		
+		mu::Ptr<mu::System> System = OperationData->UpdateContext->GetSystem();
+		const TSharedPtr<mu::Model, ESPMode::ThreadSafe> Model = OperationData->UpdateContext->GetModel();
+
+		auto EndUpdateImage = [](TSharedPtr<FMutableImageOperationData>& OperationData)
+		{
+			// The request could be cancelled in parallel from CancelCounterSafely and its value be changed
+			// between reading it and actually running Decrement() and RescheduleCallback(), so lock
+			FScopeLock Lock(&OperationData->CounterTaskLock);
+
+			if (OperationData->Counter) // If the request has been cancelled the counter will be null
+			{
+				// Make the FMutableTextureMipDataProvider continue
+				OperationData->Counter->Decrement();
+
+				if (OperationData->Counter->GetValue() == 0)
+				{
+					OperationData->RescheduleCallback();
+				}
+			}
+			
+			TRACE_END_REGION(UE_MUTABLE_MIPDATA_PROVIDER_UPDATE_IMAGE_REGION);
+		};
+
+#if WITH_EDITOR
+		// Recompiling a CO in the editor will invalidate the previously generated Model. Check that it is valid before accessing the streamed data.
+		if (!(Model && Model->IsValid()))
+		{
+			EndUpdateImage(OperationData);
+			return;
+		}
+#endif
+
+		// For now, we are forcing the recreation of mutable-side instances with every update.
+		mu::Instance::ID InstanceID = System->NewInstance(Model);
+		UE_LOG(LogMutable, Verbose, TEXT("Creating Mutable instance with id [%d] for a single UpdateImage"), InstanceID)
+
+		mu::Ptr<const mu::Instance> Instance;
+
+		// Main instance generation step
+		{
+			// LOD mask, set to all ones to build all LODs
+			uint32 LODMask = 0xFFFFFFFF;
+
+			Instance = System->BeginUpdate(InstanceID, OperationData->UpdateContext->GetParameters(), OperationData->UpdateContext->GetState(), LODMask);
+			check(Instance);
+		}
+
+		const FMutableImageReference& ImageRef = OperationData->RequestedImage;
+
+		int32 SurfaceIndex = Instance->FindSurfaceById(ImageRef.Component, ImageRef.LOD, ImageRef.SurfaceId);
+		check(SurfaceIndex >= 0);
+
+		// This ID may be different than the ID obtained the first time the image was generated, because the mutable
+		// runtime cannot remember all the resources it has built, and only remembers a fixed amount.
+		mu::FResourceID MipImageID = Instance->GetImageId(ImageRef.Component, ImageRef.LOD, SurfaceIndex, ImageRef.Image);
+
+		UE::Tasks::TTask<mu::Ptr<const mu::Image>> GetImageTask = 
+				System->GetImage(InstanceID, MipImageID, ImageRef.BaseMip + OperationData->MipsToSkip, ImageRef.LOD);
+
+		UE::Tasks::AddNested(UE::Tasks::Launch(TEXT("MipDataProvider_EndUpdateImagesTask"),
+				[System, OperationData, InstanceID, StartTime, GetImageTask, EndUpdateImage]() mutable
+				{
+					check(GetImageTask.IsCompleted());
+
+					mu::Ptr<const mu::Image> ResultImage = GetImageTask.GetResult();
+
+					check(ResultImage);
+
+					int32 FullMipCount = ResultImage->GetMipmapCount(ResultImage->GetSizeX(), ResultImage->GetSizeY());
+					int32 RealMipCount = ResultImage->GetLODCount();
+
+					// Did we fail to generate the entire mipchain (if we have mips at all)?
+					bool bForceMipchain = (RealMipCount != 1) && (RealMipCount != FullMipCount);
+
+					if (bForceMipchain)
+					{
+						MUTABLE_CPUPROFILER_SCOPE(GetImage_MipFix);
+
+						UE_LOG(LogMutable, Warning, TEXT("Mutable generated an incomplete mip chain for image."));
+
+						// Force the right number of mips. The missing data will be black.
+						mu::Ptr<mu::Image> NewImage = new mu::Image(ResultImage->GetSizeX(), ResultImage->GetSizeY(), FullMipCount, ResultImage->GetFormat(), mu::EInitializationType::Black);
+
+						// Formats with BytesPerBlock == 0 will not allocate memory. This type of images are not expected here.
+						check(!NewImage->DataStorage.IsEmpty());
+
+						for (int32 L = 0; L < RealMipCount; ++L)
+						{
+							TArrayView<uint8> DestView = NewImage->DataStorage.GetLOD(L);
+							TArrayView<const uint8> SrcView = ResultImage->DataStorage.GetLOD(L);
+
+							check(DestView.Num() == SrcView.Num());
+							FMemory::Memcpy(DestView.GetData(), SrcView.GetData(), DestView.Num());
+						}
+
+						ResultImage = NewImage;
+					}
+
+					OperationData->Result = ResultImage;
+
+					// End update
+					{
+						MUTABLE_CPUPROFILER_SCOPE(EndUpdate);
+						System->EndUpdate(InstanceID);
+						System->ReleaseInstance(InstanceID);
+
+						if (UCustomizableObjectSystem::ShouldClearWorkingMemoryOnUpdateEnd())
+						{
+							System->ClearWorkingMemory();
+						}
+					}
+
+					if (FLogBenchmarkUtil::IsBenchmarkingReportingEnabled())
+					{
+						double Time = FPlatformTime::Seconds() - StartTime;
+						// Report the peak memory used by the operation
+						const int64 PeakMemory = mu::FGlobalMemoryCounter::GetPeak();
+						// Report the peak memory used during the operation (operation + baseline)
+						const int64 RealMemoryPeak = PeakMemory + OperationData->ImageUpdateStartBytes;
+
+						const FString& CustomizableObjectPathName = OperationData->UpdateContext->GetCustomizableObjectPathName();
+						const FString& InstancePathName = OperationData->UpdateContext->GetInstancePathName();
+
+						const FString& Descriptor = OperationData->UpdateContext->CapturedDescriptor;
+						const bool bDidLevelBeginPlay = OperationData->UpdateContext->bLevelBegunPlay;
+
+						ExecuteOnGameThread(UE_SOURCE_LOCATION, 
+						[CustomizableObjectPathName, InstancePathName, bDidLevelBeginPlay, Time, PeakMemory, RealMemoryPeak, Descriptor ]
+						{
+							if (!UCustomizableObjectSystem::IsCreated()) // We are shutting down
+							{
+								return;	
+							}
+							
+							UCustomizableObjectSystem* System = UCustomizableObjectSystem::GetInstance();
+							if (!System)
+							{
+								return;
+							}
+
+							System->GetPrivate()->LogBenchmarkUtil.FinishUpdateImage(CustomizableObjectPathName, InstancePathName, Descriptor, bDidLevelBeginPlay, Time, PeakMemory, RealMemoryPeak);
+						});
+					}
+
+					EndUpdateImage(OperationData);
+				},
+				UE::Tasks::Prerequisites(GetImageTask),
+				UE::Tasks::ETaskPriority::Inherit,
+				UE::Tasks::EExtendedTaskPriority::Inline)); // MipDataProvider_EndUpdateImagesTaskGetImage.
+	}
+} // namespace Impl
+
+
+int32 FMutableTextureMipDataProvider::GetMips(const FTextureUpdateContext& Context, int32 StartingMipIndex, const FTextureMipInfoArray& MipInfos, const FTextureUpdateSyncOptions& SyncOptions)
+{
+	MUTABLE_CPUPROFILER_SCOPE(FMutableTextureMipDataProvider::GetMips);
+
+	if (!UCustomizableObjectSystem::IsActive())
+	{
+		// Mutable is disabled. Skip all mip operations and mark the update task as completed.
+		AdvanceTo(ETickState::Done, ETickThread::None);
+		return CurrentFirstLODIdx;
+	}
+
+#if WITH_EDITOR
+	check(Context.Texture->HasPendingInitOrStreaming());
+	check(CustomizableObjectInstance->GetCustomizableObject());
+	if (CustomizableObjectInstance->GetCustomizableObject()->GetPrivate()->IsLocked())
+	{
+		PrintWarningAndAdvanceToCleanup();
+
+		return CurrentFirstLODIdx;
+	}
+#endif
+
+	const UTexture2D* Texture = Cast<UTexture2D>(Context.Texture);
+	check(Texture);
+	check(!Texture->NeverStream);
+	const TIndirectArray<FTexture2DMipMap>& OwnerMips = Texture->GetPlatformMips();
+
+	const int32 NumMips = OwnerMips.Num();
+	check(ImageRef.ImageID > 0);
+
+	// Maximum value to skip, will be minimized by the first mip level requested
+	int32 MipsToSkip = 256;
+
+	for (int32 MipIndex = StartingMipIndex; MipIndex < CurrentFirstLODIdx; ++MipIndex)
+	{
+		check(Context.MipsView.IsValidIndex(MipIndex) && MipInfos.IsValidIndex(MipIndex));
+
+		const FTexture2DMipMap& MipMap = *Context.MipsView[MipIndex];
+		FByteBulkData* AuxBulkData = const_cast<FByteBulkData*>(&MipMap.BulkData);
+
+		const FTextureMipInfo& MipInfo = MipInfos[MipIndex];
+		void* Dest = MipInfo.DestData;
+
+		if (AuxBulkData->GetBulkDataSize() > 0)
+		{
+			// Mips are already generated, no need for Mutable progressive mip streaming, just normal CPU->GPU streaming
+			AuxBulkData->GetCopy(&Dest, false);
+		}
+		else // Generate a Mip Request to Mutable
+		{
+			check(UCustomizableObjectSystem::GetInstance()->GetPrivate()->EnableMutableProgressiveMipStreaming == 1);
+
+			if (!OperationData.IsValid())
+			{
+				OperationData = MakeShared<FMutableImageOperationData>();
+			}
+
+			const int32 LODBias = static_cast<int32>(Context.MipsView.GetData() - Texture->GetPlatformData()->Mips.GetData());
+			check(LODBias >= 0);
+			check(NumMips == Context.MipsView.Num() + LODBias);
+
+			OperationData->Levels.Add({ MipIndex + LODBias, Dest, (int32)MipInfo.SizeX, (int32)MipInfo.SizeY, (int32)MipInfo.DataSize, MipInfo.Format });
+			MipsToSkip = FMath::Min(MipsToSkip, MipIndex + LODBias);
+		}
+	}
+
+	UCustomizableObjectSystemPrivate* CustomizableObjectSystem = UCustomizableObjectSystem::GetInstance()->GetPrivate();
+	if (CustomizableObjectSystem)
+	{
+		if (OperationData.IsValid())
+		{
+			OperationData->RequestedImage = ImageRef;
+			OperationData->UpdateContext = UpdateContext;
+			OperationData->MipsToSkip = MipsToSkip;
+
+			check(SyncOptions.Counter);
+			// Increment to stop PollMips from running until the Mutable request task finishes. 
+			// If a request completes immediately, then it will call the callback
+			// but that won't do anything because the tick would not try to acquire the lock since it is already locked.
+			SyncOptions.Counter->Increment();
+			OperationData->Counter = SyncOptions.Counter;
+			OperationData->RescheduleCallback = SyncOptions.RescheduleCallback;
+
+			MutableTaskId = CustomizableObjectSystem->MutableTaskGraph.AddMutableThreadTaskLowPriority(
+				TEXT("Mutable_MipUpdate"),
+				[OperationData = this->OperationData]()
+				{
+					Impl::Task_Mutable_UpdateImage(OperationData);
+				});
+		}
+
+		AdvanceTo(ETickState::PollMips, ETickThread::Async);
+	}
+	else
+	{
+		AdvanceTo(ETickState::CleanUp, ETickThread::Async);
+	}
+
+	return CurrentFirstLODIdx;
+}
+
+
+bool FMutableTextureMipDataProvider::PollMips(const FTextureUpdateSyncOptions& SyncOptions)
+{
+	MUTABLE_CPUPROFILER_SCOPE(FMutableTextureMipDataProvider::PollMips)
+
+	// Once this point is reached, even if the task has not been completed, we know that all the work we need from it has been completed.
+	// Furthermore, checking if the task is completed is incorrect since PollMips could have been called by RescheduleCallback (before completing the task).
+	
+#if WITH_EDITOR
+	check(CustomizableObjectInstance->GetCustomizableObject());
+	if (CustomizableObjectInstance->GetCustomizableObject()->GetPrivate()->IsLocked())
+	{
+		PrintWarningAndAdvanceToCleanup();
+
+		return false;
+	}
+#endif
+
+	if (bRequestAborted)
+	{
+		OperationData = nullptr;
+		AdvanceTo(ETickState::CleanUp, ETickThread::Async);
+		return false;
+	}
+	
+	if (OperationData && OperationData->Levels.Num())
+	{
+		// The counter must be zero meaning the Mutable image operation has finished
+		check(SyncOptions.Counter->GetValue() == 0);
+
+		mu::Ptr<const mu::Image> Image = OperationData->Result;
+		
+		int32 ImageLODCount = 0;
+		if (Image)
+		{
+			ImageLODCount = Image->GetLODCount();
+			// check(Image->GetLODCount() == OperationData->Levels.Num()); TODO PRP
+			
+			// No longer true, since missing data may mean we generate smaller images.
+			//check(Image->GetSizeX() == OperationData->Levels[0].SizeX);
+			//check(Image->GetSizeY() == OperationData->Levels[0].SizeY);
+			if (Image->GetSizeX() != OperationData->Levels[0].SizeX
+				||
+				Image->GetSizeY() != OperationData->Levels[0].SizeY)
+			{
+				OperationData = nullptr;
+				AdvanceTo(ETickState::CleanUp, ETickThread::Async);
+				return false;
+			}
+		}
+
+		int32 MipIndex = 0;
+		for (FMutableMipUpdateLevel& Level : OperationData->Levels)
+		{
+			void* Dest = Level.Dest;
+
+			if (MipIndex < ImageLODCount)
+			{
+				int32 MipDataSize = Image->GetLODDataSize(MipIndex);
+
+				// Check Mip DataSize for consistency, but skip if 0 because it's optional and might be zero in cooked mips
+				bool bCorrectDataSize = (Level.DataSize == 0 || MipDataSize == Level.DataSize);
+				if (bCorrectDataSize)
+				{
+					FMemory::Memcpy(Dest, Image->GetMipData(MipIndex), MipDataSize);
+				}
+				else
+				{
+					UE_LOG(LogMutable, Warning, TEXT("Mip data has incorrect size."));
+					FMemory::Memzero(Dest, Level.DataSize);
+				}
+			}
+			else
+			{
+				// Mutable didn't generate all the expected mips
+				UE_LOG(LogMutable, Warning, TEXT("Mutable image is missing mips."));
+				FMemory::Memzero(Dest, Level.DataSize);
+			}
+			++MipIndex;
+		}
+
+		// Force the immediate release of the image memory to reduce the transient memory usage
+		Image = nullptr;
+		OperationData->Result = nullptr;
+	}
+
+	OperationData = nullptr;
+	AdvanceTo(ETickState::Done, ETickThread::None);
+	return true;
+}
+
+
+void FMutableTextureMipDataProvider::CleanUp(const FTextureUpdateSyncOptions& SyncOptions)
+{
+	check(!SyncOptions.Counter || SyncOptions.Counter->GetValue() == 0);
+	AdvanceTo(ETickState::Done, ETickThread::None);
+}
+
+
+void FMutableTextureMipDataProvider::Cancel(const FTextureUpdateSyncOptions& SyncOptions)
+{
+	bRequestAborted = true;
+
+	CancelCounterSafely();
+
+	if (UCustomizableObjectSystem* System = UCustomizableObjectSystem::GetInstance())
+	{
+		System->GetPrivate()->MutableTaskGraph.CancelMutableThreadTaskLowPriority(MutableTaskId);	
+	}
+}
+
+
+FTextureMipDataProvider::ETickThread FMutableTextureMipDataProvider::GetCancelThread() const
+{
+	//return OperationData ? FTextureMipDataProvider::ETickThread::Async : FTextureMipDataProvider::ETickThread::None;
+	return FTextureMipDataProvider::ETickThread::None;
+}
+
+
+void FMutableTextureMipDataProvider::AbortPollMips()
+{
+	bRequestAborted = true;
+
+	CancelCounterSafely();
+
+	if (UCustomizableObjectSystem* System = UCustomizableObjectSystem::GetInstance())
+	{
+		System->GetPrivate()->MutableTaskGraph.CancelMutableThreadTaskLowPriority(MutableTaskId);	
+	}
+}
+
+
+void FMutableTextureMipDataProvider::CancelCounterSafely()
+{
+	if (OperationData.IsValid())
+	{
+		// The Counter could be read in parallel from Task_Mutable_UpdateImage, so lock
+		FScopeLock Lock(&OperationData->CounterTaskLock);
+
+		if (OperationData->Counter)
+		{
+			OperationData->Counter->Set(0);
+			OperationData->Counter = nullptr;
+		}
+	}
+}
diff --git a/Engine/Plugins/Mutable/Source/CustomizableObject/Private/MuCO/CustomizableSkeletalComponent.cpp b/Engine/Plugins/Mutable/Source/CustomizableObject/Private/MuCO/CustomizableSkeletalComponent.cpp
new file mode 100644
--- /dev/null
+++ b/Engine/Plugins/Mutable/Source/CustomizableObject/Private/MuCO/CustomizableSkeletalComponent.cpp
@@ -0,0 +1,162 @@
+// Copyright Epic Games, Inc. All Rights Reserved.
+
+#include "MuCO/CustomizableSkeletalComponent.h"
+#include "MuCO/CustomizableObjectInstanceUsage.h"
+
+#include "UObject/UObjectGlobals.h"
+#include "Components/SkeletalMeshComponent.h"
+#include "MuCO/CustomizableObject.h"
+
+
+void UCustomizableSkeletalComponent::CreateCustomizableObjectInstanceUsage()
+{
+	if (CustomizableObjectInstanceUsage)
+	{
+		// CustomizableObjectInstanceUsage may already exist if duplicated from an existing Customizable Skeletal Component
+		if (CustomizableObjectInstanceUsage->CustomizableSkeletalComponent != this)
+		{
+			CustomizableObjectInstanceUsage = nullptr;
+		}
+	}
+
+	AActor* RootActor = GetAttachmentRootActor();
+	bool bIsDefaultActor = RootActor ?
+		RootActor->HasAnyFlags(RF_ClassDefaultObject) :
+		false;
+
+	if (!CustomizableObjectInstanceUsage && !HasAnyFlags(RF_ClassDefaultObject) && !bIsDefaultActor)
+	{
+		CustomizableObjectInstanceUsage = NewObject<UCustomizableObjectInstanceUsage>(this, TEXT("InstanceUsage"), RF_Transient);
+		CustomizableObjectInstanceUsage->CustomizableSkeletalComponent = this;
+	}
+}
+
+
+void UCustomizableSkeletalComponent::Callbacks() const
+{
+	if (CustomizableObjectInstanceUsage)
+	{
+		CustomizableObjectInstanceUsage->Callbacks();
+	}
+}
+
+
+USkeletalMesh* UCustomizableSkeletalComponent::GetSkeletalMesh() const
+{
+	if (CustomizableObjectInstanceUsage)
+	{
+		return CustomizableObjectInstanceUsage->GetSkeletalMesh();
+	}
+
+	return nullptr;
+}
+
+
+void UCustomizableSkeletalComponent::SetSkeletalMesh(USkeletalMesh* SkeletalMesh)
+{
+	if (CustomizableObjectInstanceUsage)
+	{
+		CustomizableObjectInstanceUsage->SetSkeletalMesh(SkeletalMesh);
+	}
+}
+
+
+void UCustomizableSkeletalComponent::SetPhysicsAsset(UPhysicsAsset* PhysicsAsset)
+{
+	if (CustomizableObjectInstanceUsage)
+	{
+		CustomizableObjectInstanceUsage->SetPhysicsAsset(PhysicsAsset);
+	}
+}
+
+
+USkeletalMesh* UCustomizableSkeletalComponent::GetAttachedSkeletalMesh() const
+{
+	if (CustomizableObjectInstanceUsage)
+	{
+		return CustomizableObjectInstanceUsage->GetAttachedSkeletalMesh();
+	}
+
+	return nullptr;
+}
+
+
+void UCustomizableSkeletalComponent::SetComponentName(const FName& Name)
+{
+	ComponentIndex = INDEX_NONE;
+	ComponentName = Name;
+}
+
+
+FName UCustomizableSkeletalComponent::GetComponentName() const
+{
+	if (ComponentIndex == INDEX_NONE)
+	{
+		return ComponentName;	
+	}
+	else
+	{
+		return FName(FString::FromInt(ComponentIndex));
+	}
+}
+
+
+void UCustomizableSkeletalComponent::UpdateSkeletalMeshAsync(bool bNeverSkipUpdate)
+{
+	if (CustomizableObjectInstanceUsage)
+	{
+		CustomizableObjectInstanceUsage->UpdateSkeletalMeshAsync(bNeverSkipUpdate);
+	}
+}
+
+
+void UCustomizableSkeletalComponent::UpdateSkeletalMeshAsyncResult(FInstanceUpdateDelegate Callback, bool bIgnoreCloseDist, bool bForceHighPriority)
+{
+	if (CustomizableObjectInstanceUsage)
+	{
+		CustomizableObjectInstanceUsage->UpdateSkeletalMeshAsyncResult(Callback, bIgnoreCloseDist, bForceHighPriority);
+	}
+}
+
+
+#if WITH_EDITOR
+
+void UCustomizableSkeletalComponent::EditorUpdateComponent()
+{
+	if (CustomizableObjectInstanceUsage)
+	{
+		CustomizableObjectInstanceUsage->EditorUpdateComponent();
+	}
+}
+#endif
+
+
+void UCustomizableSkeletalComponent::OnAttachmentChanged()
+{
+	Super::OnAttachmentChanged();
+
+	USkeletalMeshComponent* Parent = Cast<USkeletalMeshComponent>(GetAttachParent());
+
+	if (Parent && CustomizableObjectInstanceUsage)
+	{
+		CustomizableObjectInstanceUsage->SetPendingSetSkeletalMesh(true);
+	}
+	else if(!GetAttachParent())
+	{
+		DestroyComponent();
+	}
+}
+
+
+void UCustomizableSkeletalComponent::PostInitProperties()
+{
+	Super::PostInitProperties();
+	CreateCustomizableObjectInstanceUsage();
+}
+
+void UCustomizableSkeletalComponent::PostReinitProperties()
+{
+	Super::PostReinitProperties();
+	CreateCustomizableObjectInstanceUsage();
+}
+
diff --git a/Engine/Plugins/Mutable/Source/CustomizableObject/Private/MuCO/LogInformationUtil.cpp b/Engine/Plugins/Mutable/Source/CustomizableObject/Private/MuCO/LogInformationUtil.cpp
new file mode 100644
--- /dev/null
+++ b/Engine/Plugins/Mutable/Source/CustomizableObject/Private/MuCO/LogInformationUtil.cpp
@@ -0,0 +1,312 @@
+// Copyright Epic Games, Inc. All Rights Reserved.
+
+#include "MuCO/LogInformationUtil.h"
+#include "MuCO/CustomizableObjectSystemPrivate.h"
+#include "Engine/SkeletalMesh.h"
+#include "Engine/World.h"
+#include "GameFramework/PlayerController.h"
+#include "MuCO/CustomizableObject.h"
+#include "MuCO/CustomizableObjectInstancePrivate.h"
+#include "TextureResource.h"
+
+int LogInformationUtil::CountLOD0 = 0;
+int LogInformationUtil::CountLOD1 = 0;
+int LogInformationUtil::CountLOD2 = 0;
+
+void LogInformationUtil::PrintGeneratedTextures(const TArray<FGeneratedTexture> GeneratedTextures, FString& Log, bool DoPrintInitialMessage)
+{
+	if (GeneratedTextures.Num() == 0)
+	{
+		return;
+	}
+
+	if (DoPrintInitialMessage)
+	{
+		Log += "Generated textures: ";
+	}
+
+	const int Max = GeneratedTextures.Num();
+	int i;
+
+	bool UseTabulation = (Max > 1);
+
+	if (!DoPrintInitialMessage)
+	{
+		Log += "\n\t\t";
+	}
+	
+
+	if (Max > 1)
+	{
+		UE_LOG(LogMutable, Log, TEXT("Case"));
+	}
+
+	for (i = 0; i < Max; ++i)
+	{
+		Log += FString::Printf(TEXT("Texture name=%s, "), *GeneratedTextures[i].Name);
+		Log += FString::Printf(TEXT("id=%d, "), GeneratedTextures[i].Key.Resource);
+		Log += "\n\t\t\t";
+		PrintTextureData(Cast<UTexture2D>(GeneratedTextures[i].Texture), Log, false);
+		if (i < (Max - 1))
+		{
+			Log += "\n\t\t";
+		}
+	}
+}
+
+
+void LogInformationUtil::PrintGeneratedMaterial(const TArray<struct FGeneratedMaterial> GeneratedMaterials, FString& Log)
+{
+	if (GeneratedMaterials.Num() == 0)
+	{
+		return;
+	}
+
+	Log += "\tGenerated materials:\n";
+
+	FString MessageChunk = "\t\t";
+
+	const int Max = GeneratedMaterials.Num();
+	int i;
+	int entryLength = 25;
+	bool AnyHasTextures = false;
+
+	MessageChunk += "Materials with no textures: ";
+
+	// First only no texture materials
+	for (i = 0; i < Max; ++i)
+	{
+		if (GeneratedMaterials[i].Textures.Num() == 0)
+		{
+			MessageChunk += FString::Printf(TEXT("%d, "), i);
+		}
+		else
+		{
+			AnyHasTextures = true;
+		}
+	}
+
+	// Second texture materials
+	for (i = 0; i < Max; ++i)
+	{
+		if (GeneratedMaterials[i].Textures.Num() > 0)
+		{
+			MessageChunk += "\n\t\t";
+			MessageChunk += FString::Printf(TEXT("%d"), i);
+			MessageChunk += FString::Printf(TEXT(" has %d texture(s)"), GeneratedMaterials[i].Textures.Num());
+			PrintGeneratedTextures(GeneratedMaterials[i].Textures, MessageChunk, false);
+		}
+	}
+
+	Log += MessageChunk;
+}
+
+
+void LogInformationUtil::LogShowInstanceData(const UCustomizableObjectInstance* CustomizableObjectInstance, FString& LogData)
+{
+	LogData += FString::Printf(TEXT("Name=%s "), *CustomizableObjectInstance->GetName());
+	LogData += FString::Printf(TEXT("Priority=%f "), FMath::Sqrt(CustomizableObjectInstance->GetPrivate()->LastMinSquareDistFromComponentToPlayer));
+
+	const uint32 InstanceFlags = CustomizableObjectInstance->GetPrivate()->GetCOInstanceFlags();
+	LogData += FString::Printf(TEXT("bShouldUpdateLODs=%d "), InstanceFlags & ECOInstanceFlags::PendingLODsUpdate);
+
+	LogData += FString::Printf(TEXT("CurrentMinLOD=%d "), CustomizableObjectInstance->GetCurrentMinLOD());
+	LogData += FString::Printf(TEXT("MinLODToLoad=%d "), CustomizableObjectInstance->GetMinLODToLoad());
+}
+
+
+// TODO: [Acosin] put in a different file and include it
+void LogInformationUtil::LogShowInstanceDataFull(const UCustomizableObjectInstance* CustomizableObjectInstance, bool ShowMaterialInfo)
+{
+	FString LogData = "\n\n";
+
+	LogData += FString::Printf(TEXT("CustomizableObjectInstance %s, "), *CustomizableObjectInstance->GetName());
+
+	LogData += "Descriptor:\n";
+	LogData += CustomizableObjectInstance->GetPrivate()->CommittedDescriptor.ToString();
+	LogData += "\n";
+	
+	if (ShowMaterialInfo)
+	{
+		PrintGeneratedMaterial(CustomizableObjectInstance->GetPrivate()->GeneratedMaterials, LogData);
+
+		if (CustomizableObjectInstance->GetPrivate()->GeneratedTextures.Num() > 0)
+		{
+			LogData += "\n\t";
+		}
+
+		PrintGeneratedTextures(CustomizableObjectInstance->GetPrivate()->GeneratedTextures, LogData, true);
+	}
+
+	FString MessageChunk;
+
+	MessageChunk += "\n\t";
+	MessageChunk += FString::Printf(TEXT("        bShowOnlyRuntimeParameters = %d\n"), CustomizableObjectInstance->GetPrivate()->bShowOnlyRuntimeParameters);
+	MessageChunk += FString::Printf(TEXT("        bShowOnlyRelevantParameters = %d\n"), CustomizableObjectInstance->GetPrivate()->bShowOnlyRelevantParameters);
+	MessageChunk += FString::Printf(TEXT("        MinSquareDistFromComponentToPlayer = %.2f\n"), CustomizableObjectInstance->GetPrivate()->MinSquareDistFromComponentToPlayer);
+	LogData += MessageChunk;
+
+	const uint32 InstanceFlags = CustomizableObjectInstance->GetPrivate()->GetCOInstanceFlags();
+	MessageChunk = FString::Printf(TEXT("        bIsBeingUsedByComponent = %d\n"), InstanceFlags & ECOInstanceFlags::UsedByComponent);
+	MessageChunk += FString::Printf(TEXT("        bIsBeingUsedByComponentInPlay = %d\n"), InstanceFlags & ECOInstanceFlags::UsedByComponentInPlay);
+	MessageChunk += FString::Printf(TEXT("        bIsPlayerOrNearIt = %d\n"), InstanceFlags & ECOInstanceFlags::UsedByPlayerOrNearIt);
+	MessageChunk += FString::Printf(TEXT("        bIsDiscardedByNumGeneratedInstancesLimit = %d\n"), InstanceFlags & ECOInstanceFlags::DiscardedByNumInstancesLimit);
+	MessageChunk += FString::Printf(TEXT("        LastMinSquareDistFromComponentToPlayer = %.2f\n"), CustomizableObjectInstance->GetPrivate()->LastMinSquareDistFromComponentToPlayer);
+	LogData += MessageChunk;
+
+	MessageChunk = FString::Printf(TEXT("        CurrentMinLOD = %d\n"), CustomizableObjectInstance->GetCurrentMinLOD());
+	MessageChunk += FString::Printf(TEXT("        MinLODToLoad = %d\n"), CustomizableObjectInstance->GetMinLODToLoad());
+	MessageChunk += FString::Printf(TEXT("        bShouldUpdateLODs = %d\n"), InstanceFlags & ECOInstanceFlags::PendingLODsUpdate);
+	LogData += MessageChunk;
+
+	MessageChunk += FString::Printf(TEXT("        bIsDowngradeLODUpdate = %d\n"), InstanceFlags & ECOInstanceFlags::PendingLODsDowngrade);
+	LogData += MessageChunk;
+
+	for (int32 ComponentIndex = 0; ComponentIndex < CustomizableObjectInstance->GetNumComponents(); ++ComponentIndex)
+	{
+		if (CustomizableObjectInstance->GetSkeletalMesh(ComponentIndex))
+		{
+			if (CustomizableObjectInstance->GetCurrentMinLOD() < 1)
+			{
+				CountLOD0++;
+			}
+			else if (CustomizableObjectInstance->GetCurrentMinLOD() < 2)
+			{
+				CountLOD1++;
+			}
+			else
+			{
+				CountLOD2++;
+			}
+		}
+	}
+
+	UE_LOG(LogMutable, Log, TEXT("%s"), *LogData);
+
+	UWorld* World = GWorld;
+
+	if (World)
+	{
+		APlayerController* PlayerController = World->GetFirstPlayerController();
+		if (PlayerController)
+		{
+			PlayerController->ClientMessage(LogData);
+		}
+	}
+}
+
+
+void LogInformationUtil::PrintTextureData(const UTexture2D* Texture, FString& Log, bool UseTabulation)
+{
+	if (!Texture)
+	{
+		return;
+	}
+
+	FString MessageChunk;
+	MessageChunk += FString::Printf(TEXT("        Name=%s\n"), *Texture->GetName());
+	if (UseTabulation)
+	{
+		FillToLength(MessageChunk, 37);
+	}
+	MessageChunk += FString::Printf(TEXT("                w=%d\n"), int32(Texture->GetSurfaceWidth()));
+	if (UseTabulation)
+	{
+		FillToLength(MessageChunk, 65);
+	}
+	MessageChunk += FString::Printf(TEXT("                h=%d\n"), int32(Texture->GetSurfaceHeight()));
+	if (UseTabulation)
+	{
+		FillToLength(MessageChunk, 95);
+	}
+	MessageChunk += FString::Printf(TEXT("                bForceMiplevelsToBeResident=%d\n"), Texture->bForceMiplevelsToBeResident);
+	if (UseTabulation)
+	{
+		MessageChunk += "\n\t\t\t";
+	}
+	MessageChunk += FString::Printf(TEXT("                bGlobalForceMipLevelsToBeResident=%d\n"), Texture->bGlobalForceMipLevelsToBeResident);
+	if (UseTabulation)
+	{
+		FillToLength(MessageChunk, 37);
+	}
+	MessageChunk += FString::Printf(TEXT("                bHasStreamingUpdatePending=%d\n"), Texture->bHasStreamingUpdatePending);
+	if (UseTabulation)
+	{
+		FillToLength(MessageChunk, 65);
+	}
+	MessageChunk += FString::Printf(TEXT("                bIgnoreStreamingMipBias=%d\n"), Texture->bIgnoreStreamingMipBias);
+	if (UseTabulation)
+	{
+		MessageChunk += "\n\t\t\t";
+	}
+	MessageChunk += FString::Printf(TEXT("                FirstResourceMemMip=%d\n"), Texture->FirstResourceMemMip);
+	if (UseTabulation)
+	{
+		FillToLength(MessageChunk, 37);
+	}
+	MessageChunk += FString::Printf(TEXT("                LevelIndex=%d\n"), Texture->LevelIndex);
+	if (UseTabulation)
+	{
+		FillToLength(MessageChunk, 65);
+	}
+	MessageChunk += FString::Printf(TEXT("                CompressionSettings=%s\n"), *UEnum::GetValueAsString(Texture->CompressionSettings));
+	if (UseTabulation)
+	{
+		FillToLength(MessageChunk, 95);
+	}
+	MessageChunk += FString::Printf(TEXT("                Filter=%s\n"), *UEnum::GetValueAsString(Texture->Filter));
+	if (UseTabulation)
+	{
+		FillToLength(MessageChunk, 128);
+	}
+	MessageChunk += FString::Printf(TEXT("                LODGroup=%s\n"), *UEnum::GetValueAsString(Texture->LODGroup));
+
+	Log += MessageChunk;
+}
+
+
+void LogInformationUtil::FillToLength(FString& Data, int32 NewLength)
+{
+	int32 i;
+	const int32 Max = NewLength - Data.Len();
+
+	for (i = 0; i < Max; ++i)
+	{
+		Data += " ";
+	}
+}
+
+
+void LogInformationUtil::ResetCounters()
+{
+	CountLOD0 = 0;
+	CountLOD1 = 0;
+	CountLOD2 = 0;
+}
+
+void LogInformationUtil::PrintImageToPlatformDataMap(const TMap<uint32, FTexturePlatformData*>& ImageToPlatformDataMap, FString& Log)
+{
+	if (ImageToPlatformDataMap.Num() == 0)
+	{
+		return;
+	}
+
+	Log += "\n\tImageToPlatformDataMap:";
+	for (const TPair<uint32, FTexturePlatformData*>& Elem : ImageToPlatformDataMap)
+	{
+		Log += "\n\t\t";
+		Log = FString::Printf(TEXT("key=%d, "), Elem.Key);
+		Log = FString::Printf(TEXT("value: SizeX=%d, "), Elem.Value->SizeX);
+		Log = FString::Printf(TEXT("SizeY=%d, "), Elem.Value->SizeY);
+		Log = FString::Printf(TEXT("PixelFormat=%d, "), GetPixelFormatString(Elem.Value->PixelFormat));
+		int i;
+		const int Max = Elem.Value->Mips.Num();
+		for (i = 0; i < Max; ++i)
+		{
+			Log += "\n\t\t\t";
+			Log = FString::Printf(TEXT("Mips[%d] "), i);
+			Log = FString::Printf(TEXT("SizeX=%d, "), Elem.Value->Mips[i].SizeX);
+			Log = FString::Printf(TEXT("SizeY=%d"), Elem.Value->Mips[i].SizeY);
+		}
+	}
+}
diff --git a/Engine/Plugins/Mutable/Source/CustomizableObject/Private/MuCO/LogInformationUtil.h b/Engine/Plugins/Mutable/Source/CustomizableObject/Private/MuCO/LogInformationUtil.h
new file mode 100644
--- /dev/null
+++ b/Engine/Plugins/Mutable/Source/CustomizableObject/Private/MuCO/LogInformationUtil.h
@@ -0,0 +1,52 @@
+// Copyright Epic Games, Inc. All Rights Reserved.
+
+#pragma once
+
+#include "Containers/Map.h"
+
+class FString;
+class UCustomizableObjectInstance;
+struct FTexturePlatformData;
+
+#define STRINGIZE(A) STRINGIZE_NX(A)
+#define STRINGIZE_NX(A) #A
+#define ADD_INTEGER_LOG_VARIABLE(String, info, booleanValue) {String += FString::Printf(TEXT(STRINGIZE_NX(info %d)), #booleanValue);}
+#define ADD_FLOAT_LOG_VARIABLE(String, info, floatValue) {String += FString::Printf(TEXT(STRINGIZE_NX(info %d)), #floatValue);}
+#define ADD_STRING_LOG_VARIABLE(String, info, stringValue) {String += FString::Printf(TEXT(STRINGIZE_NX(info %s)), #stringValue);}
+#define ADD_STRING_LOG(String, stringValue) {String += FString::Printf(TEXT(STRINGIZE_NX(stringValues)), #stringValue);}
+
+class LogInformationUtil
+{
+public:
+	// Show full data about all UCustomizableObjectInstance existing elements
+	static void LogShowInstanceDataFull(const UCustomizableObjectInstance* CustomizableObjectInstance, bool ShowMaterialInfo);
+
+	// Show data about all UCustomizableObjectInstance existing elements
+	static void LogShowInstanceData(const UCustomizableObjectInstance* CustomizableObjectInstance, FString& LogData);
+
+	// Utility method just to log information, adds to the given FString containing the
+	// log message the log information of GeneratedMaterials
+	static void PrintGeneratedMaterial(const TArray<struct FGeneratedMaterial> GeneratedMaterials, FString& Log);
+
+	// Utility method just to log information, adds to the given FString containing the
+	// log message the log information of GeneratedTextures
+	static void PrintGeneratedTextures(const TArray<struct FGeneratedTexture> GeneratedTextures, FString& Log, bool DoPrintInitialMessage);
+
+	// Utility method just to log information, adds to the given FString containing the
+	// log message the log information of Texture
+	static void PrintTextureData(const class UTexture2D* Texture, FString& Log, bool UseTabulation);
+
+	// Reset static lod counters
+	static void ResetCounters();
+
+	// Fills with empty spaces the Data string so it reaches a kength of NewLength (is lower than current length, nothing is done)
+	static void FillToLength(FString& Data, int32 NewLength);
+
+	// Utility method just to log information, adds to the given FString containing the
+	// log message the log information of ImageToPlatformDataMap
+	static void PrintImageToPlatformDataMap(const TMap<uint32, FTexturePlatformData*>& ImageToPlatformDataMap, FString& Log);
+
+	static int CountLOD0;
+	static int CountLOD1;
+	static int CountLOD2;
+};
diff --git a/Engine/Plugins/Mutable/Source/CustomizableObject/Private/MuCO/MultilayerProjector.cpp b/Engine/Plugins/Mutable/Source/CustomizableObject/Private/MuCO/MultilayerProjector.cpp
new file mode 100644
--- /dev/null
+++ b/Engine/Plugins/Mutable/Source/CustomizableObject/Private/MuCO/MultilayerProjector.cpp
@@ -0,0 +1,342 @@
+// Copyright Epic Games, Inc. All Rights Reserved.
+
+#include "MuCO/MultilayerProjector.h"
+
+#include "MuCO/CustomizableObject.h"
+#include "MuCO/CustomizableObjectInstanceDescriptor.h"
+
+#include UE_INLINE_GENERATED_CPP_BY_NAME(MultilayerProjector)
+
+
+const FString FMultilayerProjector::MULTILAYER_PROJECTOR_PARAMETERS_INVALID = TEXT("Invalid Multilayer Projector Parameters.");
+
+
+void FMultilayerProjectorLayer::Read(const FCustomizableObjectInstanceDescriptor& Descriptor, const FString& ParamName, const int32 LayerIndex)
+{
+	if (!Descriptor.IsMultilayerProjector(ParamName))
+	{
+		ensureAlwaysMsgf(false, TEXT("%s"), *FMultilayerProjector::MULTILAYER_PROJECTOR_PARAMETERS_INVALID);
+		return;
+	}
+
+	check(LayerIndex >= 0 && LayerIndex < Descriptor.NumProjectorLayers(*ParamName)); // Layer out of range.
+
+	{
+		ECustomizableObjectProjectorType DummyType;
+		Descriptor.GetProjectorValue(ParamName, Position, Direction, Up, Scale, Angle, DummyType, LayerIndex);
+	}
+
+	{
+		const int32 ImageParamIndex = Descriptor.FindTypedParameterIndex(ParamName + FMultilayerProjector::IMAGE_PARAMETER_POSTFIX, EMutableParameterType::Int);
+		Image = Descriptor.GetIntParameters()[ImageParamIndex].ParameterRangeValueNames[LayerIndex];
+	}
+
+	{
+		const int32 OpacityParamIndex = Descriptor.FindTypedParameterIndex(ParamName + FMultilayerProjector::OPACITY_PARAMETER_POSTFIX, EMutableParameterType::Float);
+		Opacity = Descriptor.GetFloatParameters()[OpacityParamIndex].ParameterRangeValues[LayerIndex];
+	}
+}
+
+
+void FMultilayerProjectorLayer::Write(FCustomizableObjectInstanceDescriptor& Descriptor, const FString& ParamName, int32 LayerIndex) const
+{
+	if (!Descriptor.IsMultilayerProjector(ParamName))
+	{
+		ensureAlwaysMsgf(false, TEXT("%s"), *FMultilayerProjector::MULTILAYER_PROJECTOR_PARAMETERS_INVALID);
+		return;
+	}
+
+	check(LayerIndex >= 0 && LayerIndex < Descriptor.NumProjectorLayers(*ParamName)); // Layer out of range.
+
+
+	Descriptor.SetProjectorValue(ParamName, Position, Direction, Up, Scale, Angle, LayerIndex);
+	Descriptor.SetIntParameterSelectedOption(ParamName + FMultilayerProjector::IMAGE_PARAMETER_POSTFIX, Image, LayerIndex);
+	Descriptor.SetFloatParameterSelectedOption(ParamName + FMultilayerProjector::OPACITY_PARAMETER_POSTFIX, Opacity, LayerIndex);
+
+}
+
+
+uint32 GetTypeHash(const FMultilayerProjectorLayer& Key)
+{
+	uint32 Hash = GetTypeHash(Key.Position);
+
+	Hash = HashCombine(Hash, GetTypeHash(Key.Direction));
+	Hash = HashCombine(Hash, GetTypeHash(Key.Up));
+	Hash = HashCombine(Hash, GetTypeHash(Key.Scale));
+	Hash = HashCombine(Hash, GetTypeHash(Key.Angle));
+	Hash = HashCombine(Hash, GetTypeHash(Key.Image));
+	Hash = HashCombine(Hash, GetTypeHash(Key.Opacity));
+
+	return Hash;
+}
+
+
+FMultilayerProjectorVirtualLayer::FMultilayerProjectorVirtualLayer(const FMultilayerProjectorLayer& Layer, const bool bEnabled, const int32 Order):
+	FMultilayerProjectorLayer(Layer),
+	bEnabled(bEnabled),
+	Order(Order)
+{
+}
+
+
+const FString FMultilayerProjector::NUM_LAYERS_PARAMETER_POSTFIX = FString("_NumLayers");
+
+const FString FMultilayerProjector::OPACITY_PARAMETER_POSTFIX = FString("_Opacity");
+
+const FString FMultilayerProjector::IMAGE_PARAMETER_POSTFIX = FString("_SelectedImages");
+
+const FString FMultilayerProjector::POSE_PARAMETER_POSTFIX = FString("_SelectedPoses");
+
+TArray<FName> FMultilayerProjector::GetVirtualLayers() const
+{
+	TArray<FName> VirtualLayers;
+	VirtualLayersMapping.GetKeys(VirtualLayers);
+	return VirtualLayers;
+}
+
+
+void FMultilayerProjector::CreateVirtualLayer(FCustomizableObjectInstanceDescriptor& Descriptor, const FName& Id)
+{
+	if (!VirtualLayersMapping.Contains(Id))
+	{
+		const int32 Index = Descriptor.NumProjectorLayers(ParamName);
+		
+		Descriptor.CreateLayer(ParamName, Index);
+		VirtualLayersMapping.Add(Id, Index);
+		VirtualLayersOrder.Add(Id, NEW_VIRTUAL_LAYER_ORDER);
+	}
+}
+
+
+FMultilayerProjectorVirtualLayer FMultilayerProjector::FindOrCreateVirtualLayer(FCustomizableObjectInstanceDescriptor& Descriptor, const FName& Id)
+{
+	FMultilayerProjectorLayer Layer;
+	bool bEnabled;
+	int32 Order;
+
+	if (const int32* Index = VirtualLayersMapping.Find(Id))
+	{
+		if (*Index == VIRTUAL_LAYER_DISABLED)
+		{
+			Layer = DisableVirtualLayers[Id];
+			bEnabled = false;
+		}
+		else
+		{
+			Layer = Descriptor.GetLayer(ParamName, *Index);
+			bEnabled = true;
+		}
+
+		Order = VirtualLayersOrder[Id];
+	}
+	else
+	{
+		const int32 NewIndex = Descriptor.NumProjectorLayers(ParamName);
+		constexpr int32 NewOrder = NEW_VIRTUAL_LAYER_ORDER;
+		
+		Descriptor.CreateLayer(ParamName, NewIndex);
+		VirtualLayersMapping.Add(Id, NewIndex);
+		VirtualLayersOrder.Add(Id, NewOrder);
+
+		Layer = Descriptor.GetLayer(ParamName, NewIndex);
+		bEnabled = true;
+		Order = NewOrder;
+	}
+
+	return FMultilayerProjectorVirtualLayer(Layer, bEnabled, Order);
+}
+
+
+void FMultilayerProjector::RemoveVirtualLayer(FCustomizableObjectInstanceDescriptor& Descriptor, const FName& Id)
+{
+	const int32* Index = VirtualLayersMapping.Find(Id);
+	check(Index); // Virtual Layer not created.
+	
+	if (*Index == VIRTUAL_LAYER_DISABLED)
+	{
+		DisableVirtualLayers.Remove(Id);
+	}
+	else
+	{
+		Descriptor.RemoveLayerAt(ParamName, *Index);
+		
+		for (TMap<FName, int32>::TIterator It =  VirtualLayersMapping.CreateIterator(); It; ++It)
+		{
+			if (It.Key() == Id)
+			{
+				It.RemoveCurrent();
+			}
+			else if (It.Value() > *Index) // Update following Layers.
+			{
+				--It.Value();
+			}
+		}
+	}
+
+	VirtualLayersOrder.Remove(Id);
+}
+
+
+FMultilayerProjectorVirtualLayer FMultilayerProjector::GetVirtualLayer(const FCustomizableObjectInstanceDescriptor& Descriptor, const FName& Id) const
+{
+	const int32* Index = VirtualLayersMapping.Find(Id);
+	check(Index); // Virtual Layer not created.
+
+	const FMultilayerProjectorLayer Layer = Descriptor.GetLayer(ParamName, *Index);
+	const bool bEnabled = *Index != VIRTUAL_LAYER_DISABLED;
+	const int32 Order = VirtualLayersOrder[Id];
+	
+	return FMultilayerProjectorVirtualLayer(Layer, bEnabled, Order);
+}
+
+
+void FMultilayerProjector::UpdateVirtualLayer(FCustomizableObjectInstanceDescriptor& Descriptor, const FName& Id, const FMultilayerProjectorVirtualLayer& Layer)
+{
+	const int32* Index = VirtualLayersMapping.Find(Id);
+	check(Index); // Virtual Layer not created.
+
+	const bool bEnabled = *Index != VIRTUAL_LAYER_DISABLED;
+	
+	if (!bEnabled)
+	{
+		DisableVirtualLayers[Id] = static_cast<FMultilayerProjectorLayer>(Layer); // Update disabled layer.
+		VirtualLayersOrder[Id] = Layer.Order;
+	}
+	else
+	{
+		int32* Order = VirtualLayersOrder.Find(Id);
+		if (*Order != Layer.Order) // Order changed, check if it needs to be moved.
+		{
+			const int32 OldIndex = *Index;
+			int32 NewIndex = CalculateVirtualLayerIndex(Id, Layer.Order);
+			if (OldIndex != NewIndex) // Move required. Could be optimized by moving only the in-between values.
+			{
+				Descriptor.RemoveLayerAt(ParamName, OldIndex);
+				UpdateMappingVirtualLayerDisabled(Id, OldIndex);
+
+				if (OldIndex < NewIndex)
+				{
+					NewIndex -= 1;
+				}
+				
+				Descriptor.CreateLayer(ParamName, NewIndex);
+				UpdateMappingVirtualLayerEnabled(Id, NewIndex);
+			}
+			
+			*Order = Layer.Order;
+		}
+
+		Descriptor.UpdateLayer(ParamName, *Index, static_cast<FMultilayerProjectorLayer>(Layer)); // Update enabled layer.
+	}
+	
+	// Enable or disable virtual layer.
+	if (Layer.bEnabled && !bEnabled)
+	{
+		const int32 NewIndex = CalculateVirtualLayerIndex(Id, VirtualLayersOrder[Id]);
+
+		Descriptor.CreateLayer(ParamName, NewIndex);
+		UpdateMappingVirtualLayerEnabled(Id, NewIndex);
+
+		Descriptor.UpdateLayer(ParamName, NewIndex, static_cast<FMultilayerProjectorLayer>(Layer));
+		
+		DisableVirtualLayers.Remove(Id);
+	}
+	else if (!Layer.bEnabled && bEnabled)
+	{
+		Descriptor.RemoveLayerAt(ParamName, *Index);
+		UpdateMappingVirtualLayerDisabled(Id, *Index);
+		
+		DisableVirtualLayers.Add(Id, Layer);
+	}
+}
+
+
+FMultilayerProjector::FMultilayerProjector(const FName& InParamName) :
+	ParamName(InParamName)
+{
+}
+
+
+int32 FMultilayerProjector::CalculateVirtualLayerIndex(const FName& Id, const int32 InsertOrder) const
+{
+	int32 LayerBeforeIndex = -1;
+	int32 LayerBeforeOrder = -1;
+	
+	for (const TTuple<FName, int>& MappingTuple : VirtualLayersMapping) // Find closest smallest layer.
+	{
+		if (MappingTuple.Value != VIRTUAL_LAYER_DISABLED && MappingTuple.Key != Id)
+		{
+			const int32 LayerOrder = VirtualLayersOrder[MappingTuple.Key];
+			if (LayerOrder <= InsertOrder)
+			{
+				if ((LayerOrder > LayerBeforeOrder) ||
+					(LayerOrder == LayerBeforeOrder && MappingTuple.Value > LayerBeforeIndex))
+				{
+					LayerBeforeIndex = MappingTuple.Value;
+					LayerBeforeOrder = LayerOrder;
+				}
+			}
+		}
+	}
+	
+	return LayerBeforeIndex + 1;
+}
+
+
+void FMultilayerProjector::UpdateMappingVirtualLayerEnabled(const FName& Id, const int32 Index)
+{
+	for (TTuple<FName, int>& Tuple : VirtualLayersMapping)
+	{
+		if (Tuple.Key == Id)
+		{
+			Tuple.Value = Index;
+		}
+		else if (Tuple.Value >= Index) // Update following Layers.
+		{
+			++Tuple.Value;
+		}
+	}
+}
+
+
+void FMultilayerProjector::UpdateMappingVirtualLayerDisabled(const FName& Id, const int32 Index)
+{
+	for (TTuple<FName, int>& Tuple : VirtualLayersMapping)
+	{
+		if (Tuple.Key == Id)
+    	{
+			Tuple.Value = VIRTUAL_LAYER_DISABLED;
+    	}
+		else if (Tuple.Value > Index) // Update following Layers.
+		{
+			--Tuple.Value;
+		}
+	}
+}
+
+
+uint32 GetTypeHash(const FMultilayerProjector& Key)
+{
+	uint32 Hash = GetTypeHash(Key.ParamName);
+
+	for (const TTuple<FName, int32>& Pair : Key.VirtualLayersMapping)
+	{
+		Hash = HashCombine(Hash, GetTypeHash(Pair.Key));
+		Hash = HashCombine(Hash, GetTypeHash(Pair.Value));
+	}
+	
+	for (const TTuple<FName, int32>& Pair : Key.VirtualLayersOrder)
+	{
+		Hash = HashCombine(Hash, GetTypeHash(Pair.Key));
+		Hash = HashCombine(Hash, GetTypeHash(Pair.Value));
+	}
+
+	for (const TTuple<FName, FMultilayerProjectorLayer>& Pair : Key.DisableVirtualLayers)
+	{
+		Hash = HashCombine(Hash, GetTypeHash(Pair.Key));
+		Hash = HashCombine(Hash, GetTypeHash(Pair.Value));
+	}
+	
+	return Hash;
+}
+
diff --git a/Engine/Plugins/Mutable/Source/CustomizableObject/Private/MuCO/MutableMeshBufferUtils.cpp b/Engine/Plugins/Mutable/Source/CustomizableObject/Private/MuCO/MutableMeshBufferUtils.cpp
new file mode 100644
--- /dev/null
+++ b/Engine/Plugins/Mutable/Source/CustomizableObject/Private/MuCO/MutableMeshBufferUtils.cpp
@@ -0,0 +1,245 @@
+// Copyright Epic Games, Inc. All Rights Reserved.
+
+#include "MuCO/MutableMeshBufferUtils.h"
+
+#include "MuR/MeshBufferSet.h"
+#include "Rendering/PositionVertexBuffer.h"
+#include "Rendering/StaticMeshVertexBuffer.h"
+
+
+void MutableMeshBufferUtils::SetupVertexPositionsBuffer(const int32& InCurrentVertexBuffer, mu::FMeshBufferSet& OutTargetVertexBuffers)
+{
+	using namespace mu;
+	const int32 ElementSize = sizeof(FPositionVertex);
+	constexpr int32 ChannelCount = 1;
+	const EMeshBufferSemantic Semantics[ChannelCount] = {MBS_POSITION};
+	const int32 SemanticIndices[ChannelCount] = {0};
+	const EMeshBufferFormat Formats[ChannelCount] = {MBF_FLOAT32};
+	const int32 Components[ChannelCount] = {3};
+	const int32 Offsets[ChannelCount] =
+	{
+		STRUCT_OFFSET(FPositionVertex, Position)
+	};
+
+	OutTargetVertexBuffers.SetBuffer(InCurrentVertexBuffer, ElementSize, ChannelCount, Semantics, SemanticIndices,
+									  Formats, Components, Offsets);
+}
+
+
+void MutableMeshBufferUtils::SetupTangentBuffer(const int32& InCurrentVertexBuffer,
+                                                mu::FMeshBufferSet& OutTargetVertexBuffers)
+{
+	// \todo: support for high precision?
+	typedef TStaticMeshVertexTangentDatum<TStaticMeshVertexTangentTypeSelector<
+		EStaticMeshVertexTangentBasisType::Default>::TangentTypeT> TangentType;
+
+	using namespace mu;
+	const int32 ElementSize = sizeof(TangentType);
+	constexpr int32 ChannelCount = 2;
+	const EMeshBufferSemantic Semantics[ChannelCount] = {MBS_TANGENT, MBS_NORMAL};
+	const int32 SemanticIndices[ChannelCount] = {0, 0};
+	const EMeshBufferFormat Formats[ChannelCount] = {MBF_PACKEDDIRS8, MBF_PACKEDDIRS8_W_TANGENTSIGN};
+	const int32 Components[ChannelCount] = {4, 4};
+	const int32 Offsets[ChannelCount] =
+	{
+		STRUCT_OFFSET(TangentType, TangentX),
+		STRUCT_OFFSET(TangentType, TangentZ)
+	};
+
+	OutTargetVertexBuffers.SetBuffer(InCurrentVertexBuffer, ElementSize, ChannelCount, Semantics, SemanticIndices,
+	                                  Formats, Components, Offsets);
+}
+
+
+void MutableMeshBufferUtils::SetupTexCoordinatesBuffer(const int32& InCurrentVertexBuffer, const int32& InChannelCount,
+                                                       mu::FMeshBufferSet& OutTargetVertexBuffers,const int32* InTextureSemanticIndicesOverride /* = nullptr */)
+{
+	// \todo: support for half precision?
+	typedef TStaticMeshVertexUVsDatum<TStaticMeshVertexUVsTypeSelector<
+		EStaticMeshVertexUVType::HighPrecision>::UVsTypeT> UVType;
+
+	using namespace mu;
+	const int32 ElementSize = sizeof(UVType) * InChannelCount;
+	constexpr int32 MaxChannelCount = MaxTexCordChannelCount;
+	const EMeshBufferSemantic Semantics[MaxChannelCount] = {
+		MBS_TEXCOORDS, MBS_TEXCOORDS, MBS_TEXCOORDS, MBS_TEXCOORDS
+	};
+	const EMeshBufferFormat Formats[MaxChannelCount] = {MBF_FLOAT32, MBF_FLOAT32, MBF_FLOAT32, MBF_FLOAT32};
+	const int32 Components[MaxChannelCount] = {2, 2, 2, 2};
+	const int32 Offsets[MaxChannelCount] =
+	{
+		STRUCT_OFFSET(UVType, UVs) + 0 * sizeof(UVType),
+		STRUCT_OFFSET(UVType, UVs) + 1 * sizeof(UVType),
+		STRUCT_OFFSET(UVType, UVs) + 2 * sizeof(UVType),
+		STRUCT_OFFSET(UVType, UVs) + 3 * sizeof(UVType),
+	};
+
+	if (InTextureSemanticIndicesOverride)
+	{
+		OutTargetVertexBuffers.SetBuffer(InCurrentVertexBuffer, ElementSize, InChannelCount, Semantics,
+		                                  InTextureSemanticIndicesOverride, Formats, Components, Offsets);
+	}
+	else
+	{
+		const int32 SemanticIndices[MaxChannelCount] = {0, 1, 2, 3};
+		OutTargetVertexBuffers.SetBuffer(InCurrentVertexBuffer, ElementSize, InChannelCount, Semantics,
+		                                  SemanticIndices, Formats, Components, Offsets);
+	}
+}
+
+
+void MutableMeshBufferUtils::SetupSkinBuffer(const int32& InCurrentVertexBuffer,
+	const int32& MaxBoneIndexTypeSizeBytes,
+	const int32& MaxBoneWeightTypeSizeBytes,
+	const int32& MaxNumBonesPerVertex,
+	mu::FMeshBufferSet& OutTargetVertexBuffers)
+{
+	using namespace mu;
+	const int32 ElementSize = (MaxBoneWeightTypeSizeBytes + MaxBoneIndexTypeSizeBytes) * MaxNumBonesPerVertex;
+	constexpr int32 ChannelCount = 2;
+	const EMeshBufferSemantic Semantics[ChannelCount] = {MBS_BONEINDICES, MBS_BONEWEIGHTS};
+	const int32 SemanticIndices[ChannelCount] = {0, 0};
+
+	EMeshBufferFormat Formats[ChannelCount] = {MBF_UINT8, MBF_NUINT8};
+	switch (MaxBoneIndexTypeSizeBytes)
+	{
+	case 0: // Fallback to something in this case.
+	case 1: Formats[0] = mu::MBF_UINT8;
+		break;
+	case 2: Formats[0] = mu::MBF_UINT16;
+		break;
+	case 4: Formats[0] = mu::MBF_UINT32;
+		break;
+	default:
+		// unsupported bone index type
+		check(false);
+		Formats[0] = mu::MBF_NONE;
+		break;
+	}
+
+	switch (MaxBoneWeightTypeSizeBytes)
+	{
+	case 0: // Fallback to something in this case.
+	case 1: Formats[1] = mu::MBF_NUINT8;
+		break;
+	case 2: Formats[1] = mu::MBF_NUINT16;
+		break;
+	case 4: Formats[1] = mu::MBF_NUINT32;
+		break;
+	default:
+		// unsupported bone weight type
+		check(false);
+		Formats[1] = mu::MBF_NONE;
+		break;
+	}
+
+	int32 Components[ChannelCount];
+	Components[0] = Components[1] = MaxNumBonesPerVertex;
+
+	int32 Offsets[ChannelCount];
+	Offsets[0] = 0;
+	Offsets[1] = MaxBoneIndexTypeSizeBytes * MaxNumBonesPerVertex;
+
+	OutTargetVertexBuffers.SetBuffer(InCurrentVertexBuffer, ElementSize, ChannelCount, Semantics, SemanticIndices,
+	                                  Formats, Components, Offsets);
+}
+
+
+void MutableMeshBufferUtils::SetupVertexColorBuffer(const int32& InCurrentVertexBuffer,
+                                                    mu::FMeshBufferSet& OutTargetVertexBuffers)
+{
+	using namespace mu;
+	const int32 ElementSize = sizeof(FColor);
+	constexpr int32 ChannelCount = 1;
+	const EMeshBufferSemantic Semantics[ChannelCount] = {MBS_COLOUR};
+	const int32 SemanticIndices[ChannelCount] = {0};
+	const EMeshBufferFormat Formats[ChannelCount] = {MBF_NUINT8};
+	const int32 Components[ChannelCount] = {4};
+	const int32 Offsets[ChannelCount] = {0};
+	check(ElementSize == 4);
+
+	OutTargetVertexBuffers.SetBuffer(InCurrentVertexBuffer, ElementSize, ChannelCount, Semantics, SemanticIndices,
+	                                  Formats, Components, Offsets);
+}
+
+
+void MutableMeshBufferUtils::SetupIndexBuffer(mu::FMeshBufferSet& OutTargetIndexBuffers)
+{
+	OutTargetIndexBuffers.SetBufferCount(1);
+
+	using namespace mu;
+	const int32 ElementSize = sizeof(uint32);
+	//SkeletalMesh->GetImportedResource()->LODModels[LOD].MultiSizeIndexContainer.GetDataTypeSize();
+	constexpr int32 ChannelCount = 1;
+	const EMeshBufferSemantic Semantics[ChannelCount] = {MBS_VERTEXINDEX};
+	const int32 SemanticIndices[ChannelCount] = {0};
+	// We force 32 bit indices, since merging meshes may create vertex buffers bigger than the initial mesh
+	// and for now the mutable runtime doesn't handle it.
+	// \TODO: go back to 16-bit indices when possible.
+	EMeshBufferFormat Formats[ChannelCount] = {MBF_UINT32};
+	const int32 Components[ChannelCount] = {1};
+	const int32 Offsets[ChannelCount] = {0};
+
+	OutTargetIndexBuffers.SetBuffer(0, ElementSize, ChannelCount, Semantics, SemanticIndices, Formats, Components,
+	                                 Offsets);
+}
+
+
+void MutableMeshBufferUtils::SetupSkinWeightProfileBuffer(const int32& InCurrentVertexBuffer,
+	const int32& MaxBoneIndexTypeSizeBytes,
+	const int32& MaxBoneWeightTypeSizeBytes,
+	const int32& MaxNumBonesPerVertex,
+	const int32 SemanticsIndex,
+	mu::FMeshBufferSet& OutTargetVertexBuffers)
+{
+	using namespace mu;
+	const int32 ElementSize = sizeof(int32) + (MaxBoneIndexTypeSizeBytes + 1) * MaxNumBonesPerVertex;
+	constexpr int32 ChannelCount = 3;
+	const EMeshBufferSemantic Semantics[ChannelCount] = { MBS_ALTSKINWEIGHT, MBS_BONEINDICES, MBS_BONEWEIGHTS };
+	const int32 SemanticIndices[ChannelCount] = { SemanticsIndex, SemanticsIndex, SemanticsIndex };
+
+	EMeshBufferFormat Formats[ChannelCount] = { MBF_INT32, MBF_UINT8, MBF_NUINT8 };
+	switch (MaxBoneIndexTypeSizeBytes)
+	{
+	case 1: Formats[1] = mu::MBF_UINT8;
+		break;
+	case 2: Formats[1] = mu::MBF_UINT16;
+		break;
+	case 4: Formats[1] = mu::MBF_UINT32;
+		break;
+	default:
+		// unsupported bone index type
+		check(false);
+		Formats[1] = mu::MBF_NONE;
+		break;
+	}
+
+	switch (MaxBoneWeightTypeSizeBytes)
+	{
+	case 1: Formats[2] = mu::MBF_NUINT8;
+		break;
+	case 2:
+	{
+		unimplemented()
+		Formats[2] = mu::MBF_NUINT16;
+		break;
+	}
+	case 4: Formats[2] = mu::MBF_NUINT32;
+		break;
+	default:
+		// unsupported bone weight type
+		check(false);
+		Formats[2] = mu::MBF_NONE;
+		break;
+	}
+
+	int32 Components[ChannelCount] = { 1, MaxNumBonesPerVertex, MaxNumBonesPerVertex };
+
+	int32 Offsets[ChannelCount];
+	Offsets[0] = 0;
+	Offsets[1] = sizeof(int32);
+	Offsets[2] = sizeof(int32) + MaxBoneIndexTypeSizeBytes * MaxNumBonesPerVertex;
+
+	OutTargetVertexBuffers.SetBuffer(InCurrentVertexBuffer, ElementSize, ChannelCount, Semantics, SemanticIndices,
+		Formats, Components, Offsets);
+}
diff --git a/Engine/Plugins/Mutable/Source/CustomizableObject/Public/MuCO/CustomizableInstanceLODManagement.h b/Engine/Plugins/Mutable/Source/CustomizableObject/Public/MuCO/CustomizableInstanceLODManagement.h
new file mode 100644
--- /dev/null
+++ b/Engine/Plugins/Mutable/Source/CustomizableObject/Public/MuCO/CustomizableInstanceLODManagement.h
@@ -0,0 +1,68 @@
+// Copyright Epic Games, Inc. All Rights Reserved.
+
+#pragma once
+#include "GameFramework/Actor.h"
+
+#include "CustomizableInstanceLODManagement.generated.h"
+
+typedef TMap<const class UCustomizableObjectInstance*, class FMutableUpdateCandidate> FMutableInstanceUpdateMap;
+
+// This is an abstract base class, override it to create a new Instance LOD management system and register with UCustomizableObjectSystem::SetInstanceLODManagement
+UCLASS(Abstract, Blueprintable, BlueprintType)
+class CUSTOMIZABLEOBJECT_API UCustomizableInstanceLODManagementBase : public UObject
+{
+	GENERATED_BODY()
+public:
+	UCustomizableInstanceLODManagementBase() : UObject() {};
+	virtual ~UCustomizableInstanceLODManagementBase() {};
+
+	// WARNING! The following methods must be overriden in derived classes
+	virtual void UpdateInstanceDistsAndLODs(FMutableInstanceUpdateMap& InOutRequestedUpdates) PURE_VIRTUAL(UCustomizableInstanceLODManagementBase::UpdateInstanceDistsAndLODs, return;);
+	virtual int32 GetNumGeneratedInstancesLimitFullLODs() const PURE_VIRTUAL(UCustomizableInstanceLODManagementBase::GetNumGeneratedInstancesLimitFullLODs, return 0;);
+	virtual int32 GetNumGeneratedInstancesLimitLOD1() const PURE_VIRTUAL(UCustomizableInstanceLODManagementBase::GetNumGeneratedInstancesLimitLOD1, return 0;);
+	virtual int32 GetNumGeneratedInstancesLimitLOD2() const PURE_VIRTUAL(UCustomizableInstanceLODManagementBase::GetNumGeneratedInstancesLimitLOD2, return 0;);
+	virtual float GetOnlyUpdateCloseCustomizableObjectsDist() const PURE_VIRTUAL(UCustomizableInstanceLODManagementBase::GetOnlyUpdateCloseCustomizableObjectsDist, return 0.f;);
+	virtual bool IsOnlyUpdateCloseCustomizableObjectsEnabled() const PURE_VIRTUAL(UCustomizableInstanceLODManagementBase::IsOnlyUpdateCloseCustomizableObjectsEnabled, return false;);
+	virtual bool IsOnlyGenerateRequestedLODLevelsEnabled() const PURE_VIRTUAL(UCustomizableInstanceLODManagementBase::IsOnlyGenerateRequestedLODLevelsEnabled, return false;);
+};
+
+
+UCLASS(Blueprintable, BlueprintType)
+class CUSTOMIZABLEOBJECT_API UCustomizableInstanceLODManagement : public UCustomizableInstanceLODManagementBase
+{
+	GENERATED_BODY()
+public:
+	UCustomizableInstanceLODManagement();
+	virtual ~UCustomizableInstanceLODManagement();
+
+	virtual void UpdateInstanceDistsAndLODs(FMutableInstanceUpdateMap& InOutRequestedUpdates) override;
+
+	virtual int32 GetNumGeneratedInstancesLimitFullLODs() const override;
+	virtual int32 GetNumGeneratedInstancesLimitLOD1() const override;
+	virtual int32 GetNumGeneratedInstancesLimitLOD2() const override;
+
+	// Used to define the actors that will be considered the centers of LOD distance calculations. If left empty it will revert to the default, the first player controller actor
+	void AddViewCenter(const AActor* const InCentralActor) { ViewCenters.Add(TWeakObjectPtr<const AActor>(InCentralActor)); };
+	void RemoveViewCenter(const AActor* const InCentralActor) { ViewCenters.Remove(TWeakObjectPtr<const AActor>(InCentralActor)); };
+	void ClearViewCenters() { ViewCenters.Empty(); }
+
+	// Sets how many of the nearest instances to the player will have updates with priority over LOD changes, only works if LODs are enabled with b.NumGeneratedInstancesLimitLODs
+	void SetNumberOfPriorityUpdateInstances(int32 NumPriorityUpdateInstances);
+	int32 GetNumberOfPriorityUpdateInstances() const;
+
+	virtual void SetCustomizableObjectsUpdateDistance(float Distance);
+	virtual float GetOnlyUpdateCloseCustomizableObjectsDist() const override;
+
+	virtual bool IsOnlyUpdateCloseCustomizableObjectsEnabled() const override;
+	virtual bool IsOnlyGenerateRequestedLODLevelsEnabled() const override { return true; };
+
+private:
+	// The list of actors that define a view radius (CloseDist) around them used for LOD management
+	TSet<TWeakObjectPtr<const AActor>> ViewCenters;
+	
+	// Sets how many of the nearest instances to the player will have updates with priority over LOD changes
+	int32 NumPriorityUpdateInstances = 1;
+
+	float CloseCustomizableObjectsDist;	
+};
+
