[mutable] Moved the Mutable plugin out of Experimental status into Beta.

diff --git a/Engine/Plugins/Mutable/Source/CustomizableObjectEditor/Private/MuCOE/CustomizableObjectCompiler.h b/Engine/Plugins/Mutable/Source/CustomizableObjectEditor/Private/MuCOE/CustomizableObjectCompiler.h
new file mode 100644
--- /dev/null
+++ b/Engine/Plugins/Mutable/Source/CustomizableObjectEditor/Private/MuCOE/CustomizableObjectCompiler.h
@@ -0,0 +1,194 @@
+// Copyright Epic Games, Inc. All Rights Reserved.
+
+#pragma once
+
+#include "AssetRegistry/AssetData.h"
+#include "MuCO/CustomizableObjectPrivate.h"
+#include "MuCO/CustomizableObjectCompilerTypes.h"
+#include "MuCO/UnrealToMutableTextureConversionUtils.h"
+#include "MuCOE/CompilationMessageCache.h"
+#include "MuCOE/GenerateMutableSource/GenerateMutableSource.h"
+#include "MuCOE/CustomizableObjectEditorLogger.h"
+#include "UObject/GCObject.h"
+#include "TickableEditorObject.h"
+
+#include "Framework/Notifications/NotificationManager.h"
+
+class FCustomizableObjectCompileRunnable;
+class FCustomizableObjectSaveDDRunnable;
+class FReferenceCollector;
+class FRunnableThread;
+class FText;
+class UCustomizableObject;
+class UCustomizableObjectNode;
+
+class FCustomizableObjectCompiler : public FTickableEditorObject, public FTickableCookObject, public FGCObject
+{
+public:
+	virtual ~FCustomizableObjectCompiler() override {}
+
+	/** Check for pending compilation process. Returns true if an object has been updated. */
+	bool Tick(bool bBlocking = false);
+	int32 GetNumRemainingWork() const;
+
+	// FTickableGameObject interface
+	virtual ETickableTickType GetTickableTickType() const override { return ETickableTickType::Conditional; };
+	virtual bool IsTickable() const override;
+	virtual void Tick( float InDeltaTime ) override;
+	virtual TStatId GetStatId() const override;
+
+	// FTickableCookObject interface
+	virtual void TickCook(float DeltaTime, bool bCookCompete) override;
+
+
+	/** Generate the Mutable Graph from the Unreal Graph. */
+	mu::NodePtr Export(UCustomizableObject* Object, const FCompilationOptions& Options, TArray<TSoftObjectPtr<const UTexture>>& OutRuntimeReferencedTextures, TArray<FMutableSourceTextureData>& OutCompilerReferencedTextures);
+
+	void CompilerLog(const FText& Message, const TArray<const UObject*>& UObject, const EMessageSeverity::Type MessageSeverity = EMessageSeverity::Warning, const bool bAddBaseObjectInfo = true, const ELoggerSpamBin SpamBin = ELoggerSpamBin::ShowAll);
+	void CompilerLog(const FText& Message, const UObject* Context = nullptr, const EMessageSeverity::Type MessageSeverity = EMessageSeverity::Warning, const bool bAddBaseObjectInfo = true, const ELoggerSpamBin SpamBin = ELoggerSpamBin::ShowAll);
+	void NotifyCompilationErrors() const;
+
+	void FinishCompilationTask();
+	void FinishSavingDerivedDataTask();
+
+	void ForceFinishCompilation();
+	void ClearCompileRequests();
+
+	void AddCompileNotification(const FText& CompilationStep) const;
+	static void RemoveCompileNotification();
+
+	/** Load required assets and compile.
+	 *
+	 * Loads assets which reference Object's package asynchronously before calling ProcessChildObjectsRecursively. */
+	void Compile(const TSharedRef<FCompilationRequest>& InCompileRequest);
+	void Compile(const TArray<TSharedRef<FCompilationRequest>>& InCompileRequests);
+	
+	bool IsRequestQueued(const TSharedRef<FCompilationRequest>& InCompileRequest) const;
+
+	/** FSerializableObject interface */
+	virtual void AddReferencedObjects(FReferenceCollector& Collector) override;
+	virtual FString GetReferencerName() const override
+	{
+		return TEXT("FCustomizableObjectCompiler");
+	}
+
+	/** Simply add CO elements from ArrayAssetData to ArrayGCProtect when they've been loaded from ArrayAssetData */
+	void UpdateArrayGCProtect();
+
+private:
+
+	// Object containing all error and warning logs raised during compilation.
+	FCompilationMessageCache CompilationLogsContainer;
+
+	void SetCompilationState(ECompilationStatePrivate State, ECompilationResultPrivate Result) const;
+	
+	void CompileInternal(bool bAsync = false);
+
+	void CompleteRequest(ECompilationStatePrivate State, ECompilationResultPrivate Result);
+	bool TryPopCompileRequest();
+
+	bool TryLoadCompiledDataFromDDC(UCustomizableObject& CustomizableObject);
+
+	void PreloadingReferencerAssetsCallback(bool bAsync);
+	
+	void ProcessChildObjectsRecursively(UCustomizableObject* Object, FMutableGraphGenerationContext &GenerationContext);
+	
+	// Will output to Mutable Log the warning and error messages generated during the CO compilation
+	// and update the values of NumWarnings and NumErrors
+	void UpdateCompilerLogData();
+
+	/** If duplicated elements are found in each entry of ParameterNamesMap, a warning for
+	the parameters with repeated name will be generated */
+	void DisplayParameterWarning(FMutableGraphGenerationContext& GenerationContext);
+	
+	/** If duplicated node ids are found, usually due to duplicating CustomizableObjects Assets, a warning
+	for the nodes with repeated ids will be generated */
+	void DisplayDuplicatedNodeIdsWarning(FMutableGraphGenerationContext& GenerationContext);
+	
+	/** Display warnings for unnamed node objects */
+	void DisplayUnnamedNodeObjectWarning(FMutableGraphGenerationContext& GenerationContext);
+	
+	/** Display a warning for each node contains an orphan pin. */
+	void DisplayOrphanNodesWarning(FMutableGraphGenerationContext& GenerationContext);
+	
+	mu::NodeObjectPtr GenerateMutableRoot(UCustomizableObject* Object, FMutableGraphGenerationContext& GenerationContext, FText& ErrorMessage, bool& bOutIsRootObject);
+
+
+	/** Launches the compile task in another thread when compiling a CO in the editor
+	* @param bShowNotification [in] whether to show the compiling CO notification or not
+	* @return nothing */
+	void LaunchMutableCompile();
+
+	/** Launches the save derived data task in another thread after compiling a CO in the
+	* editor
+	* @param bShowNotification [in] whether to show the saving DD notification or not
+	* @return nothing */
+	void SaveCODerivedData();
+
+	/** Add to ArrayAssetData the FAssetData information of all referencers of static class type UCustomizableObject::StaticClass()
+	* that reference the package given by the PathName parameter
+	* @param PathName            [in]  path to the CO to be analyzed (for instance, CO->GetOuter()->GetPathName())
+	* @param ArrayReferenceNames [out] array with the package names which are referenced by the CO with PathName given as parameter
+	* @return nothing */
+	void AddCachedReferencers(const FName& PathName, TArray<FName>& ArrayReferenceNames);
+
+	/** Just used to clean ArrayAssetData
+	* @return nothing */
+	void CleanCachedReferencers();
+
+	/** Will test if package path given by PackageName parameter is one of ArrayAssetData's elements FAssetData::PackageName value
+	* @param PackageName [in] package name to test
+	* @return true if cached, false otherwise */
+	bool IsCachedInAssetData(const FString& PackageName);
+
+	/** Find FAssetData ArrayAssetData with PackageName given by parameter
+	* @param PackageName [in] package name to find in ArrayAssetData
+	* @return pointer to element if any found, nullptr otherwise */
+	FAssetData* GetCachedAssetData(const FString& PackageName);
+
+	ECompilationResultPrivate GetCompilationResult() const;
+
+	/** Pointer to the Asynchronous Preloading process call back */
+	TSharedPtr<FStreamableHandle> AsynchronousStreamableHandlePtr;
+	TArray<FSoftObjectPath> ArrayAssetToStream;
+
+	/** Compile task and thread. */
+	TSharedPtr<FCustomizableObjectCompileRunnable> CompileTask;
+	TSharedPtr<FRunnableThread> CompileThread;
+
+	/** SaveDD task and thread. */
+	TSharedPtr<FCustomizableObjectSaveDDRunnable> SaveDDTask;
+	TSharedPtr<FRunnableThread> SaveDDThread;
+
+	/** Array where to put the names of the already processed child in ProcessChildObjectsRecursively */
+	TArray<FName> ArrayAlreadyProcessedChild;
+
+	/** Array with all the packages used to compile current Customizable Object */
+	TArray<FAssetData> ArrayAssetData;
+
+	/** Array used to protect from garbage collection those COs loaded asynchronously */
+	TArray<TObjectPtr<UCustomizableObject>> ArrayGCProtect;
+
+	TSharedPtr<mu::Model, ESPMode::ThreadSafe> Model;
+	TSharedPtr<FModelStreamableBulkData> ModelStreamableBulkData;
+
+	// Protected from GC with FCustomizableObjectCompiler::AddReferencedObjects
+	TObjectPtr<UCustomizableObject> CurrentObject = nullptr;
+
+	FCompilationOptions CurrentOptions;
+
+	/** Current Compilation request. */
+	TSharedPtr<FCompilationRequest> CurrentRequest;
+
+	/** Pending requests. */
+	TArray<TSharedRef<FCompilationRequest>> CompileRequests;
+
+	uint32 NumCompilationRequests = 0;
+
+	/** Compilation progress bar handle */
+	FProgressNotificationHandle CompileNotificationHandle;
+
+	/** Compilation start time in seconds. */
+	double CompilationStartTime = 0;
+
+};
diff --git a/Engine/Plugins/Mutable/Source/CustomizableObjectEditor/Private/MuCOE/CustomizableObjectDetails.cpp b/Engine/Plugins/Mutable/Source/CustomizableObjectEditor/Private/MuCOE/CustomizableObjectDetails.cpp
new file mode 100644
--- /dev/null
+++ b/Engine/Plugins/Mutable/Source/CustomizableObjectEditor/Private/MuCOE/CustomizableObjectDetails.cpp
@@ -0,0 +1,386 @@
+// Copyright Epic Games, Inc. All Rights Reserved.
+
+#include "MuCOE/CustomizableObjectDetails.h"
+
+#include "DetailLayoutBuilder.h"
+#include "DetailWidgetRow.h"
+#include "IDetailsView.h"
+#include "MuCO/CustomizableObject.h"
+#include "MuCO/CustomizableObjectPrivate.h"
+#include "MuCOE/GraphTraversal.h"
+#include "MuCOE/Nodes/CustomizableObjectNodeObject.h"
+#include "MuR/Model.h"
+#include "MuR/Parameters.h"
+#include "MuR/Ptr.h"
+
+class ITableRow;
+
+
+#define LOCTEXT_NAMESPACE "CustomizableObjectDetails"
+
+TSharedRef<IDetailCustomization> FCustomizableObjectDetails::MakeInstance()
+{
+	return MakeShareable( new FCustomizableObjectDetails );
+}
+
+
+class FStateDetailsNode
+{
+public:
+
+	FStateDetailsNode()
+	{
+		StateIndex = -1;
+		ParameterIndex = -1;
+	}
+
+	TSharedPtr<mu::Model, ESPMode::ThreadSafe> Model;
+	int StateIndex;
+	int ParameterIndex;
+
+	TSharedPtr<FString> GetName()
+	{
+		FString res;
+
+		if (ParameterIndex<0)
+		{
+			res = Model->GetStateName( StateIndex );
+		}
+		else
+		{
+			mu::ParametersPtr TempParams = mu::Model::NewParameters(Model);
+			res = TempParams->GetName( ParameterIndex );
+		}
+
+		return MakeShareable( new FString( res ) );
+	}
+
+
+	TArray<TSharedPtr<FStateDetailsNode>> GetChildrenList()
+	{
+		TArray<TSharedPtr<FStateDetailsNode>> ChildrenList;
+
+		// Is it a state node?
+		if (ParameterIndex<0)
+		{
+			int ParameterCount = Model->GetStateParameterCount( StateIndex );
+			for ( int i=0; i<ParameterCount; ++i )
+			{
+				TSharedPtr<FStateDetailsNode> SlateDetailsNode = MakeShareable( new FStateDetailsNode );
+
+				SlateDetailsNode->Model = Model;
+				SlateDetailsNode->StateIndex = StateIndex;
+				SlateDetailsNode->ParameterIndex = Model->GetStateParameterIndex( StateIndex, i );
+
+				ChildrenList.Add( SlateDetailsNode );
+			}
+		}
+
+		return ChildrenList;
+	}
+};
+
+
+void FCustomizableObjectDetails::CustomizeDetails( IDetailLayoutBuilder& DetailBuilder )
+{
+	const IDetailsView* DetailsView = DetailBuilder.GetDetailsView();
+	if (DetailsView->GetSelectedObjects().Num())
+	{
+		CustomizableObject = Cast<UCustomizableObject>(DetailsView->GetSelectedObjects()[0].Get());
+	}
+
+	IDetailCategoryBuilder& CustomizableObjectCategory = DetailBuilder.EditCategory("CustomizableObject");
+	IDetailCategoryBuilder& StatesCategory = DetailBuilder.EditCategory( "States" );
+	IDetailCategoryBuilder& VersioningCategory = DetailBuilder.EditCategory( "Versioning" );
+	//StatesCategory.CategoryIcon( "ActorClassIcon.CustomizableObject" );
+
+	DetailBuilder.HideProperty("States");
+	DetailBuilder.HideProperty("CustomizableObjectClassTags");
+	DetailBuilder.HideProperty("PopulationClassTags");
+
+	// Make the tree get automatically updated each time we compile the host CO
+	if (CustomizableObject && !CustomizableObject->GetPrivate()->PostCompileDelegate.IsBoundToObject(this))
+	{
+		CustomizableObject->GetPrivate()->PostCompileDelegate.AddSP(this, &FCustomizableObjectDetails::UpdateTree);
+	}
+	
+	// Cache the states defined in the CO
+	UpdateTree();
+	
+	StatesCategory.AddCustomRow( LOCTEXT("FCustomizableObjectDetails States", "States") )
+	[
+		SNew( SVerticalBox )
+		+ SVerticalBox::Slot()
+		.Padding( 2.0f )
+		[
+			SAssignNew(StatesTree, STreeView<TSharedPtr<FStateDetailsNode>>)
+			.SelectionMode(ESelectionMode::Single)
+			.TreeItemsSource( &RootTreeItems )
+			// Called to child items for any given parent item
+			.OnGetChildren( this, &FCustomizableObjectDetails::OnGetChildrenForStateTree )
+			// Generates the actual widget for a tree item
+			.OnGenerateRow( this, &FCustomizableObjectDetails::OnGenerateRowForStateTree )
+			.HeaderRow
+			(
+				SNew(SHeaderRow)
+				.Visibility(EVisibility::Collapsed)
+				+ SHeaderRow::Column(TEXT("State"))
+				.DefaultLabel(NSLOCTEXT("CustomizableObjectDetails","State","State"))
+			)
+		]
+	];
+
+	StatesTree->SetIsRightClickScrollingEnabled(false);
+
+	TSharedRef<IPropertyHandle> VersionBridgeProperty = DetailBuilder.GetProperty("VersionBridge");
+
+	if (VersionBridgeProperty->IsValidHandle() && CustomizableObject)
+	{
+		if (CustomizableObject->IsChildObject())
+		{
+			VersionBridgeProperty->MarkHiddenByCustomization();
+		}
+		else
+		{
+			VersionBridgeProperty->MarkResetToDefaultCustomized();
+		}
+	}
+
+	TSharedRef<IPropertyHandle> VersionStructProperty = DetailBuilder.GetProperty("VersionStruct");
+
+	if (VersionStructProperty->IsValidHandle() && CustomizableObject)
+	{
+		if (CustomizableObject->IsChildObject())
+		{
+			VersionStructProperty->MarkResetToDefaultCustomized();
+		}
+		else
+		{
+			VersionStructProperty->MarkHiddenByCustomization();
+		}
+	}
+}
+
+
+void FCustomizableObjectDetails::UpdateTree()
+{
+	UCustomizableObject* RootObject = GetRootObject(CustomizableObject);
+
+	if (RootObject && RootObject->GetPrivate()->GetModel() )
+	{
+		RootTreeItems.SetNumUninitialized(0);
+		const uint32 NumElements = RootObject->GetPrivate()->GetModel()->GetStateCount();
+		for ( uint32 i=0; i<NumElements; ++i )
+		{
+			TSharedPtr<FStateDetailsNode> SlateDetailsNode = MakeShareable( new FStateDetailsNode );
+
+			SlateDetailsNode->Model = RootObject->GetPrivate()->GetModel();
+			SlateDetailsNode->StateIndex = i;
+
+			RootTreeItems.Add( SlateDetailsNode );
+		}
+	}
+
+	if (StatesTree)
+	{
+		StatesTree->RequestTreeRefresh();
+	}
+}
+
+
+/** The item used for visualizing the class in the tree. */
+class SStateItem : public STableRow< TSharedPtr<FString> >
+{
+public:
+	
+	SLATE_BEGIN_ARGS( SStateItem )
+		: _StateName()
+		, _OwnerTableView()
+		, _HighlightText(NULL)
+		, _TextColor(FLinearColor(1.0f, 1.0f, 1.0f, 1.0f))
+		{}
+
+		/** The classname this item contains. */
+		SLATE_ARGUMENT( TSharedPtr<FString>, StateName )
+		/** The table that owns this item. */
+		SLATE_ARGUMENT( TSharedPtr<STableViewBase>, OwnerTableView )
+		/** The text this item should highlight, if any. */
+		SLATE_ARGUMENT( const FString*, HighlightText )
+		/** The color text this item will use. */
+		SLATE_ARGUMENT( FSlateColor, TextColor )
+		/** The node this item is associated with. */
+		SLATE_ARGUMENT( TSharedPtr<FStateDetailsNode>, AssociatedNode)
+
+	SLATE_END_ARGS()
+
+	/**
+	 * Construct the widget
+	 *
+	 * @param InArgs   A declaration from which to construct the widget
+	 */
+	void Construct( const FArguments& InArgs )
+	{
+		StateName = InArgs._StateName;
+		AssociatedNode = InArgs._AssociatedNode;
+
+		FMargin Padding = { 0.0f, 3.0f, 6.0f, 3.0f };
+		
+		if (AssociatedNode->ParameterIndex != INDEX_NONE)
+		{
+			Padding = { 10.0f, 3.0f, 6.0f, 3.0f };
+		}
+
+		this->ChildSlot
+		[
+			SNew(SHorizontalBox)
+				
+			+SHorizontalBox::Slot()
+			.AutoWidth()
+			.VAlign(VAlign_Center)
+			.HAlign(HAlign_Right)
+			[
+				SNew( SExpanderArrow, SharedThis(this) )
+			]
+
+			+SHorizontalBox::Slot()
+			.AutoWidth()
+			.Padding( Padding )
+			.VAlign(VAlign_Center)
+			.HAlign(HAlign_Left)
+			[
+				SNew( STextBlock )
+				.Text( FText::FromString(*StateName.Get()) )
+				//.HighlightString(*InArgs._HighlightText)
+				//.ColorAndOpacity( this, &SClassItem::GetTextColor)
+			]
+
+			//+SHorizontalBox::Slot()
+			//	.HAlign(HAlign_Right)
+			//	.VAlign(VAlign_Center)
+			//	.Padding( 0.0f, 0.0f, 6.0f, 0.0f )
+			//	[
+			//		SNew( SComboButton )
+			//			.ContentPadding(FMargin(2.0f))
+			//			//.OnGetMenuContent(this, &SClassItem::GenerateDropDown)
+			//	]
+		];
+		
+		TextColor = InArgs._TextColor;
+
+		STableRow< TSharedPtr<FString> >::ConstructInternal( 
+			STableRow::FArguments()
+				.ShowSelection(true),
+			InArgs._OwnerTableView.ToSharedRef() );	
+	}
+
+private:
+
+	/**
+	 * Generates the drop down menu for the item.
+	 *
+	 * @return		The drop down menu widget.
+	 */
+	/*
+	TSharedRef<SWidget> GenerateDropDown()
+	{
+		// Empty list of commands.
+		TSharedPtr< FUICommandList > Commands;
+
+		const bool bShouldCloseWindowAfterMenuSelection = true;	// Set the menu to automatically close when the user commits to a choice
+		FMenuBuilder MenuBuilder( bShouldCloseWindowAfterMenuSelection, Commands);
+		{ 
+			if (UClass* Class = AssociatedNode->Class.Get())
+			{
+				bool bIsBlueprint(false);
+				bool bHasBlueprint(false);
+
+				ClassViewer::Helpers::GetClassInfo(Class, bIsBlueprint, bHasBlueprint);
+
+				bHasBlueprint = AssociatedNode->Blueprint.IsValid();
+
+				if (bIsBlueprint)
+				{
+					FUIAction Action( FExecuteAction::CreateStatic( &ClassViewer::Helpers::CreateBlueprint, Class ) );
+					MenuBuilder.AddMenuEntry(TEXT("Create Blueprint..."), TEXT("Creates a new Blueprint from the current one."), NAME_None, Action);
+				}
+
+				if (bHasBlueprint)
+				{
+					MenuBuilder.AddMenuSeparator();
+
+					FUIAction Action( FExecuteAction::CreateStatic( &ClassViewer::Helpers::OpenBlueprintTool, ClassViewer::Helpers::GetBlueprint(Class) ) );
+					MenuBuilder.AddMenuEntry(TEXT("Open Blueprint Editor"), TEXT("Open the Blueprint in the editor."), NAME_None, Action);
+				}
+
+				if (bHasBlueprint)
+				{
+					MenuBuilder.AddMenuSeparator();
+
+					FUIAction Action( FExecuteAction::CreateStatic( &ClassViewer::Helpers::FindInContentBrowser, ClassViewer::Helpers::GetBlueprint(Class), Class ) );
+					MenuBuilder.AddMenuEntry(TEXT("Find In Content Browser"), TEXT("Find in Content Browser"), NAME_None, Action);
+				}
+			}
+		}
+
+		return MenuBuilder.MakeWidget();
+	}
+	*/
+
+	/** Returns the text color for the item based on if it is selected or not. */
+	FSlateColor GetTextColor() const
+	{
+		const TSharedPtr< ITypedTableView< TSharedPtr<FString> > > OwnerWidget = OwnerTablePtr.Pin();
+		const TSharedPtr<FString>* MyItem = OwnerWidget->Private_ItemFromWidget( this );
+		const bool bIsSelected = OwnerWidget->Private_IsItemSelected( *MyItem );
+
+		if(bIsSelected)
+		{
+			return FSlateColor::UseForeground();
+		}
+
+		return TextColor;
+	}
+
+private:
+
+	/** The class name for which this item is associated with. */
+	TSharedPtr<FString> StateName;
+
+	/** The text color for this item. */
+	FSlateColor TextColor;
+
+	/** The Class Viewer Node this item is associated with. */
+	TSharedPtr< FStateDetailsNode > AssociatedNode;
+
+};
+
+
+TSharedRef< ITableRow > FCustomizableObjectDetails::OnGenerateRowForStateTree( TSharedPtr<FStateDetailsNode> Item, const TSharedRef< STableViewBase >& OwnerTable )
+{	
+	// If the item was accepted by the filter, leave it bright, otherwise dim it.
+	float AlphaValue = 1.0f;
+	TSharedRef< ITableRow > ReturnRow = SNew(SStateItem)
+		.StateName(Item->GetName())
+		//.HighlightText(&SearchBox->GetText())
+		.TextColor(FLinearColor(1.0f, 1.0f, 1.0f, AlphaValue))
+		.AssociatedNode(Item)
+		.OwnerTableView( OwnerTable );
+
+	// Expand the item if needed.
+	//bool* bIsExpanded = ExpansionStateMap.Find( Item->GetName() );
+	//if( bIsExpanded && *bIsExpanded )
+	//{
+	//	StatesTree->SetItemExpansion( Item, *bIsExpanded );
+	//}
+
+	return ReturnRow;
+}
+
+
+void FCustomizableObjectDetails::OnGetChildrenForStateTree( TSharedPtr<FStateDetailsNode> InParent, TArray< TSharedPtr< FStateDetailsNode > >& OutChildren )
+{
+	// Simply return the children, it's already setup.
+	OutChildren = InParent->GetChildrenList();
+}
+
+#undef LOCTEXT_NAMESPACE
diff --git a/Engine/Plugins/Mutable/Source/CustomizableObjectEditor/Private/MuCOE/CustomizableObjectDetails.h b/Engine/Plugins/Mutable/Source/CustomizableObjectEditor/Private/MuCOE/CustomizableObjectDetails.h
new file mode 100644
--- /dev/null
+++ b/Engine/Plugins/Mutable/Source/CustomizableObjectEditor/Private/MuCOE/CustomizableObjectDetails.h
@@ -0,0 +1,37 @@
+// Copyright Epic Games, Inc. All Rights Reserved.
+
+#pragma once
+
+#include "IDetailCustomization.h"
+#include "Widgets/Views/STreeView.h"
+
+class FStateDetailsNode;
+class IDetailLayoutBuilder;
+class ITableRow;
+
+
+class FCustomizableObjectDetails : public IDetailCustomization
+{
+public:
+	/** Makes a new instance of this detail layout class for a specific detail view requesting it */
+	static TSharedRef<IDetailCustomization> MakeInstance();
+
+	/** ILayoutDetails interface */
+	virtual void CustomizeDetails( IDetailLayoutBuilder& DetailBuilder ) override;
+
+private:
+
+	class UCustomizableObject* CustomizableObject = nullptr;
+	
+	TSharedPtr<STreeView<TSharedPtr<FStateDetailsNode>>> StatesTree;
+
+	/** The root items to be displayed in the state tree. */
+	TArray< TSharedPtr< FStateDetailsNode > > RootTreeItems;
+
+	/** Re-caches the states found in the CO and calls for the update of the STreeView exposing the states and parameters*/
+	void UpdateTree();
+	
+	TSharedRef< ITableRow > OnGenerateRowForStateTree( TSharedPtr<FStateDetailsNode> Item, const TSharedRef< STableViewBase >& OwnerTable );
+	void OnGetChildrenForStateTree( TSharedPtr<FStateDetailsNode> InParent, TArray< TSharedPtr< FStateDetailsNode > >& OutChildren );
+
+};
