[mutable] Moved the Mutable plugin out of Experimental status into Beta.

diff --git a/Engine/Plugins/Mutable/Source/CustomizableObjectEditor/Private/MuCOE/CustomizableObjectEditor.cpp b/Engine/Plugins/Mutable/Source/CustomizableObjectEditor/Private/MuCOE/CustomizableObjectEditor.cpp
new file mode 100644
--- /dev/null
+++ b/Engine/Plugins/Mutable/Source/CustomizableObjectEditor/Private/MuCOE/CustomizableObjectEditor.cpp
@@ -0,0 +1,2500 @@
+// Copyright Epic Games, Inc. All Rights Reserved.
+
+#include "MuCOE/CustomizableObjectEditor.h"
+
+#include "AssetViewerSettings.h"
+#include "ContentBrowserModule.h"
+#include "CustomizableObjectEditorPerformanceAnalyzer.h"
+#include "Animation/DebugSkelMeshComponent.h"
+#include "AssetRegistry/ARFilter.h"
+#include "AssetRegistry/AssetRegistryModule.h"
+#include "DetailsViewArgs.h"
+#include "EdGraphUtilities.h"
+#include "Engine/StaticMesh.h"
+#include "FileHelpers.h"
+#include "Framework/Commands/GenericCommands.h"
+#include "Framework/MultiBox/MultiBoxBuilder.h"
+#include "Framework/Notifications/NotificationManager.h"
+#include "GraphEditorActions.h"
+#include "IContentBrowserSingleton.h"
+#include "HAL/PlatformApplicationMisc.h"
+#include "IDetailsView.h"
+#include "MaterialGraph/MaterialGraphNode.h"
+#include "MaterialGraph/MaterialGraphNode_Comment.h"
+#include "MuCO/CustomizableObjectInstance.h"
+#include "MuCO/CustomizableObjectSystem.h"
+#include "MuCO/CustomizableSkeletalComponent.h"
+#include "MuCO/UnrealPortabilityHelpers.h"
+#include "MuCO/CustomizableObjectCompilerTypes.h"
+#include "MuCOE/CustomizableObjectCustomSettings.h"
+#include "MuCOE/CustomizableObjectEditorActions.h"
+#include "MuCOE/CustomizableObjectEditorLogger.h"
+#include "MuCO/ICustomizableObjectEditorModule.h"
+#include "MuCOE/CustomizableObjectEditorNodeContextCommands.h"
+#include "MuCOE/CustomizableObjectEditorStyle.h"
+#include "MuCOE/CustomizableObjectEditorViewportClient.h"
+#include "MuCOE/CustomizableObjectEditor_Deprecated.h"
+#include "MuCOE/CustomizableObjectGraph.h"
+#include "MuCOE/CustomizableObjectPreviewScene.h"
+#include "MuCOE/EdGraphSchema_CustomizableObject.h"
+#include "MuCOE/GraphTraversal.h"
+#include "MuCOE/Nodes/CustomizableObjectNodeEnumParameter.h"
+#include "MuCOE/Nodes/CustomizableObjectNodeFloatParameter.h"
+#include "MuCOE/Nodes/CustomizableObjectNodeGroupProjectorParameter.h"
+#include "MuCOE/Nodes/CustomizableObjectNodeMaterial.h"
+#include "MuCOE/Nodes/CustomizableObjectNodeMaterialVariation.h"
+#include "MuCOE/Nodes/CustomizableObjectNodeModifierClipMorph.h"
+#include "MuCOE/Nodes/CustomizableObjectNodeModifierClipWithMesh.h"
+#include "MuCOE/Nodes/CustomizableObjectNodeObject.h"
+#include "MuCOE/Nodes/CustomizableObjectNodeObjectGroup.h"
+#include "MuCOE/Nodes/CustomizableObjectNodeProjectorConstant.h"
+#include "MuCOE/Nodes/CustomizableObjectNodeSkeletalMesh.h"
+#include "MuCOE/Nodes/CustomizableObjectNodeStaticMesh.h"
+#include "MuCOE/Nodes/CustomizableObjectNodeTable.h"
+#include "MuCOE/SCustomizableObjectEditorAdvancedPreviewSettings.h"
+#include "MuCOE/SCustomizableObjectEditorPerformanceReport.h"
+#include "MuCOE/SCustomizableObjectEditorTagExplorer.h"
+#include "MuCOE/SCustomizableObjectEditorTextureAnalyzer.h"
+#include "MuCOE/SCustomizableObjectEditorViewport.h"
+#include "PropertyEditorModule.h"
+#include "ScopedTransaction.h"
+#include "MuCO/CustomizableObjectInstancePrivate.h"
+#include "MuCO/CustomizableObjectPrivate.h"
+#include "UObject/EnumProperty.h"
+#include "Widgets/Docking/SDockTab.h"
+#include "Widgets/Input/SSearchBox.h"
+#include "Widgets/Input/STextComboBox.h"
+#include "Widgets/Layout/SScrollBox.h"
+#include "Widgets/Notifications/SNotificationList.h"
+#include "SMutableObjectViewer.h"
+#include "MuCO/CustomizableObjectCustomVersion.h"
+#include "MuCO/CustomizableSkeletalMeshActor.h"
+#include "MuCOE/CustomizableObjectInstanceFactory.h"
+
+class FAdvancedPreviewScene;
+class FWorkspaceItem;
+class IToolkitHost;
+class SWidget;
+enum class EColorArithmeticOperation : uint8;
+
+#define LOCTEXT_NAMESPACE "CustomizableObjectEditor"
+
+DEFINE_LOG_CATEGORY_STATIC(LogCustomizableObjectEditor, Log, All);
+
+const FName FCustomizableObjectEditor::ViewportTabId( TEXT( "CustomizableObjectEditor_Viewport" ) );
+const FName FCustomizableObjectEditor::ObjectPropertiesTabId( TEXT( "CustomizableObjectEditor_ObjectProperties" ) );
+const FName FCustomizableObjectEditor::InstancePropertiesTabId( TEXT( "CustomizableObjectEditor_InstanceProperties" ) );
+const FName FCustomizableObjectEditor::GraphTabId( TEXT( "CustomizableObjectEditor_Graph" ) );
+const FName FCustomizableObjectEditor::GraphNodePropertiesTabId( TEXT( "CustomizableObjectEditor_GraphNodeProperties" ) );
+const FName FCustomizableObjectEditor::AdvancedPreviewSettingsTabId(TEXT("CustomizableObjectEditor_AdvancedPreviewSettings"));
+const FName FCustomizableObjectEditor::TextureAnalyzerTabId(TEXT("CustomizableObjectEditor_TextureAnalyzer"));
+const FName FCustomizableObjectEditor::PerformanceReportTabId(TEXT("CustomizableObjectEditor_PerformanceReport"));
+const FName FCustomizableObjectEditor::PerformanceAnalyzerTabId(TEXT("CustomizableObjectEditor_MewPerformanceReport"));
+const FName FCustomizableObjectEditor::TagExplorerTabId(TEXT("CustomizableObjectEditor_TagExplorer"));
+const FName FCustomizableObjectEditor::ObjectDebuggerTabId(TEXT("CustomizableObjectEditor_ObjectDebugger"));
+const FName FCustomizableObjectEditor::PopulationClassTagManagerTabId(TEXT("CustomizableObjectEditor_PopulationClassTabManager"));
+
+
+void UUpdateClassWrapper::DelegatedCallback(UCustomizableObjectInstance* Instance)
+{
+	Delegate.ExecuteIfBound();
+}
+
+
+void FCustomizableObjectEditor::RegisterTabSpawners(const TSharedRef<class FTabManager>& InTabManager)
+{
+	WorkspaceMenuCategory = InTabManager->AddLocalWorkspaceMenuCategory(LOCTEXT("WorkspaceMenu_CustomizableObjectEditor", "Customizable Object Editor"));
+	TSharedRef<FWorkspaceItem> WorkspaceMenuCategoryRef = WorkspaceMenuCategory.ToSharedRef();
+	
+	FAssetEditorToolkit::RegisterTabSpawners(InTabManager);
+
+	InTabManager->RegisterTabSpawner(ViewportTabId, FOnSpawnTab::CreateSP(this, &FCustomizableObjectEditor::SpawnTab_Viewport))
+		.SetDisplayName(LOCTEXT("ViewportTab", "Viewport"))
+		.SetGroup(WorkspaceMenuCategoryRef);
+	
+	InTabManager->RegisterTabSpawner(ObjectPropertiesTabId, FOnSpawnTab::CreateSP(this, &FCustomizableObjectEditor::SpawnTab_ObjectProperties))
+		.SetDisplayName(LOCTEXT("ObjectPropertiesTab", "Object Properties"))
+		.SetGroup(WorkspaceMenuCategoryRef);
+	
+	InTabManager->RegisterTabSpawner(InstancePropertiesTabId, FOnSpawnTab::CreateSP(this, &FCustomizableObjectEditor::SpawnTab_InstanceProperties))
+		.SetDisplayName(LOCTEXT("InstancePropertiesTab", "Instance Properties"))
+		.SetGroup(WorkspaceMenuCategoryRef);
+	
+	InTabManager->RegisterTabSpawner(GraphTabId, FOnSpawnTab::CreateSP(this, &FCustomizableObjectEditor::SpawnTab_Graph))
+		.SetDisplayName(LOCTEXT("GraphTab", "Object Graph"))
+		.SetGroup(WorkspaceMenuCategoryRef);
+	
+	InTabManager->RegisterTabSpawner(GraphNodePropertiesTabId, FOnSpawnTab::CreateSP(this, &FCustomizableObjectEditor::SpawnTab_GraphNodeProperties))
+		.SetDisplayName(LOCTEXT("GraphNodePropertiesTab", "Object Graph Node Properties"))
+		.SetGroup(WorkspaceMenuCategoryRef);
+
+	InTabManager->RegisterTabSpawner(AdvancedPreviewSettingsTabId, FOnSpawnTab::CreateSP(this, &FCustomizableObjectEditor::SpawnTab_AdvancedPreviewSettings))
+		.SetDisplayName(LOCTEXT("AdvancedPreviewSettingsTab", "Advanced Preview Settings"))
+		.SetGroup(WorkspaceMenuCategoryRef);
+
+	InTabManager->RegisterTabSpawner(TextureAnalyzerTabId, FOnSpawnTab::CreateSP(this, &FCustomizableObjectEditor::SpawnTab_TextureAnalyzer))
+		.SetDisplayName(LOCTEXT("TextureAnalyzer", "Texture Analyzer"))
+		.SetGroup(WorkspaceMenuCategoryRef);
+
+	InTabManager->RegisterTabSpawner(PerformanceReportTabId, FOnSpawnTab::CreateSP(this, &FCustomizableObjectEditor::SpawnTab_PerformanceReport))
+		.SetDisplayName(LOCTEXT("PerformanceReport", "Performance Report"))
+		.SetGroup(WorkspaceMenuCategoryRef);
+
+	InTabManager->RegisterTabSpawner(PerformanceAnalyzerTabId, FOnSpawnTab::CreateSP(this, &FCustomizableObjectEditor::SpawnTab_PerformanceAnalyzer))
+		.SetDisplayName(LOCTEXT("PerformanceAnalyzer", "Performance Analyzer"))
+		.SetGroup(WorkspaceMenuCategoryRef);
+	
+	InTabManager->RegisterTabSpawner(TagExplorerTabId, FOnSpawnTab::CreateSP(this, &FCustomizableObjectEditor::SpawnTab_TagExplorer))
+		.SetDisplayName(LOCTEXT("TagExplorerTab", "Tag Explorer"))
+		.SetGroup(WorkspaceMenuCategoryRef);
+}
+
+
+void FCustomizableObjectEditor::UnregisterTabSpawners(const TSharedRef<class FTabManager>& InTabManager)
+{
+	InTabManager->UnregisterTabSpawner( ViewportTabId );
+	InTabManager->UnregisterTabSpawner( ObjectPropertiesTabId );
+	InTabManager->UnregisterTabSpawner( InstancePropertiesTabId );
+	InTabManager->UnregisterTabSpawner( GraphTabId );
+	InTabManager->UnregisterTabSpawner( GraphNodePropertiesTabId );
+	InTabManager->UnregisterTabSpawner(AdvancedPreviewSettingsTabId);
+	InTabManager->UnregisterTabSpawner(TextureAnalyzerTabId);
+	InTabManager->UnregisterTabSpawner(PerformanceReportTabId);
+	InTabManager->UnregisterTabSpawner(PerformanceAnalyzerTabId);
+}	
+
+
+
+FCustomizableObjectEditor::~FCustomizableObjectEditor()
+{
+	if (PreviewInstance)
+	{
+		if (PreviewInstance->GetPrivate()->bSelectedProfileDirty && PreviewInstance->GetPrivate()->SelectedProfileIndex != INDEX_NONE)
+		{
+			PreviewInstance->GetPrivate()->SaveParametersToProfile(PreviewInstance->GetPrivate()->SelectedProfileIndex);
+		}
+	}
+	
+	CustomizableObjectDetailsView.Reset();
+	GEditor->UnregisterForUndo(this);
+
+	FCoreUObjectDelegates::OnObjectModified.RemoveAll(this);
+
+	CustomizableObject->GetPrivate()->Status.GetOnStateChangedDelegate().RemoveAll(this);
+
+	FAssetRegistryModule& AssetRegistryModule = FModuleManager::LoadModuleChecked<FAssetRegistryModule>(TEXT("AssetRegistry"));
+	AssetRegistryModule.Get().OnFilesLoaded().RemoveAll(this);
+
+	FCoreUObjectDelegates::OnObjectPropertyChanged.RemoveAll(this);
+
+	GEngine->ForceGarbageCollection(true);
+}
+
+
+void FCustomizableObjectEditor::InitCustomizableObjectEditor(const EToolkitMode::Type Mode, const TSharedPtr<IToolkitHost>& InitToolkitHost)
+{
+	ProjectorParameter = NewObject<UProjectorParameter>();
+
+	CustomSettings = NewObject<UCustomSettings>();
+	CustomSettings->SetEditor(SharedThis(this));
+	
+	EditorProperties = NewObject<UCustomizableObjectEditorProperties>();
+
+	// Support undo/redo
+	CustomizableObject->SetFlags(RF_Transactional);
+
+	GEditor->RegisterForUndo(this);
+
+	// Register our commands. This will only register them if not previously registered
+	FGraphEditorCommands::Register();
+	FCustomizableObjectEditorCommands::Register();
+	FCustomizableObjectEditorViewportCommands::Register();
+	FCustomizableObjectEditorNodeContextCommands::Register();
+
+	BindCommands();
+
+	FPropertyEditorModule& PropPlugin = FModuleManager::LoadModuleChecked<FPropertyEditorModule>("PropertyEditor");
+
+	FDetailsViewArgs DetailsViewArgs;
+	DetailsViewArgs.NotifyHook = this;
+	DetailsViewArgs.NameAreaSettings = FDetailsViewArgs::ENameAreaSettings::HideNameArea;
+	DetailsViewArgs.bAllowSearch = true;
+	//DetailsViewArgs.bShowActorLabel = false;
+	DetailsViewArgs.bShowObjectLabel = false;
+	DetailsViewArgs.bShowScrollBar = false;
+
+	CustomizableObjectDetailsView = PropPlugin.CreateDetailView( DetailsViewArgs );
+
+	CustomizableInstanceDetailsView = PropPlugin.CreateDetailView( DetailsViewArgs );
+	GraphNodeDetailsView = PropPlugin.CreateDetailView(DetailsViewArgs);
+
+	Viewport = SNew(SCustomizableObjectEditorViewportTabBody)
+		.CustomizableObjectEditor(SharedThis(this));
+
+	Viewport->SetCustomizableObject(CustomizableObject);
+	ViewportClient = Viewport->GetViewportClient();
+
+	// \TODO: Create only when needed?
+	TextureAnalyzer = SNew(SCustomizableObjecEditorTextureAnalyzer).CustomizableObjectEditor(this).CustomizableObjectInstanceEditor(nullptr);
+
+	// \TODO: Create only when needed?
+	TagExplorer = SNew(SCustomizableObjectEditorTagExplorer).CustomizableObjectEditor(this);
+	
+	TSharedPtr<FAdvancedPreviewScene> AdvancedPreviewScene = StaticCastSharedPtr<FAdvancedPreviewScene>(Viewport->GetPreviewScene());
+
+	CustomizableObjectEditorAdvancedPreviewSettings =
+		SNew(SCustomizableObjectEditorAdvancedPreviewSettings, AdvancedPreviewScene.ToSharedRef())
+		.CustomSettings(CustomSettings)
+		.CustomizableObjectEditor(SharedThis(this).ToWeakPtr());
+	CustomizableObjectEditorAdvancedPreviewSettings->LoadProfileEnvironment();
+	AdvancedPreviewSettingsWidget = CustomizableObjectEditorAdvancedPreviewSettings;
+
+	const TSharedRef<FTabManager::FLayout> StandaloneDefaultLayout = FTabManager::NewLayout( "Standalone_CustomizableObjectEditor_Layout_v1.4" )
+	->AddArea
+	(
+		FTabManager::NewPrimaryArea()->SetOrientation(Orient_Horizontal)
+		->Split
+		(
+			FTabManager::NewStack()
+			->SetSizeCoefficient(0.6f)
+			->SetHideTabWell(true)
+			->AddTab(GraphTabId, ETabState::OpenedTab)
+		)
+		->Split
+		(
+			FTabManager::NewSplitter() ->SetOrientation(Orient_Horizontal) 
+			->SetSizeCoefficient(0.4f)
+			->Split
+			(
+				FTabManager::NewSplitter() ->SetOrientation(Orient_Vertical) 
+				->SetSizeCoefficient(0.5f)
+				->Split
+				(
+					FTabManager::NewStack()
+					->SetSizeCoefficient(0.4f)
+					->AddTab(ObjectPropertiesTabId, ETabState::OpenedTab)
+					->AddTab(TagExplorerTabId, ETabState::OpenedTab)
+					->SetForegroundTab(ObjectPropertiesTabId)
+				)
+				->Split
+				(
+					FTabManager::NewStack()
+					->SetSizeCoefficient(0.6f)
+					->AddTab(GraphNodePropertiesTabId, ETabState::OpenedTab)
+				)
+			)
+			->Split
+			(
+				FTabManager::NewSplitter() ->SetOrientation(Orient_Vertical) 
+				->SetSizeCoefficient(0.5f)
+				->Split
+				(
+					FTabManager::NewStack()
+					->SetSizeCoefficient(0.5f)
+					->AddTab(InstancePropertiesTabId, ETabState::OpenedTab)
+					->AddTab(AdvancedPreviewSettingsTabId, ETabState::OpenedTab)
+					->SetForegroundTab(InstancePropertiesTabId)
+				)
+				->Split
+				(
+					FTabManager::NewStack()
+					->SetSizeCoefficient(0.5f)
+					->AddTab(ViewportTabId, ETabState::OpenedTab)
+					->SetHideTabWell(true)
+				)
+			)
+		)	
+	);
+
+	const bool bCreateDefaultStandaloneMenu = true;
+	const bool bCreateDefaultToolbar = true;
+	FAssetEditorToolkit::InitAssetEditor( Mode, InitToolkitHost, CustomizableObjectEditorAppIdentifier, StandaloneDefaultLayout, bCreateDefaultStandaloneMenu, bCreateDefaultToolbar, CustomizableObject);
+	
+	CustomizableObjectDetailsView->SetObject(CustomizableObject); // Can only be called after initializing the Asset Editor
+
+	ExtendToolbar();
+	RegenerateMenusAndToolbars();
+
+	// Clears selection highlight.
+	OnObjectPropertySelectionChanged(NULL);
+	OnInstancePropertySelectionChanged(NULL);
+	FCoreUObjectDelegates::OnObjectModified.AddRaw(this, &FCustomizableObjectEditor::OnObjectModified);
+	
+	UCustomizableObjectPrivate* CustomizableObjectPrivate = CustomizableObject->GetPrivate();
+	
+	PreviewInstance = CustomizableObject->CreateInstance();	
+	PreviewInstance->UpdatedNativeDelegate.AddSP(SharedThis(this), &FCustomizableObjectEditor::OnUpdatePreviewInstance);
+	PreviewInstance->SetBuildParameterRelevancy(true);
+
+	CustomizableInstanceDetailsView->SetObject(PreviewInstance, true);
+	
+	CustomizableObjectPrivate->Status.GetOnStateChangedDelegate().AddRaw(this, &FCustomizableObjectEditor::OnCustomizableObjectStatusChanged);
+	OnCustomizableObjectStatusChanged(FCustomizableObjectStatusTypes::EState::Loading, CustomizableObjectPrivate->Status.Get());  // Fake we are still in the loading phase.
+	
+	CustomizableObject->GetPostCompileDelegate().AddSP(this, &FCustomizableObjectEditor::OnPostCompile); // Must be attached after creating the Instance since the Instance also does some work in this delegate.
+}
+
+
+FName FCustomizableObjectEditor::GetToolkitFName() const
+{
+	return FName("CustomizableObjectEditor");
+}
+
+
+FText FCustomizableObjectEditor::GetBaseToolkitName() const
+{
+	return LOCTEXT("ToolkitName", "Customizable Object Editor");
+}
+
+
+void FCustomizableObjectEditor::SelectNode(const UEdGraphNode* Node)
+{
+	GraphEditor->JumpToNode(Node);
+}
+
+
+void FCustomizableObjectEditor::AddReferencedObjects( FReferenceCollector& Collector )
+{
+	Collector.AddReferencedObject( CustomizableObject );
+	Collector.AddReferencedObject( PreviewInstance );
+	Collector.AddReferencedObject( ProjectorParameter );
+	Collector.AddReferencedObject( CustomSettings );
+	Collector.AddReferencedObject(EditorProperties);
+}
+
+
+FCustomizableObjectEditor::FCustomizableObjectEditor(UCustomizableObject& ObjectToEdit) :
+	CustomizableObject(&ObjectToEdit) {}
+
+
+TSharedRef<SDockTab> FCustomizableObjectEditor::SpawnTab_Viewport( const FSpawnTabArgs& Args )
+{
+	check( Args.GetTabId() == ViewportTabId );
+
+	TSharedRef<SDockTab> DockTab = SNew(SDockTab)
+		.Label( FText::FromString( GetTabPrefix() + LOCTEXT("CustomizableObjectViewport_TabTitle", "Viewport").ToString() ) )
+		[
+			SNew( SVerticalBox )
+
+			+ SVerticalBox::Slot()
+			.Padding( 2.0f )
+			.FillHeight(1.0f)
+			[
+				Viewport.ToSharedRef()
+			]
+		];
+
+	DockTab->SetTabIcon(FCustomizableObjectEditorStyle::Get().GetBrush("CustomizableObjectEditor.Tabs.Preview"));
+
+	return DockTab;
+}
+
+TSharedRef<SDockTab> FCustomizableObjectEditor::SpawnTab_ObjectProperties( const FSpawnTabArgs& Args )
+{
+	check( Args.GetTabId() == ObjectPropertiesTabId );
+
+	TSharedRef<SScrollBox> ScrollBox = SNew(SScrollBox)
+		+ SScrollBox::Slot()
+		[
+			CustomizableObjectDetailsView.ToSharedRef()
+		];
+
+	ScrollBox->SetScrollBarRightClickDragAllowed(true);
+
+	TSharedRef<SDockTab> DockTab = SNew(SDockTab)
+		.Label( FText::FromString( GetTabPrefix() + LOCTEXT( "CustomizableObjectProperties_TabTitle", "Object Properties" ).ToString() ) )
+		[
+			ScrollBox
+		];
+
+	DockTab->SetTabIcon(FCustomizableObjectEditorStyle::Get().GetBrush("CustomizableObjectEditor.Tabs.CustomizableObjectProperties"));
+
+	return DockTab;
+}
+
+
+TSharedRef<SDockTab> FCustomizableObjectEditor::SpawnTab_InstanceProperties( const FSpawnTabArgs& Args )
+{
+	check( Args.GetTabId() == InstancePropertiesTabId );
+
+	TSharedRef<SScrollBox> ScrollBox = SNew(SScrollBox)
+		+ SScrollBox::Slot()
+		[
+			CustomizableInstanceDetailsView.ToSharedRef()
+		];
+
+	ScrollBox->SetScrollBarRightClickDragAllowed(true);
+
+	TSharedRef<SDockTab> DockTab = SNew(SDockTab)
+		.Label( FText::FromString( GetTabPrefix() + LOCTEXT( "CustomizableInstanceProperties_TabTitle", "Preview Instance" ).ToString() ) )
+		[
+			ScrollBox
+		];
+
+	DockTab->SetTabIcon(FCustomizableObjectEditorStyle::Get().GetBrush("CustomizableObjectEditor.Tabs.CustomizableInstanceProperties"));
+
+	return DockTab;
+}
+
+
+/** Create new tab for the supplied graph - don't call this directly, call SExplorer->FindTabForGraph.*/
+void FCustomizableObjectEditor::CreateGraphEditorWidget(UEdGraph* InGraph)
+{
+	UCustomizableObjectGraph* CustomizableObjectGraph = Cast<UCustomizableObjectGraph>(InGraph);
+	check(CustomizableObjectGraph != NULL);
+
+	// Add essential graph nodes if they do not exist
+	CustomizableObjectGraph->AddEssentialGraphNodes();
+
+	GraphEditorCommands = MakeShareable(new FUICommandList);
+
+	TSharedRef<SWidget> TitleBarWidget =
+		SNew(SHorizontalBox)
+		+ SHorizontalBox::Slot()
+		.HAlign(HAlign_Left)
+		.FillWidth(10.f)
+		.Padding(5.f)
+		[
+			SNew(SSearchBox)
+			.HintText(LOCTEXT("Search", "Search..."))
+			.ToolTipText(LOCTEXT("Search Nodes, Properties or Values that contain the inserted words", "Search Nodes, Properties or Values that contain the inserted words"))
+			.OnTextCommitted(FOnTextCommitted::CreateSP(this, &FCustomizableObjectEditor::OnEnterText))
+			.SelectAllTextWhenFocused(true)
+		];
+
+	// Create the appearance info
+	FGraphAppearanceInfo AppearanceInfo;	
+	AppearanceInfo.CornerText = LOCTEXT("ApperanceCornerText", "MUTABLE");
+
+	SGraphEditor::FGraphEditorEvents InEvents;
+	InEvents.OnSelectionChanged = SGraphEditor::FOnSelectionChanged::CreateSP( this, &FCustomizableObjectEditor::OnSelectedGraphNodesChanged );
+	InEvents.OnTextCommitted = FOnNodeTextCommitted::CreateSP(this, &FCustomizableObjectEditor::OnNodeTitleCommitted);
+
+	// Make full graph editor
+	GraphEditor = SNew(SGraphEditor)
+		.AdditionalCommands(GraphEditorCommands)
+		.Appearance(AppearanceInfo)
+		.GraphToEdit(InGraph)
+		.GraphEvents(InEvents)
+		.TitleBar(TitleBarWidget)
+		.ShowGraphStateOverlay(false); // Removes graph state overlays (border and text) such as "SIMULATING" and "READ-ONLY"
+
+	// Editing commands
+	GraphEditorCommands->MapAction(FGenericCommands::Get().Delete,
+		FExecuteAction::CreateSP(this, &FCustomizableObjectEditor::DeleteSelectedNodes),
+		FCanExecuteAction::CreateSP(this, &FCustomizableObjectEditor::CanDeleteNodes));
+
+	GraphEditorCommands->MapAction( FGenericCommands::Get().Copy,
+		FExecuteAction::CreateSP(this, &FCustomizableObjectEditor::CopySelectedNodes),
+		FCanExecuteAction::CreateSP(this, &FCustomizableObjectEditor::CanCopyNodes));
+
+	GraphEditorCommands->MapAction(FGenericCommands::Get().Paste,
+		FExecuteAction::CreateSP(this, &FCustomizableObjectEditor::PasteNodes),
+		FCanExecuteAction::CreateSP(this, &FCustomizableObjectEditor::CanPasteNodes));
+
+	GraphEditorCommands->MapAction(FGenericCommands::Get().Cut,
+		FExecuteAction::CreateSP(this, &FCustomizableObjectEditor::CutSelectedNodes),
+		FCanExecuteAction::CreateSP(this, &FCustomizableObjectEditor::CanCutNodes));
+
+	GraphEditorCommands->MapAction(FGenericCommands::Get().Duplicate,
+		FExecuteAction::CreateSP(this, &FCustomizableObjectEditor::DuplicateSelectedNodes),
+		FCanExecuteAction::CreateSP(this, &FCustomizableObjectEditor::CanDuplicateSelectedNodes));
+
+	GraphEditorCommands->MapAction(FGenericCommands::Get().Rename,
+		FExecuteAction::CreateSP(this, &FCustomizableObjectEditor::OnRenameNode),
+		FCanExecuteAction::CreateSP(this, &FCustomizableObjectEditor::CanRenameNodes));
+
+	GraphEditorCommands->MapAction(FCustomizableObjectEditorNodeContextCommands::Get().CreateComment,
+		FExecuteAction::CreateSP(this, &FCustomizableObjectEditor::CreateCommentBoxFromKey));
+
+	// Alignment Commands
+	GraphEditorCommands->MapAction(FGraphEditorCommands::Get().AlignNodesTop,
+		FExecuteAction::CreateSP(GraphEditor.Get(), &SGraphEditor::OnAlignTop));
+
+	GraphEditorCommands->MapAction(FGraphEditorCommands::Get().AlignNodesMiddle,
+		FExecuteAction::CreateSP(GraphEditor.Get(), &SGraphEditor::OnAlignMiddle));
+
+	GraphEditorCommands->MapAction(FGraphEditorCommands::Get().AlignNodesBottom,
+		FExecuteAction::CreateSP(GraphEditor.Get(), &SGraphEditor::OnAlignBottom));
+
+	GraphEditorCommands->MapAction(FGraphEditorCommands::Get().AlignNodesLeft,
+		FExecuteAction::CreateSP(GraphEditor.Get(), &SGraphEditor::OnAlignLeft));
+
+	GraphEditorCommands->MapAction(FGraphEditorCommands::Get().AlignNodesCenter,
+		FExecuteAction::CreateSP(GraphEditor.Get(), &SGraphEditor::OnAlignCenter));
+
+	GraphEditorCommands->MapAction(FGraphEditorCommands::Get().AlignNodesRight,
+		FExecuteAction::CreateSP(GraphEditor.Get(), &SGraphEditor::OnAlignRight));
+
+	GraphEditorCommands->MapAction(FGraphEditorCommands::Get().StraightenConnections,
+		FExecuteAction::CreateSP(GraphEditor.Get(), &SGraphEditor::OnStraightenConnections));
+
+	// Distribution Commands
+	GraphEditorCommands->MapAction(FGraphEditorCommands::Get().DistributeNodesHorizontally,
+		FExecuteAction::CreateSP(GraphEditor.Get(), &SGraphEditor::OnDistributeNodesH));
+	
+	GraphEditorCommands->MapAction(FGraphEditorCommands::Get().DistributeNodesVertically,
+		FExecuteAction::CreateSP(GraphEditor.Get(), &SGraphEditor::OnDistributeNodesV));
+}
+
+
+TSharedRef<SDockTab> FCustomizableObjectEditor::SpawnTab_Graph( const FSpawnTabArgs& Args )
+{
+	check( Args.GetTabId().TabType == GraphTabId );
+
+	CreateGraphEditorWidget(CustomizableObject->GetPrivate()->GetSource());
+
+	TSharedRef<SDockTab> DockTab = SNew(SDockTab)
+		.Label( FText::FromString( GetTabPrefix() + LOCTEXT( "SourceGraph", "Source Graph" ).ToString() ) )
+		.TabColorScale( GetTabColorScale() )
+		[
+			GraphEditor.ToSharedRef()
+		];
+
+	DockTab->SetTabIcon(FCustomizableObjectEditorStyle::Get().GetBrush("CustomizableObjectEditor.Tabs.NodeGraph"));
+
+	return DockTab;
+}
+
+
+TSharedRef<SDockTab> FCustomizableObjectEditor::SpawnTab_GraphNodeProperties( const FSpawnTabArgs& Args )
+{
+	check( Args.GetTabId().TabType == GraphNodePropertiesTabId );
+
+	TSharedRef<SScrollBox> ScrollBox = SNew(SScrollBox)
+		+ SScrollBox::Slot()
+		[
+			GraphNodeDetailsView.ToSharedRef()
+		];
+
+	ScrollBox->SetScrollBarRightClickDragAllowed(true);
+
+	TSharedRef<SDockTab> DockTab = SNew(SDockTab)
+		.Label(FText::FromString(GetTabPrefix() + LOCTEXT("Graph Node Properties", "Node Properties").ToString()))
+		.TabColorScale(GetTabColorScale())
+		[
+			ScrollBox
+		];
+
+	DockTab->SetTabIcon(FCustomizableObjectEditorStyle::Get().GetBrush("CustomizableObjectEditor.Tabs.NodeProperties"));
+
+	return DockTab;
+}
+
+
+TSharedRef<SDockTab> FCustomizableObjectEditor::SpawnTab_AdvancedPreviewSettings(const FSpawnTabArgs& Args)
+{
+	check(Args.GetTabId() == AdvancedPreviewSettingsTabId);
+	TSharedRef<SDockTab> DockTab = SNew(SDockTab)
+		.Label(LOCTEXT("StaticMeshPreviewScene_TabTitle", "Preview Scene Settings"))
+		[
+			AdvancedPreviewSettingsWidget.ToSharedRef()
+		];
+
+	DockTab->SetTabIcon(FCustomizableObjectEditorStyle::Get().GetBrush("CustomizableObjectEditor.Tabs.PreviewSettings"));
+
+	return DockTab;
+}
+
+
+UCustomizableObjectInstance* FCustomizableObjectEditor::GetPreviewInstance()
+{
+	return PreviewInstance;
+}
+
+
+void FCustomizableObjectEditor::BindCommands()
+{
+	const FCustomizableObjectEditorCommands& Commands = FCustomizableObjectEditorCommands::Get();
+
+	// Toolbar
+	// Compile and options
+	ToolkitCommands->MapAction(
+		Commands.Compile,
+		FExecuteAction::CreateSP(this, &FCustomizableObjectEditor::CompileObject, false, false),
+		FCanExecuteAction::CreateStatic(&UCustomizableObjectSystem::IsActive),
+		FIsActionChecked());
+
+	ToolkitCommands->MapAction(
+		Commands.CompileOnlySelected,
+		FExecuteAction::CreateSP(this, &FCustomizableObjectEditor::CompileObject, true, false),
+		FCanExecuteAction::CreateStatic(&UCustomizableObjectSystem::IsActive),
+		FIsActionChecked());
+
+	// Compile and options
+	ToolkitCommands->MapAction(
+		Commands.ResetCompileOptions,
+		FExecuteAction::CreateSP(this, &FCustomizableObjectEditor::ResetCompileOptions),
+		FCanExecuteAction(),
+		FIsActionChecked());
+
+	ToolkitCommands->MapAction(
+		Commands.CompileOptions_UseDiskCompilation,
+		FExecuteAction::CreateSP(this, &FCustomizableObjectEditor::CompileOptions_UseDiskCompilation_Toggled),
+		FCanExecuteAction(),
+		FIsActionChecked::CreateSP(this, &FCustomizableObjectEditor::CompileOptions_UseDiskCompilation_IsChecked));
+
+	// Debug and options
+	ToolkitCommands->MapAction(
+		Commands.Debug,
+		FExecuteAction::CreateSP(this, &FCustomizableObjectEditor::DebugObject),
+		FCanExecuteAction(),
+		FIsActionChecked());
+
+	// References
+	ToolkitCommands->MapAction(
+		Commands.CompileGatherReferences,
+		FExecuteAction::CreateSP(this, &FCustomizableObjectEditor::CompileObject, false, true),
+		FCanExecuteAction(),
+		FIsActionChecked());
+
+	ToolkitCommands->MapAction(
+		Commands.ClearGatheredReferences,
+		FExecuteAction::CreateSP(this, &FCustomizableObjectEditor::ClearGatheredReferences),
+		FCanExecuteAction(),
+		FIsActionChecked());
+	
+	// Texture Analyzer
+	ToolkitCommands->MapAction(
+		Commands.TextureAnalyzer,
+		FExecuteAction::CreateSP(this, &FCustomizableObjectEditor::OpenTextureAnalyzerTab),
+		FCanExecuteAction(),
+		FIsActionChecked());
+
+	// Performance Report
+	ToolkitCommands->MapAction(
+		Commands.PerformanceReport,
+		FExecuteAction::CreateSP(this, &FCustomizableObjectEditor::OpenPerformanceReportTab),
+		FCanExecuteAction(),
+		FIsActionChecked());
+
+	// Performance Analyzer
+	ToolkitCommands->MapAction(
+		Commands.PerformanceAnalyzer,
+		FExecuteAction::CreateSP(this, &FCustomizableObjectEditor::OpenPerformanceAnalyzerTab),
+		FCanExecuteAction(),
+		FIsActionChecked());
+	
+	// Undo-Redo
+	ToolkitCommands->MapAction(
+		FGenericCommands::Get().Undo,
+		FExecuteAction::CreateSP(this, &FCustomizableObjectEditor::UndoGraphAction));
+
+	ToolkitCommands->MapAction(
+		FGenericCommands::Get().Redo,
+		FExecuteAction::CreateSP(this, &FCustomizableObjectEditor::RedoGraphAction));
+}
+
+
+void FCustomizableObjectEditor::UndoGraphAction()
+{
+	GEditor->UndoTransaction();
+}
+
+
+void FCustomizableObjectEditor::RedoGraphAction()
+{
+	// Clear selection, to avoid holding refs to nodes that go away
+	GraphEditor->ClearSelectionSet();
+
+	GEditor->RedoTransaction();
+}
+
+
+bool FCustomizableObjectEditor::GroupNodeIsLinkedToParentByName(UCustomizableObjectNodeObject* Node, UCustomizableObject* Test, const FString& ParentGroupName)
+{
+	TArray<UCustomizableObjectNodeObjectGroup*> GroupNodes;
+	Test->GetPrivate()->GetSource()->GetNodesOfClass<UCustomizableObjectNodeObjectGroup>(GroupNodes);
+
+	for (UCustomizableObjectNodeObjectGroup* GroupNode : GroupNodes)
+	{
+		if ((Node->ParentObjectGroupId == GroupNode->NodeGuid) && (ParentGroupName == GroupNode->GroupName))
+		{
+			return true;
+		}
+	}
+
+	return false;
+}
+
+
+// TODO FutureGMT, use graph traversal abstraction instead of a hardcoded implementation.
+void FCustomizableObjectEditor::ReconstructAllChildNodes(UCustomizableObjectNode& StartNode, const UClass& NodeType)
+{
+	UCustomizableObject* Object = CastChecked<UCustomizableObject>(StartNode.GetCustomizableObjectGraph()->GetOuter());
+	const TMultiMap<FGuid, UCustomizableObjectNodeObject*> Mapping = GetNodeGroupObjectNodeMapping(Object);
+	
+	TArray<UCustomizableObjectNode*> NodesToVisit;
+	NodesToVisit.Add(&StartNode);
+	
+	while (!NodesToVisit.IsEmpty())
+	{
+		UCustomizableObjectNode* Node = NodesToVisit.Pop();
+
+		if (&NodeType == Node->GetClass())
+		{
+			Node->UCustomizableObjectNode::ReconstructNode();							
+		}
+
+		if (const UCustomizableObjectNodeObjectGroup* GroupNode = Cast<UCustomizableObjectNodeObjectGroup>(Node))
+		{
+			TArray<UCustomizableObjectNodeObject*> ObjectNodes;
+			Mapping.MultiFind(GroupNode->NodeGuid, ObjectNodes);
+			
+			for (UCustomizableObjectNodeObject* ObjectNode : ObjectNodes)
+			{
+				NodesToVisit.Add(ObjectNode);	
+			}
+		}
+		
+		for (const UEdGraphPin* Pin : Node->GetAllPins()) // Not using GetAllNonOrphanPins on purpose since we want want to be able to reconstruct nodes that have non-orphan pins.
+		{
+			if (Pin->Direction != EGPD_Input)
+			{
+				continue;
+			}
+
+			for (const UEdGraphPin* ConnectedPin : FollowInputPinArray(*Pin))
+			{
+				if (UCustomizableObjectNode* TypedNode = Cast<UCustomizableObjectNode>(ConnectedPin->GetOwningNode()))
+				{
+					NodesToVisit.Add(TypedNode);						
+				}
+			}
+		}
+	}
+}
+
+
+UProjectorParameter* FCustomizableObjectEditor::GetProjectorParameter()
+{
+	return ProjectorParameter;
+}
+
+
+UCustomSettings* FCustomizableObjectEditor::GetCustomSettings()
+{
+	return CustomSettings;
+}
+
+
+void FCustomizableObjectEditor::SelectSingleNode(UCustomizableObjectNode& Node)
+{
+	FGraphPanelSelectionSet SelectedNodes = GraphEditor->GetSelectedNodes();
+	if (SelectedNodes.Num() != 1 || Cast<UCustomizableObjectNode>(*SelectedNodes.CreateIterator()) != &Node)
+	{
+		GraphEditor->ClearSelectionSet();
+		GraphEditor->SetNodeSelection(&Node, true);
+	}
+}
+
+
+void FCustomizableObjectEditor::HideGizmo()
+{
+	HideGizmoProjectorNodeProjectorConstant();
+	HideGizmoProjectorNodeProjectorParameter();
+	HideGizmoProjectorParameter();
+	HideGizmoClipMorph();
+	HideGizmoClipMesh();
+	HideGizmoLight();
+}
+
+
+void FCustomizableObjectEditor::ShowGizmoProjectorNodeProjectorConstant(UCustomizableObjectNodeProjectorConstant& Node)
+{
+	if (GizmoType != EGizmoType::NodeProjectorConstant)
+	{
+		HideGizmo();
+	}
+
+	GizmoType = EGizmoType::NodeProjectorConstant;
+
+	SelectSingleNode(Node);
+	
+	FProjectorTypeDelegate ProjectorTypeDelegate;
+	ProjectorTypeDelegate.BindUObject(&Node, &UCustomizableObjectNodeProjectorConstant::GetProjectorType);		
+
+	FWidgetColorDelegate WidgetColorDelegate;
+	WidgetColorDelegate.BindLambda([]() { return FColor::Red; });
+
+	FWidgetLocationDelegate WidgetLocationDelegate;
+	WidgetLocationDelegate.BindUObject(&Node, &UCustomizableObjectNodeProjectorConstant::GetProjectorPosition);
+
+	FOnWidgetLocationChangedDelegate OnWidgetLocationChangedDelegate;
+	OnWidgetLocationChangedDelegate.BindUObject(&Node, &UCustomizableObjectNodeProjectorConstant::SetProjectorPosition);
+
+	FWidgetDirectionDelegate WidgetDirectionDelegate;
+	WidgetDirectionDelegate.BindUObject(&Node, &UCustomizableObjectNodeProjectorConstant::GetProjectorDirection);
+
+	FOnWidgetDirectionChangedDelegate OnWidgetDirectionChangedDelegate;
+	OnWidgetDirectionChangedDelegate.BindUObject(&Node, &UCustomizableObjectNodeProjectorConstant::SetProjectorDirection);
+
+	FWidgetUpDelegate WidgetUpDelegate;
+	WidgetUpDelegate.BindUObject(&Node, &UCustomizableObjectNodeProjectorConstant::GetProjectorUp);
+
+	FOnWidgetUpChangedDelegate OnWidgetUpChangedDelegate;
+	OnWidgetUpChangedDelegate.BindUObject(&Node, &UCustomizableObjectNodeProjectorConstant::SetProjectorUp);
+
+	FWidgetScaleDelegate WidgetScaleDelegate;
+	WidgetScaleDelegate.BindUObject(&Node, &UCustomizableObjectNodeProjectorConstant::GetProjectorScale);
+
+	FOnWidgetScaleChangedDelegate OnWidgetScaleChangedDelegate;
+	OnWidgetScaleChangedDelegate.BindUObject(&Node, &UCustomizableObjectNodeProjectorConstant::SetProjectorScale);
+
+	FWidgetAngleDelegate WidgetAngleDelegate;
+	WidgetAngleDelegate.BindUObject(&Node, &UCustomizableObjectNodeProjectorConstant::GetProjectorAngle);
+
+	FWidgetTrackingStartedDelegate WidgetTrackingStartedDelegate;
+	WidgetTrackingStartedDelegate.BindLambda([WeakNode = MakeWeakObjectPtr(&Node)]()
+	{
+		if (UCustomizableObjectNodeProjectorConstant* Node = WeakNode.Get())
+		{
+			Node->Modify();
+		}
+	});
+	
+	Viewport->ShowGizmoProjector(WidgetLocationDelegate, OnWidgetLocationChangedDelegate,
+		WidgetDirectionDelegate, OnWidgetDirectionChangedDelegate,
+		WidgetUpDelegate, OnWidgetUpChangedDelegate,
+		WidgetScaleDelegate, OnWidgetScaleChangedDelegate,
+		WidgetAngleDelegate,
+		ProjectorTypeDelegate,
+		WidgetColorDelegate,
+		WidgetTrackingStartedDelegate);
+}
+
+
+void FCustomizableObjectEditor::HideGizmoProjectorNodeProjectorConstant()
+{
+	if (GizmoType != EGizmoType::NodeProjectorConstant)
+	{
+		return;
+	}
+
+	GizmoType = EGizmoType::Hidden;
+	
+	Viewport->HideGizmoProjector();
+
+	const FGraphPanelSelectionSet SelectedNodes = GraphEditor->GetSelectedNodes();
+	for (FGraphPanelSelectionSet::TConstIterator NodeIt( SelectedNodes ); NodeIt; ++NodeIt)
+	{
+		const UObject* Node = *NodeIt;
+		if (Node->IsA<UCustomizableObjectNodeProjectorConstant>())
+		{
+			GraphEditor->ClearSelectionSet();
+			break;
+		}
+	}			
+}
+
+
+void FCustomizableObjectEditor::ShowGizmoProjectorNodeProjectorParameter(UCustomizableObjectNodeProjectorParameter& Node)
+{
+	if (GizmoType != EGizmoType::NodeProjectorParameter)
+	{
+		HideGizmo();
+		GizmoType = EGizmoType::NodeProjectorParameter;
+	}
+
+	SelectSingleNode(Node);
+	
+	FProjectorTypeDelegate ProjectorTypeDelegate;
+	ProjectorTypeDelegate.BindUObject(&Node, &UCustomizableObjectNodeProjectorParameter::GetProjectorType);		
+
+	FWidgetColorDelegate WidgetColorDelegate;
+	WidgetColorDelegate.BindLambda([]() { return FColor::Red; });
+	
+	FWidgetLocationDelegate WidgetLocationDelegate;
+	WidgetLocationDelegate.BindUObject(&Node, &UCustomizableObjectNodeProjectorParameter::GetProjectorDefaultPosition);
+
+	FOnWidgetLocationChangedDelegate OnWidgetLocationChangedDelegate;
+	OnWidgetLocationChangedDelegate.BindUObject(&Node, &UCustomizableObjectNodeProjectorParameter::SetProjectorDefaultPosition);
+
+	FWidgetDirectionDelegate WidgetDirectionDelegate;
+	WidgetDirectionDelegate.BindUObject(&Node, &UCustomizableObjectNodeProjectorParameter::GetProjectorDefaultDirection);
+
+	FOnWidgetDirectionChangedDelegate OnWidgetDirectionChangedDelegate;
+	OnWidgetDirectionChangedDelegate.BindUObject(&Node, &UCustomizableObjectNodeProjectorParameter::SetProjectorDefaultDirection);
+
+	FWidgetUpDelegate WidgetUpDelegate;
+	WidgetUpDelegate.BindUObject(&Node, &UCustomizableObjectNodeProjectorParameter::GetProjectorDefaultUp);
+
+	FOnWidgetUpChangedDelegate OnWidgetUpChangedDelegate;
+	OnWidgetUpChangedDelegate.BindUObject(&Node, &UCustomizableObjectNodeProjectorParameter::SetProjectorDefaultUp);
+
+	FWidgetScaleDelegate WidgetScaleDelegate;
+	WidgetScaleDelegate.BindUObject(&Node, &UCustomizableObjectNodeProjectorParameter::GetProjectorDefaultScale);
+
+	FOnWidgetScaleChangedDelegate OnWidgetScaleChangedDelegate;
+	OnWidgetScaleChangedDelegate.BindUObject(&Node, &UCustomizableObjectNodeProjectorParameter::SetProjectorDefaultScale);
+
+	FWidgetAngleDelegate WidgetAngleDelegate;
+	WidgetAngleDelegate.BindUObject(&Node, &UCustomizableObjectNodeProjectorParameter::GetProjectorDefaultAngle);
+	
+	FWidgetTrackingStartedDelegate WidgetTrackingStartedDelegate;
+	WidgetTrackingStartedDelegate.BindLambda([WeakNode = MakeWeakObjectPtr(&Node)]()
+	{
+		if (UCustomizableObjectNodeProjectorParameter* Node = WeakNode.Get())
+		{
+			Node->Modify();
+		}
+	});
+	
+	Viewport->ShowGizmoProjector(WidgetLocationDelegate, OnWidgetLocationChangedDelegate,
+		WidgetDirectionDelegate, OnWidgetDirectionChangedDelegate,
+		WidgetUpDelegate, OnWidgetUpChangedDelegate,
+		WidgetScaleDelegate, OnWidgetScaleChangedDelegate,
+		WidgetAngleDelegate,
+		ProjectorTypeDelegate,
+		WidgetColorDelegate,
+		WidgetTrackingStartedDelegate);
+}
+
+
+void FCustomizableObjectEditor::HideGizmoProjectorNodeProjectorParameter()
+{
+	if (GizmoType != EGizmoType::NodeProjectorParameter)
+	{
+		return;
+	}
+
+	GizmoType = EGizmoType::Hidden;
+
+	Viewport->HideGizmoProjector();
+
+	const FGraphPanelSelectionSet SelectedNodes = GraphEditor->GetSelectedNodes();
+	for (FGraphPanelSelectionSet::TConstIterator NodeIt( SelectedNodes ); NodeIt; ++NodeIt)
+	{
+		const UObject* Node = *NodeIt;
+		if (Node->IsA<UCustomizableObjectNodeProjectorParameter>())
+		{
+			GraphEditor->ClearSelectionSet();
+			break;
+		}
+	}	
+}
+
+
+void FCustomizableObjectEditor::ShowGizmoProjectorParameter(const FString& ParamName, int32 RangeIndex)
+{
+	if (GizmoType != EGizmoType::ProjectorParameter)
+	{
+		HideGizmo();		
+		GizmoType = EGizmoType::ProjectorParameter;
+	}
+	
+	FCustomizableObjectInstanceEditor::ShowGizmoProjectorParameter(ParamName, RangeIndex, SharedThis(this), Viewport, CustomizableInstanceDetailsView, ProjectorParameter, PreviewInstance);
+}
+
+
+void FCustomizableObjectEditor::HideGizmoProjectorParameter()
+{
+	if (GizmoType != EGizmoType::ProjectorParameter)
+	{
+		return;	
+	}
+
+	GizmoType = EGizmoType::Hidden;
+
+	FCustomizableObjectInstanceEditor::HideGizmoProjectorParameter(SharedThis(this), Viewport, CustomizableInstanceDetailsView);
+}
+
+
+void FCustomizableObjectEditor::ShowGizmoClipMorph(UCustomizableObjectNodeModifierClipMorph& Node)
+{
+	if (Node.BoneName == FName())
+	{	
+		return;
+	}
+
+	if (GizmoType != EGizmoType::ClipMorph)
+	{
+		HideGizmo();		
+		GizmoType = EGizmoType::ClipMorph;
+	}
+	
+	SelectSingleNode(Node);
+
+	Viewport->ShowGizmoClipMorph(Node);
+}
+
+
+void FCustomizableObjectEditor::HideGizmoClipMorph()
+{
+	if (GizmoType != EGizmoType::ClipMorph)
+	{
+		return;	
+	}
+	
+	GizmoType = EGizmoType::Hidden;
+
+	Viewport->HideGizmoClipMorph();
+
+	const FGraphPanelSelectionSet SelectedNodes = GraphEditor->GetSelectedNodes();
+	for (FGraphPanelSelectionSet::TConstIterator NodeIt( SelectedNodes ); NodeIt; ++NodeIt)
+	{
+		const UObject* Node = *NodeIt;
+		if (Node->IsA<UCustomizableObjectNodeModifierClipMorph>())
+		{
+			GraphEditor->ClearSelectionSet();
+			break;
+		}
+	}	
+}
+
+
+void FCustomizableObjectEditor::ShowGizmoClipMesh(UCustomizableObjectNodeModifierClipWithMesh& Node)
+{
+	UObject* ClipMesh = nullptr;
+	int32 LODIndex = 0;
+	int32 SectionIndex = 0;
+	int32 MaterialSlotIndex = 0;
+
+	if (const UEdGraphPin* ConnectedPin = FollowInputPin(*Node.ClipMeshPin()))
+	{
+		if (const UEdGraphNode* ConnectedNode = ConnectedPin->GetOwningNode())
+		{
+			int32 DummyIndex = 0;
+
+			if (const UCustomizableObjectNodeStaticMesh* StaticMeshNode = Cast<UCustomizableObjectNodeStaticMesh>(ConnectedNode))
+			{
+				ClipMesh = StaticMeshNode->GetMesh();
+				StaticMeshNode->GetPinSection(*ConnectedPin, LODIndex, SectionIndex, DummyIndex);
+				MaterialSlotIndex = SectionIndex;
+			}
+			else if (const UCustomizableObjectNodeSkeletalMesh* SkeletalMeshNode = Cast<UCustomizableObjectNodeSkeletalMesh>(ConnectedNode))
+			{
+				ClipMesh = SkeletalMeshNode->GetMesh();
+				SkeletalMeshNode->GetPinSection(*ConnectedPin, LODIndex, SectionIndex, DummyIndex);
+				MaterialSlotIndex = SkeletalMeshNode->GetSkeletalMaterialIndexFor(*ConnectedPin);
+			}
+			else if (const UCustomizableObjectNodeTable* TableNode = Cast<UCustomizableObjectNodeTable>(ConnectedNode))
+			{
+				ClipMesh = TableNode->GetColumnDefaultAssetByType<UObject>(ConnectedPin);
+
+				TableNode->GetPinLODAndSection(ConnectedPin, LODIndex, SectionIndex);
+				MaterialSlotIndex = SectionIndex;
+
+				if (TableNode->GetPinMeshType(ConnectedPin) == ETableMeshPinType::SKELETAL_MESH)
+				{
+					MaterialSlotIndex = TableNode->GetDefaultSkeletalMaterialIndexFor(*ConnectedPin);
+				}
+			}
+		}
+	}
+
+	if (ClipMesh && LODIndex >= 0 && MaterialSlotIndex >= 0)
+	{
+		if (GizmoType != EGizmoType::ClipMesh)
+		{
+			HideGizmo();
+			GizmoType = EGizmoType::ClipMesh;
+		}
+
+		SelectSingleNode(Node);
+
+		Viewport->ShowGizmoClipMesh(Node, *ClipMesh, LODIndex, SectionIndex, MaterialSlotIndex);
+	}
+}
+
+
+void FCustomizableObjectEditor::HideGizmoClipMesh()
+{
+	if (GizmoType != EGizmoType::ClipMesh)
+	{
+		return;	
+	}
+
+	GizmoType = EGizmoType::Hidden;
+
+	Viewport->HideGizmoClipMesh();
+
+	const FGraphPanelSelectionSet SelectedNodes = GraphEditor->GetSelectedNodes();
+	for (FGraphPanelSelectionSet::TConstIterator NodeIt( SelectedNodes ); NodeIt; ++NodeIt)
+	{
+		const UObject* Node = *NodeIt;
+		if (Node->IsA<UCustomizableObjectNodeModifierClipWithMesh>())
+		{
+			GraphEditor->ClearSelectionSet();
+			break;
+		}
+	}	
+}
+
+
+void FCustomizableObjectEditor::ShowGizmoLight(ULightComponent& InSelectedLight)
+{
+	if (GizmoType != EGizmoType::Light)
+	{
+		HideGizmo();
+		GizmoType = EGizmoType::Light;
+	}
+	
+	CustomSettings->SetSelectedLight(&InSelectedLight);
+
+	Viewport->ShowGizmoLight(InSelectedLight);
+	
+	CustomizableObjectEditorAdvancedPreviewSettings->Refresh();
+}
+
+
+void FCustomizableObjectEditor::HideGizmoLight()
+{
+	if (GizmoType != EGizmoType::Light)
+	{
+		return;	
+	}
+	
+	GizmoType = EGizmoType::Hidden;
+
+	CustomSettings->SetSelectedLight(nullptr);
+
+	Viewport->HideGizmoLight();
+
+	CustomizableObjectEditorAdvancedPreviewSettings->Refresh();
+}
+
+
+UCustomizableObjectEditorProperties* FCustomizableObjectEditor::GetEditorProperties()
+{
+	return EditorProperties;
+}
+
+
+void FCustomizableObjectEditor::PostUndo(bool bSuccess)
+{
+	if (bSuccess)
+	{
+		if (CustomizableObjectDetailsView.IsValid())
+		{
+			CustomizableObjectDetailsView->RemoveInvalidObjects();
+		}
+
+		if (CustomizableInstanceDetailsView.IsValid())
+		{
+			CustomizableInstanceDetailsView->RemoveInvalidObjects();
+		}
+
+		if (GraphNodeDetailsView.IsValid())
+		{
+			GraphNodeDetailsView->RemoveInvalidObjects();
+		}
+
+		GraphEditor->NotifyGraphChanged();
+		CustomizableObject->MarkPackageDirty();
+
+		FSlateApplication::Get().DismissAllMenus();
+	}
+}
+
+
+FString FCustomizableObjectEditor::GetDocumentationLink() const
+{
+	return DocumentationURL;
+}
+
+
+void FCustomizableObjectEditor::ExtendToolbar()
+{
+	TSharedPtr<FUICommandList> CommandList = GetToolkitCommands();
+
+	struct Local
+	{
+		static void FillToolbar(FToolBarBuilder& ToolbarBuilder, FCustomizableObjectEditor* Editor, TSharedPtr<FUICommandList> CommandList)
+		{
+			ToolbarBuilder.BeginSection("Compilation");
+			ToolbarBuilder.AddToolBarButton(FCustomizableObjectEditorCommands::Get().Compile);
+			ToolbarBuilder.AddToolBarButton(FCustomizableObjectEditorCommands::Get().CompileOnlySelected);
+			ToolbarBuilder.AddComboButton(
+				FUIAction(),
+				FOnGetContent::CreateSP(Editor, &FCustomizableObjectEditor::GenerateCompileOptionsMenuContent, CommandList.ToSharedRef()),
+				LOCTEXT("Compile_Options_Label", "Compile Options"),
+				LOCTEXT("Compile_Options_Tooltip", "Change Compile Options"),
+				TAttribute<FSlateIcon>(),
+				true);
+			ToolbarBuilder.EndSection();
+			
+			ToolbarBuilder.BeginSection("Information");
+			ToolbarBuilder.AddToolBarButton(FCustomizableObjectEditorCommands::Get().TextureAnalyzer);
+			ToolbarBuilder.AddToolBarButton(FCustomizableObjectEditorCommands::Get().PerformanceReport);
+			ToolbarBuilder.AddToolBarButton(FCustomizableObjectEditorCommands::Get().PerformanceAnalyzer);
+			ToolbarBuilder.EndSection();
+		}
+	};
+
+	TSharedPtr<FExtender> ToolbarExtender = MakeShareable(new FExtender);
+
+	ToolbarExtender->AddToolBarExtension(
+		"Asset",
+		EExtensionHook::After,
+		GetToolkitCommands(),
+		FToolBarExtensionDelegate::CreateStatic(&Local::FillToolbar, this, CommandList));
+
+	AddToolbarExtender(ToolbarExtender);
+
+	ICustomizableObjectEditorModule* CustomizableObjectEditorModule = &FModuleManager::LoadModuleChecked<ICustomizableObjectEditorModule>("CustomizableObjectEditor");
+	AddToolbarExtender(CustomizableObjectEditorModule->GetCustomizableObjectEditorToolBarExtensibilityManager()->GetAllExtenders());
+}
+
+
+TSharedRef<SWidget> FCustomizableObjectEditor::GenerateCompileOptionsMenuContent(TSharedRef<FUICommandList> InCommandList)
+{
+	const bool bShouldCloseWindowAfterMenuSelection = false;
+	FMenuBuilder MenuBuilder(bShouldCloseWindowAfterMenuSelection, InCommandList);
+
+	// settings
+	MenuBuilder.BeginSection("ResetCompileOptions");
+	{
+		MenuBuilder.AddMenuEntry(FCustomizableObjectEditorCommands::Get().ResetCompileOptions);
+	}
+	MenuBuilder.EndSection();
+
+	if (!CustomizableObject)
+	{
+		return MenuBuilder.MakeWidget();
+	}
+
+	MenuBuilder.BeginSection("Optimization", LOCTEXT("MutableCompileOptimizationHeading", "Optimization"));
+	{
+		// Level
+		CompileOptimizationStrings.Empty();
+		CompileOptimizationStrings.Add(MakeShareable(new FString(LOCTEXT("OptimizationNone", "None (Disable texture streaming)").ToString())));
+		CompileOptimizationStrings.Add(MakeShareable(new FString(LOCTEXT("OptimizationMin", "Minimal").ToString())));
+		CompileOptimizationStrings.Add(MakeShareable(new FString(LOCTEXT("OptimizationMax", "Maximum").ToString())));
+		check(CompileOptimizationStrings.Num() == UE_MUTABLE_MAX_OPTIMIZATION + 1);
+
+		if (CustomizableObject)
+		{
+			int32 SelectedOptimization = FMath::Clamp(CustomizableObject->GetPrivate()->OptimizationLevel, 0, CompileOptimizationStrings.Num() - 1);
+			CompileOptimizationCombo =
+				SNew(STextComboBox)
+				.OptionsSource(&CompileOptimizationStrings)
+				.InitiallySelectedItem(CompileOptimizationStrings[SelectedOptimization])
+				.OnSelectionChanged(this, &FCustomizableObjectEditor::OnChangeCompileOptimizationLevel)
+				;
+
+			MenuBuilder.AddWidget(CompileOptimizationCombo.ToSharedRef(), LOCTEXT("MutableCompileOptimizationLevel", "Optimization Level"));
+		}
+
+		{
+			CompileTextureCompressionStrings.Empty();
+			CompileTextureCompressionStrings.Add(MakeShareable(new FString(LOCTEXT("MutableTextureCompressionNone", "None").ToString())));
+			CompileTextureCompressionStrings.Add(MakeShareable(new FString(LOCTEXT("MutableTextureCompressionFast", "Fast").ToString())));
+			CompileTextureCompressionStrings.Add(MakeShareable(new FString(LOCTEXT("MutableTextureCompressionHighQuality", "High Quality").ToString())));
+
+			int32 SelectedCompression = FMath::Clamp(int32(CustomizableObject->GetPrivate()->TextureCompression), 0, CompileTextureCompressionStrings.Num() - 1);
+			CompileTextureCompressionCombo =
+				SNew(STextComboBox)
+				.OptionsSource(&CompileTextureCompressionStrings)
+				.InitiallySelectedItem(CompileTextureCompressionStrings[SelectedCompression])
+				.OnSelectionChanged(this, &FCustomizableObjectEditor::OnChangeCompileTextureCompressionType)
+				;
+
+			MenuBuilder.AddWidget(CompileTextureCompressionCombo.ToSharedRef(), LOCTEXT("MutableCompileTextureCompressionType", "Texture Compression"));
+		}
+
+		// Image tiling
+		// Unfortunately SNumericDropDown doesn't work with integers at the time of writing.
+		TArray<SNumericDropDown<float>::FNamedValue> TilingOptions;
+		TilingOptions.Add(SNumericDropDown<float>::FNamedValue(0, FText::FromString(TEXT("0")), FText::FromString(TEXT("Disabled"))));
+		TilingOptions.Add(SNumericDropDown<float>::FNamedValue(64, FText::FromString(TEXT("64")), FText::FromString(TEXT("64"))));
+		TilingOptions.Add(SNumericDropDown<float>::FNamedValue(128, FText::FromString(TEXT("128")), FText::FromString(TEXT("128"))));
+		TilingOptions.Add(SNumericDropDown<float>::FNamedValue(256, FText::FromString(TEXT("256")), FText::FromString(TEXT("256"))));
+		TilingOptions.Add(SNumericDropDown<float>::FNamedValue(512, FText::FromString(TEXT("512")), FText::FromString(TEXT("512"))));
+
+		CompileTilingCombo = SNew(SNumericDropDown<float>)
+			.DropDownValues(TilingOptions)
+			.Value_Lambda([this]() 
+				{ 
+					return CustomizableObject ? float(CustomizableObject->GetPrivate()->ImageTiling) : 0.0f;
+				})
+			.OnValueChanged_Lambda([this](float Value) 
+				{ 
+					if (CustomizableObject)
+					{
+						CustomizableObject->GetPrivate()->ImageTiling = int32(Value);
+						CustomizableObject->Modify();
+					}
+				});
+		MenuBuilder.AddWidget(CompileTilingCombo.ToSharedRef(), LOCTEXT("MutableCompileImageTiling", "Image Tiling"));
+
+		MenuBuilder.AddMenuEntry(FCustomizableObjectEditorCommands::Get().CompileOptions_UseDiskCompilation);
+	}
+	MenuBuilder.EndSection();
+
+	MenuBuilder.BeginSection("Packaging", LOCTEXT("MutableCompilePackagingHeading", "Packaging"));
+	{
+		// Unfortunately SNumericDropDown doesn't work with integers at the time of writing.
+		TArray<SNumericDropDown<float>::FNamedValue> EmbeddedOptions;
+		EmbeddedOptions.Add(SNumericDropDown<float>::FNamedValue(0, FText::FromString(TEXT("0")), FText::FromString(TEXT("Disabled"))));
+		EmbeddedOptions.Add(SNumericDropDown<float>::FNamedValue(16, FText::FromString(TEXT("16")), FText::FromString(TEXT("16"))));
+		EmbeddedOptions.Add(SNumericDropDown<float>::FNamedValue(64, FText::FromString(TEXT("64")), FText::FromString(TEXT("64"))));
+		EmbeddedOptions.Add(SNumericDropDown<float>::FNamedValue(256, FText::FromString(TEXT("256")), FText::FromString(TEXT("256"))));
+		EmbeddedOptions.Add(SNumericDropDown<float>::FNamedValue(512, FText::FromString(TEXT("512")), FText::FromString(TEXT("512"))));
+		EmbeddedOptions.Add(SNumericDropDown<float>::FNamedValue(1024, FText::FromString(TEXT("1024")), FText::FromString(TEXT("1024"))));
+		EmbeddedOptions.Add(SNumericDropDown<float>::FNamedValue(4096, FText::FromString(TEXT("4096")), FText::FromString(TEXT("4096"))));
+
+		EmbeddedDataLimitCombo = SNew(SNumericDropDown<float>)
+			.DropDownValues(EmbeddedOptions)
+			.Value_Lambda([this]()
+				{
+					return CustomizableObject ? float(CustomizableObject->GetPrivate()->EmbeddedDataBytesLimit) : 0.0f;
+				})
+			.OnValueChanged_Lambda([this](float Value)
+				{
+					if (CustomizableObject)
+					{
+						CustomizableObject->GetPrivate()->EmbeddedDataBytesLimit = uint64(Value);
+						CustomizableObject->Modify();
+					}
+				});
+			MenuBuilder.AddWidget(EmbeddedDataLimitCombo.ToSharedRef(), LOCTEXT("MutableCompileEmbeddedLimit", "Embedded Data Limit (Bytes)"));
+
+		// Packaging file size control.
+		TArray<SNumericDropDown<float>::FNamedValue> PackagedOptions;
+		PackagedOptions.Add(SNumericDropDown<float>::FNamedValue(0, FText::FromString(TEXT("0")), FText::FromString(TEXT("Split All"))));
+		PackagedOptions.Add(SNumericDropDown<float>::FNamedValue(16   * 1024, FText::FromString(TEXT("16 KB")), FText::FromString(TEXT("16 KB"))));
+		PackagedOptions.Add(SNumericDropDown<float>::FNamedValue(64   * 1024, FText::FromString(TEXT("64 KB")), FText::FromString(TEXT("64 KB"))));
+		PackagedOptions.Add(SNumericDropDown<float>::FNamedValue(1024 * 1024, FText::FromString(TEXT("1 MB")), FText::FromString(TEXT("1 MB"))));
+		PackagedOptions.Add(SNumericDropDown<float>::FNamedValue(64   * 1024 * 1024, FText::FromString(TEXT("64 MB")), FText::FromString(TEXT("64 MB"))));
+		PackagedOptions.Add(SNumericDropDown<float>::FNamedValue(256  * 1024 * 1024, FText::FromString(TEXT("256 MB")), FText::FromString(TEXT("256 MB"))));
+		PackagedOptions.Add(SNumericDropDown<float>::FNamedValue(1024 * 1024 * 1024, FText::FromString(TEXT("1 GB")), FText::FromString(TEXT("1 GB"))));
+
+		PackagedDataLimitCombo = SNew(SNumericDropDown<float>)
+			.DropDownValues(PackagedOptions)
+			.Value_Lambda([this]()
+				{
+					return CustomizableObject ? float(CustomizableObject->GetPrivate()->PackagedDataBytesLimit) : 0.0f;
+				})
+			.OnValueChanged_Lambda([this](float Value)
+				{
+					if (CustomizableObject)
+					{
+						CustomizableObject->GetPrivate()->PackagedDataBytesLimit = uint64(Value);
+						CustomizableObject->Modify();
+					}
+				});
+			MenuBuilder.AddWidget(PackagedDataLimitCombo.ToSharedRef(), LOCTEXT("MutableCompilePackagedLimit", "Packaged Data File Max Limit (Bytes)"));
+	}
+	MenuBuilder.EndSection();
+
+	// Debugging options
+	MenuBuilder.BeginSection("Debugger", LOCTEXT("MutableDebugger", "Debugger"));
+	{
+		MenuBuilder.AddMenuEntry(FCustomizableObjectEditorCommands::Get().Debug);
+	}
+	MenuBuilder.EndSection();
+
+	MenuBuilder.BeginSection("References", LOCTEXT("References", "References"));
+	{
+		MenuBuilder.AddMenuEntry(FCustomizableObjectEditorCommands::Get().CompileGatherReferences);
+		MenuBuilder.AddMenuEntry(FCustomizableObjectEditorCommands::Get().ClearGatheredReferences);
+	}
+	MenuBuilder.EndSection();
+
+	return MenuBuilder.MakeWidget();
+}
+
+
+FText FCustomizableObjectEditor::GetToolkitName() const
+{
+	FFormatNamedArguments Args;
+	Args.Add(TEXT("ObjectName"), FText::FromString(GetEditingObject()->GetName()));
+	Args.Add(TEXT("ToolkitName"), GetBaseToolkitName());
+	return FText::Format(LOCTEXT("AppLabelWithAssetName", "{ObjectName} - {ToolkitName}"), Args);
+}
+
+
+FString FCustomizableObjectEditor::GetWorldCentricTabPrefix() const
+{
+	return LOCTEXT( "WorldCentricTabPrefix", "CustomizableObject " ).ToString();
+}
+
+
+FLinearColor FCustomizableObjectEditor::GetWorldCentricTabColorScale() const
+{
+	return FLinearColor( 0.3f, 0.2f, 0.5f, 0.5f );
+}
+
+
+UCustomizableObject* FCustomizableObjectEditor::GetCustomizableObject()
+{
+	return CustomizableObject;
+}
+
+
+void FCustomizableObjectEditor::RefreshTool()
+{
+	if (ViewportClient)
+	{
+		ViewportClient->Invalidate();
+	}
+}
+
+
+TSharedPtr<SCustomizableObjectEditorViewportTabBody> FCustomizableObjectEditor::GetViewport()
+{
+	return Viewport;
+}
+
+
+void FCustomizableObjectEditor::OnObjectPropertySelectionChanged(FProperty* InProperty)
+{
+	CustomizableObject->PostEditChange();
+
+	if (ViewportClient)
+	{
+		ViewportClient->Invalidate();
+	}
+}
+
+
+void FCustomizableObjectEditor::OnInstancePropertySelectionChanged(FProperty* InProperty)
+{
+	if (ViewportClient)
+	{
+		ViewportClient->Invalidate();
+	}
+}
+
+
+void FCustomizableObjectEditor::OnObjectModified(UObject* Object)
+{
+	if (const UCustomizableObjectInstance* Instance = Cast<UCustomizableObjectInstance>(Object); !Instance)
+	{
+		// Sometimes when another CO is open in another editor window/tab, it triggers this callback, so prevent the modification of this object by a callback triggered by another one
+		if (UCustomizableObject* AuxCustomizableObject = Cast<UCustomizableObject>(Object))
+		{
+			AuxCustomizableObject->GetPrivate()->UpdateVersionId();
+		}
+		else if (UCustomizableObjectNode* Node = Cast<UCustomizableObjectNode>(Object))
+		{
+			if (UCustomizableObjectGraph* Graph = Cast<UCustomizableObjectGraph>(Node->GetOuter()))
+			{
+				if (UCustomizableObject* AuxOuterCustomizableObject = Cast<UCustomizableObject>(Graph->GetOuter()))
+				{
+					AuxOuterCustomizableObject->GetPrivate()->UpdateVersionId();
+				}
+			}
+		}
+		else if (UCustomizableObjectGraph* Graph = Cast<UCustomizableObjectGraph>(Object))
+		{
+			if (UCustomizableObject* AuxOuterCustomizableObject = Cast<UCustomizableObject>(Graph->GetOuter()))
+			{
+				AuxOuterCustomizableObject->GetPrivate()->UpdateVersionId();
+			}
+		}
+	}
+}
+
+
+void FCustomizableObjectEditor::CompileObject(bool bOnlySelectedParameters, bool bGatherReferences)
+{
+	// Resetting viewport parameters
+	Viewport->SetDrawDefaultUVMaterial();
+
+	if (CustomizableObject->GetPrivate()->Status.Get() == FCustomizableObjectStatus::EState::Loading)
+	{
+		FNotificationInfo Info(LOCTEXT("CustomizableObjectCompileTryLater", "Please wait until Customizable Object is loaded"));
+		Info.bFireAndForget = true;
+		Info.bUseThrobber = true;
+		Info.FadeOutDuration = 1.0f;
+		Info.ExpireDuration = 2.0f;
+		FSlateNotificationManager::Get().AddNotification(Info);
+		return;
+	}
+
+	if (!CustomizableObject->GetPrivate()->GetSource())
+	{
+		return;
+	}
+
+	TSharedRef<FCompilationRequest> CompileRequest = MakeShared<FCompilationRequest>(*CustomizableObject, true);
+	CompileRequest->GetCompileOptions().bSilentCompilation = false;
+	CompileRequest->GetCompileOptions().bGatherReferences = bGatherReferences;
+
+	if (bOnlySelectedParameters && PreviewInstance)
+	{
+		const TArray<FCustomizableObjectIntParameterValue>& IntParameters = PreviewInstance->GetPrivate()->GetDescriptor().GetIntParameters();
+		TMap<FString, FString> ParamNamesToSelectedOptions;
+		ParamNamesToSelectedOptions.Reserve(IntParameters.Num());
+
+		for (const FCustomizableObjectIntParameterValue& IntParam : IntParameters)
+		{
+			ParamNamesToSelectedOptions.Add(IntParam.ParameterName, IntParam.ParameterValueName);
+		}
+
+		CompileRequest->SetParameterNamesToSelectedOptions(ParamNamesToSelectedOptions);
+	}
+
+	ICustomizableObjectEditorModule::GetChecked().CompileCustomizableObject(CompileRequest);
+}
+
+
+void FCustomizableObjectEditor::DebugObject() const
+{
+	const TSharedPtr<SDockTab> NewMutableObjectTab = SNew(SDockTab)
+	.Label(FText::FromString(TEXT("Debugger")))
+	[
+		SNew(SMutableObjectViewer, CustomizableObject)
+	];
+	
+	// Spawn the debugger tab alongside the Graph Tab 
+	TabManager->InsertNewDocumentTab(GraphTabId, FTabManager::ESearchPreference::PreferLiveTab, NewMutableObjectTab.ToSharedRef());
+}
+
+
+void FCustomizableObjectEditor::ClearGatheredReferences()
+{
+	CustomizableObject->GetPrivate()->References = {};
+	CustomizableObject->Modify();
+}
+
+
+void FCustomizableObjectEditor::ResetCompileOptions()
+{
+	const FScopedTransaction Transaction(LOCTEXT("ResetCompilationOptionsTransaction", "Reset Compilation Options"));
+	CustomizableObject->Modify();
+	
+	UCustomizableObjectPrivate* DefaultObject = Cast<UCustomizableObjectPrivate>(CustomizableObject->GetPrivate()->StaticClass()->GetDefaultObject());
+	CustomizableObject->GetPrivate()->OptimizationLevel = DefaultObject->OptimizationLevel;
+	CustomizableObject->GetPrivate()->bUseDiskCompilation = DefaultObject->bUseDiskCompilation;
+	CustomizableObject->GetPrivate()->TextureCompression = DefaultObject->TextureCompression;
+	CustomizableObject->GetPrivate()->EmbeddedDataBytesLimit = DefaultObject->EmbeddedDataBytesLimit;
+	CustomizableObject->GetPrivate()->PackagedDataBytesLimit = DefaultObject->PackagedDataBytesLimit;
+	CustomizableObject->GetPrivate()->ImageTiling = DefaultObject->ImageTiling;
+}
+
+void FCustomizableObjectEditor::OnChangeCompileOptimizationLevel(TSharedPtr<FString> NewSelection, ESelectInfo::Type SelectInfo)
+{
+	const FScopedTransaction Transaction(LOCTEXT("ChangedOptimizationLevelTransaction", "Changed Optimization Level"));
+	CustomizableObject->Modify();
+	CustomizableObject->GetPrivate()->OptimizationLevel = CompileOptimizationStrings.Find(NewSelection);
+}
+
+
+void FCustomizableObjectEditor::CompileOptions_UseDiskCompilation_Toggled()
+{
+	const FScopedTransaction Transaction(LOCTEXT("ChangedEnableCompilingUsingTheDiskAsMemory", "Changed Enable compiling using the disk as memory"));
+	CustomizableObject->Modify();
+	CustomizableObject->GetPrivate()->bUseDiskCompilation = !CustomizableObject->GetPrivate()->bUseDiskCompilation;
+}
+
+
+bool FCustomizableObjectEditor::CompileOptions_UseDiskCompilation_IsChecked()
+{
+	return CustomizableObject->GetPrivate()->bUseDiskCompilation;
+}
+
+
+void FCustomizableObjectEditor::OnChangeCompileTextureCompressionType(TSharedPtr<FString> NewSelection, ESelectInfo::Type)
+{
+	const FScopedTransaction Transaction(LOCTEXT("ChangedTextureCompressionTransaction", "Changed Texture Compression Type"));
+	CustomizableObject->Modify();
+	CustomizableObject->GetPrivate()->TextureCompression = ECustomizableObjectTextureCompression(CompileTextureCompressionStrings.Find(NewSelection));
+}
+
+
+void FCustomizableObjectEditor::SaveAsset_Execute()
+{
+	if (PreviewInstance)
+	{
+		if (PreviewInstance->GetPrivate()->IsSelectedParameterProfileDirty())
+		{
+			PreviewInstance->GetPrivate()->SaveParametersToProfile(PreviewInstance->GetPrivate()->SelectedProfileIndex);
+		}
+	}
+
+	UPackage* Package = CustomizableObject->GetOutermost();
+
+	if (Package)
+	{
+		TArray<UPackage*> PackagesToSave;
+		PackagesToSave.Add(Package);
+
+		FEditorFileUtils::PromptForCheckoutAndSave(PackagesToSave, false, false);
+	}
+}
+
+
+void FCustomizableObjectEditor::DeleteSelectedNodes()
+{
+	if (!GraphEditor.IsValid())
+	{
+		return;
+	}
+
+	const FScopedTransaction Transaction(LOCTEXT("UEdGraphSchema_CustomizableObject", "Delete Nodes"));
+
+	const FGraphPanelSelectionSet SelectedNodes = GraphEditor->GetSelectedNodes();
+	GraphEditor->ClearSelectionSet();
+
+	for (FGraphPanelSelectionSet::TConstIterator NodeIt( SelectedNodes ); NodeIt; ++NodeIt)
+	{
+		if (UEdGraphNode* Node = Cast<UEdGraphNode>(*NodeIt))
+		{
+			if (Node->CanUserDeleteNode())
+			{
+				if (const UEdGraph* GraphObj = Node->GetGraph())
+				{
+					if (const UEdGraphSchema* Schema = GraphObj->GetSchema())
+					{
+						Schema->BreakNodeLinks(*Node);  // Required to notify to all connected nodes (UEdGraphNode::PinConnectionListChanged() and UEdGraphNode::PinConnectionListChanged(...))
+					}
+				}
+
+				Node->DestroyNode();
+			}
+		}
+	}
+}
+
+
+bool FCustomizableObjectEditor::CanDeleteNodes() const
+{
+	if (GraphEditor.IsValid() && GraphEditor->GetSelectedNodes().Num() > 0)
+	{
+		for (auto Itr = GraphEditor->GetSelectedNodes().CreateConstIterator(); Itr; ++Itr)
+		{
+			UCustomizableObjectNode* Node = Cast<UCustomizableObjectNode>(*Itr);
+
+			if (Node && !Node->CanUserDeleteNode())
+			{
+				return false;
+			}
+		}
+
+		return true;
+	}
+
+	return false;
+}
+
+
+void FCustomizableObjectEditor::DuplicateSelectedNodes()
+{
+	CopySelectedNodes();
+	PasteNodes();
+}
+
+
+bool FCustomizableObjectEditor::CanDuplicateSelectedNodes() const
+{
+	return CanCopyNodes();
+}
+
+
+void FCustomizableObjectEditor::OnSelectedGraphNodesChanged(const FGraphPanelSelectionSet& NewSelection)
+{
+	TArray<UObject*> Objects;
+	for (FGraphPanelSelectionSet::TConstIterator It(NewSelection); It; ++It)
+	{
+		Objects.Add(*It);
+	}
+
+	// Standard details
+	if (GraphNodeDetailsView.IsValid())
+	{
+		GraphNodeDetailsView->SetObjects(Objects);
+	}		
+	
+	if (!bRecursionGuard) // Calling the following functions will unselect some nodes causing OnSelectedGraphNodesChanged to be called again
+	{
+		TGuardValue<bool> RecursionGuard(bRecursionGuard, true);
+
+		if (Objects.Num() != 1)
+		{
+			HideGizmoClipMorph();
+			HideGizmoClipMesh();
+			HideGizmoProjectorNodeProjectorConstant();
+			HideGizmoProjectorNodeProjectorParameter();
+
+			for (UObject* Object : Objects) // Reselect the multiple selection. Clearly showing gizmos when selecting a node is a really bad idea. Remove on MTBL-1684
+			{
+				GraphEditor->SetNodeSelection(Cast<UEdGraphNode>(Object), true);
+			}
+			
+			return;
+		}
+
+		if (UCustomizableObjectNodeModifierClipMorph* NodeModifierClipMorph = Cast<UCustomizableObjectNodeModifierClipMorph>(Objects[0]))
+		{		
+			ShowGizmoClipMorph(*NodeModifierClipMorph);
+		}
+		else if (UCustomizableObjectNodeModifierClipWithMesh* NodeModifierClipWithMesh = Cast<UCustomizableObjectNodeModifierClipWithMesh>(Objects[0]))
+		{
+			ShowGizmoClipMesh(*NodeModifierClipWithMesh);
+		}
+		else if (UCustomizableObjectNodeProjectorConstant* NodeProjectorConstant = Cast<UCustomizableObjectNodeProjectorConstant>(Objects[0]))
+		{
+			ShowGizmoProjectorNodeProjectorConstant(*NodeProjectorConstant);
+		}
+		else if (UCustomizableObjectNodeProjectorParameter* NodeProjectorParameter = Cast<UCustomizableObjectNodeProjectorParameter>(Objects[0]))
+		{
+			ShowGizmoProjectorNodeProjectorParameter(*NodeProjectorParameter);		
+		}
+		else
+		{
+			HideGizmoClipMorph();	
+			HideGizmoClipMesh();
+			HideGizmoProjectorNodeProjectorParameter();
+			HideGizmoProjectorNodeProjectorConstant();
+		}
+	}
+}
+
+
+void FCustomizableObjectEditor::NotifyPostChange( const FPropertyChangedEvent& PropertyChangedEvent, FProperty* PropertyThatChanged  )
+{
+	// Is it a source graph node?
+	const UObject* OuterObject = PropertyThatChanged->GetOwner<UObject>();
+	const UClass* OuterClass = Cast<UClass>(OuterObject);
+	if (OuterClass && OuterClass->IsChildOf(UCustomizableObjectNode::StaticClass()))
+	{
+		FPropertyChangedEvent Event(PropertyThatChanged);
+		CustomizableObject->GetPrivate()->GetSource()->PostEditChangeProperty(Event);
+		CustomizableObject->PostEditChangeProperty(Event);
+
+		if (GraphEditor.IsValid())
+		{
+			GraphEditor->NotifyGraphChanged();
+		}
+	}
+}
+
+void FCustomizableObjectEditor::CopySelectedNodes()
+{
+	// Export the selected nodes and place the text on the clipboard
+	const FGraphPanelSelectionSet SelectedNodes = GraphEditor->GetSelectedNodes();
+
+	FString ExportedText;
+
+	for (FGraphPanelSelectionSet::TConstIterator SelectedIter(SelectedNodes); SelectedIter; ++SelectedIter)
+	{
+		if(UEdGraphNode* Node = Cast<UEdGraphNode>(*SelectedIter))
+		{
+			Node->PrepareForCopying();
+		}
+	}
+
+	FEdGraphUtilities::ExportNodesToText(SelectedNodes, /*out*/ ExportedText);
+
+	FPlatformApplicationMisc::ClipboardCopy(*ExportedText);
+
+	// Make sure Material remains the owner of the copied nodes
+	for (FGraphPanelSelectionSet::TConstIterator SelectedIter(SelectedNodes); SelectedIter; ++SelectedIter)
+	{
+		if (UMaterialGraphNode* Node = Cast<UMaterialGraphNode>(*SelectedIter))
+		{
+			Node->PostCopyNode();
+		}
+		else if (UMaterialGraphNode_Comment* Comment = Cast<UMaterialGraphNode_Comment>(*SelectedIter))
+		{
+			Comment->PostCopyNode();
+		}
+	}
+}
+
+bool FCustomizableObjectEditor::CanCopyNodes() const
+{
+	// If any of the nodes can be duplicated then we should allow copying
+	const FGraphPanelSelectionSet SelectedNodes = GraphEditor->GetSelectedNodes();
+	for (FGraphPanelSelectionSet::TConstIterator SelectedIter(SelectedNodes); SelectedIter; ++SelectedIter)
+	{
+		UEdGraphNode* Node = Cast<UEdGraphNode>(*SelectedIter);
+		if ((Node != NULL) && Node->CanDuplicateNode())
+		{
+			return true;
+		}
+	}
+
+	return false;
+}
+
+
+void FCustomizableObjectEditor::PasteNodes()
+{
+	PasteNodesHere(GraphEditor->GetPasteLocation());
+}
+
+
+void FCustomizableObjectEditor::PasteNodesHere(const FVector2D& Location)
+{
+	// Undo/Redo support
+	const FScopedTransaction Transaction( LOCTEXT("CustomizableObjectEditorPaste", "Customizable Object Editor Editor: Paste") );
+	CustomizableObject->GetPrivate()->GetSource()->Modify();
+	CustomizableObject->Modify();
+
+	// Clear the selection set (newly pasted stuff will be selected)
+	GraphEditor->ClearSelectionSet();
+
+	// Grab the text to paste from the clipboard.
+	FString TextToImport;
+	FPlatformApplicationMisc::ClipboardPaste(TextToImport);
+
+	// Import the nodes
+	TSet<UEdGraphNode*> PastedNodes;
+	FEdGraphUtilities::ImportNodesFromText(CustomizableObject->GetPrivate()->GetSource(), TextToImport, /*out*/ PastedNodes);
+	
+	//Average position of nodes so we can move them while still maintaining relative distances to each other
+	FVector2D AvgNodePosition(0.0f,0.0f);
+
+	for (TSet<UEdGraphNode*>::TIterator It(PastedNodes); It; ++It)
+	{
+		UEdGraphNode* Node = *It;
+		AvgNodePosition.X += Node->NodePosX;
+		AvgNodePosition.Y += Node->NodePosY;
+	}
+
+	if ( PastedNodes.Num() > 0 )
+	{
+		float InvNumNodes = 1.0f/float(PastedNodes.Num());
+		AvgNodePosition.X *= InvNumNodes;
+		AvgNodePosition.Y *= InvNumNodes;
+	}
+
+	for (TSet<UEdGraphNode*>::TIterator It(PastedNodes); It; ++It)
+	{
+		UEdGraphNode* Node = *It;
+		if (UCustomizableObjectNode* GraphNode = Cast<UCustomizableObjectNode>(Node))
+		{
+			// There can be only one default mesh paint texture.
+			//UMaterialExpressionTextureBase* TextureSample = Cast<UMaterialExpressionTextureBase>( NewExpression );
+			//if( TextureSample )
+			//{
+			//	TextureSample->IsDefaultMeshpaintTexture = false;
+			//}
+
+			//NewExpression->UpdateParameterGuid(true, true);
+
+			//UMaterialExpressionFunctionInput* FunctionInput = Cast<UMaterialExpressionFunctionInput>( NewExpression );
+			//if( FunctionInput )
+			//{
+			//	FunctionInput->ConditionallyGenerateId(true);
+			//	FunctionInput->ValidateName();
+			//}
+
+			//UMaterialExpressionFunctionOutput* FunctionOutput = Cast<UMaterialExpressionFunctionOutput>( NewExpression );
+			//if( FunctionOutput )
+			//{
+			//	FunctionOutput->ConditionallyGenerateId(true);
+			//	FunctionOutput->ValidateName();
+			//}
+		}
+		//else if (UMaterialGraphNode_Comment* CommentNode = Cast<UMaterialGraphNode_Comment>(Node))
+		//{
+		//	CommentNode->MaterialDirtyDelegate = Material->MaterialGraph->MaterialDirtyDelegate;
+		//	CommentNode->MaterialExpressionComment->Material = Material;
+		//	Material->EditorComments.Add(CommentNode->MaterialExpressionComment);
+		//}
+
+		// Select the newly pasted stuff
+		GraphEditor->SetNodeSelection(Node, true);
+
+		Node->NodePosX = (Node->NodePosX - AvgNodePosition.X) + Location.X ;
+		Node->NodePosY = (Node->NodePosY - AvgNodePosition.Y) + Location.Y ;
+
+		Node->SnapToGrid(SNodePanel::GetSnapGridSize());
+
+		// Give new node a different Guid from the old one
+		Node->CreateNewGuid();
+	}
+
+	// Force new pasted Material Expressions to have same connections as graph nodes
+	//Material->MaterialGraph->LinkMaterialExpressionsFromGraph();
+	
+	for (UEdGraphNode* PastedNode : PastedNodes)
+	{
+		if (UCustomizableObjectNode* TypedNode = Cast<UCustomizableObjectNode>(PastedNode))
+		{
+			TypedNode->PostBackwardsCompatibleFixup();
+		}
+	}
+	
+	// Update UI
+	GraphEditor->NotifyGraphChanged();
+
+	CustomizableObject->PostEditChange();
+	CustomizableObject->MarkPackageDirty();
+}
+
+
+bool FCustomizableObjectEditor::CanPasteNodes() const
+{
+	FString ClipboardContent;
+
+	FPlatformApplicationMisc::ClipboardPaste(ClipboardContent);
+
+	return FEdGraphUtilities::CanImportNodesFromText(CustomizableObject->GetPrivate()->GetSource(), ClipboardContent);
+}
+
+
+void FCustomizableObjectEditor::CutSelectedNodes()
+{
+	CopySelectedNodes();
+	// Cut should only delete nodes that can be duplicated
+	DeleteSelectedNodes();
+}
+
+
+bool FCustomizableObjectEditor::CanCutNodes() const
+{
+	return CanCopyNodes() && CanDeleteNodes();
+}
+
+
+void FCustomizableObjectEditor::CreateCommentBoxFromKey()
+{
+	CreateCommentBox(GraphEditor->GetPasteLocation());
+}
+
+UEdGraphNode* FCustomizableObjectEditor::CreateCommentBox(const FVector2D& InTargetPosition)
+{
+	//const FScopedTransaction Transaction(LOCTEXT("UEdGraphSchema_CustomizableObject", "Add Comment Box"));
+
+	UEdGraphNode_Comment* CommentTemplate = NewObject<UEdGraphNode_Comment>();
+
+	UEdGraphNode_Comment* NewComment = nullptr;
+	{
+		CustomizableObject->Modify();
+
+		// const FGraphPanelSelectionSet& SelectionSet = GraphEditor->GetSelectedNodes();
+		FSlateRect Bounds;
+		FVector2D Location;
+		FVector2D Size;
+
+		if (GraphEditor->GetBoundsForSelectedNodes(Bounds, 50.f))
+		{
+			Location.X = Bounds.Left;
+			Location.Y = Bounds.Top;
+			Size = Bounds.GetSize();
+		}
+		else
+		{
+			Location.X = InTargetPosition.X;
+			Location.Y = InTargetPosition.Y;
+			Size.X = 400;
+			Size.Y = 100;
+		}
+
+		NewComment = FEdGraphSchemaAction_NewNode::SpawnNodeFromTemplate<UEdGraphNode_Comment>(GraphEditor->GetCurrentGraph(), CommentTemplate, InTargetPosition, true);
+		NewComment->NodePosX = Location.X;
+		NewComment->NodePosY = Location.Y;
+		NewComment->NodeWidth = Size.X;
+		NewComment->NodeHeight = Size.Y;
+		NewComment->NodeComment = FString(TEXT("Comment"));
+	}
+
+	CustomizableObject->MarkPackageDirty();
+	GraphEditor->NotifyGraphChanged();
+
+	return NewComment;
+}
+
+
+TSharedPtr<SCustomizableObjectEditorAdvancedPreviewSettings> FCustomizableObjectEditor::GetAdvancedPreviewSettings()
+{
+	return CustomizableObjectEditorAdvancedPreviewSettings;
+}
+
+
+bool FCustomizableObjectEditor::ShowLightingSettings()
+{
+	return true;
+}
+
+
+bool FCustomizableObjectEditor::ShowProfileManagementOptions()
+{
+	return true;
+}
+
+
+void FCustomizableObjectEditor::FindProperty(const FProperty* Property, const void* InContainer, const FString& FindString, const UObject& Context, bool& bFound)
+{
+	if (!Property || !InContainer)
+	{
+		return;
+	}
+
+	const FString PropertyName = Property->GetDisplayNameText().ToString();
+	if (PropertyName.Contains(FindString))
+	{
+		LogSearchResult(Context, "Property Name", bFound, *PropertyName);
+		bFound = true;
+	}
+	
+	for (int32 Index = 0; Index < Property->ArrayDim; ++Index)
+	{
+		const uint8* ValuePtr = Property->ContainerPtrToValuePtr<uint8>(InContainer, Index);
+
+		if (const FStrProperty* StringProperty = CastField<FStrProperty>(Property))
+		{
+			const FString* StringResult = StringProperty->GetPropertyValuePtr(ValuePtr);
+			if (StringResult->Contains(FindString))
+			{
+				LogSearchResult(Context, "Property Value", bFound, *StringResult);
+				bFound = true;
+			}
+		}
+		else if (const FEnumProperty* EnumProperty = CastField<FEnumProperty>(Property))
+		{
+			const UEnum* EnumResult = EnumProperty->GetEnum();
+
+			const FString StringResult = EnumResult->GetDisplayNameTextByIndex(*ValuePtr).ToString();
+			if (StringResult.Contains(FindString))
+			{
+				LogSearchResult(Context, "Property Value", bFound, StringResult);
+				bFound = true;
+			}
+		}
+		else if (const FSoftObjectProperty* SoftObjectProperty = CastField<FSoftObjectProperty>(Property))
+		{
+			const FString ObjectPath = SoftObjectProperty->GetPropertyValuePtr(ValuePtr)->ToString();
+			if (ObjectPath.Contains(FindString))
+			{
+				LogSearchResult(Context, "Property Value", bFound, ObjectPath);
+				bFound = true;
+			}
+		}
+		else if (const FObjectPropertyBase* ObjectProperty = CastField<FObjectPropertyBase>(Property))
+		{
+			if (const UObject* ObjectValue = ObjectProperty->GetObjectPropertyValue(ValuePtr))
+			{
+				const FString Name = ObjectValue->GetName();
+			
+				if (ObjectValue->GetName().Contains(FindString))
+				{
+					LogSearchResult(Context, "Property Value", bFound, Name);
+					bFound = true;
+				}				
+			}
+		}
+		else if (const FStructProperty* StructProperty = CastField<FStructProperty>(Property))
+		{
+			for (TFieldIterator<FProperty> It(StructProperty->Struct); It; ++It)
+			{
+				FindProperty(*It, ValuePtr, FindString, Context, bFound);
+			}
+		}
+		else if (const FArrayProperty* ArrayProperty = CastField<FArrayProperty>(Property))
+		{
+			FScriptArrayHelper ArrayHelper(ArrayProperty, ValuePtr);
+			for (int32 ValueIdx = 0; ValueIdx < ArrayHelper.Num(); ++ValueIdx)
+			{
+				FindProperty(ArrayProperty->Inner, ArrayHelper.GetRawPtr(ValueIdx), FindString, Context, bFound);
+			}
+		}
+		else if (const FSetProperty* SetProperty = CastField<FSetProperty>(Property))
+		{
+			FScriptSetHelper SetHelper(SetProperty, ValuePtr);
+			for (FScriptSetHelper::FIterator SetIt = SetHelper.CreateIterator(); SetIt; ++SetIt)
+			{
+				FindProperty(SetProperty->ElementProp, SetHelper.GetElementPtr(SetIt), FindString, Context, bFound);
+			}
+		}
+		else if (const FMapProperty* MapProperty = CastField<FMapProperty>(Property))
+		{
+			FScriptMapHelper MapHelper(MapProperty, ValuePtr);
+			for (FScriptMapHelper::FIterator MapIt = MapHelper.CreateIterator(); MapIt; ++MapIt)
+			{
+				const uint8* MapValuePtr = MapHelper.GetPairPtr(MapIt);
+				FindProperty(MapProperty->KeyProp, MapValuePtr, FindString, Context, bFound);
+				FindProperty(MapProperty->ValueProp, MapValuePtr, FindString, Context, bFound);
+			}
+		}
+	}
+}
+
+
+void FCustomizableObjectEditor::OnPostCompile()
+{
+	CreatePreviewActor();
+	PreviewInstance->UpdateSkeletalMeshAsync(true, true);
+}
+
+
+void FCustomizableObjectEditor::OnEnterText(const FText& NewText, ETextCommit::Type TextType)
+{
+	if (TextType != ETextCommit::OnEnter)
+	{
+		return;
+	}
+	
+	if (!GraphEditor)
+	{
+		return;
+	}
+
+	const UEdGraph* Graph = GraphEditor->GetCurrentGraph();
+	if (!Graph)
+	{
+		return;
+	}
+
+	bool bFound = false;
+
+	const FString FindString = NewText.ToString();
+
+	for (TObjectPtr<UEdGraphNode> Node : Graph->Nodes)
+	{
+		if (!Node)
+		{
+			continue;
+		}
+
+		// Node names are not in the reflection system
+		const FString NodeName = Node->GetNodeTitle(ENodeTitleType::FullTitle).ToString().Replace(TEXT("\n"), TEXT(" "));
+		if (NodeName.Contains(NewText.ToString(), ESearchCase::IgnoreCase))
+		{
+			LogSearchResult(*Node, "Node", bFound, NodeName);
+			bFound = true;
+		}
+
+		// Pins are not in the reflection system
+		for (const UEdGraphPin* Pin : Node->GetAllPins())
+		{
+			const FString PinFriendlyName = Pin->PinFriendlyName.ToString();
+			if (PinFriendlyName.Contains(FindString))
+			{
+				LogSearchResult(*Node, "Pin", bFound, PinFriendlyName);
+				bFound = true;
+			}
+		}
+
+		// Find anything marked as a UPROPERTY
+		for (TFieldIterator<FProperty> It(Node->GetClass()); It; ++It)
+		{
+			FindProperty(*It, Node, FindString, *Node, bFound);
+		}
+	}
+
+	const FText Text = bFound ?
+		LOCTEXT("SearchCompleted", "Search completed") :
+		FText::FromString("No Results for: " + NewText.ToString());
+
+	FCustomizableObjectEditorLogger::CreateLog(Text)
+		.Category(ELoggerCategory::GraphSearch)
+		.CustomNotification()
+		.Log();
+}
+
+
+void FCustomizableObjectEditor::LogSearchResult(const UObject& Context, const FString& Type, bool bIsFirst, const FString& Result) const
+{
+	if (!bIsFirst)
+	{
+		FCustomizableObjectEditorLogger::CreateLog(LOCTEXT("SearchResults", "Search Results:"))
+		.Notification(false)
+		.Log();
+	}
+	
+	FCustomizableObjectEditorLogger::CreateLog(FText::FromString(Type + ": " + Result))
+	.Context(Context)
+	.BaseObject()
+	.Notification(false)
+	.Log();
+}
+
+
+void FCustomizableObjectEditor::OnUpdatePreviewInstance(UCustomizableObjectInstance* Instance)
+{
+	if (PreviewInstance->GetPrivate()->SkeletalMeshStatus != ESkeletalMeshStatus::Success)
+	{
+		Actor->GetRootComponent()->SetVisibility(false, true);
+	}
+	Actor->GetRootComponent()->UpdateBounds();
+	
+	if (TextureAnalyzer.IsValid())
+	{
+		TextureAnalyzer->RefreshTextureAnalyzerTable(PreviewInstance);
+	}
+}
+
+
+void FCustomizableObjectEditor::UpdateObjectProperties()
+{
+	if (CustomizableObjectDetailsView.IsValid())
+	{
+		CustomizableObjectDetailsView->ForceRefresh();
+	}
+}
+
+
+void FCustomizableObjectEditor::OnNodeTitleCommitted(const FText& NewText, ETextCommit::Type CommitInfo, UEdGraphNode* NodeBeingChanged)
+{
+	if (NodeBeingChanged)
+	{
+		const FScopedTransaction Transaction(LOCTEXT("RenameNode", "Rename Node"));
+		NodeBeingChanged->Modify();
+		NodeBeingChanged->OnRenameNode(NewText.ToString());
+	}
+}
+
+
+void FCustomizableObjectEditor::UpdateGraphNodeProperties()
+{
+	OnSelectedGraphNodesChanged(FGraphPanelSelectionSet());
+	OnSelectedGraphNodesChanged(GraphEditor->GetSelectedNodes());
+}
+
+
+void FCustomizableObjectEditor::OpenTextureAnalyzerTab()
+{
+	TabManager->TryInvokeTab(TextureAnalyzerTabId);
+}
+
+
+void FCustomizableObjectEditor::OpenPerformanceReportTab()
+{
+	TabManager->TryInvokeTab(PerformanceReportTabId);
+}
+
+void FCustomizableObjectEditor::OpenPerformanceAnalyzerTab()
+{
+	TabManager->TryInvokeTab(PerformanceAnalyzerTabId);
+}
+
+TSharedRef<SDockTab> FCustomizableObjectEditor::SpawnTab_TextureAnalyzer(const FSpawnTabArgs& Args)
+{
+	check(Args.GetTabId() == TextureAnalyzerTabId);
+	
+	return SNew(SDockTab)
+	.Label(LOCTEXT("Texture Analyzer", "Texture Analyzer"))
+	[
+		TextureAnalyzer.ToSharedRef()
+	];
+}
+
+
+TSharedRef<SDockTab> FCustomizableObjectEditor::SpawnTab_PerformanceReport(const FSpawnTabArgs& Args)
+{
+	check(Args.GetTabId() == PerformanceReportTabId);
+	check(CustomizableObject);
+
+	if (!PerformanceReport.IsValid())
+	{
+		PerformanceReport = SNew(SCustomizableObjecEditorPerformanceReport).CustomizableObject(CustomizableObject);
+	}
+
+	return SNew(SDockTab)
+	.Label(LOCTEXT("Performance Report", "Performance Report"))
+	[
+		PerformanceReport.ToSharedRef()
+	];
+}
+
+
+TSharedRef<SDockTab> FCustomizableObjectEditor::SpawnTab_PerformanceAnalyzer(const FSpawnTabArgs& Args)
+{
+	check(Args.GetTabId() == PerformanceAnalyzerTabId);
+	check(CustomizableObject);
+
+	if (!PerformanceAnalyzer.IsValid())
+	{
+		PerformanceAnalyzer = SNew(SCustomizableObjectEditorPerformanceAnalyzer).CustomizableObject(CustomizableObject);
+	}
+
+	return SNew(SDockTab)
+	.Label(LOCTEXT("Performance Analyzer", "Performance Analyzer"))
+	[
+		PerformanceAnalyzer.ToSharedRef()
+	];
+}
+
+
+TSharedRef<SDockTab> FCustomizableObjectEditor::SpawnTab_TagExplorer(const FSpawnTabArgs& Args)
+{
+	check(Args.GetTabId() == TagExplorerTabId);
+
+	return SNew(SDockTab)
+	.Label(LOCTEXT("Tag_Explorer", "Tag Explorer"))
+	[
+		TagExplorer.ToSharedRef()
+	];
+}
+
+
+UCustomizableObject* FCustomizableObjectEditor::GetAbsoluteCOParent(const UCustomizableObjectNodeObject* const Root)
+{
+	if (Root->ParentObject != nullptr)
+	{
+		//Get all the NodeObjects
+		TArray<UCustomizableObjectNodeObject*> ObjectNodes;
+		Root->ParentObject->GetPrivate()->GetSource()->GetNodesOfClass<UCustomizableObjectNodeObject>(ObjectNodes);
+		if (!ObjectNodes.IsEmpty())
+		{
+			//Getting the parent of the root
+			UCustomizableObjectNodeObject* FirstObjectNode = ObjectNodes[0];
+			if (FirstObjectNode->ParentObject == nullptr)
+			{
+				return Root->ParentObject;
+			}
+
+			return GetAbsoluteCOParent(FirstObjectNode);
+		}
+	}
+
+	return nullptr;
+}
+
+
+void FCustomizableObjectEditor::AddCachedReferencers(const FName& PathName, TArray<FName>& ArrayReferenceNames, TArray<FAssetData>& ArrayAssetData)
+{
+	ArrayReferenceNames.Empty();
+	FAssetRegistryModule& AssetRegistryModule = FModuleManager::LoadModuleChecked<FAssetRegistryModule>("AssetRegistry");
+	AssetRegistryModule.Get().GetReferencers(PathName, ArrayReferenceNames);
+
+	FARFilter Filter;
+	for (const FName& ReferenceName : ArrayReferenceNames)
+	{
+		bool IsCachedInAssetData = false;
+
+		const int32 MaxIndex = ArrayAssetData.Num();
+
+		for (int32 i = 0; i < MaxIndex; ++i)
+		{
+			if (ArrayAssetData[i].PackageName.ToString() == ReferenceName.ToString())
+			{
+				IsCachedInAssetData = true;
+			}
+		}
+
+		if (!IsCachedInAssetData && !ReferenceName.ToString().StartsWith(TEXT("/TempAutosave")))
+		{
+			Filter.PackageNames.Add(ReferenceName);
+		}
+	}
+
+	Filter.bIncludeOnlyOnDiskAssets = false;
+
+	TArray<FAssetData> ArrayAssetDataTemp;
+	AssetRegistryModule.Get().GetAssets(Filter, ArrayAssetDataTemp);
+
+	// Store only those which have static class type Customizable Object, to avoid loading not needed elements
+	const int32 MaxIndex = ArrayAssetDataTemp.Num();
+	for (int32 i = 0; i < MaxIndex; ++i)
+	{
+		if (ArrayAssetDataTemp[i].GetClass() == UCustomizableObject::StaticClass())
+		{
+			ArrayAssetData.Add(ArrayAssetDataTemp[i]);
+		}
+	}
+}
+
+
+void FCustomizableObjectEditor::CreatePreviewActor()
+{
+	if (Actor)
+	{
+		Actor->Destroy();
+	}
+	
+	Actor.Reset(ViewportClient->GetWorld()->SpawnActor<ASkeletalMeshActor>());
+
+	TArray<TWeakObjectPtr<UDebugSkelMeshComponent>> PreviewSkeletalMeshComponents;
+	
+	for (int32 ObjectComponentIndex = 0; ObjectComponentIndex < CustomizableObject->GetComponentCount(); ++ObjectComponentIndex)
+	{
+		UDebugSkelMeshComponent* DebugComponent = NewObject<UDebugSkelMeshComponent>(Actor.Get(), NAME_None, RF_Transient);
+		DebugComponent->bCastInsetShadow = true; // For better quality shadows in the editor previews, more similar to the in-game ones
+		DebugComponent->bCanHighlightSelectedSections = false;
+		DebugComponent->bComponentUseFixedSkelBounds = true; // First bounds computed would be using physics asset
+		DebugComponent->MarkRenderStateDirty();
+		DebugComponent->AttachToComponent(Actor->GetRootComponent(), FAttachmentTransformRules::KeepRelativeTransform);
+		DebugComponent->RegisterComponent();
+		
+		PreviewSkeletalMeshComponents.Add(DebugComponent);
+		
+		UCustomizableSkeletalComponent* CustomizableComponent = NewObject<UCustomizableSkeletalComponent>(DebugComponent);
+		CustomizableComponent->bSkipSetReferenceSkeletalMesh = true;
+		CustomizableComponent->CustomizableObjectInstance = PreviewInstance;
+		CustomizableComponent->SetComponentName(CustomizableObject->GetComponentName(ObjectComponentIndex));
+		CustomizableComponent->AttachToComponent(DebugComponent, FAttachmentTransformRules::KeepRelativeTransform);
+		CustomizableComponent->RegisterComponent();
+	}
+	
+	Viewport->SetPreviewActor(Actor.Get(), PreviewInstance, PreviewSkeletalMeshComponents);
+}
+
+
+void RemoveRestrictedChars(FString& String)
+{
+	// Remove restricted chars, according to FPaths::ValidatePath, RestrictedChars = "/?:&\\*\"<>|%#@^ ";
+
+	String = String.Replace(TEXT("/"), TEXT(""));
+	String = String.Replace(TEXT("?"), TEXT(""));
+	String = String.Replace(TEXT(":"), TEXT(""));
+	String = String.Replace(TEXT("&"), TEXT(""));
+	String = String.Replace(TEXT("\\"), TEXT(""));
+	String = String.Replace(TEXT("*"), TEXT(""));
+	String = String.Replace(TEXT("\""), TEXT(""));
+	String = String.Replace(TEXT("<"), TEXT(""));
+	String = String.Replace(TEXT(">"), TEXT(""));
+	String = String.Replace(TEXT("|"), TEXT(""));
+	String = String.Replace(TEXT("%"), TEXT(""));
+	String = String.Replace(TEXT("#"), TEXT(""));
+	String = String.Replace(TEXT("@"), TEXT(""));
+	String = String.Replace(TEXT("^"), TEXT(""));
+	String = String.Replace(TEXT(" "), TEXT(""));
+}
+
+
+void FCustomizableObjectEditor::OnCustomizableObjectStatusChanged(FCustomizableObjectStatus::EState PreviousState, const FCustomizableObjectStatus::EState CurrentState)
+{
+	if (PreviousState == FCustomizableObjectStatusTypes::EState::Loading)
+	{
+		if (CurrentState == FCustomizableObjectStatusTypes::EState::ModelLoaded)
+		{
+			CreatePreviewActor();
+			PreviewInstance->UpdateSkeletalMeshAsync(true, true);
+		}
+		else if (CurrentState == FCustomizableObjectStatusTypes::EState::NoModel)
+		{
+			CustomizableObject->ConditionalAutoCompile();
+		}
+	}
+}
+
+
+void FCustomizableObjectEditor::OnRenameNode()
+{
+	if (GraphEditor.IsValid())
+	{
+		const FGraphPanelSelectionSet& SelectedNodes = GraphEditor->GetSelectedNodes();
+
+		for (FGraphPanelSelectionSet::TConstIterator NodeIt(SelectedNodes); NodeIt; ++NodeIt)
+		{
+			// Rename only the first valid selected node
+			UEdGraphNode* SelectedNode = Cast<UEdGraphNode>(*NodeIt);
+			if (SelectedNode && SelectedNode->GetCanRenameNode())
+			{
+				GraphEditor->IsNodeTitleVisible(SelectedNode, true);
+				break;
+			}
+		}
+	}
+}
+
+
+bool FCustomizableObjectEditor::CanRenameNodes() const
+{
+	if (GraphEditor.IsValid())
+	{
+		const FGraphPanelSelectionSet& SelectedNodes = GraphEditor->GetSelectedNodes();
+
+		for (FGraphPanelSelectionSet::TConstIterator SelectedIter(SelectedNodes); SelectedIter; ++SelectedIter)
+		{
+			if (UEdGraphNode* Node = Cast<UEdGraphNode>(*SelectedIter))
+			{
+				return Node->GetCanRenameNode();
+			}
+		}
+	}
+
+	return false;
+}
+
+
+
+#undef LOCTEXT_NAMESPACE
