[mutable] Moved the Mutable plugin out of Experimental status into Beta.

diff --git a/Engine/Plugins/Mutable/Source/CustomizableObjectEditor/Private/MuCOE/CustomizableObjectEditor.h b/Engine/Plugins/Mutable/Source/CustomizableObjectEditor/Private/MuCOE/CustomizableObjectEditor.h
new file mode 100644
--- /dev/null
+++ b/Engine/Plugins/Mutable/Source/CustomizableObjectEditor/Private/MuCOE/CustomizableObjectEditor.h
@@ -0,0 +1,468 @@
+// Copyright Epic Games, Inc. All Rights Reserved.
+
+#pragma once
+
+#include "CustomizableObjectInstanceEditor.h"
+#include "EditorUndoClient.h"
+#include "GraphEditor.h"
+#include "Misc/NotifyHook.h"
+#include "MuCOE/ICustomizableObjectEditor.h"
+#include "MuCO/CustomizableObjectPrivate.h"
+#include "Widgets/Input/SNumericDropDown.h"
+
+#include "CustomizableObjectEditor.generated.h"
+
+enum class ECustomizableObjectProjectorType : uint8;
+namespace ESelectInfo { enum Type : int; }
+
+class SCustomizableObjectEditorPerformanceAnalyzer;
+class FCustomizableObjectEditorViewportClient;
+class FProperty;
+class FSpawnTabArgs;
+class FUICommandList;
+class IToolkitHost;
+class SCustomizableObjectEditorViewportTabBody;
+class SDockTab;
+class STextComboBox;
+class SWidget;
+class UCustomizableObject;
+class UCustomizableObjectInstance;
+class UCustomizableObjectNode;
+class UCustomizableObjectNodeModifierClipMorph;
+class UCustomizableObjectNodeModifierClipWithMesh;
+class UCustomizableObjectNodeObject;
+class UCustomizableObjectNodeProjectorConstant;
+class UCustomizableObjectNodeProjectorParameter;
+class UCustomizableSkeletalComponent;
+class UDebugSkelMeshComponent;
+class UEdGraph;
+class UMaterialInterface;
+class UPoseAsset;
+class UStaticMeshComponent;
+class UTexture;
+struct FAssetData;
+struct FFrame;
+struct FPropertyChangedEvent;
+template <typename FuncType> class TFunction;
+
+DECLARE_DELEGATE(FCreatePreviewInstanceFlagDelegate);
+
+extern void RemoveRestrictedChars(FString& String);
+
+
+enum class EGizmoType : uint8
+{
+	Hidden,
+	ProjectorParameter,
+	NodeProjectorConstant,
+	NodeProjectorParameter,
+	ClipMorph,
+	ClipMesh,
+	Light,
+};
+
+
+/**
+* Wrapper UObject class for the UCustomizableObjectInstance::FObjectInstanceUpdatedDelegate dynamic multicast delegate
+*/
+UCLASS()
+class UUpdateClassWrapper : public UObject
+{
+public:
+	GENERATED_BODY()
+
+	/** Method to assign for the callback */
+	UFUNCTION()
+	void DelegatedCallback(UCustomizableObjectInstance* Instance);
+
+	FCreatePreviewInstanceFlagDelegate Delegate;
+};
+
+
+/** Statistics for the Texture Analyzer */
+UCLASS(Transient, MinimalAPI, meta = (DisplayName = "Texture Stats"))
+class UCustomizableObjectEditorTextureStats : public UObject
+{
+	GENERATED_BODY()
+
+public:
+
+	/** Texture - double click to open */
+	UPROPERTY(VisibleAnywhere, AssetRegistrySearchable, Category = "Stats", meta = (DisplayName = "Texture", ColumnWidth = "40"))
+	FString TextureName;
+
+	/** Texture - double click to open */
+	UPROPERTY(VisibleAnywhere, AssetRegistrySearchable, Category = "Stats", meta = (DisplayName = "Name", ColumnWidth = "50"))
+	FString TextureParameterName;
+
+	/** Material - double click to open */
+	UPROPERTY(VisibleAnywhere, AssetRegistrySearchable, Category = "Stats", meta = (DisplayName = "Material", ColumnWidth = "50"))
+	FString MaterialName;
+
+	/** Parent Material - double click to open */
+	UPROPERTY(VisibleAnywhere, AssetRegistrySearchable, Category = "Stats", meta = (DisplayName = "Parent", ColumnWidth = "50"))
+	FString MaterialParameterName;
+
+	/** Used to open the texture in the editor*/
+	UPROPERTY()
+	TObjectPtr<UTexture> Texture;
+
+	/** Used to open the material in the editor*/
+	UPROPERTY()
+	TObjectPtr<UMaterialInterface> Material;
+
+	/** Used to open the parent material in the editor*/
+	UPROPERTY()
+	TObjectPtr<UMaterialInterface> ParentMaterial;
+
+	/** Resolution of the texture */
+	UPROPERTY(VisibleAnywhere, AssetRegistrySearchable, Category = "Stats", meta = (DisplayName = "Resolution X", ColumnWidth = "40", DisplayRight = "true"))
+	int32 ResolutionX;
+	UPROPERTY(VisibleAnywhere, AssetRegistrySearchable, Category = "Stats", meta = (DisplayName = "Resolution Y", ColumnWidth = "40", DisplayRight = "true"))
+	int32 ResolutionY;
+
+	/** The memory used in KB */
+	UPROPERTY(VisibleAnywhere, AssetRegistrySearchable, Category = "Stats", meta = (DisplayName = "Size Kb", ColumnWidth = "90"))
+	FString Size;
+
+	/** The texture format, e.g. PF_DXT1 */
+	UPROPERTY(VisibleAnywhere, AssetRegistrySearchable, Category = "Stats", meta = (ColumnWidth = "96"))
+	TEnumAsByte<EPixelFormat> Format;
+
+	/** LOD Bias for this texture. (Texture LODBias + Texture group) */
+	UPROPERTY(VisibleAnywhere, AssetRegistrySearchable, Category = "Stats", meta = (DisplayName = "LOD Bias", ColumnWidth = "70"))
+	int32 LODBias;
+	
+	/** Says if the texture is being streamed */
+	UPROPERTY(VisibleAnywhere, AssetRegistrySearchable, Category = "Stats", meta = (DisplayName = "Streamed", ColumnWidth = "70"))
+	FString IsStreamed;
+
+	/** The Level of detail group of the texture */
+	UPROPERTY(VisibleAnywhere, AssetRegistrySearchable, Category = "Stats", meta = (DisplayName = "Texture Group", ColumnWidth = "70"))
+	TEnumAsByte<enum TextureGroup> LODGroup;
+	
+	/** The Component of the texture */
+	UPROPERTY(VisibleAnywhere, AssetRegistrySearchable, Category = "Stats", meta = (DisplayName = "Component", ColumnWidth = "70"))
+	int32 Component;
+
+};
+
+
+/**
+ * CustomizableObject Editor class
+ */
+class FCustomizableObjectEditor : 
+	public ICustomizableObjectEditor, 
+	public FGCObject,
+	public FNotifyHook,
+	public FEditorUndoClient
+{
+public:
+	explicit FCustomizableObjectEditor(UCustomizableObject& ObjectToEdit);
+
+	virtual ~FCustomizableObjectEditor() override;
+
+	virtual void RegisterTabSpawners(const TSharedRef<class FTabManager>& TabManager) override;
+	virtual void UnregisterTabSpawners(const TSharedRef<class FTabManager>& TabManager) override;
+	void CreateCommentBoxFromKey();
+
+	/**
+	 * Initialize a new Customizable Object editor. Called immediately after construction.
+	 * Required due to being unable to use SharedThis in the constructor.
+	 *
+	 * See static Create(...) function.
+	 */
+	void InitCustomizableObjectEditor(const EToolkitMode::Type Mode, const TSharedPtr<IToolkitHost>& InitToolkitHost);
+
+	// FGCObject interface
+	virtual void AddReferencedObjects( FReferenceCollector& Collector ) override; 
+	virtual FString GetReferencerName() const override
+	{
+		return TEXT("FCustomizableObjectEditor");
+	}
+
+	// IToolkit interface
+	virtual FName GetToolkitFName() const override;
+	virtual FText GetBaseToolkitName() const override;
+	virtual FText GetToolkitName() const override;
+	virtual FString GetWorldCentricTabPrefix() const override;
+	virtual FLinearColor GetWorldCentricTabColorScale() const override;
+
+	// ICustomizableObjectEditor interface
+	virtual UCustomizableObject* GetCustomizableObject() override;
+	virtual void RefreshTool() override;
+	virtual TSharedPtr<SCustomizableObjectEditorViewportTabBody> GetViewport() override;
+	virtual UCustomizableObjectInstance* GetPreviewInstance() override;
+	virtual bool CanPasteNodes() const override;
+	virtual void PasteNodesHere(const FVector2D& Location) override;
+	virtual void SelectNode(const UEdGraphNode* Node) override;
+	virtual void ReconstructAllChildNodes(UCustomizableObjectNode& StartNode, const UClass& NodeType) override;
+	virtual UProjectorParameter* GetProjectorParameter() override;
+	virtual UCustomSettings* GetCustomSettings() override;
+	virtual void HideGizmo() override;
+	virtual void ShowGizmoProjectorNodeProjectorConstant(UCustomizableObjectNodeProjectorConstant& Node) override;
+	virtual void HideGizmoProjectorNodeProjectorConstant() override;
+	virtual void ShowGizmoProjectorNodeProjectorParameter(UCustomizableObjectNodeProjectorParameter& Node) override;
+	virtual void HideGizmoProjectorNodeProjectorParameter() override;
+	virtual void ShowGizmoProjectorParameter(const FString& ParamName, int32 RangeIndex = -1) override;
+	virtual void HideGizmoProjectorParameter() override;
+	virtual void ShowGizmoClipMorph(UCustomizableObjectNodeModifierClipMorph& Node) override;
+	virtual void HideGizmoClipMorph() override;
+	virtual void ShowGizmoClipMesh(UCustomizableObjectNodeModifierClipWithMesh& Node) override;
+	virtual void HideGizmoClipMesh() override;
+	virtual void ShowGizmoLight(ULightComponent& SelectedLight) override;
+	virtual void HideGizmoLight() override;
+	virtual UCustomizableObjectEditorProperties* GetEditorProperties() override;
+	virtual TSharedPtr< SCustomizableObjectEditorAdvancedPreviewSettings> GetAdvancedPreviewSettings() override;
+	virtual bool ShowLightingSettings() override;
+	virtual bool ShowProfileManagementOptions() override;
+
+	/** Select only this node only. Do nothing if already was only selected. */
+	void SelectSingleNode(UCustomizableObjectNode& Node);
+	
+	/** Called to undo the last action */
+	void UndoGraphAction();
+
+	/** Called to redo the last undone action */
+	void RedoGraphAction();
+
+	/** Utility method: Test whether the CO Node Object given as parameter is linked to any of the CO Node Object Group nodes
+	* in the Test CO given as parameter */
+	static bool GroupNodeIsLinkedToParentByName(UCustomizableObjectNodeObject* Node, UCustomizableObject* Test, const FString& ParentGroupName);
+
+	// FEditorUndoClient Interface
+	virtual void PostUndo(bool bSuccess) override;
+	virtual void PostRedo(bool bSuccess) override { PostUndo(bSuccess); }
+
+	// FNotifyHook interface
+	virtual void NotifyPostChange( const FPropertyChangedEvent& PropertyChangedEvent, FProperty* PropertyThatChanged ) override;
+
+	// Delegates
+	void DeleteSelectedNodes();
+	bool CanDeleteNodes() const;
+	void DuplicateSelectedNodes();
+	bool CanDuplicateSelectedNodes() const;
+	void OnSelectedGraphNodesChanged(const FGraphPanelSelectionSet& NewSelection);
+
+	/**
+	 * Called when a node's title is committed for a rename
+	 *
+	 * @param	NewText				New title text
+	 * @param	CommitInfo			How text was committed
+	 * @param	NodeBeingChanged	The node being changed
+	 */
+	void OnNodeTitleCommitted(const FText& NewText, ETextCommit::Type CommitInfo, UEdGraphNode* NodeBeingChanged);
+
+	virtual void UpdateGraphNodeProperties();
+
+	/** Callback to notify the editor when the PreviewInstance has been updated */
+	void OnUpdatePreviewInstance(UCustomizableObjectInstance* Instance);
+
+	/** Called when the Object Properties needs to be updated */
+	virtual void UpdateObjectProperties() override;
+
+	/** Debug the object as a raw mutable data in the internal tools. */
+	void DebugObject() const;
+
+	/** Clear game asset references saved in the Customizable Object. */
+	void ClearGatheredReferences();
+	
+private:
+	TSharedRef<SDockTab> SpawnTab_Viewport(const FSpawnTabArgs& Args);
+	TSharedRef<SDockTab> SpawnTab_ObjectProperties(const FSpawnTabArgs& Args);
+	TSharedRef<SDockTab> SpawnTab_InstanceProperties(const FSpawnTabArgs& Args);
+	TSharedRef<SDockTab> SpawnTab_Graph(const FSpawnTabArgs& Args);
+	TSharedRef<SDockTab> SpawnTab_GraphNodeProperties(const FSpawnTabArgs& Args);
+	TSharedRef<SDockTab> SpawnTab_AdvancedPreviewSettings(const FSpawnTabArgs& Args);
+	TSharedRef<SDockTab> SpawnTab_TextureAnalyzer(const FSpawnTabArgs& Args);
+	TSharedRef<SDockTab> SpawnTab_PerformanceReport(const FSpawnTabArgs& Args);
+	TSharedRef<SDockTab> SpawnTab_PerformanceAnalyzer(const FSpawnTabArgs& Args);
+	TSharedRef<SDockTab> SpawnTab_TagExplorer(const FSpawnTabArgs& Args);
+	
+	/** Binds commands associated with the Static Mesh Editor. */
+	void BindCommands();
+	
+	/** Command list for the graph editor */
+	TSharedPtr<FUICommandList> GraphEditorCommands;
+
+	/** Compile the Customizable Object.
+	 * 
+	 * @param bOnlySelectedParameters If true, compile only selected int parameters.
+	 * @param bGatherReferences If true, also gather asset references and save them in the Customizable Object. Marks the objects as modified. */
+	void CompileObject(bool bOnlySelectedParameters, bool bGatherReferences);
+	
+	// Compile options menu callbacks
+	TSharedRef<SWidget> GenerateCompileOptionsMenuContent(TSharedRef<FUICommandList> InCommandList);
+	void ResetCompileOptions();
+	TSharedPtr<STextComboBox> CompileOptimizationCombo;
+	TArray< TSharedPtr<FString> > CompileOptimizationStrings;
+	TSharedPtr<STextComboBox> CompileTextureCompressionCombo;
+	TArray< TSharedPtr<FString> > CompileTextureCompressionStrings;
+	TSharedPtr<SNumericDropDown<float>> CompileTilingCombo;
+	TSharedPtr<SNumericDropDown<float>> EmbeddedDataLimitCombo;
+	TSharedPtr<SNumericDropDown<float>> PackagedDataLimitCombo;
+
+	void CompileOptions_UseDiskCompilation_Toggled();
+	bool CompileOptions_UseDiskCompilation_IsChecked();
+	void OnChangeCompileOptimizationLevel(TSharedPtr<FString> NewSelection, ESelectInfo::Type);
+	void OnChangeCompileTextureCompressionType(TSharedPtr<FString> NewSelection, ESelectInfo::Type);
+
+	/** Save Customizable Object open in editor */
+	void SaveAsset_Execute() override;
+
+	/** Callback when selection changes in the Property Tree. */
+	void OnObjectPropertySelectionChanged(FProperty* InProperty);
+
+	/** Callback when selection changes in the Property Tree. */
+	void OnInstancePropertySelectionChanged(FProperty* InProperty);
+
+	/** Callback for the object modified event */
+	void OnObjectModified(UObject* Object);
+		
+	void CreateGraphEditorWidget(UEdGraph* InGraph);
+
+	/** Copy the currently selected nodes */
+	void CopySelectedNodes();
+	/** Whether we are able to copy the currently selected nodes */
+	bool CanCopyNodes() const;
+
+	/** Paste the contents of the clipboard */
+	void PasteNodes();
+
+	/** Cut the currently selected nodes */
+	void CutSelectedNodes();
+	/** Whether we are able to cut the currently selected nodes */
+	bool CanCutNodes() const;
+
+	/** On starting to rename node */
+	void OnRenameNode();
+
+	/** Check if node can be renamed */
+	bool CanRenameNodes() const;
+
+	virtual UEdGraphNode* CreateCommentBox(const FVector2D& NodePos) override;
+	
+	/** Searches a node that contains the inserted word */
+	void OnEnterText(const FText& NewText, ETextCommit::Type TextType);
+
+	/** Logs the search results of the search
+	 * @param Context The UObject we have found to be related with the searched string.
+	 * @param Type The type of relation with the searched word. It is a node, a value or maybe a variable?
+	 * @param bIsFirst Is this the first time we encountered something during our search?
+	 * @param Result The string containing the search word we are looking for in Node
+	 */
+	void LogSearchResult(const UObject& Context, const FString& Type, bool bIsFirst, const FString& Result) const;
+
+	/** Open the Texture Analyzer tab */
+	void OpenTextureAnalyzerTab();
+
+	/** Open the Performance Report tab */
+	void OpenPerformanceReportTab();
+
+	/** Open the Performance Analyzer tab */
+	void OpenPerformanceAnalyzerTab();
+	
+	/** Creates the actor and necessary components for the preview of the CO instance */
+	void CreatePreviewActor();
+
+	/** Recursively find any property that its name or value contains the given string.
+	  * @param Property Root property.
+	  * @param Container Root property container (address of the property value).
+	  * @param FindString String to find for.
+	  * @param Context UObject Context where this string has been found.
+	  * @param bFound Mark as true if any property has been found. */
+	void FindProperty(const FProperty* Property, const void* Container, const FString& FindString, const UObject& Context, bool& bFound);
+
+	void OnPostCompile();
+
+public:
+	void OnCustomizableObjectStatusChanged(FCustomizableObjectStatus::EState PreviousState, FCustomizableObjectStatus::EState CurrentState);
+
+	// Helpers to get the absolute parent of a Customizable Object
+	static UCustomizableObject* GetAbsoluteCOParent(const UCustomizableObjectNodeObject* const Root);
+	static void AddCachedReferencers(const FName& PathName, TArray<FName>& ArrayReferenceNames, TArray<FAssetData>& ArrayAssetData);
+
+	/**	The tab ids for all the tabs used */
+	static const FName ViewportTabId;
+	static const FName ObjectPropertiesTabId;
+	static const FName InstancePropertiesTabId;
+	static const FName GraphTabId;
+	static const FName GraphNodePropertiesTabId;
+	static const FName SystemPropertiesTabId;
+	static const FName AdvancedPreviewSettingsTabId;
+	static const FName TextureAnalyzerTabId;
+	static const FName PerformanceReportTabId;
+	static const FName PerformanceAnalyzerTabId;
+	static const FName TagExplorerTabId;
+	static const FName ObjectDebuggerTabId;
+	static const FName PopulationClassTagManagerTabId;
+
+private:
+	/** The currently viewed object. */
+	TObjectPtr<UCustomizableObject> CustomizableObject = nullptr;
+	TObjectPtr<UCustomizableObjectInstance> PreviewInstance = nullptr;
+
+	/** Preview Actor. All preview components are attached to this actor. */
+	TStrongObjectPtr<AActor> Actor;
+
+	/** List of open tool panels; used to ensure only one exists at any one time */
+	TMap< FName, TWeakPtr<class SDockableTab> > SpawnedToolPanels;
+
+	/** Preview Viewport widget */
+	TSharedPtr<SCustomizableObjectEditorViewportTabBody> Viewport;
+	TSharedPtr<FCustomizableObjectEditorViewportClient> ViewportClient;
+
+	TSharedPtr<IDetailsView> CustomizableInstanceDetailsView;
+
+
+	/** Property View */
+	TSharedPtr<class IDetailsView> CustomizableObjectDetailsView;
+
+	/** */
+	TSharedPtr<SGraphEditor> GraphEditor;
+	TSharedPtr<class IDetailsView> GraphNodeDetailsView;
+
+	/** Widget to select which node pins are visible. */
+	TSharedPtr<class SCustomizableObjectNodePinViewer> NodePinViewer;
+	
+	/** UObject class to be able to use the update callback */
+	TObjectPtr<UUpdateClassWrapper> HelperCallback;
+	
+	/** Scene preview settings widget, upcast of CustomizableObjectEditorAdvancedPreviewSettings */
+	TSharedPtr<SWidget> AdvancedPreviewSettingsWidget;
+
+	/** Scene preview settings widget */
+	TSharedPtr<class SCustomizableObjectEditorAdvancedPreviewSettings> CustomizableObjectEditorAdvancedPreviewSettings;
+	
+	/** Texture Analyzer table widget which shows the information of the transient textures used in the customizable object instance */
+	TSharedPtr<class SCustomizableObjecEditorTextureAnalyzer> TextureAnalyzer;
+
+	/** Performance report widget to test and analyze the current customizable object resource demands */
+	TSharedPtr<class SCustomizableObjecEditorPerformanceReport> PerformanceReport;
+
+	/** New performance analyzer widget */
+	TSharedPtr<SCustomizableObjectEditorPerformanceAnalyzer> PerformanceAnalyzer;
+	
+	/** Widget to explore all the tags related with the Customizable Object open in the editor */
+	TSharedPtr<class SCustomizableObjectEditorTagExplorer> TagExplorer;
+
+	/** Adds the customizable Object Editor commands to the default toolbar */
+	void ExtendToolbar();
+
+	/** URL to open when pressing the documentation button generated by UE */
+	const FString DocumentationURL{ TEXT("[REDACTED_PROFILE_URL]/Mutable-Documentation/wiki") };
+	
+	TObjectPtr<UProjectorParameter> ProjectorParameter = nullptr;
+
+	TObjectPtr<UCustomSettings> CustomSettings = nullptr;
+
+	TObjectPtr<UCustomizableObjectEditorProperties> EditorProperties = nullptr;
+	
+	bool bRecursionGuard = false;
+	
+	EGizmoType GizmoType = EGizmoType::Hidden;
+
+protected:
+	/** @return the documentation location for this editor */
+	virtual FString GetDocumentationLink() const override;
+};
diff --git a/Engine/Plugins/Mutable/Source/CustomizableObjectEditor/Private/MuCOE/CustomizableObjectEditorModule.cpp b/Engine/Plugins/Mutable/Source/CustomizableObjectEditor/Private/MuCOE/CustomizableObjectEditorModule.cpp
new file mode 100644
--- /dev/null
+++ b/Engine/Plugins/Mutable/Source/CustomizableObjectEditor/Private/MuCOE/CustomizableObjectEditorModule.cpp
@@ -0,0 +1,843 @@
+// Copyright Epic Games, Inc. All Rights Reserved.
+
+#include "MuCOE/CustomizableObjectEditorModule.h"
+
+#include "AssetToolsModule.h"
+#include "GameFramework/Pawn.h"
+#include "ISettingsModule.h"
+#include "ISettingsSection.h"
+#include "MessageLogModule.h"
+#include "Misc/MessageDialog.h"
+#include "Misc/DateTime.h"
+#include "Misc/Timespan.h"
+#include "HAL/FileManager.h"
+#include "MuCO/CustomizableObjectSystem.h"		// For defines related to memory function replacements.
+#include "MuCO/CustomizableObjectInstanceUsage.h"
+#include "MuCO/CustomizableSkeletalMeshActor.h"
+#include "MuCO/ICustomizableObjectModule.h"		// For instance editor command utility function
+#include "MuCO/UnrealPortabilityHelpers.h"
+#include "MuCOE/CustomizableInstanceDetails.h"
+#include "MuCOE/CustomizableObjectCustomSettings.h"
+#include "MuCOE/CustomizableObjectCustomSettingsDetails.h"
+#include "MuCOE/CustomizableObjectDetails.h"
+#include "MuCOE/CustomizableObjectEditor.h"
+#include "MuCOE/CustomizableObjectEditorLogger.h"
+#include "MuCOE/CustomizableObjectEditorSettings.h"
+#include "MuCOE/CustomizableObjectEditorStyle.h"
+#include "MuCOE/CustomizableObjectIdentifierCustomization.h"
+#include "MuCOE/CustomizableObjectInstanceEditor.h"
+#include "MuCOE/CustomizableObjectInstanceFactory.h"
+#include "MuCOE/CustomizableObjectVersionBridge.h"
+#include "MuCOE/CustomizableObjectNodeObjectGroupDetails.h"
+#include "MuCOE/Nodes/CustomizableObjectNodeCopyMaterial.h"
+#include "MuCOE/Nodes/CustomizableObjectNodeExternalPin.h"
+#include "MuCOE/Nodes/CustomizableObjectNodeExternalPinDetails.h"
+#include "MuCOE/Nodes/CustomizableObjectNodeLayoutBlocks.h"
+#include "MuCOE/Nodes/CustomizableObjectNodeMaterial.h"
+#include "MuCOE/Nodes/CustomizableObjectNodeMeshSectionDetails.h"
+#include "MuCOE/Nodes/CustomizableObjectNodeModifierEditMeshSection.h"
+#include "MuCOE/Nodes/CustomizableObjectNodeModifierEditMeshSectionDetails.h"
+#include "MuCOE/Nodes/CustomizableObjectNodeModifierExtendMeshSection.h"
+#include "MuCOE/Nodes/CustomizableObjectNodeModifierExtendMeshSectionDetails.h"
+#include "MuCOE/Nodes/CustomizableObjectNodeModifierMorphMeshSection.h"
+#include "MuCOE/Nodes/CustomizableObjectNodeModifierMorphMeshSectionDetails.h"
+#include "MuCOE/Nodes/CustomizableObjectNodeModifierRemoveMeshBlocks.h"
+#include "MuCOE/Nodes/CustomizableObjectNodeModifierRemoveMeshBlocksDetails.h"
+#include "MuCOE/Nodes/CustomizableObjectNodeModifierRemoveMesh.h"
+#include "MuCOE/Nodes/CustomizableObjectNodeModifierRemoveMeshDetails.h"
+#include "MuCOE/Nodes/CustomizableObjectNodeModifierClipMorph.h"
+#include "MuCOE/Nodes/CustomizableObjectNodeModifierClipMorphDetails.h"
+#include "MuCOE/Nodes/CustomizableObjectNodeModifierClipWithMesh.h"
+#include "MuCOE/Nodes/CustomizableObjectNodeModifierClipWithMeshDetails.h"
+#include "MuCOE/Nodes/CustomizableObjectNodeModifierClipWithUVMask.h"
+#include "MuCOE/Nodes/CustomizableObjectNodeModifierClipDeform.h"
+#include "MuCOE/Nodes/CustomizableObjectNodeMeshMorph.h"
+#include "MuCOE/Nodes/CustomizableObjectNodeMeshMorphDetails.h"
+#include "MuCOE/Nodes/CustomizableObjectNodeMeshReshapeCommon.h"
+#include "MuCOE/Nodes/CustomizableObjectNodeMeshReshapeSelectionDetails.h"
+#include "MuCOE/Nodes/CustomizableObjectNodeObjectDetails.h"
+#include "MuCOE/Nodes/CustomizableObjectNodeObjectGroup.h"
+#include "MuCOE/Nodes/CustomizableObjectNodeDetails.h"
+#include "MuCOE/Nodes/CustomizableObjectNodeProjectorConstant.h"
+#include "MuCOE/Nodes/CustomizableObjectNodeProjectorParameterDetails.h"
+#include "MuCOE/Nodes/CustomizableObjectNodeSkeletalMesh.h"
+#include "MuCOE/Nodes/CustomizableObjectNodeSkeletalMeshDetails.h"
+#include "MuCOE/Nodes/CustomizableObjectNodeStaticMesh.h"
+#include "MuCOE/Nodes/CustomizableObjectNodeTable.h"
+#include "MuCOE/Nodes/CustomizableObjectNodeTableDetails.h"
+#include "MuCOE/Nodes/CustomizableObjectNodeComponentMeshDetails.h"
+#include "MuCOE/Widgets/CustomizableObjectLODReductionSettings.h"
+#include "PropertyEditorModule.h"
+#include "AssetRegistry/AssetRegistryModule.h"
+#include "MuCO/CustomizableObjectInstance.h"
+#include "MuCOE/Nodes/CustomizableObjectNodeGroupProjectorParameter.h"
+#include "UObject/UObjectIterator.h"
+#include "Subsystems/PlacementSubsystem.h"
+#include "Components/SkeletalMeshComponent.h"
+#include "MuCOE/GraphTraversal.h"
+#include "MuCOE/CustomizableObjectInstanceBaker.h"
+#include "Editor.h"
+#include "Nodes/CustomizableObjectNodeComponentMeshDetails.h"
+
+class AActor;
+class FString;
+class ICustomizableObjectDebugger;
+class ICustomizableObjectEditor;
+class ICustomizableObjectInstanceEditor;
+class IToolkitHost;
+class UObject;
+
+
+const FName CustomizableObjectEditorAppIdentifier = FName(TEXT("CustomizableObjectEditorApp"));
+const FName CustomizableObjectInstanceEditorAppIdentifier = FName(TEXT("CustomizableObjectInstanceEditorApp"));
+const FName CustomizableObjectDebuggerAppIdentifier = FName(TEXT("CustomizableObjectDebuggerApp"));
+
+#define LOCTEXT_NAMESPACE "MutableSettings"
+
+/** Max timespan in days before a Saved/MutableStreamedDataEditor file is deleted. */
+constexpr int32 MaxAccessTimespan = 30;
+
+constexpr float ShowOnScreenCompileWarningsTickerTime = 1.0f;
+
+
+void ShowOnScreenCompileWarnings()
+{
+	TSet<const UCustomizableObject*> Objects;
+	
+	for (TObjectIterator<UCustomizableObjectInstanceUsage> CustomizableObjectInstanceUsage; CustomizableObjectInstanceUsage; ++CustomizableObjectInstanceUsage)
+	{
+		if (!IsValid(*CustomizableObjectInstanceUsage) || CustomizableObjectInstanceUsage->IsTemplate())
+		{
+			continue;
+		}
+		
+		const UCustomizableObjectInstance* Instance = CustomizableObjectInstanceUsage->GetCustomizableObjectInstance();
+		if (!Instance)
+		{
+			continue;
+		}
+
+		const UCustomizableObject* Object = Cast<UCustomizableObject>(Instance->GetCustomizableObject());
+		if (!Object)
+		{
+			continue;
+		}
+		
+		const USkeletalMeshComponent* Parent = Cast<USkeletalMeshComponent>(CustomizableObjectInstanceUsage->GetAttachParent());
+		if (!Parent)
+		{
+			continue;
+		}
+
+		const UWorld* World = Parent->GetWorld();
+		if (!World)
+		{
+			continue;
+		}
+
+		if (World->WorldType != EWorldType::PIE)
+		{
+			continue;
+		}
+
+		Objects.Add(Object);
+	}
+
+	for (const UCustomizableObject* Object : Objects)
+	{
+		if (Object->GetPrivate()->Status.Get() != FCustomizableObjectStatus::EState::ModelLoaded)
+		{
+			continue;
+		}
+	
+		// Show a warning if the compilation was not done with optimizations.
+		const uint64 KeyCompiledWithOptimization = reinterpret_cast<uint64>(Object);
+		if (!Object->GetPrivate()->GetModelResources().bIsCompiledWithOptimization)
+		{
+			FString Msg = FString::Printf(TEXT("Customizable Object [%s] was compiled without optimization."), *Object->GetName());
+			GEngine->AddOnScreenDebugMessage(KeyCompiledWithOptimization, ShowOnScreenCompileWarningsTickerTime * 2.0f, FColor::Yellow, Msg);
+		}
+		else
+		{
+			GEngine->RemoveOnScreenDebugMessage(KeyCompiledWithOptimization);
+		}
+		
+		const uint64 KeyCompiledOutOfDate = reinterpret_cast<uint64>(Object) + KEY_OFFSET_COMPILATION_OUT_OF_DATE; // Offset added to avoid collision with bIsCompiledWithOptimization warning
+		TArray<FName> OutOfDatePackages;
+		if (Object->GetPrivate()->IsCompilationOutOfDate(&OutOfDatePackages))
+		{
+			FString Msg = FString::Printf(TEXT("Customizable Object [%s] compilation out of date. See the Output Log for more information."), *Object->GetName());
+			GEngine->AddOnScreenDebugMessage(KeyCompiledOutOfDate, ShowOnScreenCompileWarningsTickerTime * 2.0f, FColor::Yellow, Msg);
+			
+			if (!GEngine->OnScreenDebugMessageExists(KeyCompiledOutOfDate))
+			{
+				UE_LOG(LogMutable, Verbose, TEXT("Warning: Customizable Object [%s] compilation out of date. Modified packages since last compilation:"), *Object->GetName());
+				for (const FName& OutOfDatePackage : OutOfDatePackages)
+				{
+					UE_LOG(LogMutable, Verbose, TEXT("%s"), *OutOfDatePackage.ToString());
+				}
+			}
+		}
+		else
+		{
+			GEngine->RemoveOnScreenDebugMessage(KeyCompiledOutOfDate);
+		}
+	}
+}
+
+void DeleteUnusedMutableStreamedDataEditorFiles()
+{
+	const FDateTime CurrentTime = FDateTime::Now();
+
+	const FString CompiledDataFolder = UCustomizableObjectPrivate::GetCompiledDataFolderPath();
+	const FString FileExtension = TEXT(".mut");
+
+	TArray<FString> Files;
+	IFileManager& FileManager = IFileManager::Get();
+	FileManager.FindFiles(Files, *CompiledDataFolder, *FileExtension);
+	
+	for (const FString& File : Files)
+	{
+		const FString FullFilePath = CompiledDataFolder + File;
+		const FDateTime AccessTimeStamp = FileManager.GetAccessTimeStamp(*(FullFilePath));
+		if (AccessTimeStamp == FDateTime::MinValue())
+		{
+			continue;
+		}
+
+		// Delete files that remain unused for more than MaxAccessTimespan
+		const FTimespan TimeSpan = CurrentTime - AccessTimeStamp;
+		if (TimeSpan.GetDays() > MaxAccessTimespan)
+		{
+			FileManager.Delete(*FullFilePath);
+		}
+	}
+}
+
+
+IMPLEMENT_MODULE( FCustomizableObjectEditorModule, CustomizableObjectEditor );
+
+
+static void* CustomMalloc(std::size_t Size_t, uint32_t Alignment)
+{
+	return FMemory::Malloc(Size_t, Alignment);
+}
+
+
+static void CustomFree(void* mem)
+{
+	return FMemory::Free(mem);
+}
+
+
+void FCustomizableObjectEditorModule::StartupModule()
+{
+	// Delete unused local compiled data
+	DeleteUnusedMutableStreamedDataEditorFiles();
+
+	// Register the thumbnail renderers
+	//UThumbnailManager::Get().RegisterCustomRenderer(UCustomizableObject::StaticClass(), UCustomizableObjectThumbnailRenderer::StaticClass());
+	//UThumbnailManager::Get().RegisterCustomRenderer(UCustomizableObjectInstance::StaticClass(), UCustomizableObjectInstanceThumbnailRenderer::StaticClass());
+	
+	// Property views
+	// Nodes
+	FPropertyEditorModule& PropertyModule = FModuleManager::LoadModuleChecked<FPropertyEditorModule>("PropertyEditor");
+	RegisterCustomDetails(PropertyModule, UCustomizableObjectNodeModifierEditMeshSection::StaticClass(), FOnGetDetailCustomizationInstance::CreateStatic(&FCustomizableObjectNodeModifierEditMeshSectionDetails::MakeInstance));
+	RegisterCustomDetails(PropertyModule, UCustomizableObjectNodeModifierExtendMeshSection::StaticClass(), FOnGetDetailCustomizationInstance::CreateStatic(&FCustomizableObjectNodeModifierExtendMeshSectionDetails::MakeInstance));
+	RegisterCustomDetails(PropertyModule, UCustomizableObjectNodeModifierRemoveMesh::StaticClass(), FOnGetDetailCustomizationInstance::CreateStatic(&FCustomizableObjectNodeModifierRemoveMeshDetails::MakeInstance));
+	RegisterCustomDetails(PropertyModule, UCustomizableObjectNodeModifierRemoveMeshBlocks::StaticClass(), FOnGetDetailCustomizationInstance::CreateStatic(&FCustomizableObjectNodeModifierRemoveMeshBlocksDetails::MakeInstance));
+	RegisterCustomDetails(PropertyModule, UCustomizableObjectNodeModifierMorphMeshSection::StaticClass(), FOnGetDetailCustomizationInstance::CreateStatic(&FCustomizableObjectNodeModifierMorphMeshSectionDetails::MakeInstance));
+	RegisterCustomDetails(PropertyModule, UCustomizableObjectNodeModifierClipMorph::StaticClass(), FOnGetDetailCustomizationInstance::CreateStatic(&FCustomizableObjectNodeModifierClipMorphDetails::MakeInstance));
+	RegisterCustomDetails(PropertyModule, UCustomizableObjectNodeModifierClipWithMesh::StaticClass(), FOnGetDetailCustomizationInstance::CreateStatic(&FCustomizableObjectNodeModifierClipWithMeshDetails::MakeInstance));
+	RegisterCustomDetails(PropertyModule, UCustomizableObjectNodeModifierClipWithUVMask::StaticClass(), FOnGetDetailCustomizationInstance::CreateStatic(&FCustomizableObjectNodeModifierBaseDetails::MakeInstance));
+	RegisterCustomDetails(PropertyModule, UCustomizableObjectNodeModifierClipDeform::StaticClass(), FOnGetDetailCustomizationInstance::CreateStatic(&FCustomizableObjectNodeModifierBaseDetails::MakeInstance));
+	RegisterCustomDetails(PropertyModule, UCustomizableObjectNodeObject::StaticClass(), FOnGetDetailCustomizationInstance::CreateStatic(&FCustomizableObjectNodeObjectDetails::MakeInstance));
+	RegisterCustomDetails(PropertyModule, UCustomizableObjectNodeObjectGroup::StaticClass(), FOnGetDetailCustomizationInstance::CreateStatic(&FCustomizableObjectNodeObjectGroupDetails::MakeInstance));
+	RegisterCustomDetails(PropertyModule, UCustomizableObjectNodeProjectorParameter::StaticClass(), FOnGetDetailCustomizationInstance::CreateStatic(&FCustomizableObjectNodeProjectorParameterDetails::MakeInstance));
+	RegisterCustomDetails(PropertyModule, UCustomizableObjectNodeProjectorConstant::StaticClass(), FOnGetDetailCustomizationInstance::CreateStatic(&FCustomizableObjectNodeProjectorParameterDetails::MakeInstance));
+	RegisterCustomDetails(PropertyModule, UCustomizableObjectNodeMeshMorph::StaticClass(), FOnGetDetailCustomizationInstance::CreateStatic(&FCustomizableObjectNodeMeshMorphDetails::MakeInstance));
+	RegisterCustomDetails(PropertyModule, UCustomizableObjectNodeExternalPin::StaticClass(), FOnGetDetailCustomizationInstance::CreateStatic(&FCustomizableObjectNodeExternalPinDetails::MakeInstance));
+	RegisterCustomDetails(PropertyModule, UCustomizableObjectNodeMaterial::StaticClass(), FOnGetDetailCustomizationInstance::CreateStatic(&FCustomizableObjectNodeMeshSectionDetails::MakeInstance));
+	RegisterCustomDetails(PropertyModule, UCustomizableObjectNodeSkeletalMesh::StaticClass(), FOnGetDetailCustomizationInstance::CreateStatic(&FCustomizableObjectNodeSkeletalMeshDetails::MakeInstance));
+	RegisterCustomDetails(PropertyModule, UCustomizableObjectNodeStaticMesh::StaticClass(), FOnGetDetailCustomizationInstance::CreateStatic(&FCustomizableObjectNodeDetails::MakeInstance));
+	RegisterCustomDetails(PropertyModule, UCustomizableObjectNodeTable::StaticClass(), FOnGetDetailCustomizationInstance::CreateStatic(&FCustomizableObjectNodeTableDetails::MakeInstance));
+	RegisterCustomDetails(PropertyModule, UCustomizableObjectNodeComponentMesh::StaticClass(), FOnGetDetailCustomizationInstance::CreateStatic(&FCustomizableObjectNodeComponentMeshDetails::MakeInstance));
+
+	// Other Objects
+	RegisterCustomDetails(PropertyModule, UCustomizableObject::StaticClass(), FOnGetDetailCustomizationInstance::CreateStatic(&FCustomizableObjectDetails::MakeInstance));
+	RegisterCustomDetails(PropertyModule, UCustomizableObjectInstance::StaticClass(), FOnGetDetailCustomizationInstance::CreateStatic(&FCustomizableInstanceDetails::MakeInstance));
+	RegisterCustomDetails(PropertyModule, UCustomSettings::StaticClass(), FOnGetDetailCustomizationInstance::CreateStatic(&FCustomizableObjectCustomSettingsDetails::MakeInstance));
+
+	//Custom properties
+	PropertyModule.RegisterCustomPropertyTypeLayout("CustomizableObjectIdentifier", FOnGetPropertyTypeCustomizationInstance::CreateStatic(&FCustomizableObjectIdentifierCustomization::MakeInstance));
+	PropertyModule.RegisterCustomPropertyTypeLayout(FMeshReshapeBoneReference::StaticStruct()->GetFName(), FOnGetPropertyTypeCustomizationInstance::CreateStatic(&FMeshReshapeBonesReferenceCustomization::MakeInstance));
+	PropertyModule.RegisterCustomPropertyTypeLayout(FBoneToRemove::StaticStruct()->GetFName(), FOnGetPropertyTypeCustomizationInstance::CreateStatic(&FCustomizableObjectLODReductionSettings::MakeInstance));
+	PropertyModule.RegisterCustomPropertyTypeLayout(NAME_StrProperty, FOnGetPropertyTypeCustomizationInstance::CreateStatic(&FCustomizableObjectStateParameterSelector::MakeInstance), MakeShared<FStatePropertyTypeIdentifier>());
+
+	PropertyModule.NotifyCustomizationModuleChanged();
+
+	// Register factory
+	FCoreDelegates::OnPostEngineInit.AddRaw(this,&FCustomizableObjectEditorModule::RegisterFactory);
+
+	// Additional UI style
+	FCustomizableObjectEditorStyle::Initialize();
+
+	RegisterSettings();
+
+	// Create the message log category
+	FMessageLogModule& MessageLogModule = FModuleManager::LoadModuleChecked<FMessageLogModule>("MessageLog");
+	MessageLogModule.RegisterLogListing(FName("Mutable"), LOCTEXT("MutableLog", "Mutable"));
+
+	CustomizableObjectEditor_ToolBarExtensibilityManager = MakeShareable(new FExtensibilityManager);
+	CustomizableObjectEditor_MenuExtensibilityManager = MakeShareable(new FExtensibilityManager);
+
+	LaunchCOIECommand = IConsoleManager::Get().RegisterConsoleCommand(
+		TEXT("mutable.OpenCOIE"),
+		TEXT("Looks for a Customizable Object Instance within the player pawn and opens its Customizable Object Instance Editor. Specify slot ID to control which component is edited."),
+		FConsoleCommandWithArgsDelegate::CreateStatic(&FCustomizableObjectEditorModule::OpenCOIE));
+
+	WarningsTickerHandle = FTSTicker::GetCoreTicker().AddTicker(TEXT("ShowOnScreenCompileWarnings"), ShowOnScreenCompileWarningsTickerTime, [](float)
+	{
+		ShowOnScreenCompileWarnings();
+		return true;
+	});
+
+	FEditorDelegates::PreBeginPIE.AddRaw(this, &FCustomizableObjectEditorModule::OnPreBeginPIE);
+}
+
+
+void FCustomizableObjectEditorModule::ShutdownModule()
+{
+	FEditorDelegates::PreBeginPIE.RemoveAll(this);
+
+	check(Compiler.GetNumRemainingWork() == 0);
+
+	if( FModuleManager::Get().IsModuleLoaded( "PropertyEditor" ) )
+	{
+		FPropertyEditorModule& PropertyModule = FModuleManager::GetModuleChecked<FPropertyEditorModule>("PropertyEditor");
+		
+		// Unregister Property views
+		for (const auto& ClassName : RegisteredCustomDetails)
+		{
+			PropertyModule.UnregisterCustomClassLayout(ClassName);
+		}
+
+		// Unregister Custom properties
+		PropertyModule.UnregisterCustomPropertyTypeLayout("CustomizableObjectIdentifier");
+		PropertyModule.UnregisterCustomPropertyTypeLayout(FMeshReshapeBoneReference::StaticStruct()->GetFName());
+		PropertyModule.UnregisterCustomPropertyTypeLayout(FBoneToRemove::StaticStruct()->GetFName());
+		PropertyModule.UnregisterCustomPropertyTypeLayout(NAME_StrProperty);
+
+		PropertyModule.NotifyCustomizationModuleChanged();
+	}
+
+	CustomizableObjectEditor_ToolBarExtensibilityManager.Reset();
+	CustomizableObjectEditor_MenuExtensibilityManager.Reset();
+
+	FCoreDelegates::OnPostEngineInit.RemoveAll(this);
+
+	FCustomizableObjectEditorStyle::Shutdown();
+
+	FTSTicker::GetCoreTicker().RemoveTicker(WarningsTickerHandle);
+}
+
+
+FCustomizableObjectEditorLogger& FCustomizableObjectEditorModule::GetLogger()
+{
+	return Logger;
+}
+
+
+bool FCustomizableObjectEditorModule::HandleSettingsSaved()
+{
+	UCustomizableObjectEditorSettings* CustomizableObjectSettings = GetMutableDefault<UCustomizableObjectEditorSettings>();
+
+	if (CustomizableObjectSettings != nullptr)
+	{
+		CustomizableObjectSettings->SaveConfig();
+		
+		FEditorCompileSettings CompileSettings;
+		CompileSettings.bIsMutableEnabled = !CustomizableObjectSettings->bDisableMutableCompileInEditor;
+		CompileSettings.bEnableAutomaticCompilation = CustomizableObjectSettings->bEnableAutomaticCompilation;
+		CompileSettings.bCompileObjectsSynchronously = CustomizableObjectSettings->bCompileObjectsSynchronously;
+		CompileSettings.bCompileRootObjectsOnStartPIE = CustomizableObjectSettings->bCompileRootObjectsOnStartPIE;
+		
+		UCustomizableObjectSystem::GetInstance()->EditorSettingsChanged(CompileSettings);
+	}
+
+    return true;
+}
+
+
+void FCustomizableObjectEditorModule::RegisterSettings()
+{
+	ISettingsModule* SettingsModule = FModuleManager::GetModulePtr<ISettingsModule>("Settings");
+
+	if (SettingsModule != nullptr)
+    {
+        ISettingsSectionPtr SettingsSectionPtr = SettingsModule->RegisterSettings("Project", "Plugins", "CustomizableObjectSettings",
+            LOCTEXT("MutableSettings_Setting", "Mutable"),
+            LOCTEXT("MutableSettings_Setting_Desc", "Mutable Settings"),
+            GetMutableDefault<UCustomizableObjectEditorSettings>()
+        );
+
+        if (SettingsSectionPtr.IsValid())
+        {
+            SettingsSectionPtr->OnModified().BindRaw(this, &FCustomizableObjectEditorModule::HandleSettingsSaved);
+        }
+
+		if (UCustomizableObjectSystem::GetInstance() != nullptr)
+		{
+			UCustomizableObjectEditorSettings* CustomizableObjectSettings = GetMutableDefault<UCustomizableObjectEditorSettings>();
+			if (CustomizableObjectSettings != nullptr)
+			{
+				FEditorCompileSettings CompileSettings;
+				CompileSettings.bIsMutableEnabled = !CustomizableObjectSettings->bDisableMutableCompileInEditor;
+				CompileSettings.bEnableAutomaticCompilation = CustomizableObjectSettings->bEnableAutomaticCompilation;
+				CompileSettings.bCompileObjectsSynchronously = CustomizableObjectSettings->bCompileObjectsSynchronously;
+				CompileSettings.bCompileRootObjectsOnStartPIE = CustomizableObjectSettings->bCompileRootObjectsOnStartPIE;
+				
+				UCustomizableObjectSystem::GetInstance()->EditorSettingsChanged(CompileSettings);
+			}
+		}
+    }
+}
+
+
+void FCustomizableObjectEditorModule::RegisterCustomDetails(FPropertyEditorModule& PropertyModule, const UClass* Class, FOnGetDetailCustomizationInstance DetailLayoutDelegate)
+{
+	const FName ClassName = FName(Class->GetName());
+	PropertyModule.RegisterCustomClassLayout(ClassName, DetailLayoutDelegate);
+
+	RegisteredCustomDetails.Add(ClassName);
+}
+
+
+void FCustomizableObjectEditorModule::OpenCOIE(const TArray<FString>& Arguments)
+{
+	int32 SlotID = INDEX_NONE;
+	if (Arguments.Num() >= 1)
+	{
+		SlotID = FCString::Atoi(*Arguments[0]);
+	}
+
+	const UWorld* CurrentWorld = []() -> const UWorld*
+	{
+		UWorld* WorldForCurrentCOI = nullptr;
+		const TIndirectArray<FWorldContext>& WorldContexts = GEngine->GetWorldContexts();
+		for (const FWorldContext& Context : WorldContexts)
+		{
+			if ((Context.WorldType == EWorldType::Game) && (Context.World() != NULL))
+			{
+				WorldForCurrentCOI = Context.World();
+			}
+		}
+		// Fall back to GWorld if we don't actually have a world.
+		if (WorldForCurrentCOI == nullptr)
+		{
+			WorldForCurrentCOI = GWorld;
+		}
+		return WorldForCurrentCOI;
+	}();
+	const int32 PlayerIndex = 0;
+
+	// Open the Customizable Object Instance Editor
+	if (UCustomizableObjectInstanceUsage* SelectedCustomizableObjectInstanceUsage = GetPlayerCustomizableObjectInstanceUsage(SlotID, CurrentWorld, PlayerIndex))
+	{
+		if (UCustomizableObjectInstance* COInstance = SelectedCustomizableObjectInstanceUsage->GetCustomizableObjectInstance())
+		{
+			FAssetToolsModule& AssetToolsModule = FModuleManager::LoadModuleChecked<FAssetToolsModule>(TEXT("AssetTools"));
+			TWeakPtr<IAssetTypeActions> WeakAssetTypeActions = AssetToolsModule.Get().GetAssetTypeActionsForClass(UCustomizableObjectInstance::StaticClass());
+
+			if (TSharedPtr<IAssetTypeActions> AssetTypeActions = WeakAssetTypeActions.Pin())
+			{
+				TArray<UObject*> AssetsToEdit;
+				AssetsToEdit.Add(COInstance);
+				AssetTypeActions->OpenAssetEditor(AssetsToEdit);
+			}
+		}
+	}
+}
+
+
+void FCustomizableObjectEditorModule::RegisterFactory()
+{
+	if (GEditor)
+	{
+		GEditor->ActorFactories.Add(NewObject<UCustomizableObjectInstanceFactory>());
+		if (UPlacementSubsystem* PlacementSubsystem = GEditor->GetEditorSubsystem<UPlacementSubsystem>())
+		{
+			PlacementSubsystem->RegisterAssetFactory(NewObject<UCustomizableObjectInstanceFactory>());
+		}
+	}
+}
+
+
+/** Recursively get all Customizable Objects that reference the given Customizable Object. */
+void GetReferencingCustomizableObjects(FName CustomizableObjectName, TArray<FName>& VisitedObjectNames, TArray<FName>& ObjectNames)
+{
+	if (VisitedObjectNames.Contains(CustomizableObjectName))
+	{
+		return;
+	}
+
+	VisitedObjectNames.Add(CustomizableObjectName);
+
+	TArray<FName> ReferencedObjectNames;
+	
+	const FAssetRegistryModule& AssetRegistryModule = FModuleManager::LoadModuleChecked<FAssetRegistryModule>("AssetRegistry");
+	AssetRegistryModule.Get().GetReferencers(CustomizableObjectName, ReferencedObjectNames, UE::AssetRegistry::EDependencyCategory::Package, UE::AssetRegistry::EDependencyQuery::Hard);
+
+	TArray<FAssetData> AssetDataArray;
+
+	FARFilter Filter;
+	Filter.PackageNames = MoveTemp(ReferencedObjectNames);
+	
+	AssetRegistryModule.Get().GetAssets(Filter, AssetDataArray);
+
+	for (FAssetData AssetData : AssetDataArray)
+	{
+		if (AssetData.GetClass() == UCustomizableObject::StaticClass())
+		{
+			FName ReferencedObjectName = AssetData.GetPackage()->GetFName();
+	
+			ObjectNames.Add(ReferencedObjectName);
+
+			GetReferencingCustomizableObjects(ReferencedObjectName, VisitedObjectNames, ObjectNames);
+		}			
+	}
+}
+
+
+void GetReferencingPackages(const UCustomizableObject& Object, TArray<FName>& ObjectNames)
+{
+	// Gather all child CustomizableObjects
+	TArray<FName> VisitedObjectNames;
+	GetReferencingCustomizableObjects(Object.GetPackage()->GetFName(), VisitedObjectNames, ObjectNames);
+
+	// Gather all tables which will composite the final tables
+	TArray<FName> CustomizableObjectNames = ObjectNames;
+	for (const FName& CustomizableObjectName : CustomizableObjectNames)
+	{
+		const TSoftObjectPtr<UObject> SoftObjectPtr = TSoftObjectPtr<UObject>(FSoftObjectPath(CustomizableObjectName.ToString()));
+
+		const UCustomizableObject* ChildCustomizableObject = Cast<UCustomizableObject>(SoftObjectPtr.LoadSynchronous());
+		if (!ChildCustomizableObject)
+		{
+			continue;
+		}
+
+		TArray<UCustomizableObjectNodeTable*> TableNodes;
+		ChildCustomizableObject->GetPrivate()->GetSource()->GetNodesOfClass(TableNodes);
+
+		FARFilter Filter;
+		Filter.ClassPaths.Add(FTopLevelAssetPath(UDataTable::StaticClass()));
+
+		for (const UCustomizableObjectNodeTable* TableNode : TableNodes)
+		{
+			TArray<FAssetData> DataTableAssets = TableNode->GetParentTables();
+
+			for (const FAssetData& DataTableAsset : DataTableAssets)
+			{
+				if (DataTableAsset.IsValid())
+				{
+					ObjectNames.AddUnique(DataTableAsset.PackageName);
+				}
+			}
+		}		
+	}
+}
+
+
+bool FCustomizableObjectEditorModule::IsCompilationOutOfDate(const UCustomizableObject& Object, TArray<FName>* OutOfDatePackages) const
+{
+	if (!Object.GetPrivate()->DirtyParticipatingObjects.IsEmpty())
+	{
+		if (OutOfDatePackages)
+		{
+			OutOfDatePackages->Append(Object.GetPrivate()->DirtyParticipatingObjects);
+		}
+		else
+		{
+			return true;
+		}
+	}
+	
+	const FAssetRegistryModule& AssetRegistryModule = FModuleManager::LoadModuleChecked<FAssetRegistryModule>("AssetRegistry");
+
+	const FModelResources& ModelResources = Object.GetPrivate()->GetModelResources();
+	for (const TTuple<FName, FGuid>& ParticipatingObject : ModelResources.ParticipatingObjects)
+	{
+		TSoftObjectPtr<UObject> SoftObjectPtr = TSoftObjectPtr<UObject>(FSoftObjectPath(ParticipatingObject.Key.ToString()));
+		if (SoftObjectPtr) // If loaded
+		{
+			PRAGMA_DISABLE_DEPRECATION_WARNINGS
+			const FGuid PackageGuid = SoftObjectPtr->GetPackage()->GetGuid();
+			PRAGMA_ENABLE_DEPRECATION_WARNINGS
+			
+			if (PackageGuid != ParticipatingObject.Value)
+			{
+				if (OutOfDatePackages)
+				{
+					OutOfDatePackages->AddUnique(ParticipatingObject.Key);
+				}
+				else
+				{
+					return true;
+				}
+			}
+		}
+		else // Not loaded
+		{
+			FAssetPackageData AssetPackageData;
+			const UE::AssetRegistry::EExists Result = AssetRegistryModule.Get().TryGetAssetPackageData(ParticipatingObject.Key, AssetPackageData);
+				
+			if (Result != UE::AssetRegistry::EExists::Exists)
+			{
+				if (OutOfDatePackages)
+				{
+					OutOfDatePackages->AddUnique(ParticipatingObject.Key);
+				}
+				else
+				{
+					return true;
+				}
+			}
+
+			PRAGMA_DISABLE_DEPRECATION_WARNINGS
+			const FGuid PackageGuid = AssetPackageData.PackageGuid;
+			PRAGMA_ENABLE_DEPRECATION_WARNINGS
+			
+			if (PackageGuid != ParticipatingObject.Value)
+			{
+				if (OutOfDatePackages)
+				{
+					OutOfDatePackages->AddUnique(ParticipatingObject.Key);
+				}
+				else
+				{
+					return true;
+				}
+			}
+		}
+	}
+
+	TArray<FName> ReferencingObjectNames;
+	GetReferencingPackages(Object, ReferencingObjectNames);
+	
+	for (const FName& ObjectName : ReferencingObjectNames)
+	{
+		TSoftObjectPtr<UObject> ReferencingObject = TSoftObjectPtr<UObject>(FSoftObjectPath(ObjectName.ToString())); 
+
+		if ((ReferencingObject && ReferencingObject->GetPackage()->IsDirty()) ||
+			!ModelResources.ParticipatingObjects.Contains(ObjectName)) // Must be in the participating objects, if not it means it did not exist when compiling the object.
+		{
+			if (OutOfDatePackages)
+			{
+				OutOfDatePackages->AddUnique(ObjectName);
+			}
+			else
+			{
+				return true;
+			}
+		}
+	}
+
+	if (OutOfDatePackages)
+	{
+		return !OutOfDatePackages->IsEmpty();	
+	}
+	else
+	{
+		return false;
+	}
+}
+
+
+bool FCustomizableObjectEditorModule::IsRootObject(const UCustomizableObject& Object) const
+{
+	return GraphTraversal::IsRootObject(Object);
+}
+
+FString FCustomizableObjectEditorModule::GetCurrentContentVersionForObject(const UCustomizableObject& Object) const
+{
+	if (Object.VersionBridge && Object.VersionBridge->GetClass()->ImplementsInterface(UCustomizableObjectVersionBridgeInterface::StaticClass()))
+	{
+		ICustomizableObjectVersionBridgeInterface* CustomizableObjectVersionBridgeInterface = Cast<ICustomizableObjectVersionBridgeInterface>(Object.VersionBridge);
+
+		if (CustomizableObjectVersionBridgeInterface)
+		{
+			return CustomizableObjectVersionBridgeInterface->GetCurrentVersionAsString();
+		}
+	}
+
+	return FString();
+}
+
+
+void FCustomizableObjectEditorModule::BakeCustomizableObjectInstance(UCustomizableObjectInstance* InTargetInstance, const FBakingConfiguration& InBakingConfig)
+{
+	UCustomizableObjectInstanceBaker* InstanceBaker = NewObject<UCustomizableObjectInstanceBaker>();
+
+	// Add the heap object to the root so we prevent it from being removed. It will get removed from there once it finishes it's work.
+	InstanceBaker->AddToRoot();
+	
+	// On baker operation completed just remove it from the root so it gets eventually destroyed by the GC system
+	const TSharedPtr<FOnBakerFinishedWork> OnBakerFinishedWorkCallback = MakeShared<FOnBakerFinishedWork>();
+	OnBakerFinishedWorkCallback->BindLambda([InstanceBaker]
+	{
+		InstanceBaker->RemoveFromRoot();
+	});
+	
+	// Ask for the baking of the instance
+	InstanceBaker->BakeInstance(InTargetInstance, InBakingConfig, OnBakerFinishedWorkCallback);
+}
+
+
+USkeletalMesh* FCustomizableObjectEditorModule::GetReferenceSkeletalMesh(const UCustomizableObject& Object, const FName& ComponentName) const
+{
+	UCustomizableObject* RootObject = GetRootObject(const_cast<UCustomizableObject*>(&Object));
+
+	TSet<UCustomizableObject*> Objects;
+	GetAllObjectsInGraph(RootObject, Objects);
+	
+	for (const UCustomizableObject* CurrentObject : Objects)
+	{
+		for (UEdGraphNode* Node : CurrentObject->GetPrivate()->GetSource()->Nodes)
+		{
+			if (UCustomizableObjectNodeComponentMesh* NodeComponentMesh = Cast<UCustomizableObjectNodeComponentMesh>(Node))
+			{
+				if (NodeComponentMesh->ComponentName == ComponentName)
+				{
+					return NodeComponentMesh->ReferenceSkeletalMesh;
+				}
+			}
+		}	
+	}
+
+	return {};
+}
+
+
+void FCustomizableObjectEditorModule::CompileCustomizableObject(const TSharedRef<FCompilationRequest>& InCompilationRequest, bool bForceRequest)
+{
+	if (IsRunningGame())
+	{
+		return;
+	}
+
+	CompileCustomizableObjects({ InCompilationRequest }, bForceRequest);
+}
+
+void FCustomizableObjectEditorModule::CompileCustomizableObjects(const TArray<TSharedRef<FCompilationRequest>>& InCompilationRequests, bool bForceRequests)
+{
+	check(IsInGameThread());
+
+	if (IsRunningGame())
+	{
+		return;
+	}
+
+	TArray<TSharedRef<FCompilationRequest>> FilteredAsyncRequests;
+	FilteredAsyncRequests.Reserve(InCompilationRequests.Num());
+
+	for (const TSharedRef<FCompilationRequest>& Request : InCompilationRequests)
+	{
+		const UCustomizableObject* CustomizableObject = Request->GetCustomizableObject();
+		if (!CustomizableObject)
+		{
+			continue;
+		}
+		
+		if (!Request->IsAsyncCompilation())
+		{
+			FCustomizableObjectCompiler* SyncCompiler = new FCustomizableObjectCompiler();
+			SyncCompiler->Compile(Request);
+			delete SyncCompiler;
+		}
+
+		else if (bForceRequests ||
+			(!CustomizableObject->GetPrivate()->IsLocked() && !Compiler.IsRequestQueued(Request)))
+		{
+			FilteredAsyncRequests.Add(Request);
+		}
+	}
+
+	Compiler.Compile(FilteredAsyncRequests);
+}
+
+
+int32 FCustomizableObjectEditorModule::Tick(bool bBlocking)
+{
+	Compiler.Tick(bBlocking);
+	return Compiler.GetNumRemainingWork();
+}
+
+
+void FCustomizableObjectEditorModule::CancelCompileRequests()
+{
+	Compiler.ForceFinishCompilation();
+	Compiler.ClearCompileRequests();
+}
+
+
+int32 FCustomizableObjectEditorModule::GetNumCompileRequests()
+{
+	return Compiler.GetNumRemainingWork();
+}
+
+
+void FCustomizableObjectEditorModule::OnPreBeginPIE(const bool bIsSimulatingInEditor)
+{
+	if (IsRunningGame() || !UCustomizableObjectSystem::IsActive())
+	{
+		return;
+	}
+
+	UCustomizableObjectSystem* System = UCustomizableObjectSystem::GetInstanceChecked();
+	if (!System->EditorSettings.bCompileRootObjectsOnStartPIE)
+	{
+		return;
+	}
+
+	// Find root customizable objects
+	FARFilter AssetRegistryFilter;
+	UE_MUTABLE_GET_CLASSPATHS(AssetRegistryFilter).Add(UE_MUTABLE_TOPLEVELASSETPATH(TEXT("/Script/CustomizableObject"), TEXT("CustomizableObject")));
+	AssetRegistryFilter.TagsAndValues.Add(FName("IsRoot"), FString::FromInt(1));
+
+	TArray<FAssetData> OutAssets;
+	const FAssetRegistryModule& AssetRegistryModule = FModuleManager::LoadModuleChecked<FAssetRegistryModule>("AssetRegistry");
+	AssetRegistryModule.Get().GetAssets(AssetRegistryFilter, OutAssets);
+
+	TArray<TSharedRef<FCompilationRequest>> Requests;
+	for (const FAssetData& Asset : OutAssets)
+	{
+		// If it is referenced by PIE it should be loaded
+		if (!Asset.IsAssetLoaded())
+		{
+			continue;
+		}
+
+		UCustomizableObject* Object = Cast<UCustomizableObject>(Asset.GetAsset());
+		if (!Object || Object->IsCompiled() || Object->GetPrivate()->IsLocked())
+		{
+			continue;
+		}
+
+		// Add uncompiled objects to the objects to cook list
+		TSharedRef<FCompilationRequest> NewRequest = MakeShared<FCompilationRequest>(*Object, true);
+		NewRequest->GetCompileOptions().bSilentCompilation = true;
+		Requests.Add(NewRequest);
+	}
+
+	if (!Requests.IsEmpty())
+	{
+		const FText Msg = FText::FromString(TEXT("Warning: one or more Customizable Objects used in PIE are uncompiled.\n\nDo you want to compile them?"));
+		if (FMessageDialog::Open(EAppMsgType::OkCancel, Msg) == EAppReturnType::Ok)
+		{
+			CompileCustomizableObjects(Requests);
+		}
+	}
+}
+
+#undef LOCTEXT_NAMESPACE
diff --git a/Engine/Plugins/Mutable/Source/CustomizableObjectEditor/Private/MuCOE/CustomizableObjectEditorSettings.cpp b/Engine/Plugins/Mutable/Source/CustomizableObjectEditor/Private/MuCOE/CustomizableObjectEditorSettings.cpp
new file mode 100644
--- /dev/null
+++ b/Engine/Plugins/Mutable/Source/CustomizableObjectEditor/Private/MuCOE/CustomizableObjectEditorSettings.cpp
@@ -0,0 +1,9 @@
+// Copyright Epic Games, Inc. All Rights Reserved.
+
+#include "MuCOE/CustomizableObjectEditorSettings.h"
+
+
+UCustomizableObjectEditorSettings::UCustomizableObjectEditorSettings()
+{
+
+}
diff --git a/Engine/Plugins/Mutable/Source/CustomizableObjectEditor/Private/MuCOE/CustomizableObjectInstanceFactory.cpp b/Engine/Plugins/Mutable/Source/CustomizableObjectEditor/Private/MuCOE/CustomizableObjectInstanceFactory.cpp
new file mode 100644
--- /dev/null
+++ b/Engine/Plugins/Mutable/Source/CustomizableObjectEditor/Private/MuCOE/CustomizableObjectInstanceFactory.cpp
@@ -0,0 +1,187 @@
+// Copyright Epic Games, Inc. All Rights Reserved.
+
+#include "MuCOE/CustomizableObjectInstanceFactory.h"
+
+#include "AssetRegistry/AssetData.h"
+#include "AssetRegistry/AssetRegistryModule.h"
+#include "Components/SkeletalMeshComponent.h"
+#include "Engine/SkeletalMesh.h"
+#include "Modules/ModuleManager.h"
+#include "MuCO/CustomizableObjectInstance.h"
+#include "MuCO/CustomizableObjectPrivate.h"
+#include "MuCO/CustomizableSkeletalComponent.h"
+#include "MuCO/CustomizableSkeletalMeshActor.h"
+#include "MuCOE/Nodes/CustomizableObjectNodeObject.h"
+#include "MuCO/UnrealPortabilityHelpers.h"
+
+#define LOCTEXT_NAMESPACE "CustomizableObjectInstanceFactory"
+
+
+UCustomizableObjectInstanceFactory::UCustomizableObjectInstanceFactory(const FObjectInitializer& ObjectInitializer)
+	: Super(ObjectInitializer)
+{
+    DisplayName = LOCTEXT("CustomizableObjectInstanceDisplayName", "Customizable Object Instance");
+    NewActorClass = ACustomizableSkeletalMeshActor::StaticClass();
+    bUseSurfaceOrientation = true;
+}
+
+
+void UCustomizableObjectInstanceFactory::PostSpawnActor(UObject* Asset, AActor* NewActor)
+{
+	Super::PostSpawnActor(Asset, NewActor);
+	UCustomizableObjectInstance* COInstance = Cast<UCustomizableObjectInstance>(Asset);
+	ACustomizableSkeletalMeshActor* NewCSMActor = CastChecked<ACustomizableSkeletalMeshActor>(NewActor);
+	if (NewCSMActor && COInstance)
+	{	
+		int32 NumComponents = COInstance->GetNumComponents();
+
+		for (int32 ComponentIndex = 0; ComponentIndex < NumComponents; ++ComponentIndex)
+		{
+			USkeletalMesh* SkeletalMesh = GetSkeletalMeshFromAsset(Asset, ComponentIndex);
+
+			if (ComponentIndex > 0 && ComponentIndex > NewCSMActor->GetNumComponents() - 1)
+			{
+				NewCSMActor->AttachNewComponent();
+			}
+
+			if (USkeletalMeshComponent* SkeletalMeshComp = NewCSMActor->GetSkeletalMeshComponentAt(ComponentIndex))
+			{
+				SkeletalMeshComp->UnregisterComponent();
+				UE_MUTABLE_SETSKINNEDASSET(SkeletalMeshComp, SkeletalMesh);
+
+				if (ComponentIndex == 0 && NewCSMActor->GetWorld()->IsGameWorld())
+				{
+					NewCSMActor->ReplicatedMesh = SkeletalMesh;
+				}
+
+				if (UCustomizableSkeletalComponent* CustomSkeletalComp = NewCSMActor->GetCustomizableSkeletalComponent(ComponentIndex))
+				{
+					CustomSkeletalComp->UnregisterComponent();
+					CustomSkeletalComp->CustomizableObjectInstance = COInstance;
+					CustomSkeletalComp->ComponentIndex = ComponentIndex;
+					//CustomSkeletalComp->ComponentName = GetComponentName(Asset,ComponentIndex);
+					CustomSkeletalComp->SetSkeletalMesh(SkeletalMesh);
+					CustomSkeletalComp->UpdateSkeletalMeshAsync();
+					CustomSkeletalComp->RegisterComponent();
+				}
+
+				SkeletalMeshComp->RegisterComponent();
+			}
+		}
+	}
+}
+
+UObject* UCustomizableObjectInstanceFactory::GetAssetFromActorInstance(AActor* ActorInstance)
+{
+	if (ACustomizableSkeletalMeshActor* CSMActor = CastChecked<ACustomizableSkeletalMeshActor>(ActorInstance))
+	{
+		if (CSMActor->GetNumComponents() > 0)
+		{
+			if (UCustomizableSkeletalComponent* CustomSkeletalComp = CSMActor->GetCustomizableSkeletalComponent(0))
+			{
+				return CustomSkeletalComp->CustomizableObjectInstance;
+			}
+		}
+	}
+	return nullptr;
+}
+
+bool UCustomizableObjectInstanceFactory::CanCreateActorFrom(const FAssetData& AssetData, FText& OutErrorMsg)
+{
+    if (!AssetData.IsValid() ||
+        (!AssetData.GetClass()->IsChildOf(UCustomizableObjectInstance::StaticClass())))
+    {
+        OutErrorMsg = LOCTEXT("NoCOISeq", "A valid customizable object instance must be specified.");
+        return false;
+    }
+
+    FAssetRegistryModule& AssetRegistryModule = FModuleManager::LoadModuleChecked<FAssetRegistryModule>(TEXT("AssetRegistry"));
+    FAssetData CustomizableObjectInstanceData;
+
+    if (AssetData.GetClass()->IsChildOf(UCustomizableObjectInstance::StaticClass()))
+    {
+        if (UCustomizableObjectInstance* CustomizableObjectInstance = Cast<UCustomizableObjectInstance>(AssetData.GetAsset()))
+        {
+            if (USkeletalMesh* SkeletalMesh = CustomizableObjectInstance->GetSkeletalMesh())
+            {
+                return true;
+            }
+            else
+            {
+                if (UCustomizableObject* CustomizableObject = CustomizableObjectInstance->GetCustomizableObject())
+                {
+                    return true;
+                }
+                else
+                {
+                    OutErrorMsg = LOCTEXT("NoCustomizableObjectInstance", "The UCustomizableObjectInstance does not have a customizableObject.");
+                    return false;
+                }
+            }
+        }
+        else
+        {
+            OutErrorMsg = LOCTEXT("NoCustomizableObjectInstanceIsNull", "The CustomizableObjectInstance is null.");
+        }
+    }
+
+    if (USkeletalMesh* SkeletalMeshCDO = Cast<USkeletalMesh>(AssetData.GetClass()->GetDefaultObject()))
+    {
+        if (SkeletalMeshCDO->HasCustomActorFactory())
+        {
+            return false;
+        }
+    }
+
+    return true;
+}
+
+
+USkeletalMesh* UCustomizableObjectInstanceFactory::GetSkeletalMeshFromAsset(UObject* Asset, int32 ComponentIndex) const
+{
+    UCustomizableObjectInstance* CustomizableObjectInstance = Cast<UCustomizableObjectInstance>(Asset);
+    if (!CustomizableObjectInstance)
+    {
+    	return nullptr;
+    }
+
+	UCustomizableObject* CustomizableObject = CustomizableObjectInstance->GetCustomizableObject();
+	if (!CustomizableObject)
+	{
+		return nullptr;
+	}
+	
+	const FName& ComponentName = CustomizableObject->GetComponentName(ComponentIndex);
+
+	if (USkeletalMesh* SkeletalMesh = CustomizableObjectInstance->GetComponentMeshSkeletalMesh(ComponentName))
+	{
+		return SkeletalMesh;
+	}
+    		
+	return CustomizableObject->GetComponentMeshReferenceSkeletalMesh(ComponentName);
+}
+
+
+//FName UCustomizableObjectInstanceFactory::GetComponentName(UObject* Asset, int32 ComponentIndex) const
+//{
+//	UCustomizableObjectInstance* CustomizableObjectInstance = Cast<UCustomizableObjectInstance>(Asset);
+//
+//	if (CustomizableObjectInstance != nullptr)
+//	{
+//		if (UCustomizableObject* CustomizableObject = CustomizableObjectInstance->GetCustomizableObject())
+//		{
+//			return CustomizableObject->GetPrivate()->Components[ComponentIndex].Name;
+//		}
+//	}
+//
+//	return FName();
+//}
+
+
+FQuat UCustomizableObjectInstanceFactory::AlignObjectToSurfaceNormal(const FVector& InSurfaceNormal, const FQuat& ActorRotation) const
+{
+    // Meshes align the Z (up) axis with the surface normal
+    return FindActorAlignmentRotation(ActorRotation, FVector(0.f, 0.f, 1.f), InSurfaceNormal);
+}
+
+#undef LOCTEXT_NAMESPACE
