[mutable] Moved the Mutable plugin out of Experimental status into Beta.

diff --git a/Engine/Plugins/Mutable/Source/CustomizableObjectEditor/Private/MuCOE/EdGraphSchema_CustomizableObject.cpp b/Engine/Plugins/Mutable/Source/CustomizableObjectEditor/Private/MuCOE/EdGraphSchema_CustomizableObject.cpp
new file mode 100644
--- /dev/null
+++ b/Engine/Plugins/Mutable/Source/CustomizableObjectEditor/Private/MuCOE/EdGraphSchema_CustomizableObject.cpp
@@ -0,0 +1,1430 @@
+// Copyright Epic Games, Inc. All Rights Reserved.
+
+#include "MuCOE/EdGraphSchema_CustomizableObject.h"
+
+#include "CustomizableObjectConnectionDrawingPolicy.h"
+#include "EdGraphNode_Comment.h"
+#include "Engine/SkeletalMesh.h"
+#include "Engine/StaticMesh.h"
+#include "Framework/Application/SlateApplication.h"
+#include "Framework/Commands/GenericCommands.h"
+#include "GraphEditorActions.h"
+#include "Materials/MaterialInterface.h"
+#include "MuCO/CustomizableObjectExtension.h"
+#include "MuCO/ICustomizableObjectModule.h"
+#include "MuCOE/CustomizableObjectEditorNodeContextCommands.h"
+#include "MuCOE/ICustomizableObjectEditor.h"
+#include "MuCOE/Nodes/CustomizableObjectNodeAnimationPose.h"
+#include "MuCOE/Nodes/CustomizableObjectNodeColorArithmeticOp.h"
+#include "MuCOE/Nodes/CustomizableObjectNodeColorConstant.h"
+#include "MuCOE/Nodes/CustomizableObjectNodeColorFromFloats.h"
+#include "MuCOE/Nodes/CustomizableObjectNodeColorParameter.h"
+#include "MuCOE/Nodes/CustomizableObjectNodeColorSwitch.h"
+#include "MuCOE/Nodes/CustomizableObjectNodeColorVariation.h"
+#include "MuCOE/Nodes/CustomizableObjectNodeCopyMaterial.h"
+#include "MuCOE/Nodes/CustomizableObjectNodeCurve.h"
+#include "MuCOE/Nodes/CustomizableObjectNodeEnumParameter.h"
+#include "MuCOE/Nodes/CustomizableObjectNodeExposePin.h"
+#include "MuCOE/Nodes/CustomizableObjectNodeExternalPin.h"
+#include "MuCOE/Nodes/CustomizableObjectNodeFloatConstant.h"
+#include "MuCOE/Nodes/CustomizableObjectNodeFloatParameter.h"
+#include "MuCOE/Nodes/CustomizableObjectNodeFloatSwitch.h"
+#include "MuCOE/Nodes/CustomizableObjectNodeFloatVariation.h"
+#include "MuCOE/Nodes/CustomizableObjectNodeFloatArithmeticOp.h"
+#include "MuCOE/Nodes/CustomizableObjectNodeGroupProjectorParameter.h"
+#include "MuCOE/Nodes/CustomizableObjectNodeLayoutBlocks.h"
+#include "MuCOE/Nodes/CustomizableObjectNodeMaterialVariation.h"
+#include "MuCOE/Nodes/CustomizableObjectNodeMaterialSwitch.h"
+#include "MuCOE/Nodes/CustomizableObjectNodeModifierExtendMeshSection.h"
+#include "MuCOE/Nodes/CustomizableObjectNodeModifierEditMeshSection.h"
+#include "MuCOE/Nodes/CustomizableObjectNodeModifierRemoveMesh.h"
+#include "MuCOE/Nodes/CustomizableObjectNodeModifierRemoveMeshBlocks.h"
+#include "MuCOE/Nodes/CustomizableObjectNodeModifierMorphMeshSection.h"
+#include "MuCOE/Nodes/CustomizableObjectNodeModifierClipDeform.h"
+#include "MuCOE/Nodes/CustomizableObjectNodeModifierClipMorph.h"
+#include "MuCOE/Nodes/CustomizableObjectNodeModifierClipWithMesh.h"
+#include "MuCOE/Nodes/CustomizableObjectNodeModifierClipWithUVMask.h"
+#include "MuCOE/Nodes/CustomizableObjectNodeMeshGeometryOperation.h"
+#include "MuCOE/Nodes/CustomizableObjectNodeMeshMorph.h"
+#include "MuCOE/Nodes/CustomizableObjectNodeMeshMorphStackApplication.h"
+#include "MuCOE/Nodes/CustomizableObjectNodeMeshMorphStackDefinition.h"
+#include "MuCOE/Nodes/CustomizableObjectNodeMeshReshape.h"
+#include "MuCOE/Nodes/CustomizableObjectNodeMeshSwitch.h"
+#include "MuCOE/Nodes/CustomizableObjectNodeMeshVariation.h"
+#include "MuCOE/Nodes/CustomizableObjectNodeObjectChild.h"
+#include "MuCOE/Nodes/CustomizableObjectNodeObjectGroup.h"
+#include "MuCOE/Nodes/CustomizableObjectNodeComponentPassthroughMesh.h"
+#include "MuCOE/Nodes/CustomizableObjectNodeComponentVariation.h"
+#include "MuCOE/Nodes/CustomizableObjectNodeComponentSwitch.h"
+#include "MuCOE/Nodes/CustomizableObjectNodeProjectorConstant.h"
+#include "MuCOE/Nodes/CustomizableObjectNodeSkeletalMesh.h"
+#include "MuCOE/Nodes/CustomizableObjectNodeStaticMesh.h"
+#include "MuCOE/Nodes/CustomizableObjectNodeTable.h"
+#include "MuCOE/Nodes/CustomizableObjectNodeTexture.h"
+#include "MuCOE/Nodes/CustomizableObjectNodePassThroughTexture.h"
+#include "MuCOE/Nodes/CustomizableObjectNodeTextureBinarise.h"
+#include "MuCOE/Nodes/CustomizableObjectNodeTextureColourMap.h"
+#include "MuCOE/Nodes/CustomizableObjectNodeTextureFromChannels.h"
+#include "MuCOE/Nodes/CustomizableObjectNodeTextureFromColor.h"
+#include "MuCOE/Nodes/CustomizableObjectNodeTextureFromFloats.h"
+#include "MuCOE/Nodes/CustomizableObjectNodeTextureInterpolate.h"
+#include "MuCOE/Nodes/CustomizableObjectNodeTextureInvert.h"
+#include "MuCOE/Nodes/CustomizableObjectNodeTextureLayer.h"
+#include "MuCOE/Nodes/CustomizableObjectNodeTextureParameter.h"
+#include "MuCOE/Nodes/CustomizableObjectNodeTextureProject.h"
+#include "MuCOE/Nodes/CustomizableObjectNodeTextureSample.h"
+#include "MuCOE/Nodes/CustomizableObjectNodeTextureSwitch.h"
+#include "MuCOE/Nodes/CustomizableObjectNodePassThroughTextureSwitch.h"
+#include "MuCOE/Nodes/CustomizableObjectNodeTextureToChannels.h"
+#include "MuCOE/Nodes/CustomizableObjectNodeTextureTransform.h"
+#include "MuCOE/Nodes/CustomizableObjectNodeTextureSaturate.h"
+#include "MuCOE/Nodes/CustomizableObjectNodeTextureVariation.h"
+#include "MuCOE/Nodes/CustomizableObjectNodePassThroughTextureVariation.h"
+#include "ScopedTransaction.h"
+#include "Settings/EditorStyleSettings.h"
+#include "ToolMenu.h"
+#include "ToolMenuSection.h"
+#include "MuCOE/Nodes/CustomizableObjectNodeReroute.h"
+#include "Toolkits/ToolkitManager.h"
+#include "PropertyEditorModule.h"
+#include "MuCO/CustomizableObjectCustomVersion.h"
+#include "MuCOE/GraphTraversal.h"
+#include "MuCOE/Nodes/CustomizableObjectNodeComponentMeshAddTo.h"
+
+class IToolkit;
+
+
+#define LOCTEXT_NAMESPACE "CustomizableObjectSchema"
+
+
+#define SNAP_GRID (16) // @todo ensure this is the same as SNodePanel::GetSnapGridSize()
+
+namespace MutablePrivate
+{
+	// Maximum distance a drag can be off a node edge to require 'push off' from node
+	const int32 NodeDistance = 60;
+}
+
+
+UEdGraphNode* FCustomizableObjectSchemaAction_NewNode::PerformAction(class UEdGraph* ParentGraph, UEdGraphPin* FromPin, const FVector2D Location, bool bSelectNewNode)
+{
+	UEdGraphNode* ResultNode = nullptr;
+
+	// If there is a template, we actually use it
+	if (NodeTemplate)
+	{
+		const FScopedTransaction Transaction(LOCTEXT("AddNode", "Add Node"));
+		ParentGraph->Modify();
+
+		if (FromPin)
+		{
+			FromPin->Modify();
+		}
+
+		ResultNode = FCustomizableObjectSchemaAction_NewNode::CreateNode(ParentGraph, FromPin, Location, NodeTemplate);
+	}
+
+	return ResultNode;
+}
+
+
+UEdGraphNode* FCustomizableObjectSchemaAction_NewNode::PerformAction(class UEdGraph* ParentGraph, TArray<UEdGraphPin*>& FromPins, const FVector2D Location, bool bSelectNewNode) 
+{
+	UEdGraphNode* ResultNode = 0;
+
+	if (FromPins.Num() > 0)
+	{
+		ResultNode = PerformAction(ParentGraph, FromPins[0], Location,bSelectNewNode);
+
+		// Try autowiring the rest of the pins
+		for (int32 Index = 1; Index < FromPins.Num(); ++Index)
+		{
+			ResultNode->AutowireNewNode(FromPins[Index]);
+		}
+	}
+	else
+	{
+		ResultNode = PerformAction(ParentGraph, NULL, Location, bSelectNewNode);
+	}
+
+	return ResultNode;
+}
+
+
+UEdGraphNode* FCustomizableObjectSchemaAction_NewNode::CreateNode(class UEdGraph* ParentGraph, UEdGraphPin* FromPin, const FVector2D Location, class UEdGraphNode* InNodeTemplate)
+{
+	// UE Code from FSchemaAction_NewNode::CreateNode(...). Overlap calculations performed before AutowireNewNode(...).
+	
+	// Duplicate template node to create new node
+	UEdGraphNode* ResultNode = DuplicateObject<UEdGraphNode>(InNodeTemplate, ParentGraph);
+
+	ResultNode->SetFlags(RF_Transactional);
+
+	ParentGraph->AddNode(ResultNode, true);
+
+	ResultNode->CreateNewGuid();
+	ResultNode->PostPlacedNewNode();
+	if (UCustomizableObjectNode* TypedResultNode = Cast<UCustomizableObjectNode>(ResultNode))
+	{
+		TypedResultNode->BeginConstruct();
+		TypedResultNode->PostBackwardsCompatibleFixup();
+	}
+	ResultNode->ReconstructNode(); // Mutable node lifecycle always starts at ReconstructNode.
+
+	// For input pins, new node will generally overlap node being dragged off
+	// Work out if we want to visually push away from connected node
+	int32 XLocation = Location.X;
+	if (FromPin && FromPin->Direction == EGPD_Input)
+	{
+		UEdGraphNode* PinNode = FromPin->GetOwningNode();
+		const float XDelta = FMath::Abs(PinNode->NodePosX - Location.X);
+
+		if (XDelta < MutablePrivate::NodeDistance)
+		{
+			// Set location to edge of current node minus the max move distance
+			// to force node to push off from connect node enough to give selection handle
+			XLocation = PinNode->NodePosX - MutablePrivate::NodeDistance;
+		}
+	}
+
+	ResultNode->AutowireNewNode(FromPin);
+
+	ResultNode->NodePosX = XLocation;
+	ResultNode->NodePosY = Location.Y;
+	ResultNode->SnapToGrid(GetDefault<UEditorStyleSettings>()->GridSnapSize);
+
+	return ResultNode;
+}
+
+
+void FCustomizableObjectSchemaAction_NewNode::AddReferencedObjects( FReferenceCollector& Collector )
+{
+	FEdGraphSchemaAction::AddReferencedObjects( Collector );
+
+	// These don't get saved to disk, but we want to make sure the objects don't get GC'd while the action array is around
+	Collector.AddReferencedObject( NodeTemplate );
+}
+
+
+////////////////////////////////////////
+// FCustomizableObjectSchemaAction_Paste
+
+UEdGraphNode* FCustomizableObjectSchemaAction_Paste::PerformAction(class UEdGraph* ParentGraph, UEdGraphPin* FromPin, const FVector2D Location, bool bSelectNewNode/* = true*/)
+{
+	TSharedPtr<ICustomizableObjectEditor> CustomizableObjectEditor = UEdGraphSchema_CustomizableObject::GetCustomizableObjectEditor(ParentGraph);
+
+	if (CustomizableObjectEditor.IsValid() && CustomizableObjectEditor->CanPasteNodes())
+	{
+		CustomizableObjectEditor->PasteNodesHere(Location);
+	}
+	
+	return nullptr;
+}
+
+
+//////////////////////////////////////////////////////////////////////////
+// DO NOT change the values because it will break the external pin nodes!
+const FName UEdGraphSchema_CustomizableObject::PC_Object("object");
+const FName UEdGraphSchema_CustomizableObject::PC_Component("component");
+const FName UEdGraphSchema_CustomizableObject::PC_Material("material");
+const FName UEdGraphSchema_CustomizableObject::PC_Modifier("modifier");
+const FName UEdGraphSchema_CustomizableObject::PC_Mesh("mesh");
+const FName UEdGraphSchema_CustomizableObject::PC_PassThroughMesh("passThroughMesh");
+const FName UEdGraphSchema_CustomizableObject::PC_Layout("layout");
+const FName UEdGraphSchema_CustomizableObject::PC_Image("image");
+const FName UEdGraphSchema_CustomizableObject::PC_PassThroughImage("passThroughImage");
+const FName UEdGraphSchema_CustomizableObject::PC_Projector("projector");
+const FName UEdGraphSchema_CustomizableObject::PC_GroupProjector("groupProjector");
+const FName UEdGraphSchema_CustomizableObject::PC_Color("color");
+const FName UEdGraphSchema_CustomizableObject::PC_Float("float");
+const FName UEdGraphSchema_CustomizableObject::PC_Bool("bool");
+const FName UEdGraphSchema_CustomizableObject::PC_Enum("enum");
+const FName UEdGraphSchema_CustomizableObject::PC_Stack("stack");
+const FName UEdGraphSchema_CustomizableObject::PC_MaterialAsset("materialAsset");
+const FName UEdGraphSchema_CustomizableObject::PC_Wildcard("wildcard");
+const FName UEdGraphSchema_CustomizableObject::PC_PoseAsset("poseAsset");
+
+
+TSharedPtr<FCustomizableObjectSchemaAction_NewNode> UEdGraphSchema_CustomizableObject::AddNewNodeAction(FGraphActionListBuilderBase& ContextMenuBuilder, const FString& Category, const FText& MenuDesc, const FText& Tooltip, const int32 Grouping, const FString& Keywords)
+{
+	TSharedPtr<FCustomizableObjectSchemaAction_NewNode> NewActionNode = TSharedPtr<FCustomizableObjectSchemaAction_NewNode>(new FCustomizableObjectSchemaAction_NewNode(Category, MenuDesc, Tooltip, Grouping, FText::FromString(Keywords)));
+
+	ContextMenuBuilder.AddAction( NewActionNode );
+
+	return NewActionNode;
+}
+
+
+TSharedPtr<FCustomizableObjectSchemaAction_NewNode> AddNewNodeAction(TArray< TSharedPtr<FEdGraphSchemaAction> >& OutTypes, const FString& Category, const FText& MenuDesc, const FText& Tooltip)
+{
+	return *(new (OutTypes) TSharedPtr<FCustomizableObjectSchemaAction_NewNode>(new FCustomizableObjectSchemaAction_NewNode(Category, MenuDesc, Tooltip, 0)));
+}
+
+
+namespace 
+{
+	bool PinRelevancyFilter(UCustomizableObjectNode* TemplateNode, FGraphContextMenuBuilder& ContextMenuBuilder)
+	{
+		const UEdGraphPin* FromPin = ContextMenuBuilder.FromPin;
+		if (!FromPin)
+		{
+			return true;
+		}
+
+		if (TemplateNode->ProvidesCustomPinRelevancyTest())
+		{
+			return TemplateNode->IsPinRelevant(ContextMenuBuilder.FromPin);
+		}
+
+		TemplateNode->BeginConstruct();
+		TemplateNode->ReconstructNode();
+		
+		for (const UEdGraphPin* Pin : TemplateNode->GetAllNonOrphanPins())
+		{
+			const UEdGraphPin* InputPin = nullptr;
+			const UEdGraphPin* OutputPin = nullptr;
+
+			if (!UEdGraphSchema_CustomizableObject::CategorizePinsByDirection(Pin, FromPin, InputPin, OutputPin))
+			{
+				continue;
+			}
+
+			const UCustomizableObjectNode* InputNode = Cast<UCustomizableObjectNode>(InputPin->GetOwningNode());
+			bool bOtherNodeIsBlocklisted = false;
+			bool bArePinsCompatible = false;
+			if (InputNode->CanConnect(InputPin, OutputPin, bOtherNodeIsBlocklisted, bArePinsCompatible))
+			{
+				return true;
+			}
+		}
+
+		return false;
+	}
+
+	template<class FilterFn>
+	void AddNewNodeActionFiltered(UCustomizableObjectNode* TemplateNode, FGraphContextMenuBuilder& ContextMenuBuilder, const FString& Category, const FText& MenuDesc, const int32 Grouping, FilterFn&& Filter)
+	{
+		if (!Filter(TemplateNode, ContextMenuBuilder))
+		{
+			return;
+		}
+
+		TSharedPtr<FCustomizableObjectSchemaAction_NewNode> Action = UEdGraphSchema_CustomizableObject::AddNewNodeAction(ContextMenuBuilder, Category, MenuDesc, FText(), Grouping);
+		Action->NodeTemplate = TemplateNode;
+	}
+
+	template<class FilterFn>
+	void AddNewNodeActionFiltered(UCustomizableObjectNode* TemplateNode, FGraphContextMenuBuilder& ContextMenuBuilder, const FString& Category, const int32 Grouping, FilterFn&& Filter)
+	{
+		AddNewNodeActionFiltered(TemplateNode, ContextMenuBuilder, Category, TemplateNode->GetNodeTitle(ENodeTitleType::ListView), Grouping, Forward<FilterFn>(Filter));
+	}
+
+	template<size_t N, class FilterFn>
+	void AddNewNodeCategoryActionsFiltered(UCustomizableObjectNode* (&TemplateNodes)[N], FGraphContextMenuBuilder& ContextMenuBuilder, const FString& Category, const int32 Grouping, FilterFn&& Filter)
+	{
+		for (size_t i = 0; i < N; ++i)
+		{
+			AddNewNodeActionFiltered(TemplateNodes[i], ContextMenuBuilder, Category, Grouping, Forward<FilterFn>(Filter));
+		}
+	}
+}
+
+
+void UEdGraphSchema_CustomizableObject::GetGraphContextActions(FGraphContextMenuBuilder& ContextMenuBuilder) const
+{
+	//const UCustomizableObjectGraph* CustomizableObjectGraph = CastChecked<UCustomizableObjectGraph>(ContextMenuBuilder.CurrentGraph);
+	//TSharedPtr<FCustomizableObjectSchemaAction_NewNode> Action;
+	int32 GeneralGrouping = 2;
+
+	const bool bDisableFilter = false;
+
+	// return true if Filter is passed.
+	const auto Filter = [bDisableFilter](UCustomizableObjectNode* TemplateNode, FGraphContextMenuBuilder& ContextMenuBuilder) -> bool
+	{
+		if (!ContextMenuBuilder.FromPin || bDisableFilter)
+		{
+			return true;
+		}
+
+		return PinRelevancyFilter(TemplateNode, ContextMenuBuilder);
+	};
+
+	{
+		UCustomizableObjectNode* Node = ContextMenuBuilder.CreateTemplateNode<UCustomizableObjectNodeObject>();
+		AddNewNodeActionFiltered(Node, ContextMenuBuilder, TEXT("Object"), LOCTEXT("Base_Group", "Base Object"), GeneralGrouping, 
+			[&Filter](UCustomizableObjectNode* TemplateNode, FGraphContextMenuBuilder& ContextMenuBuilder) -> bool
+			{
+				// Only let user add a base node if there isn't one in the graph
+				for (const TObjectPtr<UEdGraphNode>& AuxNode : ContextMenuBuilder.CurrentGraph->Nodes)
+				{
+					UCustomizableObjectNodeObject* CustomizableObjectNodeObject = Cast<UCustomizableObjectNodeObject>(AuxNode);
+
+					if (CustomizableObjectNodeObject && CustomizableObjectNodeObject->bIsBase)
+					{
+						return false;
+					}
+				}
+
+				return Filter(TemplateNode, ContextMenuBuilder);
+			}	
+		);
+	}
+
+	{
+		UCustomizableObjectNode* Node = NewObject<UCustomizableObjectNodeObjectGroup>(); //ContextMenuBuilder.CreateTemplateNode<UCustomizableObjectNodeObjectGroup>();
+		AddNewNodeActionFiltered(Node, ContextMenuBuilder, TEXT("Object"), LOCTEXT("Child_Group", "Object Group"), GeneralGrouping, Filter);
+	}
+	
+	{
+		UCustomizableObjectNode* Node = ContextMenuBuilder.CreateTemplateNode<UCustomizableObjectNodeObjectChild>();
+		AddNewNodeActionFiltered(Node, ContextMenuBuilder, TEXT("Object"), LOCTEXT("Child_Object", "Child Object"), GeneralGrouping, Filter);
+	}
+
+	{
+		UCustomizableObjectNode* ObjectTemplateNodes[]
+		{
+			ContextMenuBuilder.CreateTemplateNode<UCustomizableObjectNodeMaterial>(),
+			ContextMenuBuilder.CreateTemplateNode<UCustomizableObjectNodeMaterialVariation>(),
+			ContextMenuBuilder.CreateTemplateNode<UCustomizableObjectNodeMaterialSwitch>(),
+			ContextMenuBuilder.CreateTemplateNode<UCustomizableObjectNodeCopyMaterial>(),
+			ContextMenuBuilder.CreateTemplateNode<UCustomizableObjectNodeTable>(),
+		};
+
+		AddNewNodeCategoryActionsFiltered(ObjectTemplateNodes, ContextMenuBuilder, TEXT("Object"), GeneralGrouping, Filter);
+	}
+
+	{
+		UCustomizableObjectNode* ModifierTemplateNodes[]
+		{
+			ContextMenuBuilder.CreateTemplateNode<UCustomizableObjectNodeModifierExtendMeshSection>(),
+			ContextMenuBuilder.CreateTemplateNode<UCustomizableObjectNodeModifierRemoveMesh>(),
+			ContextMenuBuilder.CreateTemplateNode<UCustomizableObjectNodeModifierRemoveMeshBlocks>(),
+			ContextMenuBuilder.CreateTemplateNode<UCustomizableObjectNodeModifierEditMeshSection>(),
+			ContextMenuBuilder.CreateTemplateNode<UCustomizableObjectNodeModifierMorphMeshSection>(),
+			ContextMenuBuilder.CreateTemplateNode<UCustomizableObjectNodeModifierClipMorph>(),
+			ContextMenuBuilder.CreateTemplateNode<UCustomizableObjectNodeModifierClipWithMesh>(),
+			ContextMenuBuilder.CreateTemplateNode<UCustomizableObjectNodeModifierClipDeform>(),
+			ContextMenuBuilder.CreateTemplateNode<UCustomizableObjectNodeModifierClipWithUVMask>(),
+		};
+
+		AddNewNodeCategoryActionsFiltered(ModifierTemplateNodes, ContextMenuBuilder, TEXT("Modifier"), GeneralGrouping, Filter);
+	}
+
+	{
+		UCustomizableObjectNode* ComponentTemplateNodes[]
+		{
+			ContextMenuBuilder.CreateTemplateNode<UCustomizableObjectNodeComponentMesh>(),
+			ContextMenuBuilder.CreateTemplateNode<UCustomizableObjectNodeComponentMeshAddTo>(),
+			ContextMenuBuilder.CreateTemplateNode<UCustomizableObjectNodeComponentPassthroughMesh>(),
+			ContextMenuBuilder.CreateTemplateNode<UCustomizableObjectNodeComponentVariation>(),
+			ContextMenuBuilder.CreateTemplateNode<UCustomizableObjectNodeComponentSwitch>(),
+		};
+
+		AddNewNodeCategoryActionsFiltered(ComponentTemplateNodes, ContextMenuBuilder, TEXT("Component"), GeneralGrouping, Filter);
+	}
+
+	{
+		UCustomizableObjectNode* MeshTemplateNodes[]
+		{
+			ContextMenuBuilder.CreateTemplateNode<UCustomizableObjectNodeSkeletalMesh>(),
+			ContextMenuBuilder.CreateTemplateNode<UCustomizableObjectNodeStaticMesh>(),
+			ContextMenuBuilder.CreateTemplateNode<UCustomizableObjectNodeMeshMorph>(),
+			ContextMenuBuilder.CreateTemplateNode<UCustomizableObjectNodeMeshSwitch>(),
+			ContextMenuBuilder.CreateTemplateNode<UCustomizableObjectNodeMeshVariation>(),
+			ContextMenuBuilder.CreateTemplateNode<UCustomizableObjectNodeMeshMorphStackDefinition>(),
+			ContextMenuBuilder.CreateTemplateNode<UCustomizableObjectNodeMeshMorphStackApplication>(),
+			ContextMenuBuilder.CreateTemplateNode<UCustomizableObjectNodeMeshGeometryOperation>(),
+			ContextMenuBuilder.CreateTemplateNode<UCustomizableObjectNodeMeshReshape>(),
+		};
+
+		AddNewNodeCategoryActionsFiltered(MeshTemplateNodes, ContextMenuBuilder, TEXT("Mesh"), GeneralGrouping, Filter);
+	}
+
+	{
+		UCustomizableObjectNode* PoseMeshTemplateNodes[]
+		{
+			ContextMenuBuilder.CreateTemplateNode<UCustomizableObjectNodeAnimationPose>(),
+		};
+
+		AddNewNodeCategoryActionsFiltered(PoseMeshTemplateNodes, ContextMenuBuilder, TEXT("PoseMesh"), GeneralGrouping, Filter);
+	}
+
+	{
+		UCustomizableObjectNode* TextureTemplateNodes[] =
+		{
+			ContextMenuBuilder.CreateTemplateNode<UCustomizableObjectNodeTexture>(),
+			ContextMenuBuilder.CreateTemplateNode<UCustomizableObjectNodePassThroughTexture>(),
+			ContextMenuBuilder.CreateTemplateNode<UCustomizableObjectNodeTextureBinarise>(),
+			ContextMenuBuilder.CreateTemplateNode<UCustomizableObjectNodeTextureInterpolate>(),
+			ContextMenuBuilder.CreateTemplateNode<UCustomizableObjectNodeTextureLayer>(),
+			ContextMenuBuilder.CreateTemplateNode<UCustomizableObjectNodeTextureSwitch>(),
+			ContextMenuBuilder.CreateTemplateNode<UCustomizableObjectNodePassThroughTextureSwitch>(),
+			ContextMenuBuilder.CreateTemplateNode<UCustomizableObjectNodeTextureVariation>(),
+			ContextMenuBuilder.CreateTemplateNode<UCustomizableObjectNodePassThroughTextureVariation>(),
+			ContextMenuBuilder.CreateTemplateNode<UCustomizableObjectNodeTextureToChannels>(),
+			ContextMenuBuilder.CreateTemplateNode<UCustomizableObjectNodeTextureFromChannels>(),
+			ContextMenuBuilder.CreateTemplateNode<UCustomizableObjectNodeTextureFromColor>(),
+			ContextMenuBuilder.CreateTemplateNode<UCustomizableObjectNodeTextureFromFloats>(),
+			ContextMenuBuilder.CreateTemplateNode<UCustomizableObjectNodeTextureProject>(),
+			ContextMenuBuilder.CreateTemplateNode<UCustomizableObjectNodeTextureParameter>(),
+			ContextMenuBuilder.CreateTemplateNode<UCustomizableObjectNodeTextureInvert>(),
+			ContextMenuBuilder.CreateTemplateNode<UCustomizableObjectNodeTextureColourMap>(),
+			ContextMenuBuilder.CreateTemplateNode<UCustomizableObjectNodeTextureTransform>(),
+			ContextMenuBuilder.CreateTemplateNode<UCustomizableObjectNodeTextureSaturate>(),
+		};
+
+		AddNewNodeCategoryActionsFiltered(TextureTemplateNodes, ContextMenuBuilder, TEXT("Texture"), GeneralGrouping, Filter);
+	}
+
+	{
+		UCustomizableObjectNode* ColorTemplateNodes[] =
+		{
+			ContextMenuBuilder.CreateTemplateNode<UCustomizableObjectNodeColorConstant>(),
+			ContextMenuBuilder.CreateTemplateNode<UCustomizableObjectNodeColorParameter>(),
+			ContextMenuBuilder.CreateTemplateNode<UCustomizableObjectNodeColorSwitch>(),
+			ContextMenuBuilder.CreateTemplateNode<UCustomizableObjectNodeColorVariation>(),
+			ContextMenuBuilder.CreateTemplateNode<UCustomizableObjectNodeTextureSample>(),
+			ContextMenuBuilder.CreateTemplateNode<UCustomizableObjectNodeColorArithmeticOp>(),
+			ContextMenuBuilder.CreateTemplateNode<UCustomizableObjectNodeColorFromFloats>(),
+		};
+
+		AddNewNodeCategoryActionsFiltered(ColorTemplateNodes, ContextMenuBuilder, TEXT("Color"), GeneralGrouping, Filter);
+	}
+
+	{
+		UCustomizableObjectNode* EnumTemplateNodes[] =
+		{
+			ContextMenuBuilder.CreateTemplateNode<UCustomizableObjectNodeEnumParameter>(),
+		};
+
+		AddNewNodeCategoryActionsFiltered(EnumTemplateNodes, ContextMenuBuilder, TEXT("Enum"), GeneralGrouping, Filter);
+	}
+
+	{
+		UCustomizableObjectNode* FloatTemplateNodes[] =
+		{
+			ContextMenuBuilder.CreateTemplateNode<UCustomizableObjectNodeFloatConstant>(),
+			ContextMenuBuilder.CreateTemplateNode<UCustomizableObjectNodeFloatParameter>(),
+			ContextMenuBuilder.CreateTemplateNode<UCustomizableObjectNodeFloatArithmeticOp>(),
+			ContextMenuBuilder.CreateTemplateNode<UCustomizableObjectNodeFloatSwitch>(),
+			ContextMenuBuilder.CreateTemplateNode<UCustomizableObjectNodeFloatVariation>(),
+			ContextMenuBuilder.CreateTemplateNode<UCustomizableObjectNodeCurve>(),
+		};
+
+
+		AddNewNodeCategoryActionsFiltered(FloatTemplateNodes, ContextMenuBuilder, TEXT("Float"), GeneralGrouping, Filter);
+
+		UCustomizableObjectNodeFloatArithmeticOp* Node = NewObject<UCustomizableObjectNodeFloatArithmeticOp>();
+		Node->Operation = EFloatArithmeticOperation::E_Add;
+		AddNewNodeActionFiltered(Node, ContextMenuBuilder, TEXT("Float"), LOCTEXT("Float_Addition", "Float Addition +"), GeneralGrouping, Filter);
+
+		Node = NewObject<UCustomizableObjectNodeFloatArithmeticOp>();
+		Node->Operation = EFloatArithmeticOperation::E_Sub;
+		AddNewNodeActionFiltered(Node, ContextMenuBuilder, TEXT("Float"), LOCTEXT("Float_Subtraction", "Float Subtraction -"), GeneralGrouping, Filter);
+
+		Node = NewObject<UCustomizableObjectNodeFloatArithmeticOp>();
+		Node->Operation = EFloatArithmeticOperation::E_Mul;
+		AddNewNodeActionFiltered(Node, ContextMenuBuilder, TEXT("Float"), LOCTEXT("Float_Multiplication", "Float Multiplication x"), GeneralGrouping, Filter);
+
+		Node = NewObject<UCustomizableObjectNodeFloatArithmeticOp>();
+		Node->Operation = EFloatArithmeticOperation::E_Div;
+		AddNewNodeActionFiltered(Node, ContextMenuBuilder, TEXT("Float"), LOCTEXT("Float_Division", "Float Division /"), GeneralGrouping, Filter);
+	}
+
+	{
+		UCustomizableObjectNode* ProjectorTemplateNodes[] =
+		{
+			ContextMenuBuilder.CreateTemplateNode<UCustomizableObjectNodeProjectorConstant>(),
+			ContextMenuBuilder.CreateTemplateNode<UCustomizableObjectNodeProjectorParameter>(),
+			ContextMenuBuilder.CreateTemplateNode<UCustomizableObjectNodeGroupProjectorParameter>(),
+		};
+
+		AddNewNodeCategoryActionsFiltered(ProjectorTemplateNodes, ContextMenuBuilder, TEXT("Projector"), GeneralGrouping, Filter);
+	}
+
+	{
+		// External Pin Nodes
+		TArray<FName> PinTypes({ PC_Material, PC_Modifier, PC_Mesh, PC_Image, PC_Projector, PC_GroupProjector, PC_Color, PC_Float, PC_Bool, PC_Enum, PC_Stack, PC_PassThroughImage, PC_MaterialAsset, PC_PoseAsset, PC_Component });
+
+		// Add pin types from extensions
+		for (const FRegisteredCustomizableObjectPinType& PinType : ICustomizableObjectModule::Get().GetExtendedPinTypes())
+		{
+			PinTypes.AddUnique(PinType.PinType.Name);
+		}
+
+		for (const FName& PinCategory : PinTypes)
+		{
+			UCustomizableObjectNodeExternalPin* CustomizableObjectNodeExternalPin = ContextMenuBuilder.CreateTemplateNode<UCustomizableObjectNodeExternalPin>();
+			CustomizableObjectNodeExternalPin->PinType = PinCategory;
+			
+			AddNewNodeActionFiltered(CustomizableObjectNodeExternalPin, ContextMenuBuilder, TEXT("Import Pin"), GeneralGrouping, Filter);
+		}
+
+		for (const FName& PinCategory : PinTypes)
+		{
+			UCustomizableObjectNodeExposePin* CustomizableObjectNodeExposePin = ContextMenuBuilder.CreateTemplateNode<UCustomizableObjectNodeExposePin>();
+			CustomizableObjectNodeExposePin->PinType = PinCategory;
+			
+			AddNewNodeActionFiltered(CustomizableObjectNodeExposePin, ContextMenuBuilder, TEXT("Export Pin"), GeneralGrouping, Filter);
+		}
+	}
+
+	// Search for all subclasses of UCustomizableObjectNode
+	//
+	// Iterate over the Class Default Objects instead of their corresponding UClasses, as this allows
+	// us to filter the TObjectIterator to UCustomizableObjectNode instead of UClass, which should
+	// produce far fewer results to iterate through.
+	for (TObjectIterator<UCustomizableObjectNode> It(RF_NoFlags); It; ++It)
+	{
+		const UCustomizableObjectNode* Node = *It;
+		if (!IsValid(Node) || !Node->HasAllFlags(RF_ClassDefaultObject) || Node->GetClass()->HasAnyClassFlags(CLASS_Abstract))
+		{
+			// Only interested in non-abstract CDOs
+			continue;
+		}
+
+		FText Category;
+		if (Node->ShouldAddToContextMenu(Category))
+		{
+			UCustomizableObjectNode* TemplateNode = ContextMenuBuilder.CreateTemplateNode<UCustomizableObjectNode>(Node->GetClass());
+			AddNewNodeActionFiltered(TemplateNode, ContextMenuBuilder, Category.ToString(), GeneralGrouping, Filter);
+		}
+	}
+	
+	{
+		UEdGraphNode_Comment* Node = NewObject<UEdGraphNode_Comment>();
+		TSharedPtr<FCustomizableObjectSchemaAction_NewNode> Action = AddNewNodeAction(ContextMenuBuilder, FString(), Node->GetNodeTitle(ENodeTitleType::ListView), FText(), 1);
+		Action->NodeTemplate = Node;
+	}
+
+	{
+		UCustomizableObjectNodeReroute* Node = NewObject<UCustomizableObjectNodeReroute>();
+		TSharedPtr<FCustomizableObjectSchemaAction_NewNode> Action = AddNewNodeAction(ContextMenuBuilder, FString(), Node->GetNodeTitle(ENodeTitleType::ListView), FText(), 1);
+		Action->NodeTemplate = Node;
+	}
+
+	// Add Paste here if appropriate
+	if (!ContextMenuBuilder.FromPin)
+	{
+		const FText PasteDesc = LOCTEXT("PasteDesc", "Paste Here");
+		const FText PasteToolTip = LOCTEXT("PasteToolTip", "Pastes copied items at this location.");
+		TSharedPtr<FCustomizableObjectSchemaAction_Paste> PasteAction(new FCustomizableObjectSchemaAction_Paste(FText::GetEmpty(), PasteDesc, PasteToolTip.ToString(), 0));
+		ContextMenuBuilder.AddAction(PasteAction);
+	}
+}
+
+
+const FPinConnectionResponse UEdGraphSchema_CustomizableObject::CanCreateConnection(const UEdGraphPin* PinA, const UEdGraphPin* PinB) const
+{
+	// Make sure the pins are not on the same node
+	if (PinA->GetOwningNode() == PinB->GetOwningNode())
+	{
+		return FPinConnectionResponse(CONNECT_RESPONSE_DISALLOW, TEXT("Both are on the same node"));
+	}
+
+	// Check both pins support connections
+	if(PinA->bNotConnectable || PinB->bNotConnectable)
+	{
+		return FPinConnectionResponse(CONNECT_RESPONSE_DISALLOW, TEXT("Pin doesn't support connections"));
+	}
+
+	// Compare the directions
+	const UEdGraphPin* InputPin = NULL;
+	const UEdGraphPin* OutputPin = NULL;
+
+	if (!CategorizePinsByDirection(PinA, PinB, /*out*/ InputPin, /*out*/ OutputPin))
+	{
+		return FPinConnectionResponse(CONNECT_RESPONSE_DISALLOW, TEXT("Directions are not compatible"));
+	}
+
+	// Type categories must match and the nodes need to be compatible with each other
+	bool bArePinsCompatible = false;
+	bool bIsOtherNodeBlocklisted = false;
+
+	UCustomizableObjectNode* InputNode = CastChecked<UCustomizableObjectNode>(InputPin->GetOwningNode());
+	if (!InputNode->CanConnect(InputPin, OutputPin, bIsOtherNodeBlocklisted,bArePinsCompatible))
+	{
+		if (!bArePinsCompatible)
+		{
+			return FPinConnectionResponse(CONNECT_RESPONSE_DISALLOW, TEXT("Types are not compatible"));
+		}
+		else if (bIsOtherNodeBlocklisted)
+		{
+			return FPinConnectionResponse(CONNECT_RESPONSE_DISALLOW, TEXT("Direct connections between these nodes are not allowed"));
+		}
+	}
+
+	// Some special nodes can only have one output
+	bool bBreakExistingDueToDataOutput = false;
+	if (UCustomizableObjectNode* n = Cast<UCustomizableObjectNode>(OutputPin->GetOwningNode()))
+	{
+		bBreakExistingDueToDataOutput = (OutputPin->LinkedTo.Num() > 0) && n->ShouldBreakExistingConnections(InputPin, OutputPin);
+	}
+
+	// See if we want to break existing connections (if its an input with an existing connection)
+	const bool bBreakExistingDueToDataInput = (InputPin->LinkedTo.Num() > 0) && !InputPin->PinType.IsArray();
+
+	if (bBreakExistingDueToDataOutput&&bBreakExistingDueToDataInput)
+	{
+		return FPinConnectionResponse(CONNECT_RESPONSE_BREAK_OTHERS_AB, TEXT("Replace connections at both ends"));
+	}
+	
+	if (bBreakExistingDueToDataInput)
+	{
+		const ECanCreateConnectionResponse ReplyBreakInputs = (PinA == InputPin) ? CONNECT_RESPONSE_BREAK_OTHERS_A : CONNECT_RESPONSE_BREAK_OTHERS_B;
+		return FPinConnectionResponse(ReplyBreakInputs, TEXT("Replace existing input connections"));
+	}
+
+	if (bBreakExistingDueToDataOutput)
+	{
+		const ECanCreateConnectionResponse ReplyBreakOutputs = (PinA == OutputPin) ? CONNECT_RESPONSE_BREAK_OTHERS_A : CONNECT_RESPONSE_BREAK_OTHERS_B;
+		return FPinConnectionResponse(ReplyBreakOutputs, TEXT("Replace existing output connections"));
+	}
+
+	return FPinConnectionResponse(CONNECT_RESPONSE_MAKE, TEXT(""));
+}
+
+
+FLinearColor UEdGraphSchema_CustomizableObject::GetPinTypeColor(const FName& TypeString)
+{
+	if (TypeString == PC_Enum)
+	{
+		return FLinearColor(0.357667f, 0.500000f, 0.060000f, 1.000000f); // Light green
+	}
+	else if (TypeString == PC_Float)
+	{
+		return FLinearColor(0.357667f, 1.000000f, 0.060000f, 1.000000f); // Green
+	}
+	else if (TypeString == PC_Color)
+	{
+		return FLinearColor(1.000000f, 0.591255f, 0.016512f, 1.000000f); // Yellow
+	}
+	else if (TypeString == PC_Bool)
+	{
+		return FLinearColor(0.470000f, 0.0f, 0.000000f, 1.000000f); // Red
+	}
+	else if (TypeString == PC_Projector)
+	{
+		return FLinearColor(1.000000f, 0.500000f, 1.000000f, 0.600000f); // Light pink
+	}
+	else if (TypeString == PC_GroupProjector)
+	{
+		return FLinearColor(1.000000f, 0.172585f, 0.000000f, 1.000000f); // Orange
+	}
+	else if (TypeString == PC_Mesh)
+	{
+		return FLinearColor(0.100000f, 0.000000f, 0.500000f, 1.000000f); // Purple
+	}
+	else if (TypeString == PC_Layout)
+	{
+		return FLinearColor(0.500000f, 0.500000f, 0.100000f, 1.000000f); // Light yellow
+	}
+	else if (TypeString == PC_Image)
+	{
+		return FLinearColor(0.353393f, 0.454175f, 1.000000f, 1.000000f); // Grey
+	}
+	else if (TypeString == PC_PassThroughImage)
+	{
+		return FLinearColor(0.353393f, 0.454175f, 0.353393f, 1.000000f); // Dark green
+	}
+	else if (TypeString == PC_Material)
+	{
+		return FLinearColor(0.000000f, 0.100000f, 0.600000f, 1.000000f); // Blue
+	}
+	else if (TypeString == PC_Modifier)
+	{
+		return FColor::Silver;
+	}
+	else if (TypeString == PC_Object)
+	{
+		return FLinearColor(0.000000f, 0.400000f, 0.910000f, 1.000000f); // Light blue
+	}
+	else if (TypeString == PC_Component)
+	{
+		return FLinearColor(FColorList::MediumOrchid); 
+	}
+	else if (TypeString == PC_Stack)
+	{
+		return FLinearColor(1.000000f, 0.000000f, 0.800000f, 1.000000f); // Pink
+	}
+	else if (TypeString == PC_MaterialAsset)
+	{
+		return FLinearColor(0.000000f, 1.000000f, 0.1000f, 1.000000f); // Cian
+	}
+	else if (TypeString == PC_Wildcard)
+	{
+		return FLinearColor(1.000000f, 1.000000f, 1.000000f, 1.000000f); // White
+	}
+	else if (TypeString == PC_PoseAsset)
+	{
+		return FLinearColor(0.700000f, 0.000000f, 0.000000f, 1.000000f); // Dark Red
+	}
+
+	for (const FRegisteredCustomizableObjectPinType& PinType : ICustomizableObjectModule::Get().GetExtendedPinTypes())
+	{
+		if (PinType.PinType.Name == TypeString)
+		{
+			return PinType.PinType.Color;
+		}
+	}
+
+	return FLinearColor(0.750000f, 0.600000f, 0.400000f, 1.000000f);
+}
+
+
+FLinearColor UEdGraphSchema_CustomizableObject::GetPinTypeColor(const FEdGraphPinType& PinType) const
+{
+	const FName TypeName = PinType.PinCategory;
+
+	return GetPinTypeColor(TypeName);
+}
+
+
+bool UEdGraphSchema_CustomizableObject::ShouldHidePinDefaultValue(UEdGraphPin* Pin) const
+{
+	check(Pin != NULL);
+
+	if (Pin->bDefaultValueIsIgnored)
+	{
+		return true;
+	}
+
+	return false;
+}
+
+
+void UEdGraphSchema_CustomizableObject::GetContextMenuActionsReconstructAllChildNodes(UToolMenu* Menu, TWeakObjectPtr<UGraphNodeContextMenuContext> WeakContext) const
+{
+	FToolMenuSection& SubSection = Menu->AddSection("Section");
+
+	TArray<UClass*> NodeTypes;
+
+	for (TObjectIterator<UCustomizableObjectNode> It(RF_NoFlags); It; ++It)
+	{
+		const UCustomizableObjectNode* Node = *It;
+		if (!IsValid(Node) || !Node->HasAllFlags(RF_ClassDefaultObject))
+		{
+			continue; // Only interested in CDOs
+		}
+
+		UClass* Class = Node->GetClass();
+		if (Class->HasAnyClassFlags(CLASS_Abstract))
+		{
+			continue;
+		}
+		
+		NodeTypes.Add(Class);
+	}
+	
+	NodeTypes.Sort([](UClass& A, UClass& B) -> bool
+	{
+		const UCustomizableObjectNode* NodeA = CastChecked<UCustomizableObjectNode>(A.GetDefaultObject());
+		const UCustomizableObjectNode* NodeB = CastChecked<UCustomizableObjectNode>(B.GetDefaultObject());
+		
+		return NodeA->GetNodeTitle(ENodeTitleType::ListView).CompareTo(NodeB->GetNodeTitle(ENodeTitleType::ListView)) < 0;
+	});
+	
+	for (UClass* NodeType : NodeTypes)
+	{
+		const UCustomizableObjectNode* Node = Cast<UCustomizableObjectNode>(NodeType->GetDefaultObject());
+
+		auto Call = [](TWeakObjectPtr<UGraphNodeContextMenuContext> WeakContext, const TWeakObjectPtr<UClass>& WeakNodeType)
+		{
+			const UGraphNodeContextMenuContext* Context = WeakContext.Get();
+			if (!Context)
+			{
+				return;
+			}
+
+			UCustomizableObjectNode* Node = const_cast<UCustomizableObjectNode*>(Cast<UCustomizableObjectNode>(Context->Node));
+			if (!Node)
+			{
+				return;
+			}
+
+			const UClass* NodeType = WeakNodeType.Get();
+			if (!NodeType)
+			{
+				return;
+			}
+
+			const TSharedRef<ICustomizableObjectEditor> Editor = GetCustomizableObjectEditor(Context->Graph).ToSharedRef();
+			Editor->ReconstructAllChildNodes(*Node, *NodeType);			
+		};
+
+		SubSection.AddMenuEntry(
+			Node->GetFName(),
+			Node->GetNodeTitle(ENodeTitleType::ListView),
+			FText(),
+			FSlateIcon(),
+			FExecuteAction::CreateLambda(Call, WeakContext, MakeWeakObjectPtr(NodeType)));
+	}
+}
+
+
+void UEdGraphSchema_CustomizableObject::GetContextMenuActions(UToolMenu* Menu, UGraphNodeContextMenuContext* Context) const
+{
+	if (!Context || !Context->Node)
+	{
+		return;
+	}
+	
+	if (!Context->Pin) // On Node right click
+	{
+		if (!Context->bIsDebugging)
+		{
+			// Node contextual actions
+			FToolMenuSection& Section = Menu->AddSection("EdGraphSchemaNodeActions", LOCTEXT("NodeActionsMenuHeader", "Node Actions"));
+			Section.AddMenuEntry(FGenericCommands::Get().Delete);
+			Section.AddMenuEntry(FGenericCommands::Get().Cut);
+			Section.AddMenuEntry(FGenericCommands::Get().Copy);
+			Section.AddMenuEntry(FGenericCommands::Get().Duplicate);
+			Section.AddMenuEntry(FGraphEditorCommands::Get().ReconstructNodes);
+			Section.AddMenuEntry(FGraphEditorCommands::Get().BreakNodeLinks);
+
+			Section.AddSubMenu("Alignment", LOCTEXT("AlignmentHeader", "Alignment"), FText(), FNewToolMenuDelegate::CreateLambda([](UToolMenu* AlignmentMenu)
+			{
+				{
+					FToolMenuSection& SubSection = AlignmentMenu->AddSection("EdGraphSchemaAlignment", LOCTEXT("AlignHeader", "Align"));
+					SubSection.AddMenuEntry(FGraphEditorCommands::Get().AlignNodesTop);
+					SubSection.AddMenuEntry(FGraphEditorCommands::Get().AlignNodesMiddle);
+					SubSection.AddMenuEntry(FGraphEditorCommands::Get().AlignNodesBottom);
+					SubSection.AddMenuEntry(FGraphEditorCommands::Get().AlignNodesLeft);
+					SubSection.AddMenuEntry(FGraphEditorCommands::Get().AlignNodesCenter);
+					SubSection.AddMenuEntry(FGraphEditorCommands::Get().AlignNodesRight);
+					SubSection.AddMenuEntry(FGraphEditorCommands::Get().StraightenConnections);
+				}
+
+				{
+					FToolMenuSection& SubSection = AlignmentMenu->AddSection("EdGraphSchemaDistribution", LOCTEXT("DistributionHeader", "Distribution"));
+					SubSection.AddMenuEntry(FGraphEditorCommands::Get().DistributeNodesHorizontally);
+					SubSection.AddMenuEntry(FGraphEditorCommands::Get().DistributeNodesVertically);
+				}
+			}));
+			
+			Section.AddSubMenu(
+				"ReconstructAllNodes",
+				LOCTEXT("ReconstructChildAllNodes", "Refresh All Child Nodes"),
+				LOCTEXT("ReconstructAllChildNodes_Tooltip", "Refresh all child nodes from the selected ones (inclusive)."),
+				FNewToolMenuDelegate::CreateUObject(this, &UEdGraphSchema_CustomizableObject::GetContextMenuActionsReconstructAllChildNodes, MakeWeakObjectPtr(Context)));
+		}
+
+		struct SCommentUtility
+		{
+			static void CreateComment(const UEdGraphSchema_CustomizableObject* Schema, UEdGraph* Graph)
+			{
+				if (Schema && Graph)
+				{
+					Schema->AddComment(Graph, NULL, FVector2D::ZeroVector, true);
+				}
+			}
+		};
+
+		FToolMenuSection& section = Menu->AddSection("SchemaActionComment", LOCTEXT("MultiCommentHeader", "Comment Group"));
+		section.AddMenuEntry("MultiCommentDesc", LOCTEXT("MultiCommentDesc", "Create Comment from Selection"),
+			LOCTEXT("CommentToolTip", "Create a resizable comment box around selection."),
+			FSlateIcon(),
+			FUIAction(FExecuteAction::CreateStatic(SCommentUtility::CreateComment, this, const_cast<UEdGraph*>(ToRawPtr(Context->Graph)))));
+	}
+	else // On Pin right click
+	{
+		if (UCustomizableObjectNode* Node = Cast<UCustomizableObjectNode>(const_cast<UEdGraphNode*>(Context->Node.Get())))
+		{
+			if (const UEdGraphPin* Pin = Context->Pin)
+			{
+				if (Node->CanPinBeHidden(*Pin))
+				{
+					FToolMenuSection& Section = Menu->FindOrAddSection("EdGraphSchemaPinActions");
+					Section.AddMenuEntry("HidePin",
+						LOCTEXT("HidePin_Label", "Hide Pin"),
+						LOCTEXT("HidePin_Tooltip", "Hides the selected pin."),
+						FSlateIcon(),
+						FUIAction(FExecuteAction::CreateLambda([Node = const_cast<UCustomizableObjectNode*>(Node), Pin]()
+						{
+							Node->SetPinHidden(*const_cast<UEdGraphPin*>(Pin), true);
+						})));
+				}
+				
+				
+				if (TSharedPtr<IDetailsView> Widget = Node->CustomizePinDetails(*Pin))
+				{
+					Node->PostReconstructNodeDelegate.AddLambda([WeakMenu = Widget->AsWeak()](){
+						if (TSharedPtr<SWidget> Menu = WeakMenu.Pin())
+						{
+							FSlateApplication::Get().DismissMenuByWidget(Menu.ToSharedRef());
+						}
+					});
+					
+					FToolMenuSection& Section = Menu->FindOrAddSection("EdGraphSchemaPinActions");
+					Section.AddSeparator("Pin Viewer");
+					Section.AddEntry(FToolMenuEntry::InitWidget("Pin Viewer", Widget.ToSharedRef(), {}));
+				}
+			}
+		}
+	}
+}
+
+
+void UEdGraphSchema_CustomizableObject::BreakNodeLinks(UEdGraphNode& TargetNode) const
+{
+#if WITH_EDITOR
+	TArray<UEdGraphPin*> Pins;
+
+	for (UEdGraphPin* Pin : TargetNode.Pins)
+	{
+		Pins.Append(ReverseFollowPinArray(*Pin, false));
+	}
+#endif
+	
+	Super::BreakNodeLinks(TargetNode);
+
+#if WITH_EDITOR
+	NodePinConnectionListChanged(Pins);
+#endif
+}
+
+
+void UEdGraphSchema_CustomizableObject::BreakPinLinks(UEdGraphPin& TargetPin, bool bSendsNodeNotification) const
+{
+	const FScopedTransaction Transaction(NSLOCTEXT("UnrealEd", "GraphEd_BreakPinLinks", "Break Pin Links"));
+
+	TArray<UEdGraphPin*> Pins;
+	Pins.Append(FollowPinArray(TargetPin, false));
+	Pins.Append(ReverseFollowPinArray(TargetPin, false));
+	
+	Super::BreakPinLinks(TargetPin, bSendsNodeNotification);
+
+#if WITH_EDITOR
+	NodePinConnectionListChanged(Pins);
+#endif
+}
+
+
+void UEdGraphSchema_CustomizableObject::BreakSinglePinLink(UEdGraphPin* SourcePin, UEdGraphPin* TargetPin) const
+{
+	const FScopedTransaction Transaction(NSLOCTEXT("UnrealEd", "GraphEd_BreakSinglePinLink", "Break Pin Link"));
+
+	TArray<UEdGraphPin*> SourceConnectedPins = ReverseFollowPinArray(*SourcePin, false);
+	TArray<UEdGraphPin*> TargetConnectedPins = ReverseFollowPinArray(*TargetPin, false);
+	
+	Super::BreakSinglePinLink(SourcePin, TargetPin);
+
+#if WITH_EDITOR
+	NodePinConnectionListChanged(SourceConnectedPins);
+	NodePinConnectionListChanged(TargetConnectedPins);
+#endif
+}
+
+
+TSharedPtr<ICustomizableObjectEditor> UEdGraphSchema_CustomizableObject::GetCustomizableObjectEditor(const class UEdGraph* ParentGraph)
+{
+	// Find the associated Editor
+	UCustomizableObject* CustomizableObject = Cast<UCustomizableObject>(ParentGraph->GetOuter());
+
+	TSharedPtr<ICustomizableObjectEditor> CustomizableObjectEditor;
+	if (CustomizableObject)
+	{
+		TSharedPtr< IToolkit > FoundAssetEditor = FToolkitManager::Get().FindEditorForAsset(CustomizableObject);
+		if (FoundAssetEditor.IsValid())
+		{
+			return StaticCastSharedPtr<ICustomizableObjectEditor>(FoundAssetEditor);
+		}
+	}
+
+	return nullptr;
+}
+
+
+UEdGraphNode* UEdGraphSchema_CustomizableObject::AddComment(class UEdGraph* ParentGraph, UEdGraphPin* FromPin, const FVector2D Location, bool bSelectNewNode/* = true*/) const
+{
+	TSharedPtr<ICustomizableObjectEditor> CustomizableObjectEditor = GetCustomizableObjectEditor(ParentGraph);
+
+	if(CustomizableObjectEditor.IsValid())
+	{
+		return CustomizableObjectEditor->CreateCommentBox(Location);
+	}
+
+	return nullptr;
+}
+
+
+void UEdGraphSchema_CustomizableObject::GetBreakLinkToSubMenuActions( class FMenuBuilder& MenuBuilder, UEdGraphPin* InGraphPin ) const
+{
+	// Make sure we have a unique name for every entry in the list
+	TMap< FString, uint32 > LinkTitleCount;
+
+	// Add all the links we could break from
+	for(TArray<class UEdGraphPin*>::TConstIterator Links(InGraphPin->LinkedTo); Links; ++Links)
+	{
+		FString Title = FString::Printf(TEXT("%s (%s)"), *(*Links)->GetOwningNode()->GetNodeTitle(ENodeTitleType::ListView).ToString(), *(*Links)->PinFriendlyName.ToString());
+		uint32 &Count = LinkTitleCount.FindOrAdd(Title);
+		FString Description;
+		if(Count == 0)
+		{
+			Description = FString::Printf(TEXT("Break link to %s"), *Title);
+		}
+		else
+		{
+			Description = FString::Printf(TEXT("Break link to %s (%d)"), *Title, Count);
+		}
+		++Count;
+		MenuBuilder.AddMenuEntry( 
+			FText::FromString(Description), 
+			FText::FromString(Description), 
+			FSlateIcon(),
+			FUIAction( FExecuteAction::CreateUObject((UEdGraphSchema_CustomizableObject*const)this, &UEdGraphSchema_CustomizableObject::BreakSinglePinLink, const_cast< UEdGraphPin* >(InGraphPin), *Links) ) 
+			);
+	}
+}
+
+
+void UEdGraphSchema_CustomizableObject::DroppedAssetsOnGraph(const TArray<FAssetData>& Assets, const FVector2D& GraphPosition, UEdGraph* Graph) const
+{
+	// To prevent overlapping when multiple assets are dropped at the same time on the graph
+	const int PixelOffset = 20;
+	int CurrentOffset = 0;
+	
+	for (FAssetData Asset : Assets)
+	{
+		// If it is not a valid asset to be spawned then just skip it
+		ESpawnableObjectType ObjectType;
+		if (!IsSpawnableAsset(Asset,ObjectType))
+		{
+			continue;
+		}
+
+		// At this point we know we are working with an asset we can spawn as a mutable node.
+		
+		UObject* Object = Asset.GetAsset();
+		UEdGraphNode* GraphNode = nullptr;
+
+		// Depending on the UObjectType spawn one or another mutable node
+		switch (ObjectType)
+		{
+		case ESpawnableObjectType::UTexture2D:
+			{
+				UTexture2D* Texture = Cast<UTexture2D>(Object);
+				UCustomizableObjectNodeTexture* Node = NewObject<UCustomizableObjectNodeTexture>(Graph);
+				Node->Texture = Texture;
+				GraphNode = Node;
+				break;
+			}
+		case ESpawnableObjectType::USkeletalMesh: 
+			{
+				USkeletalMesh* SkeletalMesh = Cast<USkeletalMesh>(Object);
+				UCustomizableObjectNodeSkeletalMesh* Node = NewObject<UCustomizableObjectNodeSkeletalMesh>(Graph);
+				Node->SkeletalMesh = SkeletalMesh;
+				GraphNode = Node;
+				break;
+			}
+		case ESpawnableObjectType::UStaticMesh:
+			{
+				UStaticMesh* Mesh = Cast<UStaticMesh>(Object);
+				UCustomizableObjectNodeStaticMesh* Node = NewObject<UCustomizableObjectNodeStaticMesh>(Graph);
+				Node->StaticMesh = Mesh;
+				GraphNode = Node;
+				break;
+			}
+		case ESpawnableObjectType::UMaterialInterface:
+			{
+				UMaterialInterface* Material = Cast<UMaterialInterface>(Object);
+				UCustomizableObjectNodeMaterial* Node = NewObject<UCustomizableObjectNodeMaterial>(Graph);
+				Node->SetMaterial(Material);
+				GraphNode = Node;
+				break;
+			}
+		// Error : A new compatible type set on UEdGraphSchema_CustomizableObject::IsSpawnableAsset is not providing a valid ESpawnableObjectType value
+		case ESpawnableObjectType::None:
+		// Error : a switch entry is missing for a ESpawnableObjectType value
+		default:
+			{
+				UE_LOG(LogTemp,Error,TEXT("Unable to create new mutable node for target asset : Invalid ESpawnableObjectType value. "));
+				checkNoEntry();
+				break;
+			}
+		}
+
+		// A node must have been spawned at this point.
+		if (GraphNode)
+		{
+			// A new node has been instanced, add it to the graph
+			GraphNode->CreateNewGuid();
+			GraphNode->PostPlacedNewNode();
+			GraphNode->AllocateDefaultPins();
+			GraphNode->NodePosX = GraphPosition.X + CurrentOffset;
+			GraphNode->NodePosY = GraphPosition.Y + CurrentOffset;
+			Graph->AddNode(GraphNode, true);
+			CurrentOffset += PixelOffset;	
+		}
+		else
+		{
+			UE_LOG(LogTemp,Error,TEXT("Unable to add null node to graph. "))
+		}
+	}
+}
+
+
+void UEdGraphSchema_CustomizableObject::OnPinConnectionDoubleCicked(UEdGraphPin* PinA, UEdGraphPin* PinB, const FVector2D& GraphPosition) const
+{
+	const FScopedTransaction Transaction(LOCTEXT("CreateRerouteNodeOnWire", "Create Reroute Node"));
+
+	// This constant is duplicated from inside of SGraphNodeKnot
+	const FVector2D NodeSpacerSize(42.0f, 24.0f);
+	const FVector2D KnotTopLeft = GraphPosition - (NodeSpacerSize * 0.5f);
+
+	// Create a new knot
+	UEdGraph* ParentGraph = PinA->GetOwningNode()->GetGraph();
+	UCustomizableObjectNodeReroute* DefaultNodeReroute = CastChecked<UCustomizableObjectNodeReroute>(UCustomizableObjectNodeReroute::StaticClass()->GetDefaultObject());
+	UCustomizableObjectNodeReroute* NodeReroute = CastChecked<UCustomizableObjectNodeReroute>(FCustomizableObjectSchemaAction_NewNode::CreateNode(ParentGraph, nullptr, KnotTopLeft, DefaultNodeReroute));
+
+	PinA->BreakLinkTo(PinB);
+	PinA->MakeLinkTo((PinA->Direction == EGPD_Output) ? NodeReroute->GetInputPin() : NodeReroute->GetOutputPin());
+	PinB->MakeLinkTo((PinB->Direction == EGPD_Output) ? NodeReroute->GetInputPin() : NodeReroute->GetOutputPin());
+	NodeReroute->UCustomizableObjectNode::ReconstructNode();
+}
+
+
+FConnectionDrawingPolicy* UEdGraphSchema_CustomizableObject::CreateConnectionDrawingPolicy(int32 InBackLayerID, int32 InFrontLayerID, float InZoomFactor, const FSlateRect& InClippingRect, FSlateWindowElementList& InDrawElements, UEdGraph* InGraphObj) const
+{
+	return new FCustomizableObjectConnectionDrawingPolicy(InBackLayerID, InFrontLayerID, InZoomFactor, InClippingRect, InDrawElements);
+}
+
+
+bool UEdGraphSchema_CustomizableObject::IsSpawnableAsset(const FAssetData& InAsset, ESpawnableObjectType& OutObjectType) const
+{
+	UObject* Object = InAsset.GetAsset();
+
+	// Type used to know what kind of UObject this asset is
+	OutObjectType = ESpawnableObjectType::None;
+
+	// Check if the provided object can be casted to any of the UObject types we can spawn as CO Nodes
+	if (Cast<UTexture2D>(Object))
+	{
+		OutObjectType = ESpawnableObjectType::UTexture2D;
+		return true;
+	}
+	else if (Cast<USkeletalMesh>(Object))
+	{
+		OutObjectType = ESpawnableObjectType::USkeletalMesh;
+		return true;
+	}
+	else if (Cast<UStaticMesh>(Object))
+	{
+		OutObjectType = ESpawnableObjectType::UStaticMesh;
+		return true;
+	}
+	else if (Cast<UMaterialInterface>(Object))
+	{
+		OutObjectType = ESpawnableObjectType::UMaterialInterface;
+		return true;
+	}
+	// Add more compatible types here, sync it up with ESpawnableObjectType
+	else
+	{
+		// Non spawnable object
+		return false;
+	}
+}
+
+
+void UEdGraphSchema_CustomizableObject::GetAssetsGraphHoverMessage(const TArray<FAssetData>& Assets,
+	const UEdGraph* HoverGraph, FString& OutTooltipText, bool& OutOkIcon) const
+{
+	Super::GetAssetsGraphHoverMessage(Assets, HoverGraph, OutTooltipText, OutOkIcon);
+
+	// Accept entry by default
+	OutOkIcon = true;
+	uint32 AmountOfIncompatibleAssets = 0;
+
+	// Iterate over the assets
+	for (const FAssetData& Asset : Assets)
+	{
+		// On first fail abort the consequent checks and tell the user
+		ESpawnableObjectType ObjectType;
+		if (!IsSpawnableAsset(Asset, ObjectType))
+		{
+			AmountOfIncompatibleAssets++;
+			OutOkIcon = false;
+
+			// Stop checking once we know that more than one asset is not compatible, the UI output will be the same
+			if (AmountOfIncompatibleAssets > 1)
+			{
+				break;
+			}
+		}
+	}
+
+	// Output debug message depending on the quantity of incompatible objects
+	if (!OutOkIcon)
+	{
+		if (Assets.Num() == 1)
+		{
+			OutTooltipText = FString("Incompatible asset selected : No node can be created for this type of asset.");
+		}
+		else if (Assets.Num() > 1 )
+		{
+			if (Assets.Num() == AmountOfIncompatibleAssets)
+			{
+				OutTooltipText = FString("Incompatible assets selected : No node can be created for any of the selected assets.");
+			}
+			else 
+			{
+				OutTooltipText = FString("Incompatible asset selected : Some assets will not be placed as nodes on the graph.");
+			}
+		}
+	}
+}
+
+
+bool UEdGraphSchema_CustomizableObject::TryCreateConnection(UEdGraphPin* PinA, UEdGraphPin* PinB) const
+{
+	if (!PinA || !PinB)
+	{
+		return false;
+	}
+	
+	UEdGraphNode* PinAOwningNode = PinA->GetOwningNode(); // TryCreateConnection can reconstruct the node invalidating the FromPin. Get the OwningNode before.
+	UEdGraphNode* PinBOwningNode = PinB->GetOwningNode();
+
+	bool bResult = Super::TryCreateConnection(PinA, PinB);
+
+	if (bResult)
+	{
+		PinAOwningNode->NodeConnectionListChanged();
+		PinBOwningNode->NodeConnectionListChanged();
+	}
+	
+	if (!PinA || PinA->bWasTrashed || !PinB || PinB->bWasTrashed)
+	{
+		return bResult;
+	}
+	
+	UEdGraphPin* InputPin;
+	UEdGraphPin* OutputPin;
+	if (!CategorizePinsByDirection(PinA, PinB, InputPin, OutputPin))
+	{
+		return bResult;
+	}
+
+	if (UCustomizableObjectNode* Node = Cast<UCustomizableObjectNode>(InputPin->GetOwningNode()))
+	{
+		Node->BreakExistingConnectionsPostConnection(InputPin, OutputPin);
+	}
+
+	if (UCustomizableObjectNode* Node = Cast<UCustomizableObjectNode>(OutputPin->GetOwningNode()))
+	{
+		Node->BreakExistingConnectionsPostConnection(InputPin, OutputPin);
+	}
+	
+#if WITH_EDITOR
+	if (bResult)
+	{
+		NodePinConnectionListChanged(ReverseFollowPinArray(*PinA, false));
+		NodePinConnectionListChanged(ReverseFollowPinArray(*PinB, false));
+	}
+#endif
+
+	return bResult;
+}
+
+
+FPinConnectionResponse UEdGraphSchema_CustomizableObject::MovePinLinks(UEdGraphPin& MoveFromPin, UEdGraphPin& MoveToPin, bool bIsIntermediateMove, bool bNotifyLinkedNodes) const
+{
+    // Mutable graph and its super does not use it. If we ever want to use it we should call NotifyIndirectConnections.
+	unimplemented();
+	return {};
+}
+
+
+FText UEdGraphSchema_CustomizableObject::GetPinCategoryName(const FName& PinCategory)
+{
+	if (PinCategory == UEdGraphSchema_CustomizableObject::PC_Object)
+	{
+		return LOCTEXT("Object_Pin_Category", "Object");
+	}
+	else if (PinCategory == UEdGraphSchema_CustomizableObject::PC_Component)
+	{
+		return LOCTEXT("Component_Pin_Category", "Component");
+	}
+	else if (PinCategory == UEdGraphSchema_CustomizableObject::PC_Material)
+	{
+		return LOCTEXT("MeshSection_Pin_Category", "Mesh Section");
+	}
+	else if (PinCategory == UEdGraphSchema_CustomizableObject::PC_Modifier)
+	{
+		return LOCTEXT("Modifier_Pin_Category", "Modifier");
+	}
+	else if (PinCategory == UEdGraphSchema_CustomizableObject::PC_Mesh)
+	{
+		return LOCTEXT("Mesh_Pin_Category", "Mesh");
+	}
+	else if (PinCategory == UEdGraphSchema_CustomizableObject::PC_Layout)
+	{
+		return LOCTEXT("Layout_Pin_Category", "Layout");
+	}
+	else if (PinCategory == UEdGraphSchema_CustomizableObject::PC_Image)
+	{
+		return LOCTEXT("Image_Pin_Category", "Texture");
+	}
+	else if (PinCategory == UEdGraphSchema_CustomizableObject::PC_PassThroughImage)
+	{
+		return LOCTEXT("PassThrough_Image_Pin_Category", "PassThrough Texture");
+	}
+	else if (PinCategory == UEdGraphSchema_CustomizableObject::PC_Projector)
+	{
+		return LOCTEXT("Projector_Pin_Category", "Projector");
+	}
+	else if (PinCategory == UEdGraphSchema_CustomizableObject::PC_GroupProjector)
+	{
+		return LOCTEXT("Group_Projector_Pin_Category", "Group Projector");
+	}
+	else if (PinCategory == UEdGraphSchema_CustomizableObject::PC_Color)
+	{
+		return LOCTEXT("Color_Pin_Category", "Color");
+	}
+	else if (PinCategory == UEdGraphSchema_CustomizableObject::PC_Float)
+	{
+		return LOCTEXT("Float_Pin_Category", "Float");
+	}
+	else if (PinCategory == UEdGraphSchema_CustomizableObject::PC_Bool)
+	{
+		return LOCTEXT("Bool_Pin_Category", "Bool");
+	}
+	else if (PinCategory == UEdGraphSchema_CustomizableObject::PC_Enum)
+	{
+		return LOCTEXT("Enum_Pin_Category", "Enum");
+	}
+	else if (PinCategory == UEdGraphSchema_CustomizableObject::PC_Stack)
+	{
+		return LOCTEXT("Stack_Pin_Category", "Stack");
+	}
+	else if (PinCategory == UEdGraphSchema_CustomizableObject::PC_MaterialAsset)
+	{
+		return LOCTEXT("Material_Asset_Pin_Category", "Table Material");
+	}
+	else if (PinCategory == UEdGraphSchema_CustomizableObject::PC_Wildcard)
+	{
+		return LOCTEXT("Wildcard_Pin_Category", "Wildcard");
+	}
+	else if (PinCategory == UEdGraphSchema_CustomizableObject::PC_PoseAsset)
+	{
+		return LOCTEXT("Pose_Pin_Category", "PoseAsset");
+	}
+	else
+	{
+		for (const FRegisteredCustomizableObjectPinType& PinType : ICustomizableObjectModule::Get().GetExtendedPinTypes())
+		{
+			if (PinType.PinType.Name == PinCategory)
+			{
+				return PinType.PinType.DisplayName;
+			}
+		}
+
+		// Need to fail gracefully here in case a plugin that was active when this graph was
+		// created is no longer loaded.
+		return LOCTEXT("Unknown_Pin_Category", "Unknown");
+	}
+}
+
+
+#undef LOCTEXT_NAMESPACE
