[mutable] Moved the Mutable plugin out of Experimental status into Beta.

diff --git a/Engine/Plugins/Mutable/Source/CustomizableObjectEditor/Private/MuCOE/GenerateMutableSource/GenerateMutableSource.cpp b/Engine/Plugins/Mutable/Source/CustomizableObjectEditor/Private/MuCOE/GenerateMutableSource/GenerateMutableSource.cpp
new file mode 100644
--- /dev/null
+++ b/Engine/Plugins/Mutable/Source/CustomizableObjectEditor/Private/MuCOE/GenerateMutableSource/GenerateMutableSource.cpp
@@ -0,0 +1,1802 @@
+// Copyright Epic Games, Inc. All Rights Reserved.
+
+#include "MuCOE/GenerateMutableSource/GenerateMutableSource.h"
+
+#include "GenerateMutableSourceComponent.h"
+#include "Animation/AnimInstance.h"
+#include "Engine/SkeletalMeshSocket.h"
+#include "Engine/TextureLODSettings.h"
+#include "Interfaces/ITargetPlatform.h"
+#include "MuCO/CustomizableObjectPrivate.h"
+#include "MuCO/CustomizableObjectExtension.h"
+#include "MuCO/CustomizableObjectUIData.h"
+#include "MuCO/ICustomizableObjectModule.h"
+#include "MuCOE/CustomizableObjectCompiler.h"
+#include "MuCOE/EdGraphSchema_CustomizableObject.h"
+#include "MuCOE/GenerateMutableSource/GenerateMutableSourceGroupProjector.h"
+#include "MuCOE/GenerateMutableSource/GenerateMutableSourceModifier.h"
+#include "MuCOE/GenerateMutableSource/GenerateMutableSourceSurface.h"
+#include "MuCOE/GraphTraversal.h"
+#include "MuCOE/ICustomizableObjectExtensionNode.h"
+#include "MuCOE/Nodes/CustomizableObjectNodeModifierEditMeshSectionBase.h"
+#include "MuCOE/Nodes/CustomizableObjectNodeModifierExtendMeshSection.h"
+#include "MuCOE/Nodes/CustomizableObjectNodeMaterial.h"
+#include "MuCOE/Nodes/CustomizableObjectNodeMaterialVariation.h"
+#include "MuCOE/Nodes/CustomizableObjectNodeMaterialSwitch.h"
+#include "MuCOE/Nodes/CustomizableObjectNodeMesh.h"
+#include "MuCOE/Nodes/CustomizableObjectNodeModifierClipWithMesh.h"
+#include "MuCOE/Nodes/CustomizableObjectNodeObjectGroup.h"
+#include "MuCOE/Nodes/CustomizableObjectNodeTable.h"
+#include "MuCOE/Nodes/CustomizableObjectNodeTexture.h"
+#include "MuCOE/Nodes/CustomizableObjectNodePassThroughTexture.h"
+#include "MuCOE/Nodes/CustomizableObjectNodeTextureBinarise.h"
+#include "MuCOE/Nodes/CustomizableObjectNodeTextureColourMap.h"
+#include "MuCOE/Nodes/CustomizableObjectNodeTextureFromChannels.h"
+#include "MuCOE/Nodes/CustomizableObjectNodeTextureInterpolate.h"
+#include "MuCOE/Nodes/CustomizableObjectNodeTextureInvert.h"
+#include "MuCOE/Nodes/CustomizableObjectNodeTextureLayer.h"
+#include "MuCOE/Nodes/CustomizableObjectNodeTextureParameter.h"
+#include "MuCOE/Nodes/CustomizableObjectNodeTextureProject.h"
+#include "MuCOE/Nodes/CustomizableObjectNodeTextureSwitch.h"
+#include "MuCOE/Nodes/CustomizableObjectNodePassThroughTextureSwitch.h"
+#include "MuCOE/Nodes/CustomizableObjectNodeTextureToChannels.h"
+#include "MuCOE/Nodes/CustomizableObjectNodeTextureTransform.h"
+#include "MuCOE/Nodes/CustomizableObjectNodeTextureSaturate.h"
+#include "MuCOE/Nodes/CustomizableObjectNodeTextureVariation.h"
+#include "MuCOE/UnrealEditorPortabilityHelpers.h"
+#include "MuT/NodeComponentEdit.h"
+#include "MuT/NodeLOD.h"
+#include "MuT/NodeMeshConstant.h"
+#include "MuT/NodeObjectGroup.h"
+#include "MuT/NodeObjectNew.h"
+#include "MuR/Mesh.h"
+#include "PhysicsEngine/PhysicsAsset.h"
+#include "PlatformInfo.h"
+#include "Math/NumericLimits.h"
+#include "Hash/CityHash.h"
+#include "MuCOE/Nodes/CustomizableObjectNodeCopyMaterial.h"
+
+#define LOCTEXT_NAMESPACE "CustomizableObjectEditor"
+
+
+uint32 GetTypeHash(const FGraphCycleKey& Key)
+{
+	return HashCombine(GetTypeHash(&Key.Pin), GetTypeHash(Key.Id));
+}
+
+
+FGraphCycle::FGraphCycle(const FGraphCycleKey&& Key, FMutableGraphGenerationContext& Context) :
+	Key(Key),
+	Context(Context)
+{
+}
+
+
+FGraphCycle::~FGraphCycle()
+{
+	Context.VisitedPins.Remove(Key);
+}
+
+
+bool FGraphCycle::FoundCycle() const
+{
+	const UCustomizableObjectNode& Node = *Cast<UCustomizableObjectNode>(Key.Pin.GetOwningNode());
+
+	if (const UCustomizableObject** Result = Context.VisitedPins.Find(Key))
+	{
+		Context.Compiler->CompilerLog(LOCTEXT("CycleFoundNode", "Cycle detected."), &Node, EMessageSeverity::Error, true);
+		Context.CustomizableObjectWithCycle = *Result;
+		return true;
+	}
+	else
+	{
+		Context.VisitedPins.Add(Key, Node.GetGraph()->GetTypedOuter<UCustomizableObject>());
+		return false;	
+	}
+}
+
+
+/** Warn if the node has more outputs than it is meant to have. */
+void CheckNumOutputs(const UEdGraphPin& Pin, const FMutableGraphGenerationContext& GenerationContext)
+{
+	if (const UCustomizableObjectNode* Typed = Cast<UCustomizableObjectNode>(Pin.GetOwningNode()))
+	{
+		if (Typed->IsSingleOutputNode())
+		{
+			int numOutLinks = 0;
+			for (UEdGraphPin* NodePin : Typed->GetAllNonOrphanPins())
+			{
+				if (NodePin->Direction == EGPD_Output)
+				{
+					numOutLinks += NodePin->LinkedTo.Num();
+				}
+			}
+
+			if (numOutLinks > 1)
+			{
+				GenerationContext.Compiler->CompilerLog(LOCTEXT("MultipleOutgoing", "The node has several outgoing connections, but it should be limited to 1."), CastChecked<UCustomizableObjectNode>(Pin.GetOwningNode()));
+			}
+		}
+	}
+}
+
+
+FMutableGraphGenerationContext::FMutableGraphGenerationContext(UCustomizableObject* InObject, FCustomizableObjectCompiler* InCompiler, const FCompilationOptions& InOptions)
+	: Object(InObject), Compiler(InCompiler), Options(InOptions)
+	, ExtensionDataCompilerInterface(*this)
+{
+	// Default flags for mesh generation nodes.
+	MeshGenerationFlags.Push(EMutableMeshConversionFlags::None);
+
+	// Default flags for mesh layout generation.
+	LayoutGenerationFlags.Push(FLayoutGenerationFlags());
+}
+
+FMutableGraphGenerationContext::~FMutableGraphGenerationContext() = default;
+
+
+void FMutableGraphGenerationContext::AddParticipatingObject(const UObject& InObject)
+{
+	PRAGMA_DISABLE_DEPRECATION_WARNINGS
+	const FGuid PackageGuid = InObject.GetPackage()->GetGuid();
+	PRAGMA_ENABLE_DEPRECATION_WARNINGS
+
+	ParticipatingObjects.Add(InObject.GetPackage()->GetFName(), PackageGuid);
+}
+
+
+mu::Ptr<mu::Mesh> FMutableGraphGenerationContext::FindGeneratedMesh( const FGeneratedMeshData::FKey& Key )
+{
+	for (const FGeneratedMeshData& d : GeneratedMeshes)
+	{
+		if (d.Key == Key )
+		{
+			return d.Generated;
+		}
+	}
+
+	return nullptr;
+}
+
+
+int32 FMutableGraphGenerationContext::AddStreamedResource(uint32 InResourceHash, UCustomizableObjectResourceDataContainer*& OutNewResource)
+{
+	OutNewResource = nullptr;
+
+	// Return resource index if found.
+	if (int32* ResourceIndex = StreamedResourceIndices.Find(InResourceHash))
+	{
+		return *ResourceIndex;
+	}
+
+	int32 NewResourceIndex = StreamedResourceData.Num();
+	const FString ContainerName = GetNameSafe(Object) + FString::Printf(TEXT("_SR_%d"), NewResourceIndex);
+
+	UCustomizableObjectResourceDataContainer* ExistingContainer = FindObject<UCustomizableObjectResourceDataContainer>(Object, *ContainerName);
+	OutNewResource = ExistingContainer ? ExistingContainer : NewObject<UCustomizableObjectResourceDataContainer>(
+		Object,
+		FName(*ContainerName),
+		RF_Public);
+
+	StreamedResourceData.Add(OutNewResource);
+	StreamedResourceIndices.Add({ InResourceHash, NewResourceIndex });
+
+	return NewResourceIndex;
+}
+
+
+int32 FMutableGraphGenerationContext::AddAssetUserDataToStreamedResources(UAssetUserData* AssetUserData)
+{
+	check(AssetUserData);
+	const uint32 AssetIdentifier = AssetUserData->GetUniqueID();
+
+	UCustomizableObjectResourceDataContainer* NewResource = nullptr;
+	const int32 ResourceIndex = AddStreamedResource(AssetIdentifier, NewResource);
+
+	if (NewResource) // Nullptr if not new
+	{
+		FCustomizableObjectAssetUserData ResourceData;
+		ResourceData.AssetUserDataEditor = AssetUserData;
+
+		NewResource->Data.Type = ECOResourceDataType::AssetUserData;
+		NewResource->Data.Data = FInstancedStruct::Make(ResourceData);
+	}
+
+	check(StreamedResourceData[ResourceIndex].GetLoadedData().Type == ECOResourceDataType::AssetUserData);
+	return ResourceIndex;
+}
+
+
+/** Adds to ParameterNamesMap the node Node to the array of elements with name Name */
+void FMutableGraphGenerationContext::AddParameterNameUnique(const UCustomizableObjectNode* Node, FString Name)
+{
+	if (TArray<const UObject*>* ArrayResult = ParameterNamesMap.Find(Name))
+	{
+		ArrayResult->AddUnique(Node);
+	}
+	else
+	{
+		TArray<const UObject*> ArrayTemp;
+		ArrayTemp.Add(Node);
+		ParameterNamesMap.Add(Name, ArrayTemp);
+	}
+}
+
+uint32 FMutableGraphGenerationContext::GetSkinWeightProfileIdUnique(const FName InProfileName)
+{
+	const FString ProfileNameString = InProfileName.ToString();
+	
+	uint32 UniqueProfileId = CityHash32(reinterpret_cast<const char*>(*ProfileNameString), ProfileNameString.Len() * sizeof(FString::ElementType));
+
+	bool bRemappedProfile = false;
+
+	FName& ProfileName = UniqueSkinWeightProfileIds.FindOrAdd(UniqueProfileId, InProfileName);
+	while (ProfileName != InProfileName)
+	{
+		if (uint32* RemappedProfileId = RemappedSkinWeightProfileIds.Find(InProfileName))
+		{
+			UniqueProfileId = *RemappedProfileId;
+			break;
+		}
+
+		// Id collision detected
+		bRemappedProfile = true;
+
+		// Increase Id in an attempt to make it unique again.
+		++UniqueProfileId;
+		ProfileName = UniqueSkinWeightProfileIds.FindOrAdd(UniqueProfileId, InProfileName);
+	}
+
+	if (bRemappedProfile)
+	{
+		RemappedSkinWeightProfileIds.Add(InProfileName, UniqueProfileId);
+	}
+
+	return UniqueProfileId;
+}
+
+
+const FGuid FMutableGraphGenerationContext::GetNodeIdUnique(const UCustomizableObjectNode* Node)
+{
+	TArray<const UObject*>* ArrayResult = NodeIdsMap.Find(Node->NodeGuid);
+
+	if (ArrayResult == nullptr)
+	{
+		TArray<const UObject*> ArrayTemp;
+		ArrayTemp.Add(Node);
+		NodeIdsMap.Add(Node->NodeGuid, ArrayTemp);
+		return Node->NodeGuid;
+	}
+
+	ArrayResult->Add(Node);
+	return FGuid::NewGuid();
+}
+
+
+void FMutableGraphGenerationContext::GenerateClippingCOInternalTags()
+{
+	int32 i = 0;
+	int32 j;
+
+	// The following array ArrayCachedTaggedMaterialNodes and related functionality was removed because it prevented two clip nodes from
+	// clipping the same material node
+	// In the case more than one CO clips a particular CO, the associated mu::NodeSurfaceNewPtr nodes with that particular CO only need to be tagged once
+	//TArray<FGuid> ArrayCachedTaggedMaterialNodes;
+
+	for (auto It = MapClipMeshNodeToMutableClipMeshNodeArray.CreateConstIterator(); It; ++It)
+	{
+		FString TagName = FString::Printf(TEXT("ClippingTag_%d"), i);
+		i++;
+
+		for (j = 0; j < It->Value.Num(); ++j)
+		{
+			It->Value[j]->RequiredTags.Add(TagName);
+		}
+
+		UCustomizableObjectNodeModifierClipWithMesh* CustomizableObjectNodeModifierClipWithMesh = It->Key;
+
+		// Tag also those nodes in the material side of the clipping functionality (since a CO can have several materials with the same LOD,
+		// all of them are considered, so CustomizableObjectNodeMeshClipWithMesh->ArrayMaterialNodeToClipWithID contains the Guids of all
+		// the material nodes with the same name selected for being affected by the clipping operation
+		TArray<FGuid>& ArrayMaterialNodeToClipWithID = CustomizableObjectNodeModifierClipWithMesh->ArrayMaterialNodeToClipWithID;
+
+		for (j = 0; j < ArrayMaterialNodeToClipWithID.Num(); ++j)
+		{
+			// The following lines were commented for the same reasons as ArrayCachedTaggedMaterialNodes above
+			//if (ArrayCachedTaggedMaterialNodes.Find(CustomizableObjectNodeMeshClipWithMesh->ArrayMaterialNodeToClipWithID[j]) != INDEX_NONE)
+			//{
+			//    continue;
+			//}
+
+			bool FoundEntryInMap = false;
+			for (auto It2 = MapMaterialNodeToMutableSurfaceNodeArray.CreateConstIterator(); (It2 && !FoundEntryInMap); ++It2)
+			{
+				if (It2->Key->NodeGuid == ArrayMaterialNodeToClipWithID[j])
+				{
+					FoundEntryInMap = true;
+
+					for (int32 k = 0; k < It2->Value.Num(); ++k)
+					{
+						It2->Value[k]->Tags.Add(TagName);
+					}
+				}
+			}
+
+			// The following lines were commented for the same reasons as ArrayCachedTaggedMaterialNodes above
+			//ArrayCachedTaggedMaterialNodes.Add(CustomizableObjectNodeMeshClipWithMesh->ArrayMaterialNodeToClipWithID[j]);
+		}
+	}
+}
+
+
+void FMutableGraphGenerationContext::GenerateSharedSurfacesUniqueIds()
+{
+	int32 UniqueId = 0;
+
+	TArray<TArray<FSharedSurface>> NodeToSharedSurfaces;
+	SharedSurfaceIds.GenerateValueArray(NodeToSharedSurfaces);
+
+	TArray<bool> VisitedSurfaces;
+	for (TArray<FSharedSurface>& SharedSurfaces : NodeToSharedSurfaces)
+	{
+		const int32 NumSurfaces = SharedSurfaces.Num();
+		VisitedSurfaces.Init(false, NumSurfaces);
+		
+		// Iterate all surfaces for a given NodeMaterial and set the same SharedSurfaceId to those that are equal.
+		for (int32 SurfaceIndex = 0; SurfaceIndex < NumSurfaces; ++SurfaceIndex)
+		{
+			if (VisitedSurfaces[SurfaceIndex])
+			{
+				continue;
+			}
+
+			FSharedSurface& CurrentSharedSurface = SharedSurfaces[SurfaceIndex];
+			CurrentSharedSurface.NodeSurfaceNew->SharedSurfaceId = UniqueId;
+			VisitedSurfaces[SurfaceIndex] = true;
+
+			for (int32 AuxSurfaceIndex = SurfaceIndex; AuxSurfaceIndex < NumSurfaces && !CurrentSharedSurface.bMakeUnique; ++AuxSurfaceIndex)
+			{
+				if (VisitedSurfaces[AuxSurfaceIndex])
+				{
+					continue;
+				}
+
+				if (SharedSurfaces[AuxSurfaceIndex].NodeModifierIDs != CurrentSharedSurface.NodeModifierIDs)
+				{
+					continue;
+				}
+
+				SharedSurfaces[AuxSurfaceIndex].NodeSurfaceNew->SharedSurfaceId = UniqueId;
+				VisitedSurfaces[AuxSurfaceIndex] = true;
+			}
+
+			++UniqueId;
+		}
+	}
+}
+
+bool FMutableGraphGenerationContext::FindBone(const FName& InBoneName, mu::FBoneName& OutBoneName) const
+{
+	const FString BoneNameString = InBoneName.ToString().ToLower();
+	OutBoneName.Id = CityHash32(reinterpret_cast<const char*>(*BoneNameString), BoneNameString.Len() * sizeof(FString::ElementType));
+	if (UniqueBoneNames.Contains(OutBoneName))
+	{
+		return true;
+	}
+
+	if (const mu::FBoneName* BoneName = RemappedBoneNames.Find(BoneNameString))
+	{
+		OutBoneName = *BoneName;
+		return true;
+	}
+
+	return false;
+}
+
+
+mu::FBoneName FMutableGraphGenerationContext::GetBoneUnique(const FName& InBoneName)
+{
+	const FString BoneNameString = InBoneName.ToString().ToLower();
+	mu::FBoneName Bone(CityHash32(reinterpret_cast<const char*>(*BoneNameString), BoneNameString.Len() * sizeof(FString::ElementType)));
+
+	bool bNewRemappedBoneName = false;
+
+	FString& BoneName = UniqueBoneNames.FindOrAdd(Bone, BoneNameString);
+	while (BoneName != InBoneName)
+	{
+		if (mu::FBoneName* RemappedBoneName = RemappedBoneNames.Find(BoneNameString))
+		{
+			Bone.Id = RemappedBoneName->Id;
+			break;
+		}
+
+		// Id collision detected
+		bNewRemappedBoneName = true;
+
+		// Increase Id in an attempt to make it unique again.
+		++Bone.Id;
+		BoneName = UniqueBoneNames.FindOrAdd(Bone, BoneNameString);
+	}
+
+	if (bNewRemappedBoneName)
+	{
+		RemappedBoneNames.Add(BoneNameString, Bone);
+	}
+
+	return Bone;
+}
+
+
+//void FMutableGraphGenerationContext::CheckPhysicsAssetInSkeletalMesh(const USkeletalMesh* SkeletalMesh)
+//{
+//	if (SkeletalMesh && SkeletalMesh->GetPhysicsAsset() && !DiscartedPhysicsAssetMap.Find(SkeletalMesh->GetPhysicsAsset()))
+//	{
+//		for (const TObjectPtr<USkeletalBodySetup>& Bodies : SkeletalMesh->GetPhysicsAsset()->SkeletalBodySetups)
+//		{
+//			if (Bodies->BoneName != NAME_None && SkeletalMesh->GetRefSkeleton().FindBoneIndex(Bodies->BoneName) == INDEX_NONE)
+//			{
+//				DiscartedPhysicsAssetMap.Add(SkeletalMesh->GetPhysicsAsset(), DiscartedPhysicsAssetMap.Num());
+//				break;
+//			}
+//		}
+//	}
+//}
+
+
+FMutableComponentInfo* FMutableGraphGenerationContext::GetCurrentComponentInfo()
+{
+	return ComponentInfos.FindByPredicate([this](const FMutableComponentInfo& Component) { return Component.ComponentName == CurrentMeshComponent; });
+}
+
+
+FGeneratedKey::FGeneratedKey(void* InFunctionAddress, const UEdGraphPin& InPin, const UCustomizableObjectNode& Node, FMutableGraphGenerationContext& GenerationContext, const bool UseMesh, const bool InbOnlyConnectedLOD)
+{
+	FunctionAddress = InFunctionAddress;
+	Pin = &InPin;
+	LOD = Node.IsAffectedByLOD() ? GenerationContext.CurrentLOD : 0;
+
+	if (UseMesh)
+	{
+		Flags = GenerationContext.MeshGenerationFlags.Last();
+		LayoutFlags = GenerationContext.LayoutGenerationFlags.Last();
+		MeshMorphStack = GenerationContext.MeshMorphStack;
+		bOnlyConnectedLOD = InbOnlyConnectedLOD;
+	}
+}
+
+
+bool FGeneratedKey::operator==(const FGeneratedKey& Other) const
+{
+	return FunctionAddress == Other.FunctionAddress &&
+		Pin == Other.Pin &&
+		LOD == Other.LOD &&
+		Flags == Other.Flags &&
+		LayoutFlags == Other.LayoutFlags &&
+		MeshMorphStack == Other.MeshMorphStack &&
+		bOnlyConnectedLOD == Other.bOnlyConnectedLOD;
+}
+
+
+uint32 GetTypeHash(const FGeneratedKey& Key)
+{
+	uint32 Hash = GetTypeHash(Key.FunctionAddress);
+	Hash = HashCombine(Hash, GetTypeHash(Key.Pin));
+	Hash = HashCombine(Hash, GetTypeHash(Key.LOD));
+	Hash = HashCombine(Hash, GetTypeHash(Key.Flags));
+	//Hash = HashCombine(Hash, GetTypeHash(Key.LayoutFlags)); // Does not support array
+	//Hash = HashCombine(Hash, GetTypeHash(Key.MeshMorphStack)); // Does not support array
+	Hash = HashCombine(Hash, GetTypeHash(Key.bOnlyConnectedLOD));
+	
+	return Hash;
+}
+
+
+void FPinDataValue::Append(const FPinDataValue& From)
+{
+	MeshesData.Append(From.MeshesData);
+}
+
+
+FPinDataValue* FPinData::Find(const UEdGraphPin* Pin)
+{
+	return Data.Find(Pin);
+}
+
+
+FPinDataValue& FPinData::GetCurrent()
+{
+	return Data[PinStack.Last()];
+}
+
+
+void FPinData::Pop()
+{
+	const uint32 Num = PinStack.Num();
+
+	check(Num >= 1); // Pop called without a previous push
+
+	if (Num >= 2)
+	{
+		Data[PinStack[Num - 2]].Append(Data[PinStack[Num - 1]]);
+	}
+
+	PinStack.Pop();
+}
+
+
+void FPinData::Push(const UEdGraphPin* Pin)
+{
+	PinStack.Push(Pin);
+	Data.FindOrAdd(Pin);
+}
+
+
+FGraphCycleKey::FGraphCycleKey(const UEdGraphPin& Pin, const FString& Id) :
+	Pin(Pin),
+	Id(Id)
+{
+}
+
+
+bool FGraphCycleKey::operator==(const FGraphCycleKey& Other) const
+{
+	return &Pin == &Other.Pin && Id == Other.Id;
+}
+
+
+FScopedPinData::FScopedPinData(FMutableGraphGenerationContext& Context, const UEdGraphPin* Pin)
+	: Context(Context)
+{
+	Context.PinData.Push(Pin);
+}
+
+
+FScopedPinData::~FScopedPinData()
+{
+	Context.PinData.Pop();
+}
+
+
+UTexture2D* FindReferenceImage(const UEdGraphPin* Pin, FMutableGraphGenerationContext& GenerationContext)
+{
+	check(Pin)
+	RETURN_ON_CYCLE(*Pin, GenerationContext)
+
+	CheckNumOutputs(*Pin, GenerationContext);
+	
+	const UEdGraphSchema_CustomizableObject* Schema = GetDefault<UEdGraphSchema_CustomizableObject>();
+	const UCustomizableObjectNode* Node = CastChecked<UCustomizableObjectNode>(Pin->GetOwningNode());
+
+	UTexture2D* Result = nullptr;
+
+	if (const UCustomizableObjectNodeTexture* TypedNodeTex = Cast<UCustomizableObjectNodeTexture>(Node))
+	{
+		Result = TypedNodeTex->Texture;
+	}
+
+	else if (const UCustomizableObjectNodePassThroughTexture* TypedNodePassThroughTex = Cast<UCustomizableObjectNodePassThroughTexture>(Node))
+	{
+		Result = Cast<UTexture2D>(TypedNodePassThroughTex->PassThroughTexture);
+	}
+
+	else if (const UCustomizableObjectNodeTextureParameter* ParamNodeTex = Cast<UCustomizableObjectNodeTextureParameter>(Node))
+	{
+		Result = ParamNodeTex->ReferenceValue;
+	}
+
+	else if (const UCustomizableObjectNodeMesh* TypedNodeMesh = Cast<UCustomizableObjectNodeMesh>(Node))
+	{
+		Result = TypedNodeMesh->FindTextureForPin(Pin);
+	}
+
+	else if (const UCustomizableObjectNodeTextureInterpolate* TypedNodeInterp = Cast<UCustomizableObjectNodeTextureInterpolate>(Node))
+	{
+		for (int LayerIndex = 0; !Result && LayerIndex < TypedNodeInterp->GetNumTargets(); ++LayerIndex)
+		{
+			if (const UEdGraphPin* ConnectedPin = FollowInputPin(*TypedNodeInterp->Targets(LayerIndex)))
+			{
+				Result = FindReferenceImage(ConnectedPin, GenerationContext);
+			}
+		}
+	}
+
+	else if (const UCustomizableObjectNodeTextureLayer* TypedNodeLayer = Cast<UCustomizableObjectNodeTextureLayer>(Node))
+	{
+		if (const UEdGraphPin* ConnectedPin = FollowInputPin(*TypedNodeLayer->BasePin()))
+		{
+			Result = FindReferenceImage(ConnectedPin, GenerationContext);
+		}
+
+		for (int LayerIndex = 0; !Result && LayerIndex < TypedNodeLayer->GetNumLayers(); ++LayerIndex)
+		{
+			if (const UEdGraphPin* ConnectedPin = FollowInputPin(*TypedNodeLayer->LayerPin(LayerIndex)))
+			{
+				if (ConnectedPin->PinType.PinCategory == Schema->PC_Image)
+				{
+					Result = FindReferenceImage(ConnectedPin, GenerationContext);
+				}
+			}
+		}
+	}
+
+	else if (const UCustomizableObjectNodeTextureSwitch* TypedNodeSwitch = Cast<UCustomizableObjectNodeTextureSwitch>(Node))
+	{
+		for (int SelectorIndex = 0; !Result && SelectorIndex < TypedNodeSwitch->GetNumElements(); ++SelectorIndex)
+		{
+			if (const UEdGraphPin* ConnectedPin = FollowInputPin(*TypedNodeSwitch->GetElementPin(SelectorIndex)))
+			{
+				Result = FindReferenceImage(ConnectedPin, GenerationContext);
+			}
+		}
+	}
+
+	else if (const UCustomizableObjectNodePassThroughTextureSwitch* TypedNodePassThroughSwitch = Cast<UCustomizableObjectNodePassThroughTextureSwitch>(Node))
+	{
+		for (int32 SelectorIndex = 0; !Result && SelectorIndex < TypedNodePassThroughSwitch->GetNumElements(); ++SelectorIndex)
+		{
+			if (const UEdGraphPin* ConnectedPin = FollowInputPin(*TypedNodePassThroughSwitch->GetElementPin(SelectorIndex)))
+			{
+				Result = FindReferenceImage(ConnectedPin, GenerationContext);
+			}
+		}
+	}
+
+	else if (const UCustomizableObjectNodeTextureVariation* TypedNodeVariation = Cast<UCustomizableObjectNodeTextureVariation>(Node))
+	{
+		if (const UEdGraphPin* ConnectedPin = FollowInputPin(*TypedNodeVariation->DefaultPin()))
+		{
+			Result = FindReferenceImage(ConnectedPin, GenerationContext);
+		}
+
+		for (int SelectorIndex = 0; !Result && SelectorIndex < TypedNodeVariation->GetNumVariations(); ++SelectorIndex)
+		{
+			if (const UEdGraphPin* ConnectedPin = FollowInputPin(*TypedNodeVariation->VariationPin(SelectorIndex)))
+			{
+				Result = FindReferenceImage(ConnectedPin, GenerationContext);
+			}
+		}
+	}
+
+	else if (const UCustomizableObjectNodeTextureFromChannels* TypedNodeFrom = Cast<UCustomizableObjectNodeTextureFromChannels>(Node))
+	{
+		mu::NodeImagePtr RNode, GNode, BNode, ANode;
+
+		if (const UEdGraphPin* ConnectedPin = FollowInputPin(*TypedNodeFrom->RPin()))
+		{
+			Result = FindReferenceImage(ConnectedPin, GenerationContext);
+		}
+		if (!Result)
+		{
+			if (const UEdGraphPin* ConnectedPin = FollowInputPin(*TypedNodeFrom->GPin()))
+			{
+				Result = FindReferenceImage(ConnectedPin, GenerationContext);
+			}
+		}
+		if (!Result)
+		{
+			if (const UEdGraphPin* ConnectedPin = FollowInputPin(*TypedNodeFrom->BPin()))
+			{
+				Result = FindReferenceImage(ConnectedPin, GenerationContext);
+			}
+			
+		}
+		if (!Result)
+		{
+			if (const UEdGraphPin* ConnectedPin = FollowInputPin(*TypedNodeFrom->APin()))
+			{
+				Result = FindReferenceImage(ConnectedPin, GenerationContext);
+			}
+		}
+	}
+
+	else if (const UCustomizableObjectNodeTextureToChannels* TypedNodeTo = Cast<UCustomizableObjectNodeTextureToChannels>(Node))
+	{
+		if (const UEdGraphPin* ConnectedPin = FollowInputPin(*TypedNodeTo->InputPin()))
+		{
+			Result = FindReferenceImage(ConnectedPin, GenerationContext);
+		}
+	}
+
+	else if (const UCustomizableObjectNodeTextureProject* TypedNodeProj = Cast<UCustomizableObjectNodeTextureProject>(Node))
+	{
+		if (TypedNodeProj->ReferenceTexture)
+		{
+			Result = TypedNodeProj->ReferenceTexture;
+		}
+		else
+		{
+			int TexIndex = -1;// TypedNodeProj->OutputPins.Find((UEdGraphPin*)Pin);
+			for (int32 i = 0; i < TypedNodeProj->GetNumOutputs(); ++i)
+			{
+				if (TypedNodeProj->OutputPins(i) == Pin)
+				{
+					TexIndex = i;
+				}
+			}
+
+			check(TexIndex >= 0 && TexIndex < TypedNodeProj->GetNumTextures());
+
+			if (const UEdGraphPin* ConnectedPin = FollowInputPin(*TypedNodeProj->TexturePins(TexIndex)))
+			{
+				Result = FindReferenceImage(ConnectedPin, GenerationContext);
+			}
+		}
+	}
+
+	else if (const UCustomizableObjectNodeTextureBinarise* TypedNodeBin = Cast<UCustomizableObjectNodeTextureBinarise>(Node))
+	{
+		if (const UEdGraphPin* ConnectedPin = FollowInputPin(*TypedNodeBin->GetBaseImagePin()))
+		{
+			Result = FindReferenceImage(ConnectedPin, GenerationContext);
+		}
+	}
+
+	else if (const UCustomizableObjectNodeTextureInvert* TypedNodeInv = Cast<UCustomizableObjectNodeTextureInvert>(Node))
+	{
+		if (const UEdGraphPin* ConnectedPin = FollowInputPin(*TypedNodeInv->GetBaseImagePin()))
+		{
+			Result = FindReferenceImage(ConnectedPin, GenerationContext);
+		}
+	}
+
+	else if (const UCustomizableObjectNodeTextureColourMap* TypedNodeColourMap = Cast<UCustomizableObjectNodeTextureColourMap>(Node))
+	{
+		if (const UEdGraphPin* ConnectedPin = FollowInputPin(*TypedNodeColourMap->GetBasePin()))
+		{
+			Result = FindReferenceImage(ConnectedPin, GenerationContext);
+		}
+	}
+
+	else if (const UCustomizableObjectNodeTextureTransform* TypedNodeTransform = Cast<UCustomizableObjectNodeTextureTransform>(Node))
+	{
+		if ( UEdGraphPin* BaseImagePin = FollowInputPin(*TypedNodeTransform->GetBaseImagePin()) )
+		{
+			Result = FindReferenceImage(BaseImagePin, GenerationContext);
+		}
+	}
+
+	else if (const UCustomizableObjectNodeTextureSaturate* TypedNodeSaturate = Cast<UCustomizableObjectNodeTextureSaturate>(Node))
+	{
+		if ( UEdGraphPin* BaseImagePin = FollowInputPin(*TypedNodeSaturate->GetBaseImagePin()) )
+		{
+			Result = FindReferenceImage(BaseImagePin, GenerationContext);
+		}
+	}
+
+	else if (const UCustomizableObjectNodeTable* TypedNodeTable = Cast<UCustomizableObjectNodeTable>(Node))
+	{
+		if (Pin->PinType.PinCategory == Schema->PC_MaterialAsset)
+		{
+			Result = TypedNodeTable->FindReferenceTextureParameter(Pin, GenerationContext.CurrentMaterialTableParameter);
+		}
+		else
+		{
+			Result = TypedNodeTable->GetColumnDefaultAssetByType<UTexture2D>(Pin);
+		}
+	}
+	
+	return Result;
+}
+
+
+mu::NodeMeshApplyPosePtr CreateNodeMeshApplyPose(FMutableGraphGenerationContext& GenerationContext, mu::NodeMeshPtr InputMeshNode, const TArray<FName>& ArrayBoneName, const TArray<FTransform>& ArrayTransform)
+{
+	check(ArrayBoneName.Num() == ArrayTransform.Num());
+
+	mu::Ptr<mu::Mesh> MutableMesh = new mu::Mesh();
+	mu::NodeMeshConstantPtr PoseNodeMesh = new mu::NodeMeshConstant;
+	PoseNodeMesh->SetValue(MutableMesh);
+
+	mu::Ptr<mu::Skeleton> MutableSkeleton = new mu::Skeleton;
+	MutableMesh->SetSkeleton(MutableSkeleton);
+	MutableMesh->SetBonePoseCount(ArrayBoneName.Num());
+	MutableSkeleton->SetBoneCount(ArrayBoneName.Num());
+
+	for (int32 i = 0; i < ArrayBoneName.Num(); ++i)
+	{
+		const FName BoneName = ArrayBoneName[i];
+		const mu::FBoneName& BoneId = GenerationContext.GetBoneUnique(BoneName);
+
+		MutableSkeleton->SetDebugName(i, BoneName);
+		MutableSkeleton->SetBoneName(i, { BoneId });
+		MutableMesh->SetBonePose(i, BoneId, (FTransform3f)ArrayTransform[i], mu::EBoneUsageFlags::Skinning);
+	}
+
+	mu::NodeMeshApplyPosePtr NodeMeshApplyPose = new mu::NodeMeshApplyPose;
+	NodeMeshApplyPose->SetBase(InputMeshNode);
+	NodeMeshApplyPose->SetPose(PoseNodeMesh);
+
+	return NodeMeshApplyPose;
+}
+
+
+// Convert a CustomizableObject Source Graph into a mutable source graph  
+mu::Ptr<mu::NodeObject> GenerateMutableSource(const UEdGraphPin * Pin, FMutableGraphGenerationContext & GenerationContext)
+{
+	MUTABLE_CPUPROFILER_SCOPE(GenerateMutableSource);
+
+	check(Pin)
+	RETURN_ON_CYCLE(*Pin, GenerationContext)
+
+	CheckNumOutputs(*Pin, GenerationContext);
+	
+	UCustomizableObjectNode* Node = CastChecked<UCustomizableObjectNode>(Pin->GetOwningNode());
+
+	GenerationContext.AddParticipatingObject(*GetRootObject(*Node));
+	
+	const FGeneratedKey Key(reinterpret_cast<void*>(&GenerateMutableSource), *Pin, *Node, GenerationContext, true);
+	if (const FGeneratedData* Generated = GenerationContext.Generated.Find(Key))
+	{
+		return static_cast<mu::NodeObject*>(Generated->Node.get());
+	}
+
+	mu::Ptr<mu::NodeObject> Result;
+	
+	if (const UCustomizableObjectNodeObject* TypedNodeObj = Cast<UCustomizableObjectNodeObject>(Node))
+	{
+		mu::Ptr<mu::NodeObjectNew> ObjectNode = new mu::NodeObjectNew();
+		Result = ObjectNode;
+
+		ObjectNode->SetName(TypedNodeObj->ObjectName);
+		FGuid FinalGuid = GenerationContext.GetNodeIdUnique(TypedNodeObj);
+		if (FinalGuid != TypedNodeObj->NodeGuid)
+		{
+			GenerationContext.Compiler->CompilerLog(FText::FromString(TEXT("Warning: Node has a duplicated GUID. A new ID has been generated, but cooked data will not be deterministic.")), Node, EMessageSeverity::Warning);
+		}
+		ObjectNode->SetUid(FinalGuid.ToString());
+		
+		// States
+		int NumStates = TypedNodeObj->States.Num();
+		ObjectNode->SetStateCount(NumStates);
+
+		// In a partial compilation we will filter the states of the root object
+		bool bFilterStates = true;
+
+		if (GenerationContext.bPartialCompilation)
+		{
+			if (!TypedNodeObj->ParentObject)
+			{
+				bFilterStates = false;
+			}
+		}
+
+		for (int StateIndex = 0; StateIndex < NumStates && bFilterStates; ++StateIndex)
+		{
+			const FCustomizableObjectState& State = TypedNodeObj->States[StateIndex];
+			ObjectNode->SetStateName(StateIndex, State.Name);
+			for (int ParamIndex = 0; ParamIndex < State.RuntimeParameters.Num(); ++ParamIndex)
+			{
+				ObjectNode->AddStateParam(StateIndex, State.RuntimeParameters[ParamIndex]);
+			}
+
+			const ITargetPlatform* TargetPlatform = GenerationContext.Options.TargetPlatform;
+
+			int32 NumExtraLODsToBuildAfterFirstLOD = 0;
+
+			const int32* AuxNumExtraLODs = State.NumExtraLODsToBuildPerPlatform.Find(TargetPlatform->PlatformName());
+
+			if (AuxNumExtraLODs)
+			{
+				NumExtraLODsToBuildAfterFirstLOD = *AuxNumExtraLODs;
+			}
+
+			ObjectNode->SetStateProperties(StateIndex, State.TextureCompressionStrategy, State.bBuildOnlyFirstLOD, GenerationContext.FirstLODAvailable, NumExtraLODsToBuildAfterFirstLOD);
+
+			// UI Data
+			FMutableStateData StateUIData;
+			StateUIData.StateUIMetadata = State.UIMetadata;
+			StateUIData.bDisableTextureStreaming = State.bDisableTextureStreaming;
+			StateUIData.bLiveUpdateMode = State.bLiveUpdateMode;
+			StateUIData.bReuseInstanceTextures = State.bReuseInstanceTextures;
+			StateUIData.ForcedParameterValues = State.ForcedParameterValues;
+
+			GenerationContext.StateUIDataMap.Add(State.Name, StateUIData);
+		}
+		
+		// Process components.
+		//-------------------------------------------------------------------
+		TArray<mu::Ptr<mu::NodeComponent>> ExplicitComponents;
+
+		const UEdGraphPin* ComponentsPin = TypedNodeObj->ComponentsPin();
+		TArray<UEdGraphPin*> ConnectedComponentPins = FollowInputPinArray(*ComponentsPin);
+		for (const UEdGraphPin* ComponentNodePin : ConnectedComponentPins)
+		{		
+			mu::Ptr<mu::NodeComponent> ComponentNode = GenerateMutableSourceComponent(ComponentNodePin, GenerationContext);
+			ObjectNode->Components.Add(ComponentNode);
+		}
+
+		// Process modifiers.
+		//-------------------------------------------------------------------
+		if (const UEdGraphPin* ModifierPin = TypedNodeObj->ModifiersPin())
+		{
+			TArray<UEdGraphPin*> ConnectedModifierPins = FollowInputPinArray(*ModifierPin);
+			for (UEdGraphPin* const ChildNodePin : ConnectedModifierPins)
+			{
+				if (!Cast<UCustomizableObjectNodeModifierBase>(ChildNodePin->GetOwningNode()))
+				{
+					continue;
+				}
+
+				// Set it to "None" to indicate we don't care about component id.
+				FName OldCurrentMeshComponent = GenerationContext.CurrentMeshComponent;
+				GenerationContext.CurrentMeshComponent = FName();
+
+				mu::Ptr<mu::NodeModifier> ModifierNode = GenerateMutableSourceModifier(ChildNodePin, GenerationContext);
+				ObjectNode->Modifiers.Add(ModifierNode);
+
+				GenerationContext.CurrentMeshComponent = OldCurrentMeshComponent;
+			}
+		}
+		
+		// Generate inputs to Object node pins added by extensions
+		//-------------------------------------------------------------------
+		for (const FRegisteredObjectNodeInputPin& ExtensionInputPin : ICustomizableObjectModule::Get().GetAdditionalObjectNodePins())
+		{
+			const UEdGraphPin* GraphPin = TypedNodeObj->FindPin(ExtensionInputPin.GlobalPinName, EGPD_Input);
+			if (!GraphPin)
+			{
+				continue;
+			}
+
+			TArray<UEdGraphPin*> ConnectedPins = FollowInputPinArray(*GraphPin);
+
+			// If the pin isn't supposed to take more than one connection, ignore all but the first
+			// incoming connection.
+			if (!ExtensionInputPin.InputPin.bIsArray && ConnectedPins.Num() > 1)
+			{
+				FString Msg = FString::Printf(TEXT("Extension input %s has multiple incoming connections but is only expecting one connection."),
+					*ExtensionInputPin.InputPin.DisplayName.ToString());
+
+				GenerationContext.Compiler->CompilerLog(FText::FromString(Msg), Node, EMessageSeverity::Warning);
+			}
+
+			for (const UEdGraphPin* ConnectedPin : ConnectedPins)
+			{
+				const UEdGraphNode* ConnectedNode = ConnectedPin->GetOwningNode();
+				
+				if (const ICustomizableObjectExtensionNode* ExtensionNode = Cast<ICustomizableObjectExtensionNode>(ConnectedNode))
+				{
+					if (mu::Ptr<mu::NodeExtensionData> GeneratedNode = ExtensionNode->GenerateMutableNode(GenerationContext.ExtensionDataCompilerInterface))
+					{
+						ObjectNode->AddExtensionDataNode(GeneratedNode, TCHAR_TO_ANSI(*ExtensionInputPin.GlobalPinName.ToString()));
+					}
+				}
+			}
+		}
+
+		// Children
+		//-------------------------------------------------------------------
+		TArray<UEdGraphPin*> ConnectedChildrenPins = FollowInputPinArray(*TypedNodeObj->ChildrenPin());
+		ObjectNode->Children.Reserve(ConnectedChildrenPins.Num());
+		for (int32 ChildIndex = 0; ChildIndex < ConnectedChildrenPins.Num(); ++ChildIndex)
+		{
+			mu::Ptr<mu::NodeObject> ChildNode = GenerateMutableSource(ConnectedChildrenPins[ChildIndex], GenerationContext);
+			ObjectNode->Children.Add(ChildNode);
+		}
+	}
+
+	else if (const UCustomizableObjectNodeObjectGroup* TypedNodeGroup = Cast<UCustomizableObjectNodeObjectGroup>(Node))
+	{
+		mu::NodeObjectGroupPtr GroupNode = new mu::NodeObjectGroup();
+		Result = GroupNode;
+
+		// All sockets from all mesh parts plugged into this group node will have the following priority when there's a socket name clash
+		GenerationContext.SocketPriorityStack.Push(TypedNodeGroup->SocketPriority);
+
+		GenerationContext.AddParameterNameUnique(TypedNodeGroup, TypedNodeGroup->GroupName);
+		GroupNode->SetName(TypedNodeGroup->GroupName);
+		GroupNode->SetUid(TypedNodeGroup->NodeGuid.ToString());
+
+		// Get all group projectors and put them in the generation context so that they are available to the child material nodes of this group node
+		uint32 NumProjectorCountBeforeNode = GenerationContext.ProjectorGroupMap.Num();
+		UEdGraphPin* ProjectorsPin = TypedNodeGroup->GroupProjectorsPin();
+		
+		if (ProjectorsPin && TypedNodeGroup)
+		{
+			for (const UEdGraphPin* GroupProjectorNodePin : FollowInputPinArray(*ProjectorsPin))
+			{
+				GenerateMutableSourceGroupProjector(GroupProjectorNodePin, GenerationContext, TypedNodeGroup);
+			}
+		}
+
+		mu::NodeObjectGroup::CHILD_SELECTION Type = mu::NodeObjectGroup::CS_ALWAYS_ALL;
+		switch (TypedNodeGroup->GroupType)
+		{
+		case ECustomizableObjectGroupType::COGT_ALL: Type = mu::NodeObjectGroup::CS_ALWAYS_ALL; break;
+		case ECustomizableObjectGroupType::COGT_TOGGLE: Type = mu::NodeObjectGroup::CS_TOGGLE_EACH; break;
+		case ECustomizableObjectGroupType::COGT_ONE: Type = mu::NodeObjectGroup::CS_ALWAYS_ONE; break;
+		case ECustomizableObjectGroupType::COGT_ONE_OR_NONE: Type = mu::NodeObjectGroup::CS_ONE_OR_NONE; break;
+		default:
+			GenerationContext.Compiler->CompilerLog(LOCTEXT("UnsupportedGroupType", "Object Group Type not supported. Setting to 'ALL'."), Node);
+			break;
+		}
+		GroupNode->SetSelectionType(Type);
+
+		// External children
+		TArray<UCustomizableObjectNodeObject*> ExternalChildNodes;
+		GenerationContext.GroupIdToExternalNodeMap.MultiFind(TypedNodeGroup->NodeGuid, ExternalChildNodes);
+		GenerationContext.GuidToParamNameMap.Add(TypedNodeGroup->NodeGuid, TypedNodeGroup->GroupName);
+
+		// Children
+		const TArray<UEdGraphPin*> ConnectedChildrenPins = FollowInputPinArray(*TypedNodeGroup->ObjectsPin());
+		const int32 NumChildren = ConnectedChildrenPins.Num();
+		const int32 TotalNumChildren = NumChildren + ExternalChildNodes.Num();
+
+		GroupNode->SetChildCount(TotalNumChildren);
+		GroupNode->SetDefaultValue(Type == mu::NodeObjectGroup::CS_ONE_OR_NONE ? -1 : 0);
+		int32 ChildIndex = 0;
+
+		// UI data
+		FMutableParameterData ParameterUIData(
+			TypedNodeGroup->ParamUIMetadata,
+			EMutableParameterType::Int);
+
+		ParameterUIData.IntegerParameterGroupType = TypedNodeGroup->GroupType;
+
+		// In the case of a partial compilation, make sure at least one child is connected so that the param is no optimized
+		bool bAtLeastOneConnected = false;
+
+		for (; ChildIndex < NumChildren; ++ChildIndex)
+		{
+			bool bLastChildNode = (ChildIndex == NumChildren - 1) && (ExternalChildNodes.Num() == 0);
+			bool bConnectAtLeastTheLastChild = bLastChildNode && !bAtLeastOneConnected;
+
+			UCustomizableObjectNodeObject* CustomizableObjectNodeObject = Cast<UCustomizableObjectNodeObject>(ConnectedChildrenPins[ChildIndex]->GetOwningNode());
+
+			FString* SelectedOptionName = GenerationContext.ParamNamesToSelectedOptions.Find(TypedNodeGroup->GroupName); // If the param is in the map restrict to only the selected option
+			mu::NodeObjectPtr ChildNode = nullptr;
+
+			if (bConnectAtLeastTheLastChild || !SelectedOptionName || (CustomizableObjectNodeObject && *SelectedOptionName == CustomizableObjectNodeObject->ObjectName) )
+			{
+				bAtLeastOneConnected = true;
+
+				ChildNode = GenerateMutableSource(ConnectedChildrenPins[ChildIndex], GenerationContext);
+				GroupNode->SetChild(ChildIndex, ChildNode.get());
+
+				if (CustomizableObjectNodeObject)
+				{
+					FString LeftSplit = CustomizableObjectNodeObject->GetPathName();
+					LeftSplit.Split(".", &LeftSplit, nullptr);
+					GenerationContext.CustomizableObjectPathMap.Add(CustomizableObjectNodeObject->Identifier.ToString(), LeftSplit);
+					GenerationContext.GroupNodeMap.Add(CustomizableObjectNodeObject->Identifier.ToString(), FCustomizableObjectIdPair(TypedNodeGroup->GroupName, ChildNode->GetName()));
+					ParameterUIData.ArrayIntegerParameterOption.Add(
+						CustomizableObjectNodeObject->ObjectName,
+				FIntegerParameterUIData(CustomizableObjectNodeObject->ParamUIMetadata));
+
+					if (TypedNodeGroup->GroupType == ECustomizableObjectGroupType::COGT_TOGGLE)
+					{
+						GenerationContext.AddParameterNameUnique(CustomizableObjectNodeObject, CustomizableObjectNodeObject->ObjectName);
+						
+						// UI Data is only relevant when the group node is set to Toggle
+						GenerationContext.ParameterUIDataMap.Add(CustomizableObjectNodeObject->ObjectName, FMutableParameterData(
+							CustomizableObjectNodeObject->ParamUIMetadata,
+							EMutableParameterType::Int));
+					}
+				}
+			}
+			else
+			{
+				ChildNode = new mu::NodeObjectNew;
+				ChildNode->SetName(CustomizableObjectNodeObject->ObjectName);
+				GroupNode->SetChild(ChildIndex, ChildNode.get());
+			}
+
+			if ((TypedNodeGroup->GroupType == ECustomizableObjectGroupType::COGT_ONE ||
+				TypedNodeGroup->GroupType == ECustomizableObjectGroupType::COGT_ONE_OR_NONE)
+				&& TypedNodeGroup->DefaultValue == ChildNode->GetName())
+			{
+				GroupNode->SetDefaultValue(ChildIndex);
+			}
+		}
+
+		const bool bCollapseUnderParent = TypedNodeGroup->ParamUIMetadata.ExtraInformation.Find(FString("CollapseUnderParent")) != nullptr;
+		constexpr bool bHideWhenNotSelected = true; //TypedNodeGroup->ParamUIMetadata.ExtraInformation.Find(FString("HideWhenNotSelected"));
+		
+		if (bCollapseUnderParent || bHideWhenNotSelected)
+		{
+			if (const UEdGraphPin* ConnectedPin = FollowOutputPin(*Pin))
+			{
+				UCustomizableObjectNodeObject* NodeObject = Cast<UCustomizableObjectNodeObject>(ConnectedPin->GetOwningNode());
+
+				if (NodeObject)
+				{
+					const FGuid* ParentId = GenerationContext.GroupIdToExternalNodeMap.FindKey(NodeObject);
+
+					// Group objects in the same graph aren't in the GroupIdToExternalNodeMap, so follow the pins instead
+					if (!ParentId && NodeObject->OutputPin())
+					{
+						if (const UEdGraphPin* ConnectedPinToObject = FollowOutputPin(*NodeObject->OutputPin()))
+						{
+							if (UCustomizableObjectNodeObjectGroup* ParentGroupNode = Cast<UCustomizableObjectNodeObjectGroup>(ConnectedPinToObject->GetOwningNode()))
+							{
+								ParentId = &ParentGroupNode->NodeGuid;
+							}
+						}
+					}
+
+					if (ParentId)
+					{
+						FString* ParentParamName = GenerationContext.GuidToParamNameMap.Find(*ParentId);
+
+						if (ParentParamName)
+						{
+							ParameterUIData.ParamUIMetadata.ExtraInformation.Add(FString("__ParentParamName"), *ParentParamName);
+
+							if (bHideWhenNotSelected)
+							{
+								ParameterUIData.ParamUIMetadata.ExtraInformation.Add(FString("__DisplayWhenParentValueEquals"), NodeObject->ObjectName);
+							}
+
+							if (bCollapseUnderParent)
+							{
+								ParameterUIData.ParamUIMetadata.ExtraInformation.Add(FString("CollapseUnderParent"));
+
+								FMutableParameterData ParentParameterUIData;
+								ParentParameterUIData.ParamUIMetadata.ExtraInformation.Add(FString("__HasCollapsibleChildren"));
+								GenerationContext.ParameterUIDataMap.Add(*ParentParamName, ParentParameterUIData);
+							}
+						}
+					}
+				}
+			}
+		}
+
+		// Build external objects that reference this object as parent
+		const int32 NumExternalChildren = FMath::Max(0, TotalNumChildren - NumChildren);
+		for (int32 ExternalChildIndex = 0; ExternalChildIndex < NumExternalChildren; ++ExternalChildIndex)
+		{
+			const UCustomizableObjectNodeObject* ExternalChildNode = ExternalChildNodes[ExternalChildIndex];
+			bool bLastExternalChildNode = ExternalChildIndex == ExternalChildNodes.Num() - 1;
+			bool bConnectAtLeastTheLastChild = bLastExternalChildNode && !bAtLeastOneConnected;
+
+			UCustomizableObjectNodeObject* CustomizableObjectNodeObject = Cast<UCustomizableObjectNodeObject>(ExternalChildNode->OutputPin()->GetOwningNode());
+
+			FString* SelectedOptionName = GenerationContext.ParamNamesToSelectedOptions.Find(TypedNodeGroup->GroupName); // If the param is in the map restrict to only the selected option
+			mu::NodeObjectPtr ChildNode = nullptr;
+
+			if (bConnectAtLeastTheLastChild || !SelectedOptionName || (CustomizableObjectNodeObject  && *SelectedOptionName == CustomizableObjectNodeObject->ObjectName) )
+			{
+				bAtLeastOneConnected = true;
+
+				ChildNode = GenerateMutableSource(ExternalChildNode->OutputPin(), GenerationContext);
+				GroupNode->SetChild(ChildIndex, ChildNode.get());
+
+				if (CustomizableObjectNodeObject)
+				{
+					FString LeftSplit = ExternalChildNode->GetPathName();
+					LeftSplit.Split(".", &LeftSplit, nullptr);
+					GenerationContext.CustomizableObjectPathMap.Add(CustomizableObjectNodeObject->Identifier.ToString(), LeftSplit);
+					GenerationContext.GroupNodeMap.Add(CustomizableObjectNodeObject->Identifier.ToString(), FCustomizableObjectIdPair(TypedNodeGroup->GroupName, ChildNode->GetName()));
+					ParameterUIData.ArrayIntegerParameterOption.Add(
+						CustomizableObjectNodeObject->ObjectName,
+						FIntegerParameterUIData(CustomizableObjectNodeObject->ParamUIMetadata));
+
+					if (CustomizableObjectNodeObject->ObjectName.IsEmpty())
+					{
+						GenerationContext.NoNameNodeObjectArray.AddUnique(CustomizableObjectNodeObject);
+					}
+
+					if (TypedNodeGroup->GroupType == ECustomizableObjectGroupType::COGT_TOGGLE)
+					{
+						GenerationContext.AddParameterNameUnique(CustomizableObjectNodeObject, CustomizableObjectNodeObject->ObjectName);
+
+						// UI Data is only relevant when the group node is set to Toggle
+						GenerationContext.ParameterUIDataMap.Add(CustomizableObjectNodeObject->ObjectName, FMutableParameterData(
+							CustomizableObjectNodeObject->ParamUIMetadata,
+							EMutableParameterType::Int));
+					}
+				}
+			}
+			else
+			{
+				ChildNode = new mu::NodeObjectNew;
+				ChildNode->SetName(CustomizableObjectNodeObject->ObjectName);
+				GroupNode->SetChild(ChildIndex, ChildNode.get());
+			}
+
+			if ((TypedNodeGroup->GroupType == ECustomizableObjectGroupType::COGT_ONE ||
+				TypedNodeGroup->GroupType == ECustomizableObjectGroupType::COGT_ONE_OR_NONE)
+				&& TypedNodeGroup->DefaultValue == ChildNode->GetName())
+			{
+				GroupNode->SetDefaultValue(ChildIndex);
+			}
+
+			ChildIndex++;
+		}
+
+		const FMutableParameterData* ChildFilledUIData = GenerationContext.ParameterUIDataMap.Find(TypedNodeGroup->GroupName);
+		if (ChildFilledUIData && ChildFilledUIData->ParamUIMetadata.ExtraInformation.Find(FString("__HasCollapsibleChildren")))
+		{
+			// Some child param filled the HasCollapsibleChildren UI info, refill it so it's not lost
+			ParameterUIData.ParamUIMetadata.ExtraInformation.Add(FString("__HasCollapsibleChildren"));
+		}
+
+		if (TypedNodeGroup->GroupType == ECustomizableObjectGroupType::COGT_TOGGLE)
+		{
+			for (const TTuple<FString, FIntegerParameterUIData>& BooleanParam : ParameterUIData.ArrayIntegerParameterOption)
+			{
+				FMutableParameterData ParameterUIDataBoolean(
+					BooleanParam.Value.ParamUIMetadata,
+					EMutableParameterType::Bool);
+
+				ParameterUIDataBoolean.ParamUIMetadata.ExtraInformation = ParameterUIData.ParamUIMetadata.ExtraInformation;
+
+				GenerationContext.ParameterUIDataMap.Add(BooleanParam.Key, ParameterUIDataBoolean);
+			}
+		}
+		else
+		{
+			GenerationContext.ParameterUIDataMap.Add(TypedNodeGroup->GroupName, ParameterUIData);
+		}
+
+		// Remove the projectors from this node
+		GenerationContext.ProjectorGroupMap.Remove(TypedNodeGroup);
+
+		// Go back to the parent group node's socket priority if it exists
+		ensure(GenerationContext.SocketPriorityStack.Num() > 0);
+		GenerationContext.SocketPriorityStack.Pop();
+
+		check(NumProjectorCountBeforeNode == GenerationContext.ProjectorGroupMap.Num());
+	}
+	else
+	{
+		GenerationContext.Compiler->CompilerLog(LOCTEXT("UnimplementedNode", "Node type not implemented yet."), Node);
+	}
+
+	GenerationContext.Generated.Add(Key, FGeneratedData(Node, Result));
+	GenerationContext.GeneratedNodes.Add(Node);
+
+	return Result;
+}
+
+
+int32 AddTagToMutableMeshUnique(mu::Mesh& MutableMesh, const FString& Tag)
+{
+	const int32 TagCount = MutableMesh.GetTagCount();
+
+	for (int32 TagIndex = TagCount - 1; TagIndex >= 0; --TagIndex)
+	{
+		if (FString(MutableMesh.GetTag(TagIndex)) == Tag)
+		{
+			return TagIndex;
+		}
+	}
+
+	MutableMesh.SetTagCount(TagCount + 1);
+	MutableMesh.SetTag(TagCount, Tag);
+
+	return TagCount;
+}
+
+FString GenerateAnimationInstanceTag(const int32 AnimBpIndex, const FName& SlotIndex)
+{
+	return FString("__AnimBP:") + FString::Printf(TEXT("%s_Slot_"), *SlotIndex.ToString()) + FString::FromInt(AnimBpIndex);
+}
+
+
+FString GenerateGameplayTag(const FString& GameplayTag)
+{
+	return FString("__AnimBPTag:") + GameplayTag;
+}
+
+
+void PopulateReferenceSkeletalMeshesData(FMutableGraphGenerationContext& GenerationContext)
+{
+	const FString PlatformName = GenerationContext.Options.TargetPlatform->IniPlatformName();
+	
+	const uint32 LODCount = GenerationContext.NumLODsInRoot;
+	const uint32 ComponentCount = GenerationContext.NumComponents;
+
+	GenerationContext.ReferenceSkeletalMeshesData.AddDefaulted(ComponentCount);
+	for(uint32 ComponentIndex = 0; ComponentIndex < ComponentCount; ++ComponentIndex)
+	{
+		USkeletalMesh* RefSkeletalMesh = GenerationContext.ComponentInfos[ComponentIndex].RefSkeletalMesh;
+		check(RefSkeletalMesh);
+
+		FMutableRefSkeletalMeshData& Data = GenerationContext.ReferenceSkeletalMeshesData[ComponentIndex];
+
+		// Set the RefSkeletalMesh
+		Data.SkeletalMesh = TObjectPtr<USkeletalMesh>(RefSkeletalMesh);
+		Data.SoftSkeletalMesh = RefSkeletalMesh;
+
+		// Set the optional USkeletalMeshLODSettings that will be applied to the generated transient meshes or the baked meshes
+		Data.SkeletalMeshLODSettings = RefSkeletalMesh->GetLODSettings();
+
+		// Gather LODData, this may include per LOD settings such as render data config or LODDataInfoArray
+		Data.LODData.AddDefaulted(GenerationContext.NumLODsInRoot);
+		
+		const uint32 RefSkeletalMeshLODCount = RefSkeletalMesh->GetLODNum();
+
+		for(uint32 LODIndex = 0; LODIndex < LODCount; ++LODIndex)
+		{
+			FMutableRefLODData& LODData = Data.LODData[LODIndex];
+			if(LODIndex < RefSkeletalMeshLODCount)
+			{
+				// Copy LOD info data from the reference skeletal mesh
+				const FSkeletalMeshLODInfo& LODInfo = *RefSkeletalMesh->GetLODInfo(LODIndex);
+				LODData.LODInfo.ScreenSize = LODInfo.ScreenSize.GetValueForPlatform(*PlatformName);
+				LODData.LODInfo.LODHysteresis = LODInfo.LODHysteresis;
+				LODData.LODInfo.bSupportUniformlyDistributedSampling = LODInfo.bSupportUniformlyDistributedSampling;
+				LODData.LODInfo.bAllowCPUAccess = LODInfo.bAllowCPUAccess;
+
+				// Copy Render data settings from the reference skeletal mesh
+				const FSkeletalMeshLODRenderData& ReferenceLODModel = RefSkeletalMesh->GetResourceForRendering()->LODRenderData[LODIndex];
+				LODData.RenderData.bIsLODOptional = ReferenceLODModel.bIsLODOptional;
+				LODData.RenderData.bStreamedDataInlined = ReferenceLODModel.bStreamedDataInlined;
+			}
+			else
+			{
+				LODData.LODInfo.ScreenSize = 0.3f / (LODIndex + 1);
+				LODData.LODInfo.LODHysteresis = 0.02f;				
+			}
+		}
+
+		// Gather SkeletalMesh Sockets;
+		const TArray<USkeletalMeshSocket*>& RefSkeletonSockets = RefSkeletalMesh->GetMeshOnlySocketList();
+		const uint32 SocketCount = RefSkeletonSockets.Num();
+
+		Data.Sockets.AddDefaulted(SocketCount);
+		for(uint32 SocketIndex = 0 ; SocketIndex < SocketCount; ++SocketIndex)
+		{
+			const USkeletalMeshSocket* RefSocket = RefSkeletonSockets[SocketIndex];
+			check(RefSocket);
+
+			FMutableRefSocket& Socket = Data.Sockets[SocketIndex];
+			Socket.SocketName = RefSocket->SocketName;
+			Socket.BoneName = RefSocket->BoneName;
+			Socket.RelativeLocation = RefSocket->RelativeLocation;
+			Socket.RelativeRotation = RefSocket->RelativeRotation;
+			Socket.RelativeScale = RefSocket->RelativeScale;
+			Socket.bForceAlwaysAnimated = RefSocket->bForceAlwaysAnimated;
+		}
+
+		// TODO: Generate Bounds?
+		// Gather Bounds
+		Data.Bounds = RefSkeletalMesh->GetBounds();
+		
+		// Additional Settings
+		Data.Settings.bEnablePerPolyCollision = RefSkeletalMesh->GetEnablePerPolyCollision();
+
+		const TArray<FSkeletalMaterial>& Materials = RefSkeletalMesh->GetMaterials();
+		for (const FSkeletalMaterial& Material : Materials)
+		{
+			if (Material.UVChannelData.bInitialized)
+			{
+				for (int32 UVIndex = 0; UVIndex < TEXSTREAM_MAX_NUM_UVCHANNELS; ++UVIndex)
+				{
+					Data.Settings.DefaultUVChannelDensity = FMath::Max(Data.Settings.DefaultUVChannelDensity, Material.UVChannelData.LocalUVDensities[UVIndex]);
+				}
+			}
+		}
+
+		// Skeleton
+		if(USkeleton* Skeleton = RefSkeletalMesh->GetSkeleton())
+		{
+			Data.Skeleton = Skeleton;
+		}
+		
+		// Physics Asset
+		if (UPhysicsAsset* PhysicsAsset = RefSkeletalMesh->GetPhysicsAsset())
+		{
+			GenerationContext.AddParticipatingObject(*PhysicsAsset);
+			Data.PhysicsAsset = PhysicsAsset;
+		}
+
+		// Post ProcessAnimInstance
+		if(const TSubclassOf<UAnimInstance> PostProcessAnimInstance = RefSkeletalMesh->GetPostProcessAnimBlueprint())
+		{
+			GenerationContext.AddParticipatingObject(*PostProcessAnimInstance.Get());
+			Data.PostProcessAnimInst = PostProcessAnimInstance;
+		}
+		
+		// Shadow Physics Asset
+		if (UPhysicsAsset* PhysicsAsset = RefSkeletalMesh->GetShadowPhysicsAsset())
+		{
+			GenerationContext.AddParticipatingObject(*PhysicsAsset);
+			Data.ShadowPhysicsAsset = PhysicsAsset;
+		}
+
+		// Asset User Data
+		if (const TArray<UAssetUserData*>* AssetUserDataArray = RefSkeletalMesh->GetAssetUserDataArray())
+		{
+			for (UAssetUserData* AssetUserData : *AssetUserDataArray)
+			{
+				if (AssetUserData)
+				{
+					FMutableRefAssetUserData& MutAssetUserData = Data.AssetUserData.AddDefaulted_GetRef();
+					MutAssetUserData.AssetUserDataIndex = GenerationContext.AddAssetUserDataToStreamedResources(AssetUserData);
+					MutAssetUserData.AssetUserData = GenerationContext.StreamedResourceData[MutAssetUserData.AssetUserDataIndex].GetPath().Get();
+					check(MutAssetUserData.AssetUserData);
+					check(MutAssetUserData.AssetUserData->Data.Type == ECOResourceDataType::AssetUserData);
+				}
+			}
+		}
+	}
+}
+
+
+uint32 GetBaseTextureSize(const FMutableGraphGenerationContext& GenerationContext, const UCustomizableObjectNodeMaterialBase* Material, uint32 ImageIndex)
+{
+	const FGeneratedImageProperties* ImageProperties = GenerationContext.ImageProperties.Find({ Material, ImageIndex });
+	return ImageProperties ? ImageProperties->TextureSize : 0;
+}
+
+
+// Find the LODBias to apply to stay within the MaxTextureSize limit of the TargetPlatform
+int32 GetPlatformLODBias(int32 TextureSize, int32 NumMips, int32 MaxPlatformSize)
+{
+	if (MaxPlatformSize > 0 && MaxPlatformSize < TextureSize)
+	{
+		const int32 MaxMipsAllowed = FMath::CeilLogTwo(MaxPlatformSize) + 1;
+		return NumMips - MaxMipsAllowed;
+	}
+
+	return 0;
+}
+
+
+uint32 ComputeLODBiasForTexture(const FMutableGraphGenerationContext& GenerationContext, const UTexture2D& Texture, const UTexture2D* ReferenceTexture, int32 BaseTextureSize)
+{
+	constexpr int32 MaxAllowedLODBias = 6;
+
+	// Force a large LODBias for debug
+	if (GenerationContext.Options.bForceLargeLODBias)
+	{
+		return FMath::Min(GenerationContext.Options.DebugBias, MaxAllowedLODBias);
+	}
+
+	// Max size and number of mips from Texture. 
+	const int32 SourceSize = (int32)FMath::Max3(Texture.Source.GetSizeX(),Texture.Source.GetSizeY(),(int64)1);
+	const int32 NumMipsSource = FMath::CeilLogTwo(SourceSize) + 1;
+
+	// When the BaseTextureSize is known, skip mips until the texture is equal or smaller.
+	if (BaseTextureSize > 0)
+	{
+		if (BaseTextureSize < SourceSize)
+		{
+			const int32 MaxNumMipsInGame = FMath::CeilLogTwo(BaseTextureSize) + 1;
+			return FMath::Max(NumMipsSource - MaxNumMipsInGame, 0);
+		}
+
+		return 0;
+	}
+
+	const UTextureLODSettings& LODSettings = GenerationContext.Options.TargetPlatform->GetTextureLODSettings();
+
+	// Get the MaxTextureSize for the TargetPlatform.
+	const int32 MaxTextureSize = GetMaxTextureSize(ReferenceTexture ? *ReferenceTexture : Texture, LODSettings);
+
+	if (ReferenceTexture)
+	{
+		// Max size and number of mips from ReferenceTexture. 
+		const int32 MaxRefSourceSize = (uint32)FMath::Max3(ReferenceTexture->Source.GetSizeX(), ReferenceTexture->Source.GetSizeY(), (int64)1);
+		const int32 NumMipsRefSource = FMath::CeilLogTwo(MaxRefSourceSize) + 1;
+
+		// Find the LODBias to apply to stay within the MaxTextureSize limit of the TargetPlatform
+		const int32 PlatformLODBias = GetPlatformLODBias(MaxRefSourceSize, NumMipsRefSource, MaxTextureSize);
+
+		// TextureSize in-game without any additional LOD bias.
+		const int64 ReferenceTextureSize = MaxRefSourceSize >> PlatformLODBias;
+
+		// Additional LODBias of the Texture
+		const int32 ReferenceTextureLODBias = LODSettings.CalculateLODBias(ReferenceTextureSize, ReferenceTextureSize, 0,	ReferenceTexture->LODGroup,
+			ReferenceTexture->LODBias, 0, ReferenceTexture->MipGenSettings, ReferenceTexture->IsCurrentlyVirtualTextured());
+
+		return FMath::Max(NumMipsSource - NumMipsRefSource + PlatformLODBias + ReferenceTextureLODBias, 0);
+	}
+
+	// Find the LODBias to apply to stay within the MaxTextureSize limit of the TargetPlatform
+	const int32 PlatformLODBias = GetPlatformLODBias(SourceSize, NumMipsSource, MaxTextureSize);
+
+	// TextureSize in-game without any additional LOD bias.
+	const int64 TextureSize = SourceSize >> PlatformLODBias;
+
+	// Additional LODBias of the Texture
+	const int32 TextureLODBias = LODSettings.CalculateLODBias(TextureSize, TextureSize, 0, Texture.LODGroup, Texture.LODBias, 0, Texture.MipGenSettings, Texture.IsCurrentlyVirtualTextured());
+
+	return FMath::Max(PlatformLODBias + TextureLODBias, 0);
+}
+
+
+int32 GetMaxTextureSize(const UTexture2D& ReferenceTexture, const UTextureLODSettings& LODSettings)
+{
+	// Setting the maximum texture size
+	FTextureLODGroup TextureGroupSettings = LODSettings.GetTextureLODGroup(ReferenceTexture.LODGroup);
+
+	if (TextureGroupSettings.MaxLODSize > 0)
+	{
+		return ReferenceTexture.MaxTextureSize == 0 ? TextureGroupSettings.MaxLODSize : FMath::Min(TextureGroupSettings.MaxLODSize, ReferenceTexture.MaxTextureSize);
+	}
+
+	return ReferenceTexture.MaxTextureSize;
+}
+
+
+int32 GetTextureSizeInGame(const UTexture2D& Texture, const UTextureLODSettings& LODSettings, uint8 SurfaceLODBias)
+{
+	const int32 SourceSize = (uint32)FMath::Max3(Texture.Source.GetSizeX(), Texture.Source.GetSizeY(), (int64)1);
+	const int32 NumMipsSource = FMath::CeilLogTwo(SourceSize) + 1;
+
+	// Max size allowed on the TargetPlatform
+	const int32 MaxTextureSize = GetMaxTextureSize(Texture, LODSettings);
+
+	// Find the LODBias to apply to stay within the MaxTextureSize limit of the TargetPlatform
+	const int32 PlatformLODBias = GetPlatformLODBias(SourceSize, NumMipsSource, MaxTextureSize);
+	
+	// MaxTextureSize in-game without any additional LOD bias.
+	const int32 MaxTextureSizeAllowed = SourceSize >> PlatformLODBias;
+
+	// Calculate the LODBias specific for this texture 
+	const int32 TextureLODBias = LODSettings.CalculateLODBias(MaxTextureSizeAllowed, MaxTextureSizeAllowed, 0, Texture.LODGroup, Texture.LODBias, 0, Texture.MipGenSettings, Texture.IsCurrentlyVirtualTextured());
+
+	return MaxTextureSizeAllowed >> (TextureLODBias + SurfaceLODBias);
+}
+
+
+mu::FImageDesc GenerateImageDescriptor(UTexture* Texture)
+{
+	check(Texture);
+	mu::FImageDesc ImageDesc;
+
+	ImageDesc.m_size[0] = Texture->Source.GetSizeX();
+	ImageDesc.m_size[1] = Texture->Source.GetSizeY();
+	ImageDesc.m_lods = Texture->Source.GetNumMips();
+
+	mu::EImageFormat MutableFormat = mu::EImageFormat::IF_RGBA_UBYTE;
+	ETextureSourceFormat SourceFormat = Texture->Source.GetFormat();
+	switch (SourceFormat)
+	{
+	case ETextureSourceFormat::TSF_G8:
+	case ETextureSourceFormat::TSF_G16:
+	case ETextureSourceFormat::TSF_R16F:
+	case ETextureSourceFormat::TSF_R32F:
+		MutableFormat = mu::EImageFormat::IF_L_UBYTE;
+		break;
+
+	default:
+		break;
+	}
+
+	ImageDesc.m_format = MutableFormat;
+
+	return ImageDesc;
+}
+
+
+mu::Ptr<mu::Image> GenerateImageConstant(UTexture* Texture, FMutableGraphGenerationContext& GenerationContext, bool bIsReference)
+{
+	MUTABLE_CPUPROFILER_SCOPE(GenerateImageConstant);
+
+	if (!Texture)
+	{
+		return nullptr;
+	}
+
+	bool bForceLoad = false;
+	bool bIsCompileTime = false;
+	if (!bIsReference)
+	{
+		bForceLoad = true;
+		if (GenerationContext.Options.OptimizationLevel == 0)
+		{
+			bIsCompileTime = false;
+		}
+		else
+		{
+			bIsCompileTime = true;
+		}
+	}
+
+	// Create a descriptor for the image.
+	// \TODO: If passthrough (bIsReference) we should apply lod bias, and max texture size to this desc.
+	// For now it is not a problem because passthrough textures shouldn't mix with any other operations.
+	mu::FImageDesc ImageDesc = GenerateImageDescriptor(Texture);
+
+	FMutableGraphGenerationContext::FGeneratedReferencedTexture InvalidEntry;
+	InvalidEntry.ID = TNumericLimits<uint32>::Max();
+
+	FMutableGraphGenerationContext::FGeneratedReferencedTexture* Entry;
+	int32 Num = 0;
+	
+	if (bIsReference)
+	{
+		Entry = &GenerationContext.PassthroughTextureMap.FindOrAdd(Texture, InvalidEntry);
+		Num = GenerationContext.PassthroughTextureMap.Num();
+	}
+	else if (bIsCompileTime)
+	{
+        Entry = &GenerationContext.CompileTimeTextureMap.FindOrAdd(Texture, InvalidEntry);
+		Num = GenerationContext.CompileTimeTextureMap.Num();
+	}
+	else
+	{
+		Entry = &GenerationContext.RuntimeReferencedTextureMap.FindOrAdd(Texture, InvalidEntry);
+		Num = GenerationContext.RuntimeReferencedTextureMap.Num();
+	}
+
+	if (Entry->ID == TNumericLimits<uint32>::Max())
+	{
+		Entry->ID = Num - 1;
+	}
+	
+	// Compile-time references that are left should be resolved immediately (should only happen in editor).
+	mu::Ptr<mu::Image> Result = mu::Image::CreateAsReference(Entry->ID, ImageDesc, bForceLoad);
+	return Result;
+}
+
+
+mu::Ptr<mu::Mesh> GenerateMeshConstant(USkeletalMesh* Mesh, FMutableGraphGenerationContext& GenerationContext, bool bIsReference)
+{
+	if (!Mesh)
+	{
+		return nullptr;
+	}
+
+	bool bForceLoad = false;
+	bool bIsCompileTime = false;
+	if (!bIsReference)
+	{
+		bForceLoad = true;
+		if (GenerationContext.Options.OptimizationLevel == 0)
+		{
+			bIsCompileTime = false;
+		}
+		else
+		{
+			bIsCompileTime = true;
+		}
+	}
+
+
+	FMutableGraphGenerationContext::FGeneratedReferencedMesh InvalidEntry;
+	InvalidEntry.ID = TNumericLimits<uint32>::Max();
+
+	FMutableGraphGenerationContext::FGeneratedReferencedMesh* Entry = nullptr;
+	int32 Num = 0;
+
+	if (bIsReference)
+	{
+		Entry = &GenerationContext.PassthroughMeshMap.FindOrAdd(Mesh, InvalidEntry);
+		Num = GenerationContext.PassthroughMeshMap.Num();
+	}
+	// Not supported yet: this will be implemented when mesh conversion happens on-demand in the compilation process
+	// like it happens with images. For now we only implement support for pure pass-through.
+	//else if (bIsCompileTime)
+	//{
+	//	Entry = &GenerationContext.CompileTimeTextureMap.FindOrAdd(Texture, InvalidEntry);
+	//	Num = GenerationContext.CompileTimeTextureMap.Num();
+	//}
+	//else
+	//{
+	//	Entry = &GenerationContext.RuntimeReferencedTextureMap.FindOrAdd(Texture, InvalidEntry);
+	//	Num = GenerationContext.RuntimeReferencedTextureMap.Num();
+	//}
+	else
+	{
+		check(false);
+		return nullptr;
+	}
+
+	if (Entry->ID == TNumericLimits<uint32>::Max())
+	{
+		Entry->ID = Num - 1;
+	}
+
+	mu::Ptr<mu::Mesh> Result = mu::Mesh::CreateAsReference(Entry->ID, bForceLoad);
+	return Result;
+}
+
+
+void AddSocketTagsToMesh(const USkeletalMesh* SourceMesh, mu::Ptr<mu::Mesh> MutableMesh, FMutableGraphGenerationContext& GenerationContext)
+{
+	for (int32 SocketIndex = 0; SocketIndex < SourceMesh->NumSockets(); ++SocketIndex)
+	{
+		USkeletalMeshSocket* Socket = SourceMesh->GetSocketByIndex(SocketIndex);
+
+		FMutableRefSocket MutableSocket;
+		MutableSocket.SocketName = Socket->SocketName;
+		MutableSocket.BoneName = Socket->BoneName;
+		MutableSocket.RelativeLocation = Socket->RelativeLocation;
+		MutableSocket.RelativeRotation = Socket->RelativeRotation;
+		MutableSocket.RelativeScale = Socket->RelativeScale;
+		MutableSocket.bForceAlwaysAnimated = Socket->bForceAlwaysAnimated;
+
+		MutableSocket.Priority = GenerationContext.SocketPriorityStack.IsEmpty() ? 0 : GenerationContext.SocketPriorityStack.Top();
+			
+		int32 SocketArrayIndex = GenerationContext.SocketArray.AddUnique(MutableSocket);
+		FString SocketTag = FString::Printf(TEXT("__Socket:%d"), SocketArrayIndex);
+
+		AddTagToMutableMeshUnique(*MutableMesh, SocketTag);
+	}
+}
+
+
+FMutableComponentInfo::FMutableComponentInfo(FName InComponentName, USkeletalMesh* InRefSkeletalMesh)
+{
+	ComponentName = InComponentName;
+
+	if (!InRefSkeletalMesh || !InRefSkeletalMesh->GetSkeleton())
+	{
+		return;
+	}
+
+	RefSkeletalMesh = InRefSkeletalMesh;
+	RefSkeleton = RefSkeletalMesh->GetSkeleton();
+
+	const int32 NumBones = RefSkeleton->GetReferenceSkeleton().GetRawBoneNum();
+	BoneNamesToPathHash.Reserve(NumBones);
+
+	const TArray<FMeshBoneInfo>& Bones = RefSkeleton->GetReferenceSkeleton().GetRawRefBoneInfo();
+
+	for (int32 BoneIndex = 0; BoneIndex < NumBones; ++BoneIndex)
+	{
+		const FMeshBoneInfo& Bone = Bones[BoneIndex];
+
+		// Retrieve parent bone name and respective hash, root-bone is assumed to have a parent hash of 0
+		const FName ParentName = Bone.ParentIndex != INDEX_NONE ? Bones[Bone.ParentIndex].Name : NAME_None;
+		const uint32 ParentHash = Bone.ParentIndex != INDEX_NONE ? GetTypeHash(ParentName) : 0;
+
+		// Look-up the path-hash from root to the parent bone
+		const uint32* ParentPath = BoneNamesToPathHash.Find(ParentName);
+		const uint32 ParentPathHash = ParentPath ? *ParentPath : 0;
+
+		// Append parent hash to path to give full path hash to current bone
+		const uint32 BonePathHash = HashCombine(ParentPathHash, ParentHash);
+
+		// Add path hash to current bone
+		BoneNamesToPathHash.Add(Bone.Name, BonePathHash);
+	}
+}
+
+
+void FMutableComponentInfo::AccumulateBonesToRemovePerLOD(const TArray<FLODReductionSettings>& LODReductionSettings, int32 NumLODs)
+{
+	BonesToRemovePerLOD.SetNum(NumLODs);
+
+	TMap<FName, bool> BonesToRemove;
+
+	const int32 ComponentSettingsLODCount = LODReductionSettings.Num();
+	for (int32 LODIndex = 0; LODIndex < NumLODs; ++LODIndex)
+	{
+		if (LODIndex < ComponentSettingsLODCount)
+		{
+			const FLODReductionSettings& LODReductionSetting = LODReductionSettings[LODIndex];
+
+			for (const FBoneToRemove& Bone : LODReductionSetting.BonesToRemove)
+			{
+				if (bool* bOnlyRemoveChildren = BonesToRemove.Find(Bone.BoneName))
+				{
+					// Removed by a previous LOD
+					*bOnlyRemoveChildren = (*bOnlyRemoveChildren) && Bone.bOnlyRemoveChildren;
+				}
+				else
+				{
+					BonesToRemove.Add(Bone.BoneName, Bone.bOnlyRemoveChildren);
+				}
+			}
+		}
+		
+		BonesToRemovePerLOD[LODIndex] = BonesToRemove;
+	}
+}
+
+
+#undef LOCTEXT_NAMESPACE
+
+FMutableGraphGenerationContext::FSharedSurface::FSharedSurface(uint8 InLOD, const mu::Ptr<mu::NodeSurfaceNew>& InNodeSurfaceNew)
+{
+	LOD = InLOD;
+	NodeSurfaceNew = InNodeSurfaceNew;
+}
+
+
+bool FMutableGraphGenerationContext::FSharedSurface::operator==(const FSharedSurface& o) const
+{
+	return NodeModifierIDs == o.NodeModifierIDs;
+}
+
+
+bool FMutableGraphGenerationContext::FGeneratedDataTablesData::HasSameSettings(const UCustomizableObjectNodeTable* Node) const
+{
+	return Node->bDisableCheckedRows == bDisableCheckedRows && Node->VersionColumn == VersionColumn;
+}
