[mutable] Moved the Mutable plugin out of Experimental status into Beta.

diff --git a/Engine/Plugins/Mutable/Source/CustomizableObjectEditor/Private/MuCOE/GenerateMutableSource/GenerateMutableSource.h b/Engine/Plugins/Mutable/Source/CustomizableObjectEditor/Private/MuCOE/GenerateMutableSource/GenerateMutableSource.h
new file mode 100644
--- /dev/null
+++ b/Engine/Plugins/Mutable/Source/CustomizableObjectEditor/Private/MuCOE/GenerateMutableSource/GenerateMutableSource.h
@@ -0,0 +1,1036 @@
+// Copyright Epic Games, Inc. All Rights Reserved.
+
+#pragma once
+
+#include "Animation/Skeleton.h"
+#include "Animation/AnimInstance.h"
+#include "Engine/DataTable.h"
+#include "Engine/SkeletalMesh.h"
+#include "MuCO/CustomizableObject.h"
+#include "MuCO/CustomizableObjectIdentifier.h"
+#include "MuCO/CustomizableObjectCompilerTypes.h"
+#include "MuCOE/ExtensionDataCompilerInterface.h"
+#include "MuCOE/Nodes/CustomizableObjectNodeMaterialBase.h"
+#include "MuCOE/Nodes/CustomizableObjectNodeObject.h"
+#include "MuT/NodeComponentNew.h"
+#include "MuT/NodeImageConstant.h"
+#include "MuT/NodeMeshConstant.h"
+#include "MuT/NodeMeshApplyPose.h"
+#include "MuT/NodeModifierMeshClipWithMesh.h"
+#include "MuT/NodeObject.h"
+#include "MuT/NodeProjector.h"
+#include "MuT/NodeScalarEnumParameter.h"
+#include "MuT/NodeScalarParameter.h"
+#include "MuT/NodeSurfaceNew.h"
+#include "MuT/Table.h"
+#include "MuR/MutableMemory.h"
+#include "MuR/Skeleton.h"
+#include "MuR/Mesh.h"
+#include "MuR/Ptr.h"
+#include "UObject/Package.h"
+
+class UCustomizableObjectNodeMaterialBase;
+class UCustomizableObjectNodeModifierClipWithMesh;
+class UCustomizableObjectNodeTable;
+struct FCustomizableObjectClothingAssetData;
+
+class FCustomizableObjectCompiler;
+class UTextureLODSettings;
+class UAnimInstance;
+class UCompositeDataTable;
+class UCustomizableObjectNodeMaterial;
+class UCustomizableObjectNodeMeshMorph;
+class UCustomizableObjectNodeObjectGroup;
+class UEdGraphNode;
+class UMaterialInterface;
+class UObject;
+class UPhysicsAsset;
+class UTexture2D;
+struct FAnimBpOverridePhysicsAssetsInfo;
+struct FMutableGraphGenerationContext;
+struct FMutableParameterData;
+struct FMutableRefSkeletalMeshData;
+struct FMutableRefSocket;
+struct FMutableSkinWeightProfileInfo;
+struct FMorphTargetVertexData;
+enum class EPinMode;
+
+struct FGeneratedImageProperties
+{
+	/** Name in the Material. */
+	FString TextureParameterName;
+
+	/** Name in the mu::Surface. */
+	int32 ImagePropertiesIndex = INDEX_NONE;
+
+	TEnumAsByte<TextureCompressionSettings> CompressionSettings = TC_Default;
+
+	TEnumAsByte<TextureFilter> Filter = TF_Bilinear;
+
+	uint32 SRGB = 0;
+
+	uint32 bFlipGreenChannel = 0;
+
+	int32 LODBias = 0;
+
+	TEnumAsByte<TextureMipGenSettings> MipGenSettings = TMGS_SimpleAverage;
+
+	int32 MaxTextureSize = 0;
+
+	TEnumAsByte<TextureGroup> LODGroup = TEnumAsByte<TextureGroup>(TMGS_FromTextureGroup);
+
+	TEnumAsByte<TextureAddress> AddressX = TA_Clamp;
+	TEnumAsByte<TextureAddress> AddressY = TA_Clamp;
+
+	bool bIsPassThrough = false;
+
+	// ReferenceTexture source size.
+	int32 TextureSize = 0;
+};
+
+
+// Flags that can influence the mesh conversion
+enum class EMutableMeshConversionFlags : uint32
+{
+	// 
+	None = 0,
+	// Ignore the skeleton and skinning
+	IgnoreSkinning = 1 << 0,
+
+	// Ignore Physics assets
+	IgnorePhysics = 1 << 1,
+
+	// Prevent this mesh generation from adding per mesh metadata. 
+	DoNotCreateMeshMetadata = 1 << 2
+};
+
+ENUM_CLASS_FLAGS(EMutableMeshConversionFlags)
+
+struct FLayoutGenerationFlags
+{
+	bool operator==(const FLayoutGenerationFlags& Other) const = default;
+
+	// Texture pin mode per UV Channel
+	TArray<EPinMode> TexturePinModes;
+};
+
+/** 
+	Struct to store the necessary data to generate the morphs of a skeletal mesh 
+	This struct allows the stack morph nodes to use the same functions as the mesh morph nodes
+*/
+struct FMorphNodeData
+{
+	// Pointer to the node that owns this morph data
+	UCustomizableObjectNode* OwningNode;
+
+	// Name of the morph that will be applied
+	FString MorphTargetName;
+
+	// Pin to the node that generates the factor of the morph
+	const UEdGraphPin* FactorPin;
+
+	// Pin of the mesh where the morphs will ble apllied
+	const UEdGraphPin* MeshPin;
+
+	bool operator==(const FMorphNodeData& Other) const
+	{
+		return OwningNode == Other.OwningNode && MorphTargetName == Other.MorphTargetName
+			&& FactorPin == Other.FactorPin && MeshPin == Other.MeshPin;
+	}
+};
+
+
+// Key for the data stored for each processed unreal graph node.
+class FGeneratedKey
+{
+	friend uint32 GetTypeHash(const FGeneratedKey& Key);
+
+public:
+	FGeneratedKey(void* InFunctionAddress, const UEdGraphPin& InPin, const UCustomizableObjectNode& Node, FMutableGraphGenerationContext& GenerationContext, const bool UseMesh = false, bool bOnlyConnectedLOD = false);
+	
+	bool operator==(const FGeneratedKey& Other) const;
+
+private:
+	/** Used to differentiate pins being cached from different functions (e.g. a PC_Color pin cached from GenerateMutableSourceImage and GenerateMutableSourceColor). */
+	void* FunctionAddress;
+	
+	const UEdGraphPin* Pin;
+	
+	int32 LOD;
+
+	/** Flag used to generate this mesh. Bit mask of EMutableMeshConversionFlags */
+	EMutableMeshConversionFlags Flags = EMutableMeshConversionFlags::None;
+
+	/** Active morphs at the time of mesh generation. */
+	TArray<FMorphNodeData> MeshMorphStack;
+
+	/** UV Layout modes */
+	FLayoutGenerationFlags LayoutFlags;
+
+	/** When caching a generated mesh, true if we force to generate the connected LOD when using Automatic LODs From Mesh. */
+	bool bOnlyConnectedLOD = false;
+};
+
+
+uint32 GetTypeHash(const FGeneratedKey& Key);
+
+
+struct FGeneratedImageKey
+{
+	FGeneratedImageKey(const UEdGraphPin* InPin)
+	{
+		Pin = InPin;
+	}
+
+	bool operator==(const FGeneratedImageKey& Other) const
+	{
+		return Pin == Other.Pin;
+	}
+
+	const UEdGraphPin* Pin;
+};
+
+
+struct FGeneratedImagePropertiesKey
+{
+	FGeneratedImagePropertiesKey(const UCustomizableObjectNodeMaterialBase* InMaterial, uint32 InImageIndex)
+	{
+		MaterialReferenceId = (PTRINT)InMaterial;
+		ImageIndex = InImageIndex;
+	}
+
+	bool operator==(const FGeneratedImagePropertiesKey& Other) const
+	{
+		return MaterialReferenceId == Other.MaterialReferenceId && ImageIndex == Other.ImageIndex;
+	}
+
+
+	PTRINT MaterialReferenceId = 0;
+	uint32 ImageIndex = 0;
+};
+
+// Structure storing results to propagate up when generating mutable mesh node expressions.
+struct FMutableGraphMeshGenerationData
+{
+	// Did we find any mesh with vertex colours in the expression?
+	bool bHasVertexColors = false;
+	bool bHasRealTimeMorphs = false;
+	bool bHasClothing = false;
+
+	// Maximum number of texture channels found in the expression.
+	int NumTexCoordChannels = 0;
+
+	// Maximum number of bones per vertex found in the expression.
+	int MaxNumBonesPerVertex = 0;
+
+	// Maximum size of the vertex bone index type in the expression.
+	int MaxBoneIndexTypeSizeBytes = 0;
+
+	int32 MaxNumTriangles = 0;
+	int32 MinNumTriangles = TNumericLimits<int32>::Max();
+
+	TArray<int32> SkinWeightProfilesSemanticIndices;
+
+	// Combine another generated data looking for the most general case.
+	void Combine(const FMutableGraphMeshGenerationData& Other)
+	{
+		bHasVertexColors = bHasVertexColors || Other.bHasVertexColors;
+		bHasRealTimeMorphs = bHasRealTimeMorphs || Other.bHasRealTimeMorphs;
+		bHasClothing = bHasClothing || Other.bHasClothing;
+		NumTexCoordChannels = FMath::Max(Other.NumTexCoordChannels, NumTexCoordChannels);
+		MaxNumBonesPerVertex = FMath::Max(Other.MaxNumBonesPerVertex, MaxNumBonesPerVertex);
+		MaxBoneIndexTypeSizeBytes = FMath::Max(Other.MaxBoneIndexTypeSizeBytes, MaxBoneIndexTypeSizeBytes);
+		MaxNumTriangles = FMath::Max(Other.MaxNumTriangles, MaxNumTriangles);
+		MinNumTriangles = FMath::Min(Other.MinNumTriangles, MinNumTriangles);
+
+		for (int32 SemanticIndex : Other.SkinWeightProfilesSemanticIndices)
+		{
+			SkinWeightProfilesSemanticIndices.AddUnique(SemanticIndex);
+		}
+	}
+};
+
+
+// Data stored for each processed unreal graph node, stored in the cache.
+struct FGeneratedData
+{
+	FGeneratedData(const UEdGraphNode* InSource, mu::NodePtr InNode,
+		const FMutableGraphMeshGenerationData* InMeshData = nullptr)
+	{
+		Source = InSource;
+		Node = InNode;
+		if (InMeshData)
+		{
+			meshData = *InMeshData;
+		}
+	}
+
+	const UEdGraphNode* Source;
+	mu::NodePtr Node;
+
+	// Used for mesh nodes only
+	FMutableGraphMeshGenerationData meshData;
+};
+
+
+inline uint32 GetTypeHash(const FGeneratedImageKey& Key)
+{
+	uint32 GuidHash = GetTypeHash(Key.Pin->PinId);
+
+	return GuidHash;
+}
+
+
+inline uint32 GetTypeHash(const FGeneratedImagePropertiesKey& Key)
+{
+	uint32 GuidHash = HashCombineFast(GetTypeHash(Key.MaterialReferenceId), Key.ImageIndex);
+
+	return GuidHash;
+}
+
+struct FPoseBoneData
+{
+	TArray<FName> ArrayBoneName;
+	TArray<FTransform> ArrayTransform;
+};
+
+struct FRealTimeMorphMeshData
+{
+	TArray<FName> NameResolutionMap;
+	TArray<FMorphTargetVertexData> Data;
+};
+
+struct FClothingMeshData
+{
+	int32 ClothingAssetIndex = INDEX_NONE;
+	int32 ClothingAssetLOD = INDEX_NONE;
+	int32 PhysicsAssetIndex = INDEX_NONE;
+	TArray<FCustomizableObjectMeshToMeshVertData> Data;
+};
+
+struct FGroupProjectorTempData
+{
+	class UCustomizableObjectNodeGroupProjectorParameter* CustomizableObjectNodeGroupProjectorParameter;
+	mu::NodeProjectorParameterPtr NodeProjectorParameterPtr;
+	mu::NodeImagePtr NodeImagePtr;
+	mu::NodeRangePtr NodeRange;
+	mu::NodeScalarParameterPtr NodeOpacityParameter;
+
+	mu::NodeScalarEnumParameterPtr PoseOptionsParameter;
+	TArray<FPoseBoneData> PoseBoneDataArray;
+
+	bool bAlternateResStateNameWarningDisplayed = false; // Used to display this warning only once
+
+	int32 TextureSize = 512;
+};
+
+
+struct FGroupNodeIdsTempData
+{
+	FGroupNodeIdsTempData(FGuid OldGuid, FGuid NewGuid = FGuid()) :
+		OldGroupNodeId(OldGuid),
+		NewGroupNodeId(NewGuid)
+	{
+
+	}
+
+	FGuid OldGroupNodeId;
+	FGuid NewGroupNodeId;
+
+	bool operator==(const FGroupNodeIdsTempData& Other) const
+	{
+		return OldGroupNodeId == Other.OldGroupNodeId;
+	}
+};
+
+struct FGroupProjectorImageInfo
+{
+	mu::Ptr<mu::NodeImage> ImageNode;
+	mu::Ptr<mu::NodeImage> ImageResizeNode;
+	mu::Ptr<mu::NodeSurfaceNew> SurfNode;
+	UCustomizableObjectNodeMaterialBase* TypedNodeMat;
+	FString TextureName;
+	FString RealTextureName;
+	FString AlternateResStateName;
+	float AlternateProjectionResolutionFactor;
+	bool bIsAlternateResolutionResized = false;
+	int32 UVLayout = 0;
+
+	FGroupProjectorImageInfo(mu::NodeImagePtr InImageNode, const FString& InTextureName, const FString& InRealTextureName, UCustomizableObjectNodeMaterialBase* InTypedNodeMat,
+		float InAlternateProjectionResolutionFactor, const FString& InAlternateResStateName, mu::Ptr<mu::NodeSurfaceNew> InSurfNode, int32 InUVLayout)
+		: TypedNodeMat(InTypedNodeMat), TextureName(InTextureName), RealTextureName(InRealTextureName),
+		AlternateResStateName(InAlternateResStateName), AlternateProjectionResolutionFactor(InAlternateProjectionResolutionFactor), 
+		UVLayout(InUVLayout)
+	{
+		ImageNode = InImageNode;
+		SurfNode = InSurfNode;
+	}
+
+	static FString GenerateId(const UCustomizableObjectNode* TypedNodeMat, int32 ImageIndex)
+	{
+		return TypedNodeMat->GetOutermost()->GetPathName() + TypedNodeMat->NodeGuid.ToString() + FString("-") + FString::FromInt(ImageIndex);
+	}
+};
+
+
+/** Struct that defines a mesh (Mesh + LOD + MaterialIndex) */
+struct FMeshData
+{
+	/** Mesh which it may contain multiple LODs and Materials. */
+	const UObject* Mesh;
+
+	/** Specific LOD of the mesh. */
+	int LOD;
+
+	/** Specific index of the mesh material (MaterialIndex is an alias of Section). */
+	int MaterialIndex;
+
+	/** Node where the mesh is defined. Not a UCustomizableObjectNodeMesh due to Table nodes */
+	const UCustomizableObjectNode* Node;
+
+	bool operator==(const FMeshData& Other) const
+	{
+		return Mesh == Other.Mesh &&
+			LOD == Other.LOD &&
+			MaterialIndex == Other.MaterialIndex &&
+			Node == Other.Node;
+	}
+};
+
+inline uint32 GetTypeHash(const FMeshData& Key)
+{
+	uint32 MeshHash = GetTypeHash(Key.Mesh);
+	uint32 NodeHash = GetTypeHash(Key.Node->GetUniqueID());
+	
+	return HashCombine(HashCombine(HashCombine(MeshHash, Key.LOD), Key.MaterialIndex), NodeHash);
+}
+
+
+/** Struct used to store info specific to each component during compilation */
+struct FMutableComponentInfo
+{
+	FMutableComponentInfo(FName InComponentName, USkeletalMesh* InRefSkeletalMesh);
+
+	void AccumulateBonesToRemovePerLOD(const TArray<FLODReductionSettings>& LODReductionSettings, int32 NumLODs);
+
+	FName ComponentName;
+
+	// Each component must have a reference SkeletalMesh with a valid Skeleton
+	USkeletalMesh* RefSkeletalMesh = nullptr;
+	USkeleton* RefSkeleton = nullptr;
+	
+	UCustomizableObjectNodeComponentMesh* NodeComponentMesh = nullptr;
+	
+	// Map to check skeleton compatibility
+	TMap<const USkeleton*, bool> SkeletonCompatibility;
+	
+	// Hierarchy hash from parent-bone to root bone, used to check if additional skeletons are compatible with
+	// the RefSkeleton
+	TMap<FName, uint32> BoneNamesToPathHash;
+
+	// Bones to remove on each LOD, include bones on previous LODs. FName (BoneToRemove) - bool (bOnlyRemoveChildren)
+	TArray<TMap<FName, bool>> BonesToRemovePerLOD;
+};
+
+
+/** Struct of data which is behind a given output pin. 
+ *
+ * Each newly added field has should also be defined on the Append() function body.
+ */
+struct FPinDataValue
+{
+	/** Set of all meshes behind a given output pin. */
+	TSet<FMeshData> MeshesData;
+
+	/** Add all data from other PinData. */
+	void Append(const FPinDataValue& From);
+};
+
+/** Eases the managment and access of data behind a given pin.
+ * 
+ * We understand as "behind a pin" as all pins which are connected to a given pin directly and indirectly. 
+ * 
+ * Given the the following pins:
+ * C -> B -> A
+ *      D ---^
+ * 
+ * FPinData will propagate the added data in the following way:
+ * C (data added on C) -> B (data added on C, B) -> A (data added on C, B, A, D)
+ *                           D (data added on D) ---^
+ * 
+ * FPinDataValue defines which data is going to be saved behind pins.
+ *
+ * When traversing the graph:
+ * - Each time an output pin is explored recursively, call the Push() function.
+ * - When exiting the recursion, call the Pop() function.
+ * - Add data to be seen behind all the stacked pins using the GetCurrent() function.
+ * - To avoid missing any returning branch, use the SCOPE_PIN_DATA wrapper.
+ * 
+ * The all data from a behind a pin can be queried using the Find() method.
+ */
+class FPinData // DEPRECATED. DO NOT USE!
+{
+public:
+	/** Get the PinData from a given pin. */
+	FPinDataValue* Find(const UEdGraphPin* Pin);
+
+	/** Get the PinData of the current pin. The top pin in the stack. */
+	FPinDataValue& GetCurrent();
+
+	/** Pop the current pin from the stack and append all its PinData with to previous pin in the stack. */
+	void Pop();
+
+	/** Push a pin to the stack and create its PinData. */
+	void Push(const UEdGraphPin* Pin);
+
+private:
+	/** Data for each pin. */
+	TMap<const UEdGraphPin*, FPinDataValue> Data;
+
+	/** Stack of pins. Used to know which pin the data belongs to during the graph traversal recursion. */
+	TArray<const UEdGraphPin*> PinStack;
+};
+
+/** Graph cycle key.
+ *
+ * Pin is not enough since we can call multiple recursive functions with the same pin.
+ * Each function has to have an unique identifier.
+ */
+struct FGraphCycleKey
+{
+	friend uint32 GetTypeHash(const FGraphCycleKey& Key);
+
+	FGraphCycleKey(const UEdGraphPin& Pin, const FString& Id);
+
+	bool operator==(const FGraphCycleKey& Other) const;
+	
+	/** Valid pin. */
+	const UEdGraphPin& Pin;
+
+	/** Unique id. */
+	FString Id;
+};
+
+/** Graph Cycle scope.
+ *
+ * Detect a cycle during the graph traversal.
+ */
+class FGraphCycle
+{
+public:
+	explicit FGraphCycle(const FGraphCycleKey&& Key, FMutableGraphGenerationContext &Context);
+	~FGraphCycle();
+
+	/** Return true if there is a cycle. */
+	bool FoundCycle() const;
+	
+private:
+	/** Graph traversal key. */
+	FGraphCycleKey Key;
+
+	/** Generation context. */
+	FMutableGraphGenerationContext& Context;
+};
+
+/** Return the default value if there is a cycle. */
+#define RETURN_ON_CYCLE(Pin, GenerationContext) \
+	FGraphCycle GraphCycle(FGraphCycleKey(Pin, TEXT(__FILE__ PREPROCESSOR_TO_STRING(__LINE__))), GenerationContext); \
+	if (GraphCycle.FoundCycle()) \
+	{ \
+		return {}; \
+	} \
+
+struct FMutableGraphGenerationContext
+{
+	FMutableGraphGenerationContext(UCustomizableObject* CustomizableObject, class FCustomizableObjectCompiler* InCompiler, const FCompilationOptions& InOptions);
+	~FMutableGraphGenerationContext();
+
+	/** See FCustomizableObjectPrivateData::ParticipatingObjects. */
+	void AddParticipatingObject(const UObject& Object);
+
+	/** See FCustomizableObjectPrivateData::ParticipatingObjects. */
+	template<typename Type>
+	void AddParticipatingObject(const TArray<Type>& InObjects)
+	{
+		for (Type InObject : InObjects)
+		{
+			if (InObject)
+			{	
+				AddParticipatingObject(*InObject);
+			}
+		}
+	}
+	
+	UCustomizableObject* Object = nullptr;
+
+	/** Full hierarchy root. */
+	UCustomizableObjectNodeObject* Root = nullptr;
+	
+	// Non-owned reference to the compiler object
+	FCustomizableObjectCompiler* Compiler = nullptr;
+
+	// Compilation options, including target platform
+	FCompilationOptions Options;
+
+	// Cache of generated pins per LOD
+	TMap<FGeneratedKey, FGeneratedData> Generated;
+
+	/** Set of all generated nodes. */
+	TSet<UCustomizableObjectNode*> GeneratedNodes;
+
+	/** Struct that stores the relevant information of a data table generated during the compilation. 
+	e.g. all data tables must have the same compilation restrictions */
+	struct FGeneratedDataTablesData
+	{
+		// Pointer to the generated mutable Table
+		mu::Ptr<mu::Table> GeneratedTable;
+
+		// Table Node used to fill this info
+		const UCustomizableObjectNodeTable* ReferenceNode;
+
+		// Stores the names of the rows that will be compiled
+		TArray<FName> RowNames;
+
+		// Compilation Restrictions:
+		// If there is a bool column in the table, checked rows will not be compiled
+		bool bDisableCheckedRows;
+
+		// Name of the column that determines de version control
+		FName VersionColumn;
+
+		// Compare the stored compilation settings with the compilation settings of a Table Node
+		// return true if the compilation settings are equal
+		bool HasSameSettings(const UCustomizableObjectNodeTable* Node) const;
+	};
+
+	// Cache of generated Node Tables
+	TMap<FString, FGeneratedDataTablesData> GeneratedTables;
+
+	struct FGeneratedCompositeDataTablesData
+	{
+		UScriptStruct* ParentStruct = nullptr;
+		TArray<FName> FilterPaths;
+		UCompositeDataTable* GeneratedDataTable = nullptr;
+
+		bool operator==(const FGeneratedCompositeDataTablesData& Other) const
+		{
+			return ParentStruct == Other.ParentStruct && FilterPaths == Other.FilterPaths;
+		}
+	};
+
+	// Cache of generated Composited Data Tables
+	TArray<FGeneratedCompositeDataTablesData> GeneratedCompositeDataTables;
+
+	// Cache of generated images, because sometimes they are reused by LOD, we use this as a second
+	// level cache
+	TMap<FGeneratedImageKey, mu::NodeImagePtr> GeneratedImages;
+
+	/** Data stored per-generated passthrough texture. */
+	struct FGeneratedReferencedTexture
+	{
+		uint32 ID;
+		//mu::FImageDesc ImageDesc;
+	};
+
+	/** Data stored per-generated passthrough mesh. */
+	struct FGeneratedReferencedMesh
+	{
+		uint32 ID;
+	};
+
+	// Cache of runtime pass-through meshes and their IDs used in the core to identify them.
+	// These meshes will remain as external references even in optimized models.
+	TMap<TSoftObjectPtr<USkeletalMesh>, FGeneratedReferencedMesh> PassthroughMeshMap;
+
+	// Cache of runtime pass-through images and their IDs used in the core to identify them.
+	// These textures will remain as external references even in optimized models.
+	TMap<TSoftObjectPtr<UTexture>, FGeneratedReferencedTexture> PassthroughTextureMap;
+
+	// Cache of runtime images and their IDs used in the core to identify them.
+	// These textures will remain as external references even in optimized models.
+	TMap<TSoftObjectPtr<const UTexture>, FGeneratedReferencedTexture> RuntimeReferencedTextureMap;
+	
+	// Cache of runtime pass-through images and their IDs used in the core to identify them
+	// These textures will become mutable images in the compiled model.
+	TMap<TSoftObjectPtr<const UTexture>, FGeneratedReferencedTexture> CompileTimeTextureMap;
+
+	// Global morph selection overrides.
+	TArray<FRealTimeMorphSelectionOverride> RealTimeMorphTargetsOverrides;
+
+	// Mutable meshes already build for source UStaticMesh or USkeletalMesh.
+	struct FGeneratedMeshData
+	{
+		struct FKey
+		{
+			/** Source mesh data. */
+			const UObject* Mesh = nullptr;
+			int32 LOD = 0; // Mesh Data LOD (i.e., LOD where we are getting the vertices from)
+			int32 CurrentLOD = 0; // Derived data LOD (i.e., LOD where we are generating the non-Core Data like morphs)
+			int32 MaterialIndex = 0;
+
+			/** Flag used to generate this mesh. Bit mask of EMutableMeshConversionFlags */
+			EMutableMeshConversionFlags Flags = EMutableMeshConversionFlags::None;
+
+			/** Tags added at the UE level that go through the Mutable core and are merged in the generated mesh.
+			 *  Only add the tags that make the mesh unique and require it not to be cached together with the 
+			 *  same exact mesh but with different tags.
+			*/
+			FString Tags;
+
+			/**
+			* SkeletalMeshNode is needed to disambiguate realtime morph selection from diferent nodes.
+			* TODO: Consider using the actual selection.
+			*/
+			const UCustomizableObjectNode* SkeletalMeshNode = nullptr;
+
+			bool operator==( const FKey& OtherKey ) const
+			{
+				return Mesh == OtherKey.Mesh && LOD == OtherKey.LOD && CurrentLOD == OtherKey.CurrentLOD && MaterialIndex == OtherKey.MaterialIndex
+					&& Flags == OtherKey.Flags && Tags == OtherKey.Tags && SkeletalMeshNode == OtherKey.SkeletalMeshNode;
+			}
+		};
+
+		FKey Key;
+
+		/** Generated mesh. */
+		mu::Ptr<mu::Mesh> Generated;
+	};
+	TArray<FGeneratedMeshData> GeneratedMeshes;
+
+	struct FGeneratedTableImageData
+	{
+		FString PinName;
+		FName PinType;
+		const mu::Ptr<mu::Table> Table;
+		const UCustomizableObjectNodeTable* TableNode;
+
+		bool operator==(const FGeneratedTableImageData& Other) const
+		{
+			return PinName == Other.PinName && Table == Other.Table;
+		}
+	};
+	TArray<FGeneratedTableImageData> GeneratedTableImages;
+
+	// Stack of mesh generation flags. The last one is the currently valid.
+	// The value is a bit mask of EMutableMeshConversionFlags
+	TArray<EMutableMeshConversionFlags> MeshGenerationFlags;
+
+	// Stack of Layout generation flags. The last one is the currently valid.
+	TArray<FLayoutGenerationFlags> LayoutGenerationFlags;
+
+	/** Find a mesh if already generated for a given source and flags. */
+	mu::Ptr<mu::Mesh> FindGeneratedMesh(const FGeneratedMeshData::FKey& Key);
+
+	/** Add a resource to the streamed resources array.
+	  * OutStreamedResourceContainer - Container to store the streamed resource.
+	  * Returns resource index in the array of streamed resources. */
+	int32 AddStreamedResource(const uint32 InResourceHash, UCustomizableObjectResourceDataContainer*& OutStreamedResourceContainer);
+
+	/** Adds a streamed resource of type AssetUserData.
+	  * Returns resource index in the array of streamed resources. */
+	int32 AddAssetUserDataToStreamedResources(UAssetUserData* AssetUserData);
+
+	/** Adds to ParameterNamesMap the node Node to the array of elements with name Name */
+	void AddParameterNameUnique(const UCustomizableObjectNode* Node, FString Name);
+
+	uint32 GetSkinWeightProfileIdUnique(const FName ProfileName);
+
+	// Check if the Id of the node Node already exists, if it's new adds it to NodeIds array, otherwise, returns new Id
+	const FGuid GetNodeIdUnique(const UCustomizableObjectNode* Node);
+
+	/** Generates new tags for the UCustomizableObjectNodeModifierClipWithMesh nodes which have assigned a CO in the
+	* UCustomizableObjectNodeModifierClipWithMesh::CustomizableObjectToClipWith field
+	* @return nothing */
+	void GenerateClippingCOInternalTags();
+
+	/** Generates shared surface IDs for all surface nodes. If one or more nodes are equal, they will use the same SharedSurfaceId */
+	void GenerateSharedSurfacesUniqueIds();
+
+	bool FindBone(const FName& BoneName, mu::FBoneName& OutBone) const;
+
+	/** Get unique identifier for BoneName built from its FString. */
+	mu::FBoneName GetBoneUnique(const FName& BoneName);
+
+	/** Check if the PhysicsAsset of a given SkeletalMesh has any SkeletalBodySetup with BoneNames not present in the
+	* InSkeletalMesh's RefSkeleton, if so, adds the PhysicsAsset to the DiscartedPhysicsAssetMap to display a warning later on */
+	//void CheckPhysicsAssetInSkeletalMesh(const USkeletalMesh* InSkeletalMesh);
+
+	/** Get the reference skeletal mesh associated to the current mesh component being generated */
+	FMutableComponentInfo* GetCurrentComponentInfo();
+
+	TMap<FName, UCustomizableObjectNodeComponentMesh*> MeshComponents;
+
+	TArray<FMutableComponentInfo> ComponentInfos;
+	TArray<FMutableRefSkeletalMeshData> ReferenceSkeletalMeshesData;
+	
+	TArray<UMaterialInterface*> ReferencedMaterials;
+	TArray<FName> ReferencedMaterialSlotNames;
+	TMap<FGeneratedImagePropertiesKey, FGeneratedImageProperties> ImageProperties;
+	TMap<FString, TArray<const UObject*>> ParameterNamesMap;
+	TArray<const UCustomizableObjectNode*> NoNameNodeObjectArray;
+	TMap<FString, FCustomizableObjectIdPair> GroupNodeMap;
+	TMap<FString, FString> CustomizableObjectPathMap;
+	TMap<FString, FMutableParameterData> ParameterUIDataMap;
+	TMap<FString, FMutableStateData> StateUIDataMap;
+	TMap<TTuple<FString, FString>, TSet<TSoftObjectPtr<UDataTable>>> IntParameterOptionDataTable;
+	TMultiMap<const UCustomizableObjectNodeObjectGroup*, FGroupProjectorTempData> ProjectorGroupMap;
+	//TMap<UPhysicsAsset*, uint32> DiscartedPhysicsAssetMap;
+
+	TArray<const USkeleton*> ReferencedSkeletons;
+
+	// Array of unique Bone identifiers. 
+	TMap<mu::FBoneName, FString> UniqueBoneNames;
+	TMap<FString, mu::FBoneName> RemappedBoneNames; // Bone identifiers that had a collision.
+
+	// Used to aviod Nodes with duplicated ids
+	TMap<FGuid, TArray<const UObject*>> NodeIdsMap;
+	TMultiMap<const UCustomizableObject*, FGroupNodeIdsTempData> DuplicatedGroupNodeIds;
+
+	// For a given material node (the key is node package path + node uid + image index in node) stores images generated for the same node at a higher quality LOD to reuse that image node
+	TMap<FString, FGroupProjectorImageInfo> GroupProjectorLODCache;
+
+	// Data structures used for clipping feature by assigning a Customizable Object
+	// Map with pairs (Unreal Mutable material node Guid, array with its corresponding Mutable surface node) to add tags from clipping nodes
+	TMap<class UCustomizableObjectNodeMaterial*, TArray<mu::Ptr<mu::NodeSurfaceNew>>> MapMaterialNodeToMutableSurfaceNodeArray;
+
+	// Map with pairs (Unreal Mutable clip mesh node, array with its corresponding Mutable mesh modifier nodes) to add tags
+	TMap<class UCustomizableObjectNodeModifierClipWithMesh*, TArray<mu::Ptr<mu::NodeModifierMeshClipWithMesh>>> MapClipMeshNodeToMutableClipMeshNodeArray;
+
+	// Data used for MorphTarget reconstruction.
+	TMap<uint32, FRealTimeMorphMeshData> RealTimeMorphTargetPerMeshData;
+
+	// Data used for Clothing reconstruction.
+	TArray<FCustomizableObjectClothingAssetData> ClothingAssetsData;
+	TMap<uint32, FClothingMeshData> ClothingPerMeshData;
+
+	TMap<uint32, FMutableMeshMetadata> MeshMetadata;
+	TMap<uint32, FMutableSurfaceMetadata> SurfaceMetadata;
+
+	// Data used for SkinWeightProfiles reconstruction
+	TArray<FMutableSkinWeightProfileInfo> SkinWeightProfilesInfo;
+
+	TMap<uint32, FName> UniqueSkinWeightProfileIds;
+	TMap<FName, uint32> RemappedSkinWeightProfileIds;
+
+	TArray<FAnimBpOverridePhysicsAssetsInfo> AnimBpOverridePhysicsAssetsInfo;
+
+	uint8 FromLOD = 0; // LOD to append to the CurrentLOD when using AutomaticLODs. 
+	uint8 CurrentLOD = 0;
+	FName CurrentMeshComponent;
+
+	/** If this is set, we are genreating materials for a "passthrough" component, with a fixed mesh. */
+	mu::Ptr<mu::NodeMesh> ComponentMeshOverride;
+
+	uint8 NumLODsInRoot = 0;
+	uint8 NumComponents = 0;
+
+	uint8 FirstLODAvailable = MAX_MESH_LOD_COUNT;
+	uint8 NumMaxLODsToStream = MAX_MESH_LOD_COUNT;
+
+	bool bEnableLODStreaming = true;
+
+	bool bPartialCompilation = false;
+
+	// Based on the last object visited.
+	ECustomizableObjectAutomaticLODStrategy CurrentAutoLODStrategy = ECustomizableObjectAutomaticLODStrategy::Manual;
+
+	// Stores external graph root nodes to be added to the specified group nodes
+	TMultiMap<FGuid, UCustomizableObjectNodeObject*> GroupIdToExternalNodeMap;
+
+	// Easily retrieve a parameter name from its node guid
+	TMap<FGuid, FString> GuidToParamNameMap;
+
+	// Graph cycle detection
+	/** Visited nodes during the DAC recursion traversal.
+	 * It acts like stack, pushing pins when recursively exploring a new pin an popping it when exiting the recursion. */
+	TMap<FGraphCycleKey, const UCustomizableObject*> VisitedPins;
+	const UCustomizableObject* CustomizableObjectWithCycle = nullptr;
+
+	/** Stores the physics assets gathered from the SkeletalMesh nodes during compilation, to be used in mesh generation in-game */
+	TArray<TSoftObjectPtr<UPhysicsAsset>> PhysicsAssets;
+
+	/** Stores the anim BP assets gathered from the SkeletalMesh nodes during compilation, to be used in mesh generation in-game */
+	TArray<TSoftClassPtr<UAnimInstance>> AnimBPAssets;
+
+	/** Stores the sockets provided by the part skeletal meshes, to be merged in the generated meshes */
+	TArray<FMutableRefSocket> SocketArray;
+
+	/** Used to propagate the socket priority defined in group nodes to their child skeletal mesh nodes
+	* It's a stack because group nodes are recursive
+	*/
+	TArray<int32> SocketPriorityStack;
+
+	// Stores the only option of an Int Param that should be compiled in a partial compilation
+	TMap<FString, FString> ParamNamesToSelectedOptions;
+
+	// Stores what param names use a certain table as a table can be used from multiple table nodes, useful for partial compilations to restrict params
+	TMap<FString, FMutableParamNameSet> TableToParamNames;
+
+	TArray<const UEdGraphNode*> LimitedParameters;
+	int32 ParameterLimitationCount = 0;
+
+	/** Data which is behind a given output pin. 
+	 *
+	 * We only consider output pins. Do not use input pins as keys. See FPinData class definition. 
+	 */
+	FPinData PinData; // DEPRECATED. DO NOT USE!
+
+	// Stores all morphs to apply them directly to a skeletal mesh node
+	TArray<FMorphNodeData> MeshMorphStack;
+
+	// Current material parameter name to find the corresponding column in a mutable table
+	FString CurrentMaterialTableParameter;
+
+	// Current material parameter id to find the corresponding column in a mutable table
+	FString CurrentMaterialTableParameterId;
+
+	struct FSharedSurface
+	{
+		FSharedSurface(uint8 InLOD, const mu::Ptr<mu::NodeSurfaceNew>& InNodeSurfaceNew);
+
+		bool operator==(const FSharedSurface& o) const;
+
+		uint8 LOD = 0;
+		mu::Ptr<mu::NodeSurfaceNew> NodeSurfaceNew;
+
+		bool bMakeUnique = false;
+		TArray<SIZE_T> NodeModifierIDs;
+	};
+
+	// Material to SharedSurfaceId
+	TMap<UCustomizableObjectNodeMaterialBase*, TArray<FSharedSurface>> SharedSurfaceIds;
+
+	/** Resource Data constants */
+	TMap<uint32, int32> StreamedResourceIndices;
+	TArray<FCustomizableObjectStreamedResourceData> StreamedResourceData;
+
+	/** Extension Data constants are collected here */
+	FExtensionDataCompilerInterface ExtensionDataCompilerInterface;
+	TArray<FCustomizableObjectResourceData> AlwaysLoadedExtensionData;
+	TArray<UCustomizableObjectResourceDataContainer*> StreamedExtensionData;
+
+	/** See FCustomizableObjectPrivateData::ParticipatingObjects. */
+	TMap<FName, FGuid> ParticipatingObjects;
+
+	/** Map to relate a Composite Data Table Row and its original DataTable */
+	TMap<UDataTable*,TMap<FName, TArray<UDataTable*>>> CompositeDataTableRowToOriginalDataTableMap;
+
+	/** Version Bridge of the root object */
+	TObjectPtr<UObject> RootVersionBridge;
+};
+
+/** Pin Data scope wrapper. Pops the pin data on scope exit. */
+class FScopedPinData
+{
+public:
+	explicit FScopedPinData(FMutableGraphGenerationContext& Context, const UEdGraphPin* Pin);
+	~FScopedPinData();
+
+private:
+	FMutableGraphGenerationContext& Context;
+};
+
+#define SCOPED_PIN_DATA(Context, Pin) \
+	FScopedPinData ScopedPinData(Context, Pin);
+
+
+namespace Private
+{
+	template<class HashableType, class HashDataSetType, class HashFuncType, class CompareFuncType>
+	uint32 GenerateUniquePersistentHash(const HashableType& HashableData, const HashDataSetType& HashDataSet, HashFuncType&& HashFunc, CompareFuncType&& CompareFunc)
+	{
+		constexpr uint32 InvalidResourceId = 0;
+		
+		const uint32 DataHash = HashFunc(HashableData);
+
+		uint32 UniqueHash = DataHash == InvalidResourceId ? DataHash + 1 : DataHash;
+
+		const HashableType* FoundHash = HashDataSet.Find(UniqueHash);
+
+		bool bIsDataAlreadyCollected = false;
+		
+		if (FoundHash)
+		{
+			bIsDataAlreadyCollected = CompareFunc(*FoundHash, HashableData); 
+		}
+
+		// NOTE: This way of unique hash generation guarantees all valid values can be used but given its 
+		// sequential nature a cascade of changes can occur if new meshes are added. Not many hash collisions 
+		// are expected so it should not be problematic.
+		if (FoundHash && !bIsDataAlreadyCollected)
+		{
+			uint32 NumTries = 0;
+			for (; NumTries < TNumericLimits<uint32>::Max(); ++NumTries)
+			{
+				FoundHash = HashDataSet.Find(UniqueHash);
+				
+				if (!FoundHash)
+				{
+					break;
+				}
+
+				bIsDataAlreadyCollected = CompareFunc(*FoundHash, HashableData);
+
+				if (bIsDataAlreadyCollected)
+				{
+					break;
+				}
+
+				UniqueHash = UniqueHash + 1 == InvalidResourceId ? InvalidResourceId + 1 : UniqueHash + 1;
+			}
+
+			if (NumTries == TNumericLimits<uint32>::Max())
+			{
+				UniqueHash = InvalidResourceId;
+			}	
+		}
+
+		return UniqueHash;
+	}
+} //Private
+
+//
+mu::Ptr<mu::NodeObject> GenerateMutableSource(const class UEdGraphPin* Pin, FMutableGraphGenerationContext& GenerationContext);
+
+/** Populate an array with all the information related to the reference skeletal meshes we might need in-game to generate instances */
+void PopulateReferenceSkeletalMeshesData(FMutableGraphGenerationContext& GenerationContext);
+
+
+void CheckNumOutputs(const UEdGraphPin& Pin, const FMutableGraphGenerationContext& GenerationContext);
+
+
+// TODO FutureGMT Remove generation context dependency and move to GraphTraversal.
+UTexture2D* FindReferenceImage(const UEdGraphPin* Pin, FMutableGraphGenerationContext& GenerationContext);
+
+
+mu::NodeMeshApplyPosePtr CreateNodeMeshApplyPose(FMutableGraphGenerationContext& GenerationContext, mu::NodeMeshPtr InputMeshNode, const TArray<FName>& ArrayBoneName, const TArray<FTransform>& ArrayTransform);
+
+
+/** Adds Tag to MutableMesh uniquely, returns the index were the tag has been inserted or the index where an intance of the tag has been found */
+int32 AddTagToMutableMeshUnique(mu::Mesh& MutableMesh, const FString& Tag);
+
+void AddSocketTagsToMesh(const USkeletalMesh* SourceMesh, mu::Ptr<mu::Mesh> MutableMesh, FMutableGraphGenerationContext& GenerationContext);
+
+// Generates the tag for an animation instance
+FString GenerateAnimationInstanceTag(const int32 AnimInstanceIndex, const FName& SlotIndex);
+
+
+FString GenerateGameplayTag(const FString& GameplayTag);
+
+uint32 GetBaseTextureSize(const FMutableGraphGenerationContext& GenerationContext, const UCustomizableObjectNodeMaterialBase* Material, uint32 ImageIndex);
+
+// Computes the LOD bias for a texture given the current mesh LOD and automatic LOD settings, the reference texture settings
+// and whether it's being built for a server or not
+uint32 ComputeLODBiasForTexture(const FMutableGraphGenerationContext& GenerationContext, const UTexture2D& Texture, const UTexture2D* ReferenceTexture = nullptr, int32 MaxTextureSizeInGame = 0);
+
+// Max texture size to set on the ImageProperties
+int32 GetMaxTextureSize(const UTexture2D& ReferenceTexture, const UTextureLODSettings& LODSettings);
+
+// Max texture size of the texture with per platform MaxTextureSize and LODBias applied.
+int32 GetTextureSizeInGame(const UTexture2D& Texture, const UTextureLODSettings& LODSettings, uint8 SurfaceLODBias = 0);
+
+mu::Ptr<mu::Image> GenerateImageConstant(UTexture*, FMutableGraphGenerationContext&, bool bIsReference);
+mu::Ptr<mu::Mesh> GenerateMeshConstant(USkeletalMesh*, FMutableGraphGenerationContext&, bool bIsReference);
+
+/** Generates a mutable image descriptor from an unreal engine texture */
+mu::FImageDesc GenerateImageDescriptor(UTexture* Texture);
diff --git a/Engine/Plugins/Mutable/Source/CustomizableObjectEditor/Private/MuCOE/GenerateMutableSource/GenerateMutableSourceColor.cpp b/Engine/Plugins/Mutable/Source/CustomizableObjectEditor/Private/MuCOE/GenerateMutableSource/GenerateMutableSourceColor.cpp
new file mode 100644
--- /dev/null
+++ b/Engine/Plugins/Mutable/Source/CustomizableObjectEditor/Private/MuCOE/GenerateMutableSource/GenerateMutableSourceColor.cpp
@@ -0,0 +1,358 @@
+// Copyright Epic Games, Inc. All Rights Reserved.
+
+#include "MuCOE/GenerateMutableSource/GenerateMutableSourceColor.h"
+
+#include "MuCOE/CustomizableObjectCompiler.h"
+#include "MuCOE/EdGraphSchema_CustomizableObject.h"
+#include "MuCOE/GenerateMutableSource/GenerateMutableSourceFloat.h"
+#include "MuCOE/GenerateMutableSource/GenerateMutableSourceImage.h"
+#include "MuCOE/GenerateMutableSource/GenerateMutableSourceTable.h"
+#include "MuCOE/GraphTraversal.h"
+#include "MuCOE/Nodes/CustomizableObjectNodeColorArithmeticOp.h"
+#include "MuCOE/Nodes/CustomizableObjectNodeColorConstant.h"
+#include "MuCOE/Nodes/CustomizableObjectNodeColorFromFloats.h"
+#include "MuCOE/Nodes/CustomizableObjectNodeColorParameter.h"
+#include "MuCOE/Nodes/CustomizableObjectNodeColorSwitch.h"
+#include "MuCOE/Nodes/CustomizableObjectNodeColorVariation.h"
+#include "MuCOE/Nodes/CustomizableObjectNodeTable.h"
+#include "MuCOE/Nodes/CustomizableObjectNodeTextureSample.h"
+#include "MuT/NodeColourArithmeticOperation.h"
+#include "MuT/NodeColourConstant.h"
+#include "MuT/NodeColourFromScalars.h"
+#include "MuT/NodeColourParameter.h"
+#include "MuT/NodeColourSampleImage.h"
+#include "MuT/NodeColourSwitch.h"
+#include "MuT/NodeColourTable.h"
+#include "MuT/NodeColourVariation.h"
+#include "MuT/NodeImageBinarise.h"
+
+
+#define LOCTEXT_NAMESPACE "CustomizableObjectEditor"
+
+
+mu::Ptr<mu::NodeColour> GenerateMutableSourceColor(const UEdGraphPin* Pin, FMutableGraphGenerationContext& GenerationContext)
+{
+	check(Pin)
+	RETURN_ON_CYCLE(*Pin, GenerationContext)
+
+	CheckNumOutputs(*Pin, GenerationContext);
+	
+	const UEdGraphSchema_CustomizableObject* Schema = GetDefault<UEdGraphSchema_CustomizableObject>();
+
+	UCustomizableObjectNode* Node = CastChecked<UCustomizableObjectNode>(Pin->GetOwningNode());
+
+	const FGeneratedKey Key(reinterpret_cast<void*>(&GenerateMutableSourceColor), *Pin, *Node, GenerationContext);
+	if (const FGeneratedData* Generated = GenerationContext.Generated.Find(Key))
+	{
+		return static_cast<mu::NodeColour*>(Generated->Node.get());
+	}
+
+	if (Node->IsNodeOutDatedAndNeedsRefresh())
+	{
+		Node->SetRefreshNodeWarning();
+	}
+
+	mu::Ptr<mu::NodeColour> Result;
+
+	if (const UCustomizableObjectNodeColorConstant* TypedNodeColorConst = Cast<UCustomizableObjectNodeColorConstant>(Node))
+	{
+		mu::Ptr<mu::NodeColourConstant> ColorNode = new mu::NodeColourConstant();
+		Result = ColorNode;
+
+		ColorNode->Value = TypedNodeColorConst->Value;
+	}
+
+	else if (const UCustomizableObjectNodeColorParameter* TypedNodeColorParam = Cast<UCustomizableObjectNodeColorParameter>(Node))
+	{
+		mu::Ptr<mu::NodeColourParameter> ColorNode = new mu::NodeColourParameter();
+		Result = ColorNode;
+
+		GenerationContext.AddParameterNameUnique(Node, TypedNodeColorParam->ParameterName);
+
+		ColorNode->Name = TypedNodeColorParam->ParameterName;
+		ColorNode->Uid = GenerationContext.GetNodeIdUnique(Node).ToString();
+		ColorNode->DefaultValue = TypedNodeColorParam->DefaultValue;
+
+		GenerationContext.ParameterUIDataMap.Add(TypedNodeColorParam->ParameterName, FMutableParameterData(
+			TypedNodeColorParam->ParamUIMetadata,
+			EMutableParameterType::Color));
+	}
+
+	else if (const UCustomizableObjectNodeColorSwitch* TypedNodeColorSwitch = Cast<UCustomizableObjectNodeColorSwitch>(Node))
+	{
+		Result = [&]() -> mu::Ptr<mu::NodeColourSwitch>
+		{
+			if (const int32 NumParameters = FollowInputPinArray(*TypedNodeColorSwitch->SwitchParameter()).Num();
+				NumParameters != 1)
+			{
+				const FText Message = NumParameters
+					? LOCTEXT("NoEnumParamInSwitch", "Switch nodes must have an enum switch parameter. Please connect an enum and refesh the switch node.")
+					: LOCTEXT("InvalidEnumInSwitch", "Switch nodes must have a single enum with all the options inside. Please remove all the enums but one and refresh the switch node.");
+
+				GenerationContext.Compiler->CompilerLog(Message, Node);
+				return nullptr;
+			}
+
+			const UEdGraphPin* EnumPin = FollowInputPin(*TypedNodeColorSwitch->SwitchParameter());
+			mu::NodeScalarPtr SwitchParam = EnumPin ? GenerateMutableSourceFloat(EnumPin, GenerationContext) : nullptr;
+
+			// Switch Param not generated
+			if (!SwitchParam)
+			{
+				// Warn about a failure.
+				if (EnumPin)
+				{
+					const FText Message = LOCTEXT("FailedToGenerateSwitchParam", "Could not generate switch enum parameter. Please refesh the switch node and connect an enum.");
+					GenerationContext.Compiler->CompilerLog(Message, Node);
+				}
+
+				return nullptr;
+			}
+
+			if (SwitchParam->GetType() != mu::NodeScalarEnumParameter::GetStaticType())
+			{
+				const FText Message = LOCTEXT("WrongSwitchParamType", "Switch parameter of incorrect type.");
+				GenerationContext.Compiler->CompilerLog(Message, Node);
+
+				return nullptr;
+			}
+
+			const int32 NumSwitchOptions = TypedNodeColorSwitch->GetNumElements();
+
+			mu::NodeScalarEnumParameter* EnumParameter = static_cast<mu::NodeScalarEnumParameter*>(SwitchParam.get());
+			if (NumSwitchOptions != EnumParameter->GetValueCount())
+			{
+				const FText Message = LOCTEXT("MismatchedSwitch", "Switch enum and switch node have different number of options. Please refresh the switch node to make sure the outcomes are labeled properly.");
+				GenerationContext.Compiler->CompilerLog(Message, Node);
+			}
+
+			mu::Ptr<mu::NodeColourSwitch> SwitchNode = new mu::NodeColourSwitch;
+			SwitchNode->Parameter = SwitchParam;
+			SwitchNode->Options.SetNum(NumSwitchOptions);
+
+			for (int SelectorIndex = 0; SelectorIndex < NumSwitchOptions; ++SelectorIndex)
+			{
+				const UEdGraphPin* const ColorPin = TypedNodeColorSwitch->GetElementPin(SelectorIndex);
+				if (const UEdGraphPin* ConnectedPin = FollowInputPin(*ColorPin))
+				{
+					SwitchNode->Options[SelectorIndex] = GenerateMutableSourceColor(ConnectedPin, GenerationContext);
+				}
+			}
+
+			return SwitchNode;
+		}(); // invoke lambda;
+	}
+
+	else if (const UCustomizableObjectNodeTextureSample* TypedNodeTexSample = Cast<UCustomizableObjectNodeTextureSample>(Node))
+	{
+		mu::Ptr<mu::NodeColourSampleImage> ColorNode = new mu::NodeColourSampleImage();
+		Result = ColorNode;
+
+		if (const UEdGraphPin* ConnectedPin = FollowInputPin(*TypedNodeTexSample->TexturePin()))
+		{
+			mu::NodeImagePtr TextureNode = GenerateMutableSourceImage(ConnectedPin, GenerationContext, 0);
+			ColorNode->Image = TextureNode;
+		}
+
+		if (const UEdGraphPin* ConnectedPin = FollowInputPin(*TypedNodeTexSample->XPin()))
+		{
+			mu::NodeScalarPtr XNode = GenerateMutableSourceFloat(ConnectedPin, GenerationContext);
+			ColorNode->X = XNode;
+		}
+
+		if (const UEdGraphPin* ConnectedPin = FollowInputPin(*TypedNodeTexSample->YPin()))
+		{
+			mu::NodeScalarPtr YNode = GenerateMutableSourceFloat(ConnectedPin, GenerationContext);
+			ColorNode->Y = YNode;
+		}
+	}
+
+	else if (const UCustomizableObjectNodeColorArithmeticOp* TypedNodeColorArith = Cast<UCustomizableObjectNodeColorArithmeticOp>(Node))
+	{
+		mu::Ptr<mu::NodeColourArithmeticOperation> OpNode = new mu::NodeColourArithmeticOperation();
+		Result = OpNode;
+
+		if (const UEdGraphPin* ConnectedPin = FollowInputPin(*TypedNodeColorArith->XPin()))
+		{
+			OpNode->A = GenerateMutableSourceColor(ConnectedPin, GenerationContext);
+		}
+
+		if (const UEdGraphPin* ConnectedPin = FollowInputPin(*TypedNodeColorArith->YPin()))
+		{
+			OpNode->B = GenerateMutableSourceColor(ConnectedPin, GenerationContext);
+		}
+
+		switch (TypedNodeColorArith->Operation)
+		{
+		case EColorArithmeticOperation::E_Add:
+			OpNode->Operation = mu::NodeColourArithmeticOperation::EOperation::Add;
+			break;
+
+		case EColorArithmeticOperation::E_Sub:
+			OpNode->Operation = mu::NodeColourArithmeticOperation::EOperation::Subtract;
+			break;
+
+		case EColorArithmeticOperation::E_Mul:
+			OpNode->Operation = mu::NodeColourArithmeticOperation::EOperation::Multiply;
+			break;
+
+		case EColorArithmeticOperation::E_Div:
+			OpNode->Operation = mu::NodeColourArithmeticOperation::EOperation::Divide;
+			break;
+		}
+	}
+
+	else if (const UCustomizableObjectNodeColorFromFloats* TypedNodeFrom = Cast<UCustomizableObjectNodeColorFromFloats>(Node))
+	{
+		mu::Ptr<mu::NodeColourFromScalars> OpNode = new mu::NodeColourFromScalars();
+		Result = OpNode;
+
+		if (const UEdGraphPin* ConnectedPin = FollowInputPin(*TypedNodeFrom->RPin()))
+		{
+			mu::NodeScalarPtr FloatNode = GenerateMutableSourceFloat(ConnectedPin, GenerationContext);
+			OpNode->X = FloatNode;
+		}
+
+		if (const UEdGraphPin* ConnectedPin = FollowInputPin(*TypedNodeFrom->GPin()))
+		{
+			mu::NodeScalarPtr FloatNode = GenerateMutableSourceFloat(ConnectedPin, GenerationContext);
+			OpNode->Y = FloatNode;
+		}
+
+		if (const UEdGraphPin* ConnectedPin = FollowInputPin(*TypedNodeFrom->BPin()))
+		{
+			mu::NodeScalarPtr FloatNode = GenerateMutableSourceFloat(ConnectedPin, GenerationContext);
+			OpNode->Z = FloatNode;
+		}
+
+		if (const UEdGraphPin* ConnectedPin = FollowInputPin(*TypedNodeFrom->APin()))
+		{
+			mu::NodeScalarPtr FloatNode = GenerateMutableSourceFloat(ConnectedPin, GenerationContext);
+			OpNode->W = FloatNode;
+		}
+	}
+
+	else if (const UCustomizableObjectNodeColorVariation* TypedNodeColorVar = Cast<const UCustomizableObjectNodeColorVariation>(Node))
+	{
+		mu::Ptr<mu::NodeColourVariation> ColorNode = new mu::NodeColourVariation();
+		Result = ColorNode;
+
+		if (const UEdGraphPin* ConnectedPin = FollowInputPin(*TypedNodeColorVar->DefaultPin()))
+		{
+			mu::Ptr<mu::NodeColour> ChildNode = GenerateMutableSourceColor(ConnectedPin, GenerationContext);
+			if (ChildNode)
+			{
+				ColorNode->DefaultColour = ChildNode;
+			}
+			else
+			{
+				GenerationContext.Compiler->CompilerLog(LOCTEXT("ColorFailed", "Color generation failed."), Node);
+			}
+		}
+
+		const int32 NumVariations = TypedNodeColorVar->GetNumVariations();
+		ColorNode->Variations.SetNum(NumVariations);
+		for (int VariationIndex = 0; VariationIndex < NumVariations; ++VariationIndex)
+		{
+			const UEdGraphPin* VariationPin = TypedNodeColorVar->VariationPin(VariationIndex);
+			if (!VariationPin) continue;
+
+			ColorNode->Variations[VariationIndex].Tag = TypedNodeColorVar->GetVariation(VariationIndex).Tag;
+			if (const UEdGraphPin* ConnectedPin = FollowInputPin(*VariationPin))
+			{
+				mu::Ptr<mu::NodeColour> ChildNode = GenerateMutableSourceColor(ConnectedPin, GenerationContext);
+				ColorNode->Variations[VariationIndex].Colour = ChildNode;
+			}
+		}
+	}
+
+	else if (const UCustomizableObjectNodeTable* TypedNodeTable = Cast<UCustomizableObjectNodeTable>(Node))
+	{
+		//This node will add a white color in case of error
+		mu::Ptr<mu::NodeColourConstant> WhiteColorNode = new mu::NodeColourConstant();
+		WhiteColorNode->Value = FVector4f(1.0f, 1.0f, 1.0f, 1.0f);
+
+		Result = WhiteColorNode;
+
+		bool bSuccess = true;
+		UDataTable* DataTable = GetDataTable(TypedNodeTable, GenerationContext);
+
+		if (DataTable)
+		{
+			FString ColumnName = Pin->PinFriendlyName.ToString();
+			FProperty* Property = DataTable->FindTableProperty(FName(*ColumnName));
+
+			if (!Property)
+			{
+				FString Msg = FString::Printf(TEXT("Couldn't find the column [%s] in the data table's struct."), *ColumnName);
+				GenerationContext.Compiler->CompilerLog(FText::FromString(Msg), Node);
+
+				bSuccess = false;
+			}
+
+			if (bSuccess)
+			{
+				// Generating a new data table if not exists
+				mu::TablePtr Table;
+				Table = GenerateMutableSourceTable(DataTable, TypedNodeTable, GenerationContext);
+
+				if (Table)
+				{
+					mu::Ptr<mu::NodeColourTable> ColorTableNode = new mu::NodeColourTable();
+
+					// Generating a new Color column if not exists
+					if (Table->FindColumn(ColumnName) == INDEX_NONE)
+					{
+						int32 Dummy = -1; // TODO MTBL-1512
+						bool Dummy2 = false;
+						bSuccess = GenerateTableColumn(TypedNodeTable, Pin, Table, ColumnName, Property, Dummy, Dummy, GenerationContext.CurrentLOD, Dummy, Dummy2, GenerationContext);
+
+						if (!bSuccess)
+						{
+							FString Msg = FString::Printf(TEXT("Failed to generate the mutable table column [%s]"), *ColumnName);
+							GenerationContext.Compiler->CompilerLog(FText::FromString(Msg), Node);
+						}
+					}
+
+					if (bSuccess)
+					{
+						Result = ColorTableNode;
+
+						ColorTableNode->Table = Table;
+						ColorTableNode->ColumnName = ColumnName;
+						ColorTableNode->ParameterName = TypedNodeTable->ParameterName;
+						ColorTableNode->bNoneOption = TypedNodeTable->bAddNoneOption;
+						ColorTableNode->DefaultRowName = TypedNodeTable->DefaultRowName.ToString();
+					}
+				}
+				else
+				{
+					FString Msg = FString::Printf(TEXT("Couldn't generate a mutable table."), *ColumnName);
+					GenerationContext.Compiler->CompilerLog(FText::FromString(Msg), Node);
+				}
+			}
+		}
+		else
+		{
+			GenerationContext.Compiler->CompilerLog(LOCTEXT("ColorTableError", "Couldn't find the data table of the node."), Node);
+		}
+	}
+
+	else
+	{
+		GenerationContext.Compiler->CompilerLog(LOCTEXT("UnimplementedNode", "Node type not implemented yet."), Node);
+	}
+
+	FGeneratedData CacheData = FGeneratedData(Node, Result);
+	GenerationContext.Generated.Add(Key, CacheData);
+	GenerationContext.GeneratedNodes.Add(Node);
+	
+	if (Result)
+	{
+		Result->SetMessageContext(Node);
+	}
+
+	return Result;
+}
+
+#undef LOCTEXT_NAMESPACE
diff --git a/Engine/Plugins/Mutable/Source/CustomizableObjectEditor/Private/MuCOE/GenerateMutableSource/GenerateMutableSourceColor.h b/Engine/Plugins/Mutable/Source/CustomizableObjectEditor/Private/MuCOE/GenerateMutableSource/GenerateMutableSourceColor.h
new file mode 100644
--- /dev/null
+++ b/Engine/Plugins/Mutable/Source/CustomizableObjectEditor/Private/MuCOE/GenerateMutableSource/GenerateMutableSourceColor.h
@@ -0,0 +1,11 @@
+// Copyright Epic Games, Inc. All Rights Reserved.
+
+#pragma once
+
+#include "MuT/NodeColour.h"
+
+class UEdGraphPin;
+struct FMutableGraphGenerationContext;
+
+/** Convert a CustomizableObject Source Graph into a mutable source graph. */
+mu::Ptr<mu::NodeColour> GenerateMutableSourceColor(const UEdGraphPin* Pin, FMutableGraphGenerationContext& GenerationContext);
diff --git a/Engine/Plugins/Mutable/Source/CustomizableObjectEditor/Private/MuCOE/GenerateMutableSource/GenerateMutableSourceFloat.cpp b/Engine/Plugins/Mutable/Source/CustomizableObjectEditor/Private/MuCOE/GenerateMutableSource/GenerateMutableSourceFloat.cpp
new file mode 100644
--- /dev/null
+++ b/Engine/Plugins/Mutable/Source/CustomizableObjectEditor/Private/MuCOE/GenerateMutableSource/GenerateMutableSourceFloat.cpp
@@ -0,0 +1,410 @@
+// Copyright Epic Games, Inc. All Rights Reserved.
+
+#include "MuCOE/GenerateMutableSource/GenerateMutableSourceFloat.h"
+
+#include "Curves/CurveFloat.h"
+#include "Curves/CurveLinearColor.h"
+#include "Curves/CurveVector.h"
+#include "MuCOE/CustomizableObjectCompiler.h"
+#include "MuCOE/EdGraphSchema_CustomizableObject.h"
+#include "MuCOE/GenerateMutableSource/GenerateMutableSourceImage.h"
+#include "MuCOE/GenerateMutableSource/GenerateMutableSourceTable.h"
+#include "MuCOE/GraphTraversal.h"
+#include "MuCOE/Nodes/CustomizableObjectNodeCurve.h"
+#include "MuCOE/Nodes/CustomizableObjectNodeEnumParameter.h"
+#include "MuCOE/Nodes/CustomizableObjectNodeFloatConstant.h"
+#include "MuCOE/Nodes/CustomizableObjectNodeFloatParameter.h"
+#include "MuCOE/Nodes/CustomizableObjectNodeFloatSwitch.h"
+#include "MuCOE/Nodes/CustomizableObjectNodeFloatVariation.h"
+#include "MuCOE/Nodes/CustomizableObjectNodeFloatArithmeticOp.h"
+#include "MuCOE/Nodes/CustomizableObjectNodeTable.h"
+#include "MuT/NodeScalarConstant.h"
+#include "MuT/NodeScalarCurve.h"
+#include "MuT/NodeScalarSwitch.h"
+#include "MuT/NodeScalarTable.h"
+#include "MuT/NodeScalarVariation.h"
+#include "MuT/NodeScalarArithmeticOperation.h"
+
+#define LOCTEXT_NAMESPACE "CustomizableObjectEditor"
+
+
+
+mu::NodeScalarPtr GenerateMutableSourceFloat(const UEdGraphPin* Pin, FMutableGraphGenerationContext& GenerationContext)
+{
+	check(Pin)
+	RETURN_ON_CYCLE(*Pin, GenerationContext)
+
+	CheckNumOutputs(*Pin, GenerationContext);
+	
+	const UEdGraphSchema_CustomizableObject* Schema = GetDefault<UEdGraphSchema_CustomizableObject>();
+
+	UCustomizableObjectNode* Node = CastChecked<UCustomizableObjectNode>(Pin->GetOwningNode());
+
+	const FGeneratedKey Key(reinterpret_cast<void*>(&GenerateMutableSourceFloat), *Pin, *Node, GenerationContext);
+	if (const FGeneratedData* Generated = GenerationContext.Generated.Find(Key))
+	{
+		return static_cast<mu::NodeScalar*>(Generated->Node.get());
+	}
+
+	if (Node->IsNodeOutDatedAndNeedsRefresh())
+	{
+		Node->SetRefreshNodeWarning();
+	}
+
+	bool bDoNotAddToGeneratedCache = false;
+
+	mu::NodeScalarPtr Result;
+	
+	if (const UCustomizableObjectNodeFloatConstant* FloatConstantNode = Cast<UCustomizableObjectNodeFloatConstant>(Node))
+	{
+		mu::NodeScalarConstantPtr ScalarNode = new mu::NodeScalarConstant();
+		Result = ScalarNode;
+
+		ScalarNode->SetValue(FloatConstantNode->Value);
+	}
+
+	else if (const UCustomizableObjectNodeFloatParameter* FloatParameterNode = Cast<UCustomizableObjectNodeFloatParameter>(Node))
+	{
+		mu::NodeScalarParameterPtr ScalarNode = new mu::NodeScalarParameter();
+		Result = ScalarNode;
+
+		GenerationContext.AddParameterNameUnique(Node, FloatParameterNode->ParameterName);
+
+		ScalarNode->SetName(FloatParameterNode->ParameterName);
+		ScalarNode->SetUid(GenerationContext.GetNodeIdUnique(Node).ToString());
+		ScalarNode->SetDefaultValue(FloatParameterNode->DefaultValue);
+
+		GenerationContext.ParameterUIDataMap.Add(FloatParameterNode->ParameterName, FMutableParameterData(
+			FloatParameterNode->ParamUIMetadata,
+			EMutableParameterType::Float));
+	}
+
+	else if (const UCustomizableObjectNodeEnumParameter* EnumParamNode = Cast<UCustomizableObjectNodeEnumParameter>(Node))
+	{
+		mu::NodeScalarEnumParameterPtr EnumParameterNode = new mu::NodeScalarEnumParameter;
+
+		const int32 NumSelectors = EnumParamNode->Values.Num();
+
+		int32 DefaultValue = FMath::Clamp(EnumParamNode->DefaultIndex, 0, NumSelectors - 1);
+
+		GenerationContext.AddParameterNameUnique(Node, EnumParamNode->ParameterName);
+
+		EnumParameterNode->SetName(EnumParamNode->ParameterName);
+		EnumParameterNode->SetUid(GenerationContext.GetNodeIdUnique(Node).ToString());
+		EnumParameterNode->SetValueCount(NumSelectors);
+		EnumParameterNode->SetDefaultValueIndex(DefaultValue);
+
+		FMutableParameterData ParameterUIData(EnumParamNode->ParamUIMetadata, EMutableParameterType::Int);
+		ParameterUIData.IntegerParameterGroupType = ECustomizableObjectGroupType::COGT_ONE;
+
+		for (int SelectorIndex = 0; SelectorIndex < NumSelectors; ++SelectorIndex)
+		{
+			EnumParameterNode->SetValue(SelectorIndex, (float)SelectorIndex, EnumParamNode->Values[SelectorIndex].Name);
+
+			ParameterUIData.ArrayIntegerParameterOption.Add(
+				EnumParamNode->Values[SelectorIndex].Name,
+				FIntegerParameterUIData(EnumParamNode->Values[SelectorIndex].ParamUIMetadata));
+		}
+
+		Result = EnumParameterNode;
+
+		GenerationContext.ParameterUIDataMap.Add(EnumParamNode->ParameterName, ParameterUIData);
+	}
+
+	else if (const UCustomizableObjectNodeFloatSwitch* TypedNodeFloatSwitch = Cast<UCustomizableObjectNodeFloatSwitch>(Node))
+	{
+		// Using a lambda so control flow is easier to manage.
+		Result = [&]() -> mu::Ptr<mu::NodeScalar>
+		{
+			const UEdGraphPin* SwitchParameter = TypedNodeFloatSwitch->SwitchParameter();
+
+			// Check Switch Parameter arity preconditions.
+			if (const int32 NumParameters = FollowInputPinArray(*SwitchParameter).Num();
+				NumParameters != 1)
+			{
+				const FText Message = NumParameters
+					? LOCTEXT("NoEnumParamInSwitch", "Switch nodes must have an enum switch parameter. Please connect an enum and refesh the switch node.")
+					: LOCTEXT("InvalidEnumInSwitch", "Switch nodes must have a single enum with all the options inside. Please remove all the enums but one and refresh the switch node.");
+
+				GenerationContext.Compiler->CompilerLog(Message, Node);
+				return nullptr;
+			}
+
+			const UEdGraphPin* EnumPin = FollowInputPin(*SwitchParameter);
+			mu::NodeScalarPtr SwitchParam = GenerateMutableSourceFloat(EnumPin, GenerationContext);
+
+			// Switch Param not generated
+			if (!SwitchParam)
+			{
+				// Warn about a failure.
+				if (EnumPin)
+				{
+					const FText Message = LOCTEXT("FailedToGenerateSwitchParam", "Could not generate switch enum parameter. Please refesh the switch node and connect an enum.");
+					GenerationContext.Compiler->CompilerLog(Message, Node);
+				}
+
+				return nullptr;
+			}
+
+			if (SwitchParam->GetType() != mu::NodeScalarEnumParameter::GetStaticType())
+			{
+				const FText Message = LOCTEXT("WrongSwitchParamType", "Switch parameter of incorrect type.");
+				GenerationContext.Compiler->CompilerLog(Message, Node);
+
+				return nullptr;
+			}
+
+			const int32 NumSwitchOptions = TypedNodeFloatSwitch->GetNumElements();
+
+			mu::NodeScalarEnumParameter* EnumParameter = static_cast<mu::NodeScalarEnumParameter*>(SwitchParam.get());
+			if (NumSwitchOptions != EnumParameter->GetValueCount())
+			{
+				const FText Message = LOCTEXT("MismatchedSwitch", "Switch enum and switch node have different number of options. Please refresh the switch node to make sure the outcomes are labeled properly.");
+				GenerationContext.Compiler->CompilerLog(Message, Node);
+			}
+
+			mu::Ptr<mu::NodeScalarSwitch> SwitchNode = new mu::NodeScalarSwitch;
+			SwitchNode->Parameter = SwitchParam;
+			SwitchNode->Options.SetNum(NumSwitchOptions);
+
+			for (int SelectorIndex = 0; SelectorIndex < NumSwitchOptions; ++SelectorIndex)
+			{
+				if (const UEdGraphPin* const FloatPin = TypedNodeFloatSwitch->GetElementPin(SelectorIndex))
+				{
+					if (const UEdGraphPin* ConnectedPin = FollowInputPin(*FloatPin))
+					{
+						SwitchNode->Options[SelectorIndex] = GenerateMutableSourceFloat(ConnectedPin, GenerationContext);
+					}
+				}
+			}
+
+			return SwitchNode;
+		}(); // invoke lambda.
+	}
+
+	else if (const UCustomizableObjectNodeCurve* TypedNodeCurve = Cast<UCustomizableObjectNodeCurve>(Node))
+	{
+		mu::Ptr<mu::NodeScalarCurve> CurveNode = new mu::NodeScalarCurve();
+		Result = CurveNode;
+
+		if (const UEdGraphPin* ConnectedPin = FollowInputPin(*TypedNodeCurve->InputPin()))
+		{
+			CurveNode->CurveSampleValue = GenerateMutableSourceFloat(ConnectedPin, GenerationContext);
+		}
+
+		if (UCurveBase* CurveAsset = TypedNodeCurve->CurveAsset)
+		{
+			int32 PinIndex = -1;
+
+			for (int32 i = 0; i < TypedNodeCurve->GetNumCurvePins(); ++i)
+			{
+				if (TypedNodeCurve->CurvePins(i) == Pin)
+				{
+					PinIndex = i;
+					break;
+				}
+			}
+
+			if (const UCurveLinearColor* const CurveColor = Cast<UCurveLinearColor>(CurveAsset))
+			{
+				if (PinIndex >= 0 && PinIndex <= 3)
+				{
+					CurveNode->Curve = CurveColor->FloatCurves[PinIndex];
+				}
+			}
+			else if (const UCurveVector* const CurveVector = Cast<UCurveVector>(CurveAsset))
+			{
+				if (PinIndex >= 0 && PinIndex <= 2)
+				{
+					CurveNode->Curve = CurveVector->FloatCurves[PinIndex];
+				}
+			}
+			else if (const UCurveFloat* const CurveFloat = Cast<UCurveFloat>(CurveAsset))
+			{
+				if (PinIndex == 0)
+				{
+					CurveNode->Curve = CurveFloat->FloatCurve;
+				}
+			}
+		}
+	}
+
+	else if (const UCustomizableObjectNodeFloatVariation* TypedNodeFloatVar = Cast<const UCustomizableObjectNodeFloatVariation>(Node))
+	{
+		mu::NodeScalarVariationPtr FloatNode = new mu::NodeScalarVariation();
+		Result = FloatNode;
+
+		if (const UEdGraphPin* ConnectedPin = FollowInputPin(*TypedNodeFloatVar->DefaultPin()))
+		{
+			mu::NodeScalarPtr ChildNode = GenerateMutableSourceFloat(ConnectedPin, GenerationContext);
+			if (ChildNode)
+			{
+				FloatNode->SetDefaultScalar(ChildNode.get());
+			}
+			else
+			{
+				GenerationContext.Compiler->CompilerLog(LOCTEXT("FloatFailed", "Float generation failed."), Node);
+			}
+		}
+
+		const int32 NumVariations = TypedNodeFloatVar->GetNumVariations();
+		FloatNode->SetVariationCount(NumVariations);
+		for (int VariationIndex = 0; VariationIndex < NumVariations; ++VariationIndex)
+		{
+			UEdGraphPin* VariationPin = TypedNodeFloatVar->VariationPin(VariationIndex);
+			if (!VariationPin) continue;
+
+			FloatNode->SetVariationTag(VariationIndex, TypedNodeFloatVar->GetVariation(VariationIndex).Tag);
+			if (const UEdGraphPin* ConnectedPin = FollowInputPin(*VariationPin))
+			{
+				mu::NodeScalarPtr ChildNode = GenerateMutableSourceFloat(ConnectedPin, GenerationContext);
+				FloatNode->SetVariationScalar(VariationIndex, ChildNode.get());
+			}
+		}
+	}
+
+	else if (const UCustomizableObjectNodeFloatArithmeticOp* TypedNodeFloatArith = Cast<UCustomizableObjectNodeFloatArithmeticOp>(Node))
+	{
+		mu::NodeScalarArithmeticOperationPtr OpNode = new mu::NodeScalarArithmeticOperation();
+		Result = OpNode;
+
+		if (const UEdGraphPin* ConnectedPin = FollowInputPin(*TypedNodeFloatArith->XPin()))
+		{
+			mu::NodeScalarPtr XNode = GenerateMutableSourceFloat(ConnectedPin, GenerationContext);
+			OpNode->SetA(XNode);
+		}
+
+		if (const UEdGraphPin* ConnectedPin = FollowInputPin(*TypedNodeFloatArith->YPin()))
+		{
+			mu::NodeScalarPtr YNode = GenerateMutableSourceFloat(ConnectedPin, GenerationContext);
+			OpNode->SetB(YNode);
+		}
+
+		switch (TypedNodeFloatArith->Operation)
+		{
+		case EFloatArithmeticOperation::E_Add:
+			OpNode->SetOperation(mu::NodeScalarArithmeticOperation::OPERATION::AO_ADD);
+			break;
+
+		case EFloatArithmeticOperation::E_Sub:
+			OpNode->SetOperation(mu::NodeScalarArithmeticOperation::OPERATION::AO_SUBTRACT);
+			break;
+
+		case EFloatArithmeticOperation::E_Mul:
+			OpNode->SetOperation(mu::NodeScalarArithmeticOperation::OPERATION::AO_MULTIPLY);
+			break;
+
+		case EFloatArithmeticOperation::E_Div:
+			OpNode->SetOperation(mu::NodeScalarArithmeticOperation::OPERATION::AO_DIVIDE);
+			break;
+
+		default:
+			unimplemented();
+		}
+	}
+
+	else if (const UCustomizableObjectNodeTable* TypedNodeTable = Cast<UCustomizableObjectNodeTable>(Node))
+	{
+		//This node will add a default value in case of error
+		mu::NodeScalarConstantPtr ConstantValue = new mu::NodeScalarConstant();
+		ConstantValue->SetValue(1.0f);
+
+		Result = ConstantValue;
+
+		if (Pin->PinType.PinCategory == Schema->PC_MaterialAsset)
+		{
+			// Material pins have to skip the cache of nodes or they will return always the same column node
+			bDoNotAddToGeneratedCache = true;
+		}
+
+		bool bSuccess = true;
+		UDataTable* DataTable = GetDataTable(TypedNodeTable, GenerationContext);
+
+		if (DataTable)
+		{
+			FString ColumnName = Pin->PinFriendlyName.ToString();
+			FProperty* Property = DataTable->FindTableProperty(FName(*ColumnName));
+
+			if (!Property)
+			{
+				FString Msg = FString::Printf(TEXT("Couldn't find the column [%s] in the data table's struct."), *ColumnName);
+				GenerationContext.Compiler->CompilerLog(FText::FromString(Msg), Node);
+
+				bSuccess = false;
+			}
+
+			if (bSuccess)
+			{
+				// Generating a new data table if not exists
+				mu::TablePtr Table;
+				Table = GenerateMutableSourceTable(DataTable, TypedNodeTable, GenerationContext);
+
+				if (Table)
+				{
+					mu::NodeScalarTablePtr ScalarTableNode = new mu::NodeScalarTable();
+
+					if (Pin->PinType.PinCategory == Schema->PC_MaterialAsset)
+					{
+						// Materials use the parameter id as column names
+						ColumnName = GenerationContext.CurrentMaterialTableParameterId;
+					}
+
+					// Generating a new Float column if not exists
+					if (Table->FindColumn(ColumnName) == INDEX_NONE)
+					{
+						int32 Dummy = -1; // TODO MTBL-1512
+						bool Dummy2 = false;
+						bSuccess = GenerateTableColumn(TypedNodeTable, Pin, Table, ColumnName, Property, Dummy, Dummy, GenerationContext.CurrentLOD, Dummy, Dummy2, GenerationContext);
+
+						if (!bSuccess)
+						{
+							FString Msg = FString::Printf(TEXT("Failed to generate the mutable table column [%s]"), *ColumnName);
+							GenerationContext.Compiler->CompilerLog(FText::FromString(Msg), Node);
+						}
+					}
+
+					if (bSuccess)
+					{
+						Result = ScalarTableNode;
+
+						ScalarTableNode->Table = Table;
+						ScalarTableNode->SetColumn(ColumnName);
+						ScalarTableNode->SetParameterName(TypedNodeTable->ParameterName);
+						ScalarTableNode->SetNoneOption(TypedNodeTable->bAddNoneOption);
+						ScalarTableNode->SetDefaultRowName(TypedNodeTable->DefaultRowName.ToString());
+					}
+				}
+				else
+				{
+					FString Msg = FString::Printf(TEXT("Couldn't generate a mutable table."), *ColumnName);
+					GenerationContext.Compiler->CompilerLog(FText::FromString(Msg), Node);
+				}
+			}
+		}
+		else
+		{
+			GenerationContext.Compiler->CompilerLog(LOCTEXT("ScalarTableError", "Couldn't find the data table of the node."), Node);
+		}
+	}
+
+	else
+	{
+		GenerationContext.Compiler->CompilerLog(LOCTEXT("UnimplementedNode", "Node type not implemented yet."), Node);
+	}
+
+	if (!bDoNotAddToGeneratedCache)
+	{
+		GenerationContext.Generated.Add(Key, FGeneratedData(Node, Result));
+		GenerationContext.GeneratedNodes.Add(Node);
+	}
+
+	if (Result)
+	{
+		Result->SetMessageContext(Node);
+	}
+
+	return Result;
+}
+
+#undef LOCTEXT_NAMESPACE
+
diff --git a/Engine/Plugins/Mutable/Source/CustomizableObjectEditor/Private/MuCOE/GenerateMutableSource/GenerateMutableSourceFloat.h b/Engine/Plugins/Mutable/Source/CustomizableObjectEditor/Private/MuCOE/GenerateMutableSource/GenerateMutableSourceFloat.h
new file mode 100644
--- /dev/null
+++ b/Engine/Plugins/Mutable/Source/CustomizableObjectEditor/Private/MuCOE/GenerateMutableSource/GenerateMutableSourceFloat.h
@@ -0,0 +1,11 @@
+// Copyright Epic Games, Inc. All Rights Reserved.
+
+#pragma once
+
+#include "MuT/NodeScalar.h"
+
+class UEdGraphPin;
+struct FMutableGraphGenerationContext;
+
+/** Convert a CustomizableObject Source Graph into a mutable source graph. */
+mu::NodeScalarPtr GenerateMutableSourceFloat(const UEdGraphPin* Pin, FMutableGraphGenerationContext& GenerationContext);
diff --git a/Engine/Plugins/Mutable/Source/CustomizableObjectEditor/Private/MuCOE/GenerateMutableSource/GenerateMutableSourceGroupProjector.cpp b/Engine/Plugins/Mutable/Source/CustomizableObjectEditor/Private/MuCOE/GenerateMutableSource/GenerateMutableSourceGroupProjector.cpp
new file mode 100644
--- /dev/null
+++ b/Engine/Plugins/Mutable/Source/CustomizableObjectEditor/Private/MuCOE/GenerateMutableSource/GenerateMutableSourceGroupProjector.cpp
@@ -0,0 +1,545 @@
+// Copyright Epic Games, Inc. All Rights Reserved.
+
+#include "MuCOE/GenerateMutableSource/GenerateMutableSourceGroupProjector.h"
+
+#include "GenerateMutableSourceImage.h"
+#include "Materials/MaterialInterface.h"
+#include "MuCO/MultilayerProjector.h"
+#include "MuCOE/CustomizableObjectCompiler.h"
+#include "MuCOE/GraphTraversal.h"
+#include "MuCOE/GenerateMutableSource/GenerateMutableSourceProjector.h"
+#include "MuCOE/Nodes/CustomizableObjectNodeAnimationPose.h"
+#include "MuCOE/Nodes/CustomizableObjectNodeModifierExtendMeshSection.h"
+#include "MuCOE/Nodes/CustomizableObjectNodeGroupProjectorParameter.h"
+#include "MuCOE/Nodes/CustomizableObjectNodeMaterial.h"
+#include "MuCOE/Nodes/CustomizableObjectNodeObjectGroup.h"
+#include "MuT/NodeColourConstant.h"
+#include "MuT/NodeColourFromScalars.h"
+#include "MuT/NodeImageLayerColour.h"
+#include "MuT/NodeImageMultiLayer.h"
+#include "MuT/NodeImagePlainColour.h"
+#include "MuT/NodeImageProject.h"
+#include "MuT/NodeImageResize.h"
+#include "MuT/NodeImageSwitch.h"
+#include "MuT/NodeImageSwizzle.h"
+#include "MuT/NodeMeshSwitch.h"
+#include "MuT/NodeRangeFromScalar.h"
+#include "MuT/NodeScalarConstant.h"
+
+class UPoseAsset;
+
+#define LOCTEXT_NAMESPACE "CustomizableObjectEditor"
+
+
+mu::NodeImagePtr GenerateMutableGroupProjection(const int32 NodeLOD, const int32 ImageIndex, mu::NodeMeshPtr MeshNode, FMutableGraphGenerationContext& GenerationContext,
+	UCustomizableObjectNodeMaterialBase* TypedNodeMat, UCustomizableObjectNodeModifierExtendMeshSection* TypedNodeExt, bool& bShareProjectionTexturesBetweenLODs, bool& bIsGroupProjectorImage,
+	UTexture2D*& GroupProjectionReferenceTexture, TMap<FString, float>& TextureNameToProjectionResFactor, FString& AlternateResStateName)
+{
+	check(static_cast<bool>(TypedNodeMat) != static_cast<bool>(TypedNodeExt)); // XOr. TypedNodeMat valid or TypedNodeExt valid. At least one valid.
+
+	if (!GenerationContext.ProjectorGroupMap.Num() || !MeshNode.get())
+	{
+		return mu::NodeImagePtr();
+	}
+
+	TArray<mu::Ptr<mu::NodeImageProject>> ImageNodes;
+	TArray<FGroupProjectorTempData> ImageNodes_ProjectorTempData;
+
+	int32 TextureSize = 512;
+
+	for (TPair<const UCustomizableObjectNodeObjectGroup*, FGroupProjectorTempData>& Elem : GenerationContext.ProjectorGroupMap)
+	{
+		FGroupProjectorTempData& ProjectorTempData = Elem.Value;
+
+		const int32& DropProjectionTextureAtLOD = ProjectorTempData.CustomizableObjectNodeGroupProjectorParameter->DropProjectionTextureAtLOD;
+		if (DropProjectionTextureAtLOD >= 0 && NodeLOD >= DropProjectionTextureAtLOD)
+		{
+			continue;
+		}
+
+		bShareProjectionTexturesBetweenLODs |= ProjectorTempData.CustomizableObjectNodeGroupProjectorParameter->bShareProjectionTexturesBetweenLODs;
+
+		if (!GroupProjectionReferenceTexture)
+		{
+			GroupProjectionReferenceTexture = ProjectorTempData.CustomizableObjectNodeGroupProjectorParameter->ReferenceTexture;
+			
+			if (GroupProjectionReferenceTexture)
+			{
+				GenerationContext.AddParticipatingObject(*GroupProjectionReferenceTexture);
+			}
+		}
+
+		const bool bProjectToImage = [&]
+		{
+			FString ParameterName;
+			
+			if (TypedNodeMat)
+			{
+				ParameterName = TypedNodeMat->GetParameterName(EMaterialParameterType::Texture, ImageIndex).ToString();
+			}
+			else
+			{
+				const FNodeMaterialParameterId ImageId = TypedNodeExt->GetParameterId(EMaterialParameterType::Texture, ImageIndex);
+				return TypedNodeExt->UsesImage(ImageId);
+			}
+
+			return ParameterName == ProjectorTempData.CustomizableObjectNodeGroupProjectorParameter->MaterialChannelNameToConnect;
+		}();
+		
+		if (bProjectToImage)
+		{
+			const bool bWarningReplacedImage = [&]
+			{
+				if (TypedNodeMat)
+				{
+					return TypedNodeMat->IsImageMutableMode(ImageIndex);
+				}
+				else
+				{
+					const FNodeMaterialParameterId ImageId = TypedNodeExt->GetParameterId(EMaterialParameterType::Texture, ImageIndex);
+					return TypedNodeExt->UsesImage(ImageId);
+				}
+			}();
+			
+			if (bWarningReplacedImage)
+			{
+				const FString ImageName = [&]
+				{
+					if (TypedNodeMat)
+					{
+						return TypedNodeMat->GetParameterName(EMaterialParameterType::Texture, ImageIndex).ToString();
+					}
+					else
+					{
+						return TypedNodeExt->GetParameterName(EMaterialParameterType::Texture, ImageIndex).ToString();
+					}
+				}();
+				
+				FString msg = FString::Printf(TEXT("Material image [%s] is connected to an image but will be replaced by a Group Projector."), *ImageName);
+				GenerationContext.Compiler->CompilerLog(FText::FromString(msg), TypedNodeMat);
+				continue;
+			}
+			
+			mu::Ptr<mu::NodeImageProject> ImageNode = new mu::NodeImageProject();
+			bIsGroupProjectorImage = true;
+			ImageNode->SetLayout(ProjectorTempData.CustomizableObjectNodeGroupProjectorParameter->UVLayout);
+
+			{
+				if (TypedNodeMat &&
+					!ProjectorTempData.CustomizableObjectNodeGroupProjectorParameter->AlternateProjectionResolutionStateName.IsEmpty()
+					&& ProjectorTempData.CustomizableObjectNodeGroupProjectorParameter->AlternateProjectionResolutionFactor > 0)
+				{
+					TextureNameToProjectionResFactor.Add(
+						TypedNodeMat->GetParameterName(EMaterialParameterType::Texture, ImageIndex).ToString(),
+						ProjectorTempData.CustomizableObjectNodeGroupProjectorParameter->AlternateProjectionResolutionFactor);
+
+					if (!AlternateResStateName.IsEmpty() &&
+						AlternateResStateName != ProjectorTempData.CustomizableObjectNodeGroupProjectorParameter->AlternateProjectionResolutionStateName &&
+						!ProjectorTempData.bAlternateResStateNameWarningDisplayed)
+					{
+						FString msg = FString::Printf(TEXT("All 'Alternate Projection Resolution State Name' properties in Group Projector Parameter nodes connected to same Group node must have the same value or be blank. Only the value of the last connected node will be used."));
+						GenerationContext.Compiler->CompilerLog(FText::FromString(msg), ProjectorTempData.CustomizableObjectNodeGroupProjectorParameter);
+						ProjectorTempData.bAlternateResStateNameWarningDisplayed = true;
+					}
+
+					AlternateResStateName = ProjectorTempData.CustomizableObjectNodeGroupProjectorParameter->AlternateProjectionResolutionStateName;
+				}
+			}
+
+			{
+				mu::NodeScalarConstantPtr ScalarNode = new mu::NodeScalarConstant;
+				ScalarNode->SetValue(120.f);
+				ImageNode->SetAngleFadeStart(ScalarNode);
+			}
+
+			{
+				mu::NodeScalarConstantPtr ScalarNode = new mu::NodeScalarConstant;
+				ScalarNode->SetValue(150.f);
+				ImageNode->SetAngleFadeEnd(ScalarNode);
+			}
+
+			mu::NodeMeshSwitchPtr MeshSwitchNode = new mu::NodeMeshSwitch;
+			MeshSwitchNode->SetParameter(ProjectorTempData.PoseOptionsParameter);
+			MeshSwitchNode->SetOptionCount(ProjectorTempData.CustomizableObjectNodeGroupProjectorParameter->OptionPoses.Num() + 1);
+			MeshSwitchNode->SetOption(0, MeshNode);
+
+			for (int32 SelectorIndex = 0; SelectorIndex < ProjectorTempData.CustomizableObjectNodeGroupProjectorParameter->OptionPoses.Num(); ++SelectorIndex)
+			{
+				if (ProjectorTempData.CustomizableObjectNodeGroupProjectorParameter->OptionPoses[SelectorIndex].OptionPose)
+				{
+					mu::NodeMeshApplyPosePtr NodeMeshApplyPose = CreateNodeMeshApplyPose(GenerationContext, MeshNode,
+						ProjectorTempData.PoseBoneDataArray[SelectorIndex].ArrayBoneName, ProjectorTempData.PoseBoneDataArray[SelectorIndex].ArrayTransform);
+
+					if (!NodeMeshApplyPose)
+					{
+						FString msg = FString::Printf(TEXT("Couldn't get bone transform information from a Pose Asset."));
+						GenerationContext.Compiler->CompilerLog(FText::FromString(msg), TypedNodeMat);
+					}
+
+					MeshSwitchNode->SetOption(SelectorIndex + 1, NodeMeshApplyPose);
+				}
+				else
+				{
+					MeshSwitchNode->SetOption(SelectorIndex + 1, MeshNode);
+				}
+			}
+
+			ImageNode->SetMesh(MeshSwitchNode);
+			ImageNode->SetProjector(ProjectorTempData.NodeProjectorParameterPtr);
+			ImageNode->SetImage(ProjectorTempData.NodeImagePtr);
+
+			TextureSize = ProjectorTempData.TextureSize;
+			ImageNode->SetImageSize( FUintVector2(TextureSize, TextureSize) );
+
+			ImageNodes.Add(ImageNode);
+			ImageNodes_ProjectorTempData.Add(ProjectorTempData);
+		}
+	}
+
+	if (ImageNodes.Num() == 0)
+	{
+		return mu::NodeImagePtr();
+	}
+	
+	mu::Ptr<mu::NodeColourConstant> ZeroColorNode = new mu::NodeColourConstant();
+	ZeroColorNode->Value = FVector4f(0.f, 0.f, 0.f, 1.0f);
+
+	mu::NodeImagePlainColourPtr ZeroPlainColourNode = new mu::NodeImagePlainColour;
+	ZeroPlainColourNode->SetSize(TextureSize, TextureSize);
+	ZeroPlainColourNode->SetColour(ZeroColorNode);
+
+	mu::NodeImageSwizzlePtr ZeroChannelNode = new mu::NodeImageSwizzle;
+	ZeroChannelNode->SetFormat(mu::EImageFormat::IF_L_UBYTE);
+	ZeroChannelNode->SetSource(0, ZeroPlainColourNode);
+	ZeroChannelNode->SetSourceChannel(0, 2); // Just take a zeroed channel for the base alpha
+
+	mu::NodeScalarConstantPtr OneConstantNode = new mu::NodeScalarConstant;
+	OneConstantNode->SetValue(1.f);
+
+	mu::NodeImagePtr ResultAlpha = ZeroChannelNode;
+	mu::NodeImagePtr ResultImage = ZeroPlainColourNode;
+
+	for (int32 i = 0; i < ImageNodes.Num(); ++i)
+	{
+		if (i > 0) // Resize the projection texture if necessary after the first iteration
+		{
+			int32 NewTextureSize = ImageNodes_ProjectorTempData.IsValidIndex(i) ? ImageNodes_ProjectorTempData[i].CustomizableObjectNodeGroupProjectorParameter->ProjectionTextureSize : TextureSize;
+
+			if (NewTextureSize <= 0 || !FMath::IsPowerOfTwo(NewTextureSize))
+			{
+				NewTextureSize = TextureSize;
+			}
+
+			if (NewTextureSize != TextureSize)
+			{
+				TextureSize = NewTextureSize;
+
+				ZeroPlainColourNode = new mu::NodeImagePlainColour;
+				ZeroPlainColourNode->SetSize(TextureSize, TextureSize);
+				ZeroPlainColourNode->SetColour(ZeroColorNode);
+			}
+		}
+
+		mu::NodeImageSwizzlePtr ImageNodesAlphaChannelNode = new mu::NodeImageSwizzle;
+		ImageNodesAlphaChannelNode->SetFormat(mu::EImageFormat::IF_L_UBYTE);
+		ImageNodesAlphaChannelNode->SetSource(0, ImageNodes[i]);
+		ImageNodesAlphaChannelNode->SetSourceChannel(0, 3);
+
+		mu::Ptr<mu::NodeColourFromScalars> ColourFromScalars = new mu::NodeColourFromScalars;
+		ColourFromScalars->X = ImageNodes_ProjectorTempData[i].NodeOpacityParameter;
+		ColourFromScalars->Y = ImageNodes_ProjectorTempData[i].NodeOpacityParameter;
+		ColourFromScalars->Z = ImageNodes_ProjectorTempData[i].NodeOpacityParameter;
+		ColourFromScalars->W = OneConstantNode;
+
+		mu::NodeImageLayerColourPtr OpacityMultiLayerNode = new mu::NodeImageLayerColour;
+		OpacityMultiLayerNode->SetBlendType(mu::EBlendType::BT_MULTIPLY);
+		OpacityMultiLayerNode->SetColour(ColourFromScalars);
+		OpacityMultiLayerNode->SetBase(ImageNodesAlphaChannelNode);
+		//OpacityMultiLayerNode->SetMask(OneChannelNode); // No mask needed
+
+		mu::NodeImageSwizzlePtr MultiplySwizzleNode = new mu::NodeImageSwizzle;
+		MultiplySwizzleNode->SetFormat(mu::EImageFormat::IF_L_UBYTE);
+		MultiplySwizzleNode->SetSource(0, OpacityMultiLayerNode);
+		MultiplySwizzleNode->SetSourceChannel(0, 0);
+
+		mu::NodeImageMultiLayerPtr BaseAlphaMultiLayerNode = new mu::NodeImageMultiLayer;
+		BaseAlphaMultiLayerNode->SetRange(ImageNodes_ProjectorTempData[i].NodeRange);
+		BaseAlphaMultiLayerNode->SetBlendType(mu::EBlendType::BT_LIGHTEN);
+		BaseAlphaMultiLayerNode->SetBase(ResultAlpha);
+		BaseAlphaMultiLayerNode->SetBlended(MultiplySwizzleNode);
+		//BaseAlphaMultiLayerNode->SetMask(MultiplySwizzleNode); // No mask needed
+		ResultAlpha = BaseAlphaMultiLayerNode;
+
+		mu::NodeImageMultiLayerPtr BaseMultiLayerNode = new mu::NodeImageMultiLayer;
+		BaseMultiLayerNode->SetRange(ImageNodes_ProjectorTempData[i].NodeRange);
+		BaseMultiLayerNode->SetBlendType(mu::EBlendType::BT_BLEND);
+		BaseMultiLayerNode->SetBase(ResultImage);
+		BaseMultiLayerNode->SetBlended(ImageNodes[i]);
+		BaseMultiLayerNode->SetMask(MultiplySwizzleNode);
+		ResultImage = BaseMultiLayerNode;
+	}
+
+	mu::NodeImageSwizzlePtr SwizzleNodeR = new mu::NodeImageSwizzle;
+	SwizzleNodeR->SetFormat(mu::EImageFormat::IF_L_UBYTE);
+	SwizzleNodeR->SetSource(0, ResultImage);
+	SwizzleNodeR->SetSourceChannel(0, 0);
+
+	mu::NodeImageSwizzlePtr SwizzleNodeG = new mu::NodeImageSwizzle;
+	SwizzleNodeG->SetFormat(mu::EImageFormat::IF_L_UBYTE);
+	SwizzleNodeG->SetSource(0, ResultImage);
+	SwizzleNodeG->SetSourceChannel(0, 1);
+
+	mu::NodeImageSwizzlePtr SwizzleNodeB = new mu::NodeImageSwizzle;
+	SwizzleNodeB->SetFormat(mu::EImageFormat::IF_L_UBYTE);
+	SwizzleNodeB->SetSource(0, ResultImage);
+	SwizzleNodeB->SetSourceChannel(0, 2);
+
+	mu::NodeImageSwizzlePtr FinalSwizzleNode = new mu::NodeImageSwizzle;
+	FinalSwizzleNode->SetFormat(mu::EImageFormat::IF_RGBA_UBYTE);
+	FinalSwizzleNode->SetSource(0, SwizzleNodeR);
+	FinalSwizzleNode->SetSourceChannel(0, 0);
+	FinalSwizzleNode->SetSource(1, SwizzleNodeG);
+	FinalSwizzleNode->SetSourceChannel(1, 0);
+	FinalSwizzleNode->SetSource(2, SwizzleNodeB);
+	FinalSwizzleNode->SetSourceChannel(2, 0);
+	FinalSwizzleNode->SetSource(3, ResultAlpha);
+	FinalSwizzleNode->SetSourceChannel(3, 0);
+
+	return FinalSwizzleNode;
+}
+
+
+bool GenerateMutableSourceGroupProjector(const UEdGraphPin* Pin, FMutableGraphGenerationContext& GenerationContext, const UCustomizableObjectNodeObjectGroup* originalGroup)
+{
+	check(Pin)
+	check(originalGroup);
+	RETURN_ON_CYCLE(*Pin, GenerationContext)
+
+	CheckNumOutputs(*Pin, GenerationContext);
+	
+	UCustomizableObjectNode* Node = CastChecked<UCustomizableObjectNode>(Pin->GetOwningNode());
+
+	const FGeneratedKey Key(reinterpret_cast<void*>(&GenerateMutableSourceGroupProjector), *Pin, *Node, GenerationContext, true);
+	if (const FGeneratedData* Generated = GenerationContext.Generated.Find(Key))
+	{
+		return true;
+	}
+
+	mu::NodeProjectorPtr Result;
+	
+	if (UCustomizableObjectNodeGroupProjectorParameter* ProjParamNode = Cast<UCustomizableObjectNodeGroupProjectorParameter>(Node))
+	{
+		FGroupProjectorTempData GroupProjectorTempData;
+		// The static cast works here because it's already known to be a mu::NodeProjectorParameter* because of the UE5 Cast in the previous line
+		GroupProjectorTempData.NodeProjectorParameterPtr = static_cast<mu::NodeProjectorParameter*>(GenerateMutableSourceProjector(Pin, GenerationContext).get());
+
+		if (GroupProjectorTempData.NodeProjectorParameterPtr)
+		{
+			// Use the projector parameter uid + num to identify parameters derived from this node
+			FGuid NumLayersParamUid = Node->NodeGuid;
+			NumLayersParamUid.D += 1;
+			FGuid SelectedPoseParamUid = Node->NodeGuid;
+			SelectedPoseParamUid.D += 2;
+			FGuid OpacityParamUid = Node->NodeGuid;
+			OpacityParamUid.D += 3;
+			FGuid SelectedImageParamUid = Node->NodeGuid;
+			SelectedImageParamUid.D += 4;
+
+			// Add to UCustomizableObjectNodeGroupProjectorParameter::OptionTextures those textures that are present in
+			// UCustomizableObjectNodeGroupProjectorParameter::OptionTexturesDataTable avoiding any repeated element
+			TArray<FGroupProjectorParameterImage> ArrayOptionTexture = ProjParamNode->GetFinalOptionTexturesNoRepeat();
+
+			if ((ProjParamNode->OptionTexturesDataTable != nullptr) &&
+				(ProjParamNode->DataTableTextureColumnName.ToString().IsEmpty() || (ProjParamNode->DataTableTextureColumnName.ToString() == "None")))
+			{
+				FString msg = FString::Printf(TEXT("The group projection node has a table assigned to the Option Images Data Table property, but no column to read textures is specified at the Data Table Texture Column Name property."));
+				GenerationContext.Compiler->CompilerLog(FText::FromString(msg), ProjParamNode, EMessageSeverity::Error, true);
+			}
+
+			GroupProjectorTempData.CustomizableObjectNodeGroupProjectorParameter = ProjParamNode;
+
+			mu::NodeScalarParameterPtr NodeScalarParam = new mu::NodeScalarParameter;
+			FString NodeScalarParamName = ProjParamNode->ParameterName + FMultilayerProjector::NUM_LAYERS_PARAMETER_POSTFIX;
+			NodeScalarParam->SetName(NodeScalarParamName);
+			NodeScalarParam->SetUid(NumLayersParamUid.ToString());
+			GenerationContext.AddParameterNameUnique(originalGroup, NodeScalarParamName);
+			NodeScalarParam->SetDefaultValue(0.f);
+
+			GenerationContext.ParameterUIDataMap.Add(NodeScalarParamName,
+				FMutableParameterData(ProjParamNode->ParamUIMetadata, EMutableParameterType::Int));
+
+			mu::Ptr<mu::NodeRangeFromScalar> NodeRangeFromScalar = new mu::NodeRangeFromScalar;
+			NodeRangeFromScalar->SetSize(NodeScalarParam);
+			GroupProjectorTempData.NodeRange = NodeRangeFromScalar;
+			GroupProjectorTempData.NodeProjectorParameterPtr->SetRangeCount(1);
+			GroupProjectorTempData.NodeProjectorParameterPtr->SetRange(0, NodeRangeFromScalar);
+
+			
+			mu::NodeScalarEnumParameterPtr PoseEnumParameterNode = new mu::NodeScalarEnumParameter;
+			FString PoseNodeEnumParamName = ProjParamNode->ParameterName + FMultilayerProjector::POSE_PARAMETER_POSTFIX;
+			PoseEnumParameterNode->SetName(PoseNodeEnumParamName);
+			PoseEnumParameterNode->SetUid(SelectedPoseParamUid.ToString());
+			GenerationContext.AddParameterNameUnique(originalGroup, PoseNodeEnumParamName);
+			PoseEnumParameterNode->SetValueCount(ProjParamNode->OptionPoses.Num() + 1);
+			PoseEnumParameterNode->SetDefaultValueIndex(0);
+			GroupProjectorTempData.PoseOptionsParameter = PoseEnumParameterNode;
+
+			
+			GenerationContext.ParameterUIDataMap.Add(PoseNodeEnumParamName,
+				FMutableParameterData(ProjParamNode->ParamUIMetadata, EMutableParameterType::Int));
+			
+			mu::NodeScalarParameterPtr OpacityParameterNode = new mu::NodeScalarParameter;
+			FString OpacityParameterNodeName = ProjParamNode->ParameterName + FMultilayerProjector::OPACITY_PARAMETER_POSTFIX;
+			OpacityParameterNode->SetName(OpacityParameterNodeName);
+			OpacityParameterNode->SetUid(OpacityParamUid.ToString());
+			GenerationContext.AddParameterNameUnique(originalGroup, OpacityParameterNodeName);
+			OpacityParameterNode->SetDefaultValue(0.75f);
+			OpacityParameterNode->SetRangeCount(1);
+			OpacityParameterNode->SetRange(0, NodeRangeFromScalar);
+			FMutableParamUIMetadata OpacityMetadata = ProjParamNode->ParamUIMetadata;
+			OpacityMetadata.ObjectFriendlyName = FString("Opacity");
+			GroupProjectorTempData.NodeOpacityParameter = OpacityParameterNode;
+
+			{
+				FMutableParameterData ParameterUIData;
+				ParameterUIData.ParamUIMetadata = OpacityMetadata;
+				ParameterUIData.Type = EMutableParameterType::Float;
+			
+				GenerationContext.ParameterUIDataMap.Add(OpacityParameterNodeName, ParameterUIData);
+			}
+			
+			if (ArrayOptionTexture.Num() == 0)
+			{
+				FString msg = FString::Printf(TEXT("The group projection node must have at least one option image connected to a texture or at least one valid element in Option Images Data Table."));
+				GenerationContext.Compiler->CompilerLog(FText::FromString(msg), ProjParamNode, EMessageSeverity::Error, true);
+				return false;
+			}
+			
+			if (FMutableComponentInfo* CurrentComponentInfo = GenerationContext.GetCurrentComponentInfo())
+			{
+				// Poses will only affect component 0 of the CO,
+				// TODO UE-206803
+				FName OldCurrentMeshComponent = GenerationContext.CurrentMeshComponent;
+				GenerationContext.CurrentMeshComponent = GenerationContext.ComponentInfos[0].ComponentName;
+
+				PoseEnumParameterNode->SetValue(0, 0.f, "Default pose");
+
+				for (int32 PoseIndex = 0; PoseIndex < ProjParamNode->OptionPoses.Num(); ++PoseIndex)
+				{
+					PoseEnumParameterNode->SetValue(PoseIndex + 1, (float)PoseIndex + 1.f, ProjParamNode->OptionPoses[PoseIndex].PoseName);
+
+					TArray<FString> ArrayBoneName;
+					TArray<FTransform> ArrayTransform;
+					UPoseAsset* PoseAsset = ProjParamNode->OptionPoses[PoseIndex].OptionPose;
+
+					if (PoseAsset == nullptr) // Check if the slot has a selected pose. Could be left empty by the user
+					{
+						FString msg = FString::Printf(TEXT("The group projection node must have a pose assigned on each Option Poses element."));
+						GenerationContext.Compiler->CompilerLog(FText::FromString(msg), ProjParamNode, EMessageSeverity::Error, true);
+						return false;
+					}
+
+					check(GroupProjectorTempData.PoseBoneDataArray.Num() == PoseIndex);
+					GroupProjectorTempData.PoseBoneDataArray.AddDefaulted(1);
+					UCustomizableObjectNodeAnimationPose::StaticRetrievePoseInformation(PoseAsset, CurrentComponentInfo->RefSkeletalMesh,
+						GroupProjectorTempData.PoseBoneDataArray[PoseIndex].ArrayBoneName, GroupProjectorTempData.PoseBoneDataArray[PoseIndex].ArrayTransform);
+				}
+
+				GenerationContext.CurrentMeshComponent = OldCurrentMeshComponent;
+			}
+		
+			mu::NodeScalarEnumParameterPtr EnumParameterNode = new mu::NodeScalarEnumParameter;
+			FString NodeEnumParamName = ProjParamNode->ParameterName + FMultilayerProjector::IMAGE_PARAMETER_POSTFIX;
+			EnumParameterNode->SetName(NodeEnumParamName);
+			EnumParameterNode->SetUid(SelectedImageParamUid.ToString());
+			GenerationContext.AddParameterNameUnique(originalGroup, NodeEnumParamName);
+			EnumParameterNode->SetValueCount(ArrayOptionTexture.Num());
+			EnumParameterNode->SetDefaultValueIndex(0);
+			EnumParameterNode->SetRangeCount(1);
+			EnumParameterNode->SetRange(0, NodeRangeFromScalar);
+
+			{
+				FMutableParameterData ParameterUIData;
+				ParameterUIData.ParamUIMetadata = ProjParamNode->ParamUIMetadata;
+				ParameterUIData.Type = EMutableParameterType::Int;
+				ParameterUIData.IntegerParameterGroupType = ECustomizableObjectGroupType::COGT_ONE;
+				ParameterUIData.ParamUIMetadata.ExtraInformation.Add(FString("UseThumbnails"));
+
+				for (int ImageIndex = 0; ImageIndex < ArrayOptionTexture.Num(); ++ImageIndex)
+				{
+					EnumParameterNode->SetValue(ImageIndex, (float)ImageIndex, ArrayOptionTexture[ImageIndex].OptionName);
+
+					FMutableParamUIMetadata OptionMetadata = ParameterUIData.ParamUIMetadata;
+					OptionMetadata.UIThumbnail = ArrayOptionTexture[ImageIndex].OptionTexture;
+					ParameterUIData.ArrayIntegerParameterOption.Add(
+						ArrayOptionTexture[ImageIndex].OptionName,
+						FIntegerParameterUIData(OptionMetadata));
+				}
+
+				GenerationContext.ParameterUIDataMap.Add(NodeEnumParamName, ParameterUIData);
+			}
+
+			mu::NodeImageSwitchPtr SwitchNode = new mu::NodeImageSwitch;
+			SwitchNode->SetParameter(EnumParameterNode);
+			SwitchNode->SetOptionCount(ArrayOptionTexture.Num());
+
+			const uint32 AdditionalLODBias = GenerationContext.Options.bUseLODAsBias ? GenerationContext.FirstLODAvailable : 0;
+			for (int32 SelectorIndex = 0; SelectorIndex < ArrayOptionTexture.Num(); ++SelectorIndex)
+			{
+				if (const TObjectPtr<UTexture2D>& Texture = ArrayOptionTexture[SelectorIndex].OptionTexture)
+				{
+					mu::Ptr<mu::Image> ImageConstant = GenerateImageConstant(ArrayOptionTexture[SelectorIndex].OptionTexture, GenerationContext, false);
+
+					mu::Ptr<mu::NodeImageConstant> ImageNode = new mu::NodeImageConstant();
+					ImageNode->SetValue(ImageConstant.get());
+
+					const uint32 MipsToSkip = ComputeLODBiasForTexture(GenerationContext, *Texture, ProjParamNode->ReferenceTexture) + AdditionalLODBias;
+					SwitchNode->SetOption(SelectorIndex, ResizeTextureByNumMips(ImageNode, MipsToSkip));
+
+					// Calculate the number of mips to tag as high res for this image.
+					if (ProjParamNode->ReferenceTexture)
+					{
+						int32 TotalMips = mu::Image::GetMipmapCount(ProjParamNode->ReferenceTexture->GetSizeX(), ProjParamNode->ReferenceTexture->GetSizeY());
+						int32 NumMipsBeyondMin = FMath::Max(0, TotalMips - GenerationContext.Options.MinDiskMips);
+						int32 HighResMipsForThisImage = FMath::Min(NumMipsBeyondMin, GenerationContext.Options.NumHighResImageMips);
+						ImageNode->SourceDataDescriptor.SourceHighResMips = HighResMipsForThisImage;
+					}
+				}
+				else
+				{
+					FString msg = FString::Printf(TEXT("The group projection node must have a texture for all the options. Please set a texture for all the options."));
+					GenerationContext.Compiler->CompilerLog(FText::FromString(msg), ProjParamNode);
+				}
+			}
+
+			int32 TextureSize = ProjParamNode->ProjectionTextureSize > 0 ? ProjParamNode->ProjectionTextureSize : 512;
+
+			// If TextureSize is not power of two, round up to the next power of two 
+			if (!FMath::IsPowerOfTwo(TextureSize))
+			{
+				TextureSize = FMath::RoundUpToPowerOfTwo(TextureSize);
+			}
+
+			// Apply additional LODBias if necessary
+			GroupProjectorTempData.TextureSize = FMath::Max(TextureSize >> AdditionalLODBias, 1);
+
+			GroupProjectorTempData.NodeImagePtr = SwitchNode;
+
+			GenerationContext.ProjectorGroupMap.Add(originalGroup, GroupProjectorTempData);
+		}
+		else
+		{
+			GenerationContext.Compiler->CompilerLog(FText::FromString(TEXT("Error getting group projection properties.")), ProjParamNode, EMessageSeverity::Error, true);
+			return false;
+		}
+	}
+	
+
+	GenerationContext.Generated.Add(Key, FGeneratedData(Node, Result));
+	GenerationContext.GeneratedNodes.Add(Node);
+
+	return true;
+}
+
+#undef LOCTEXT_NAMESPACE
+
diff --git a/Engine/Plugins/Mutable/Source/CustomizableObjectEditor/Private/MuCOE/GenerateMutableSource/GenerateMutableSourceGroupProjector.h b/Engine/Plugins/Mutable/Source/CustomizableObjectEditor/Private/MuCOE/GenerateMutableSource/GenerateMutableSourceGroupProjector.h
new file mode 100644
--- /dev/null
+++ b/Engine/Plugins/Mutable/Source/CustomizableObjectEditor/Private/MuCOE/GenerateMutableSource/GenerateMutableSourceGroupProjector.h
@@ -0,0 +1,24 @@
+// Copyright Epic Games, Inc. All Rights Reserved.
+
+#pragma once
+
+#include "Containers/Map.h"
+#include "MuT/NodeSurfaceNew.h"
+
+class UCustomizableObjectNodeModifierExtendMeshSection;
+class FString;
+class UCustomizableObjectNodeMaterial;
+class UCustomizableObjectNodeMaterialBase;
+class UCustomizableObjectNodeObjectGroup;
+class UEdGraphPin;
+class UTexture2D;
+struct FMutableGraphGenerationContext;
+
+
+mu::NodeImagePtr GenerateMutableGroupProjection(const int32 NodeLOD, const int32 ImageIndex, mu::NodeMeshPtr MeshNode, FMutableGraphGenerationContext& GenerationContext,
+	UCustomizableObjectNodeMaterialBase* TypedNodeMat, UCustomizableObjectNodeModifierExtendMeshSection* TypedNodeExt, bool& bShareProjectionTexturesBetweenLODs, bool& bIsGroupProjectorImage,
+	UTexture2D*& GroupProjectionReferenceTexture, TMap<FString, float>& TextureNameToProjectionResFactor, FString& AlternateResStateName);
+
+
+/** Convert a CustomizableObject Source Graph into a mutable source graph. */
+bool GenerateMutableSourceGroupProjector(const UEdGraphPin* Pin, FMutableGraphGenerationContext& GenerationContext, const UCustomizableObjectNodeObjectGroup* originalGroup);
