[mutable] Moved the Mutable plugin out of Experimental status into Beta.

diff --git a/Engine/Plugins/Mutable/Source/CustomizableObjectEditor/Private/MuCOE/GenerateMutableSource/GenerateMutableSourceImage.cpp b/Engine/Plugins/Mutable/Source/CustomizableObjectEditor/Private/MuCOE/GenerateMutableSource/GenerateMutableSourceImage.cpp
new file mode 100644
--- /dev/null
+++ b/Engine/Plugins/Mutable/Source/CustomizableObjectEditor/Private/MuCOE/GenerateMutableSource/GenerateMutableSourceImage.cpp
@@ -0,0 +1,1154 @@
+// Copyright Epic Games, Inc. All Rights Reserved.
+
+#include "MuCOE/GenerateMutableSource/GenerateMutableSourceImage.h"
+
+#include "MuCOE/CustomizableObjectCompiler.h"
+#include "MuCOE/EdGraphSchema_CustomizableObject.h"
+#include "MuCOE/GenerateMutableSource/GenerateMutableSourceColor.h"
+#include "MuCOE/GenerateMutableSource/GenerateMutableSourceFloat.h"
+#include "MuCOE/GenerateMutableSource/GenerateMutableSourceMesh.h"
+#include "MuCOE/GenerateMutableSource/GenerateMutableSourceProjector.h"
+#include "MuCOE/GenerateMutableSource/GenerateMutableSourceTable.h"
+#include "MuCOE/GraphTraversal.h"
+#include "MuCOE/Nodes/CustomizableObjectNodeMesh.h"
+#include "MuCOE/Nodes/CustomizableObjectNodeTable.h"
+#include "MuCOE/Nodes/CustomizableObjectNodeTexture.h"
+#include "MuCOE/Nodes/CustomizableObjectNodePassThroughTexture.h"
+#include "MuCOE/Nodes/CustomizableObjectNodeTextureBinarise.h"
+#include "MuCOE/Nodes/CustomizableObjectNodeTextureColourMap.h"
+#include "MuCOE/Nodes/CustomizableObjectNodeTextureFromChannels.h"
+#include "MuCOE/Nodes/CustomizableObjectNodeTextureFromColor.h"
+#include "MuCOE/Nodes/CustomizableObjectNodeTextureFromFloats.h"
+#include "MuCOE/Nodes/CustomizableObjectNodeTextureInterpolate.h"
+#include "MuCOE/Nodes/CustomizableObjectNodeTextureInvert.h"
+#include "MuCOE/Nodes/CustomizableObjectNodeTextureLayer.h"
+#include "MuCOE/Nodes/CustomizableObjectNodeTextureParameter.h"
+#include "MuCOE/Nodes/CustomizableObjectNodeTextureProject.h"
+#include "MuCOE/Nodes/CustomizableObjectNodeTextureSwitch.h"
+#include "MuCOE/Nodes/CustomizableObjectNodePassThroughTextureSwitch.h"
+#include "MuCOE/Nodes/CustomizableObjectNodeTextureToChannels.h"
+#include "MuCOE/Nodes/CustomizableObjectNodeTextureTransform.h"
+#include "MuCOE/Nodes/CustomizableObjectNodeTextureSaturate.h"
+#include "MuCOE/Nodes/CustomizableObjectNodeTextureVariation.h"
+#include "MuCOE/UnrealEditorPortabilityHelpers.h"
+#include "MuCO/UnrealToMutableTextureConversionUtils.h"
+#include "MuT/NodeImageBinarise.h"
+#include "MuT/NodeImageColourMap.h"
+#include "MuT/NodeImageFormat.h"
+#include "MuT/NodeImageInterpolate.h"
+#include "MuT/NodeImageInvert.h"
+#include "MuT/NodeImageLayer.h"
+#include "MuT/NodeImageLayerColour.h"
+#include "MuT/NodeImageParameter.h"
+#include "MuT/NodeImagePlainColour.h"
+#include "MuT/NodeColourFromScalars.h"
+#include "MuT/NodeImageProject.h"
+#include "MuT/NodeImageResize.h"
+#include "MuT/NodeImageSaturate.h"
+#include "MuT/NodeImageSwitch.h"
+#include "MuT/NodeImageSwizzle.h"
+#include "MuT/NodeImageTable.h"
+#include "MuT/NodeImageTransform.h"
+#include "MuT/NodeImageVariation.h"
+#include "MuT/NodeImageConstant.h"
+#include "MuT/NodeScalarConstant.h"
+
+#define LOCTEXT_NAMESPACE "CustomizableObjectEditor"
+
+
+mu::Ptr<mu::NodeImage> ResizeTextureByNumMips(const mu::Ptr<mu::NodeImage>& ImageConstant, int32 MipsToSkip)
+{
+	if (MipsToSkip > 0)
+	{
+		mu::NodeImageResizePtr ImageResize = new mu::NodeImageResize();
+		ImageResize->SetBase(ImageConstant.get());
+		ImageResize->SetRelative(true);
+		const float factor = FMath::Pow(0.5f, MipsToSkip);
+		ImageResize->SetSize(factor, factor);
+		return ImageResize;
+	}
+
+	return ImageConstant;
+}
+
+
+mu::NodeImagePtr GenerateMutableSourceImage(const UEdGraphPin* Pin, FMutableGraphGenerationContext& GenerationContext, int32 ReferenceTextureSize)
+{
+	check(Pin)
+	RETURN_ON_CYCLE(*Pin, GenerationContext)
+
+	CheckNumOutputs(*Pin, GenerationContext);
+	
+	const UEdGraphSchema_CustomizableObject* Schema = GetDefault<UEdGraphSchema_CustomizableObject>();
+
+	UCustomizableObjectNode* Node = CastChecked<UCustomizableObjectNode>(Pin->GetOwningNode());
+
+	const FGeneratedKey Key(reinterpret_cast<void*>(&GenerateMutableSourceImage), *Pin, *Node, GenerationContext, true);
+	if (const FGeneratedData* Generated = GenerationContext.Generated.Find(Key))
+	{
+		return static_cast<mu::NodeImage*>(Generated->Node.get());
+	}
+
+	if (Node->IsNodeOutDatedAndNeedsRefresh())
+	{
+		Node->SetRefreshNodeWarning();
+	}
+
+	bool bDoNotAddToGeneratedCache = false;
+
+	mu::NodeImagePtr Result;
+	
+	if (const UCustomizableObjectNodeTexture* TypedNodeTex = Cast<UCustomizableObjectNodeTexture>(Node))
+	{
+		UTexture2D* BaseTexture = TypedNodeTex->Texture;
+		if (BaseTexture)
+		{
+			GenerationContext.AddParticipatingObject(*BaseTexture);
+			
+			// Check the specific image cache
+			FGeneratedImageKey ImageKey = FGeneratedImageKey(Pin);
+			mu::NodeImagePtr ImageNode;
+			
+			const uint32 MipsToSkip = ComputeLODBiasForTexture(GenerationContext, *BaseTexture, nullptr, ReferenceTextureSize);
+
+			if (mu::NodeImagePtr* Cached = GenerationContext.GeneratedImages.Find(ImageKey))
+			{
+				ImageNode = *Cached;
+			}
+			else
+			{
+				mu::Ptr<mu::Image> ImageConstant = GenerateImageConstant(BaseTexture, GenerationContext, false);
+
+				mu::Ptr<mu::NodeImageConstant> ConstantImageNode = new mu::NodeImageConstant;
+				ConstantImageNode->SetValue( ImageConstant.get() );
+				
+				int32 TotalMips = mu::Image::GetMipmapCount(ImageConstant->GetSizeX(), ImageConstant->GetSizeY());
+				int32 NumMipsBeyondMin = FMath::Max(0, TotalMips - int32(MipsToSkip) - GenerationContext.Options.MinDiskMips);
+				int32 HighResMipsForThisImage = FMath::Min(NumMipsBeyondMin, GenerationContext.Options.NumHighResImageMips);
+				ConstantImageNode->SourceDataDescriptor.SourceHighResMips = HighResMipsForThisImage;
+
+				ImageNode = ConstantImageNode;
+
+				GenerationContext.GeneratedImages.Add(ImageKey, ImageNode);
+			}
+
+			Result = ResizeTextureByNumMips(ImageNode, MipsToSkip);
+		}
+		else
+		{
+			GenerationContext.Compiler->CompilerLog(LOCTEXT("MissingImage", "Missing image in texture node."), Node, EMessageSeverity::Warning);
+		}
+	}
+
+	else if (const UCustomizableObjectNodeTextureParameter* TypedNodeParam = Cast<UCustomizableObjectNodeTextureParameter>(Node))
+	{
+		mu::Ptr<mu::NodeImageParameter> TextureNode = new mu::NodeImageParameter();
+
+		GenerationContext.AddParameterNameUnique(Node, TypedNodeParam->ParameterName);
+
+		TextureNode->SetName(TypedNodeParam->ParameterName);
+		TextureNode->SetUid(GenerationContext.GetNodeIdUnique(Node).ToString());
+
+		if (TypedNodeParam->DefaultValue)
+		{
+			TextureNode->SetDefaultValue(FName(TypedNodeParam->DefaultValue->GetPathName()));			
+		}
+
+		GenerationContext.ParameterUIDataMap.Add(TypedNodeParam->ParameterName, FMutableParameterData(
+			TypedNodeParam->ParamUIMetadata,
+			EMutableParameterType::Texture));
+
+		// Force the same format that the default texture if any.
+		mu::NodeImageFormatPtr FormatNode = new mu::NodeImageFormat();
+		// \TODO: Take if from default?
+		//if (TypedNodeParam->DefaultValue)
+		//{
+		//	FormatNode->SetFormat(mu::EImageFormat::IF_RGBA_UBYTE);
+		//}
+		//else
+		{
+			// Force an "easy format" on the texture.
+			FormatNode->SetFormat(mu::EImageFormat::IF_RGBA_UBYTE);
+		}
+		FormatNode->SetSource(TextureNode);
+
+		mu::NodeImageResizePtr ResizeNode = new mu::NodeImageResize();
+		ResizeNode->SetBase(FormatNode);
+		ResizeNode->SetRelative(false);
+
+		FUintVector2 TextureSize(TypedNodeParam->TextureSizeX, TypedNodeParam->TextureSizeY);
+
+		const UTexture2D* ReferenceTexture = TypedNodeParam->ReferenceValue;
+		if (ReferenceTexture)
+		{
+			GenerationContext.AddParticipatingObject(*ReferenceTexture);
+
+			const uint32 LODBias = ComputeLODBiasForTexture(GenerationContext, *TypedNodeParam->ReferenceValue, ReferenceTexture, ReferenceTextureSize);
+			TextureSize.X = FMath::Max(ReferenceTexture->Source.GetSizeX() >> LODBias, 1);
+			TextureSize.Y = FMath::Max(ReferenceTexture->Source.GetSizeY() >> LODBias, 1);
+		}
+		else
+		{
+			const int32 MaxNodeTextureSize = FMath::Max(TypedNodeParam->TextureSizeX, TypedNodeParam->TextureSizeY);
+			if (MaxNodeTextureSize <= 0)
+			{
+				TextureSize.X = TextureSize.Y = 1;
+				GenerationContext.Compiler->CompilerLog(LOCTEXT("TextureParameterSize0", "Texture size not specified. Add a reference texture or set a valid value to the Texture Size variables."), Node);
+			}
+			else if (ReferenceTextureSize > 0 && ReferenceTextureSize < MaxNodeTextureSize)
+			{
+				const uint32 MipsToSkip = FMath::CeilLogTwo(MaxNodeTextureSize) - FMath::CeilLogTwo(ReferenceTextureSize);
+				TextureSize.X = FMath::Max(TextureSize.X >> MipsToSkip, (uint32)1);
+				TextureSize.Y = FMath::Max(TextureSize.Y >> MipsToSkip, (uint32)1);
+			}
+		}
+
+		ResizeNode->SetSize(TextureSize.X, TextureSize.Y);
+
+		Result = ResizeNode;
+	}
+
+	else if (const UCustomizableObjectNodeMesh* TypedNodeMesh = Cast<UCustomizableObjectNodeMesh>(Node))
+	{
+		mu::Ptr<mu::NodeImageConstant> ImageNode = new mu::NodeImageConstant();
+		Result = ImageNode;
+
+		UTexture2D* Texture = TypedNodeMesh->FindTextureForPin(Pin);
+
+		if (Texture)
+		{
+			mu::Ptr<mu::Image> ImageConstant = GenerateImageConstant(Texture, GenerationContext, false);
+			ImageNode->SetValue(ImageConstant.get());
+
+			const uint32 MipsToSkip = ComputeLODBiasForTexture(GenerationContext, *Texture, nullptr, ReferenceTextureSize);
+			Result = ResizeTextureByNumMips(ImageNode, MipsToSkip);
+
+			int32 TotalMips = mu::Image::GetMipmapCount(ImageConstant->GetSizeX(), ImageConstant->GetSizeY());
+			int32 NumMipsBeyondMin = FMath::Max(0, TotalMips - int32(MipsToSkip) - GenerationContext.Options.MinDiskMips);
+			int32 HighResMipsForThisImage = FMath::Min(NumMipsBeyondMin, GenerationContext.Options.NumHighResImageMips);
+			ImageNode->SourceDataDescriptor.SourceHighResMips = HighResMipsForThisImage;
+		}
+		else
+		{
+			Result = mu::NodeImagePtr();
+		}
+	}
+
+	else if (const UCustomizableObjectNodeTextureInterpolate* TypedNodeInterp = Cast<UCustomizableObjectNodeTextureInterpolate>(Node))
+	{
+		mu::NodeImageInterpolatePtr ImageNode = new mu::NodeImageInterpolate();
+		Result = ImageNode;
+
+		if (const UEdGraphPin* ConnectedPin = FollowInputPin(*TypedNodeInterp->FactorPin()))
+		{
+			mu::NodeScalarPtr FactorNode = GenerateMutableSourceFloat(ConnectedPin, GenerationContext);
+			ImageNode->SetFactor(FactorNode);
+		}
+
+		int currentTarget = 0;
+		for (int LayerIndex = 0; LayerIndex < TypedNodeInterp->GetNumTargets(); ++LayerIndex)
+		{
+			if (const UEdGraphPin* ConnectedPin = FollowInputPin(*TypedNodeInterp->Targets(LayerIndex)))
+			{
+				mu::NodeImagePtr TargetNode = GenerateMutableSourceImage(ConnectedPin, GenerationContext, ReferenceTextureSize);
+
+				if (TargetNode)
+				{
+					ImageNode->SetTargetCount(currentTarget + 1);
+					ImageNode->SetTarget(currentTarget, TargetNode);
+					++currentTarget;
+				}
+			}
+		}
+	}
+
+	else if (const UCustomizableObjectNodeTextureLayer* TypedNodeLayer = Cast<UCustomizableObjectNodeTextureLayer>(Node))
+	{
+		if (const UEdGraphPin* ConnectedPin = FollowInputPin(*TypedNodeLayer->BasePin()))
+		{
+			Result = GenerateMutableSourceImage(ConnectedPin, GenerationContext, ReferenceTextureSize);
+		}
+
+		for (int LayerIndex = 0; LayerIndex < TypedNodeLayer->GetNumLayers(); ++LayerIndex)
+		{
+			if (const UEdGraphPin* OtherPin = FollowInputPin(*TypedNodeLayer->LayerPin(LayerIndex)))
+			{
+				mu::NodeImagePtr MaskNode = nullptr;
+				if (const UEdGraphPin* ConnectedPin = FollowInputPin(*TypedNodeLayer->MaskPin(LayerIndex)))
+				{
+					MaskNode = GenerateMutableSourceImage(ConnectedPin, GenerationContext, ReferenceTextureSize);
+				}
+
+				mu::EBlendType Type = mu::EBlendType::BT_BLEND;
+				switch (TypedNodeLayer->Layers[LayerIndex].Effect)
+				{
+				case COTLE_MODULATE: Type = mu::EBlendType::BT_BLEND; break;
+				case COTLE_MULTIPLY: Type = mu::EBlendType::BT_MULTIPLY; break;
+				case COTLE_SOFTLIGHT: Type = mu::EBlendType::BT_SOFTLIGHT; break;
+				case COTLE_HARDLIGHT: Type = mu::EBlendType::BT_HARDLIGHT; break;
+				case COTLE_DODGE: Type = mu::EBlendType::BT_DODGE; break;
+				case COTLE_BURN: Type = mu::EBlendType::BT_BURN; break;
+				case COTLE_SCREEN: Type = mu::EBlendType::BT_SCREEN; break;
+				case COTLE_OVERLAY: Type = mu::EBlendType::BT_OVERLAY; break;
+				case COTLE_ALPHA_OVERLAY: Type = mu::EBlendType::BT_LIGHTEN; break;
+				case COTLE_NORMAL_COMBINE: Type = mu::EBlendType::BT_NORMAL_COMBINE; break;
+				default:
+					GenerationContext.Compiler->CompilerLog(LOCTEXT("UnsupportedImageEffect", "Texture layer effect not supported. Setting to 'Blend'."), Node);
+					break;
+				}
+
+				if (Type == mu::EBlendType::BT_BLEND && !MaskNode)
+				{
+					GenerationContext.Compiler->CompilerLog(LOCTEXT("ModulateWithoutMask", "Texture layer effect uses Modulate without a mask. It will replace everything below it!"), Node);
+				}
+				
+				if (OtherPin->PinType.PinCategory == Schema->PC_Image)
+				{
+					mu::NodeImagePtr BlendNode = GenerateMutableSourceImage(OtherPin, GenerationContext, ReferenceTextureSize);
+
+					mu::NodeImageLayerPtr LayerNode = new mu::NodeImageLayer;
+					LayerNode->SetBlendType(Type);
+					LayerNode->SetBase(Result);
+					LayerNode->SetBlended(BlendNode);
+					LayerNode->SetMask(MaskNode);
+					Result = LayerNode;
+				}
+
+				else if (OtherPin->PinType.PinCategory == Schema->PC_Color)
+				{
+					mu::NodeColourPtr ColorNode = GenerateMutableSourceColor(OtherPin, GenerationContext);
+
+					mu::NodeImageLayerColourPtr LayerNode = new mu::NodeImageLayerColour;
+					LayerNode->SetBlendType(Type);
+					LayerNode->SetBase(Result);
+					LayerNode->SetColour(ColorNode);
+					LayerNode->SetMask(MaskNode);
+					Result = LayerNode;
+				}
+
+				// We need it here because we create multiple nodes.
+				Result->SetMessageContext(Node);
+			}
+		}
+	}
+
+	else if (const UCustomizableObjectNodeTextureSwitch* TypedNodeTextureSwitch = Cast<UCustomizableObjectNodeTextureSwitch>(Node))
+	{
+		Result = [&]()
+		{
+			const UEdGraphPin* SwitchParameter = TypedNodeTextureSwitch->SwitchParameter();
+
+			// Check Switch Parameter arity preconditions.
+			if (const UEdGraphPin* ConnectedPin = FollowInputPin(*SwitchParameter))
+			{
+				mu::NodeScalarPtr SwitchParam = GenerateMutableSourceFloat(ConnectedPin, GenerationContext);
+				// Switch Param not generated
+				if (!SwitchParam)
+				{
+					const FText Message = LOCTEXT("FailedToGenerateSwitchParam", "Could not generate switch enum parameter. Please refesh the switch node and connect an enum.");
+					GenerationContext.Compiler->CompilerLog(Message, Node);
+					
+					return Result;
+				}
+
+				if (SwitchParam->GetType() != mu::NodeScalarEnumParameter::GetStaticType())
+				{
+					const FText Message = LOCTEXT("WrongSwitchParamType", "Switch parameter of incorrect type.");
+					GenerationContext.Compiler->CompilerLog(Message, Node);
+
+					return Result;
+				}
+
+				const int32 NumSwitchOptions = TypedNodeTextureSwitch->GetNumElements();
+
+				mu::NodeScalarEnumParameter* EnumParameter = static_cast<mu::NodeScalarEnumParameter*>(SwitchParam.get());
+				if (NumSwitchOptions != EnumParameter->GetValueCount())
+				{
+					const FText Message = LOCTEXT("MismatchedSwitch", "Switch enum and switch node have different number of options. Please refresh the switch node to make sure the outcomes are labeled properly.");
+					GenerationContext.Compiler->CompilerLog(Message, Node);
+				}
+
+				mu::NodeImageSwitchPtr SwitchNode = new mu::NodeImageSwitch;
+				SwitchNode->SetParameter(SwitchParam);
+				SwitchNode->SetOptionCount(NumSwitchOptions);
+
+				for (int SelectorIndex = 0; SelectorIndex < NumSwitchOptions; ++SelectorIndex)
+				{
+					if (const UEdGraphPin* TexturePin = FollowInputPin(*TypedNodeTextureSwitch->GetElementPin(SelectorIndex)))
+					{
+						SwitchNode->SetOption(SelectorIndex, GenerateMutableSourceImage(TexturePin, GenerationContext, ReferenceTextureSize));
+					}
+					else
+					{
+						const FText Message = LOCTEXT("MissingTexture", "Unable to generate texture switch node. Required connection not found.");
+						GenerationContext.Compiler->CompilerLog(Message, Node);
+						return Result;
+					}
+				}
+
+				Result = SwitchNode;
+				return Result;
+			}
+			else
+			{
+				GenerationContext.Compiler->CompilerLog(LOCTEXT("NoEnumParamInSwitch", "Switch nodes must have an enum switch parameter. Please connect an enum and refesh the switch node."), Node);
+				return Result;
+			}
+		}(); // invoke lambda.
+	}
+
+	else if (const UCustomizableObjectNodeTextureVariation* TypedNodeImageVar = Cast<const UCustomizableObjectNodeTextureVariation>(Node))
+	{
+		// UCustomizableObjectNodePassThroughTextureVariation nodes are also handled here
+		mu::NodeImageVariationPtr TextureNode = new mu::NodeImageVariation();
+		Result = TextureNode;
+
+		if (const UEdGraphPin* ConnectedPin = FollowInputPin(*TypedNodeImageVar->DefaultPin()))
+		{
+			mu::NodeImagePtr ChildNode = GenerateMutableSourceImage(ConnectedPin, GenerationContext, ReferenceTextureSize);
+			if (ChildNode)
+			{
+				TextureNode->SetDefaultImage(ChildNode.get());
+			}
+			else
+			{
+				GenerationContext.Compiler->CompilerLog(LOCTEXT("TextureFailed", "Texture generation failed."), Node);
+			}
+		}
+
+		TextureNode->SetVariationCount(TypedNodeImageVar->Variations.Num());
+		for (int32 VariationIndex = 0; VariationIndex < TypedNodeImageVar->Variations.Num(); ++VariationIndex)
+		{
+			const UEdGraphPin* VariationPin = TypedNodeImageVar->VariationPin(VariationIndex);
+			if (!VariationPin)
+			{
+				continue;
+			}
+
+			TextureNode->SetVariationTag(VariationIndex, TypedNodeImageVar->Variations[VariationIndex].Tag);
+			if (const UEdGraphPin* ConnectedPin = FollowInputPin(*VariationPin))
+			{
+				mu::NodeImagePtr ChildNode = GenerateMutableSourceImage(ConnectedPin, GenerationContext, ReferenceTextureSize);
+				TextureNode->SetVariationImage(VariationIndex, ChildNode.get());
+			}
+		}
+	}
+
+	else if (const UCustomizableObjectNodeTextureFromColor* TypedNodeFromColor = Cast<UCustomizableObjectNodeTextureFromColor>(Node))
+	{
+		mu::NodeColourPtr color;
+
+		if (const UEdGraphPin* ConnectedPin = FollowInputPin(*TypedNodeFromColor->ColorPin()))
+		{
+			color = GenerateMutableSourceColor(ConnectedPin, GenerationContext);
+		}
+
+		mu::NodeImagePlainColourPtr ImageFromColour = new mu::NodeImagePlainColour;
+		Result = ImageFromColour;
+
+		if (color)
+		{
+			ImageFromColour->SetColour(color);
+		}
+
+		if (ReferenceTextureSize > 0)
+		{
+			ImageFromColour->SetSize(ReferenceTextureSize, ReferenceTextureSize);
+		}
+	}
+
+	else if (const UCustomizableObjectNodeTextureFromFloats* TypedNodeFromFloats = Cast<UCustomizableObjectNodeTextureFromFloats>(Node))
+	{		
+		mu::Ptr<mu::NodeColourFromScalars> Color = new mu::NodeColourFromScalars();
+
+		if (const UEdGraphPin* RPin = FollowInputPin(*TypedNodeFromFloats->RPin()))
+		{
+			Color->X = GenerateMutableSourceFloat(RPin, GenerationContext);
+		}
+
+		if (const UEdGraphPin* GPin = FollowInputPin(*TypedNodeFromFloats->GPin()))
+		{
+			Color->Y = GenerateMutableSourceFloat(GPin, GenerationContext);
+		}
+
+		if (const UEdGraphPin* BPin = FollowInputPin(*TypedNodeFromFloats->BPin()))
+		{
+			Color->Z = GenerateMutableSourceFloat(BPin, GenerationContext);
+		}
+
+		if (const UEdGraphPin* APin = FollowInputPin(*TypedNodeFromFloats->APin()))
+		{
+			Color->W = GenerateMutableSourceFloat(APin, GenerationContext);
+		}
+
+		mu::NodeImagePlainColourPtr ImageFromColour = new mu::NodeImagePlainColour;
+		Result = ImageFromColour;
+
+		if (Color)
+		{
+			ImageFromColour->SetColour(Color);
+		}
+
+		if (ReferenceTextureSize > 0)
+		{
+			ImageFromColour->SetSize(ReferenceTextureSize, ReferenceTextureSize);
+		}
+	}
+
+	else if (const UCustomizableObjectNodeTextureFromChannels* TypedNodeFrom = Cast<UCustomizableObjectNodeTextureFromChannels>(Node))
+	{
+		mu::NodeImagePtr RNode, GNode, BNode, ANode;
+		bool RGB = false;
+
+		if (const UEdGraphPin* ConnectedPin = FollowInputPin(*TypedNodeFrom->RPin()))
+		{
+			RNode = GenerateMutableSourceImage(ConnectedPin, GenerationContext, ReferenceTextureSize);
+			RGB = true;
+		}
+		if (const UEdGraphPin* ConnectedPin = FollowInputPin(*TypedNodeFrom->GPin()))
+		{
+			GNode = GenerateMutableSourceImage(ConnectedPin, GenerationContext, ReferenceTextureSize);
+			RGB = true;
+		}
+		if (const UEdGraphPin* ConnectedPin = FollowInputPin(*TypedNodeFrom->BPin()))
+		{
+			BNode = GenerateMutableSourceImage(ConnectedPin, GenerationContext, ReferenceTextureSize);
+			RGB = true;
+		}
+		if (const UEdGraphPin* ConnectedPin = FollowInputPin(*TypedNodeFrom->APin()))
+		{
+			ANode = GenerateMutableSourceImage(ConnectedPin, GenerationContext, ReferenceTextureSize);
+		}
+
+		mu::NodeImageSwizzlePtr SwizzleNode = new mu::NodeImageSwizzle;
+		Result = SwizzleNode;
+
+		if (RGB && ANode)
+		{
+			SwizzleNode->SetFormat(mu::EImageFormat::IF_RGBA_UBYTE);
+			SwizzleNode->SetSource(0, RNode);
+			SwizzleNode->SetSourceChannel(0, 0);
+			SwizzleNode->SetSource(1, GNode);
+			SwizzleNode->SetSourceChannel(1, 0);
+			SwizzleNode->SetSource(2, BNode);
+			SwizzleNode->SetSourceChannel(2, 0);
+			SwizzleNode->SetSource(3, ANode);
+			SwizzleNode->SetSourceChannel(3, 0);
+		}
+		else if (RGB)
+		{
+			SwizzleNode->SetFormat(mu::EImageFormat::IF_RGB_UBYTE);
+			SwizzleNode->SetSource(0, RNode);
+			SwizzleNode->SetSourceChannel(0, 0);
+			SwizzleNode->SetSource(1, GNode);
+			SwizzleNode->SetSourceChannel(1, 0);
+			SwizzleNode->SetSource(2, BNode);
+			SwizzleNode->SetSourceChannel(2, 0);
+		}
+		else if (RNode)
+		{
+			SwizzleNode->SetFormat(mu::EImageFormat::IF_L_UBYTE);
+			SwizzleNode->SetSource(0, RNode);
+			SwizzleNode->SetSourceChannel(0, 0);
+		}
+		else if (ANode)
+		{
+			SwizzleNode->SetFormat(mu::EImageFormat::IF_L_UBYTE);
+			SwizzleNode->SetSource(0, ANode);
+			SwizzleNode->SetSourceChannel(0, 0);
+		}
+	}
+
+	else if (const UCustomizableObjectNodeTextureToChannels* TypedNodeTo = Cast<UCustomizableObjectNodeTextureToChannels>(Node))
+	{
+		mu::NodeImagePtr BaseNode;
+		if (const UEdGraphPin* ConnectedPin = FollowInputPin(*TypedNodeTo->InputPin()))
+		{
+			BaseNode = GenerateMutableSourceImage(ConnectedPin, GenerationContext, ReferenceTextureSize);
+		}
+
+		mu::NodeImageSwizzlePtr SwizzleNode = new mu::NodeImageSwizzle;
+		Result = SwizzleNode;
+		SwizzleNode->SetFormat(mu::EImageFormat::IF_L_UBYTE);
+		SwizzleNode->SetSource(0, BaseNode);
+
+		if (Pin == TypedNodeTo->RPin())
+		{
+			SwizzleNode->SetSourceChannel(0, 0);
+		}
+		else if (Pin == TypedNodeTo->GPin())
+		{
+			SwizzleNode->SetSourceChannel(0, 1);
+		}
+		else if (Pin == TypedNodeTo->BPin())
+		{
+			SwizzleNode->SetSourceChannel(0, 2);
+		}
+		else if (Pin == TypedNodeTo->APin())
+		{
+			SwizzleNode->SetSourceChannel(0, 3);
+		}
+		else
+		{
+			check(false);
+		}
+	}
+
+	else if (const UCustomizableObjectNodeTextureProject* TypedNodeProject = Cast<UCustomizableObjectNodeTextureProject>(Node))
+	{
+		mu::Ptr<mu::NodeImageProject> ImageNode = new mu::NodeImageProject();
+		Result = ImageNode;
+
+		if (!FollowInputPin(*TypedNodeProject->MeshPin()))
+		{
+			GenerationContext.Compiler->CompilerLog(LOCTEXT("MissingMeshInProjector", "Texture projector does not have a Mesh. It will be ignored. "), Node, EMessageSeverity::Warning);
+			Result = nullptr;
+		}
+
+		ImageNode->SetLayout(TypedNodeProject->Layout);
+
+		// Calculate the max TextureSize allowed using the ReferenceTextureSize and the Reference texture from the node
+		int32 MaxReferenceTextureSizeInGame = ReferenceTextureSize;
+		if (TypedNodeProject->ReferenceTexture)
+		{
+			GenerationContext.AddParticipatingObject(*TypedNodeProject->ReferenceTexture);
+
+			const UTextureLODSettings& TextureLODSettings = GenerationContext.Options.TargetPlatform->GetTextureLODSettings();
+			const int32 FirstLODAvailable = GenerationContext.Options.bUseLODAsBias ? GenerationContext.FirstLODAvailable : 0;
+			MaxReferenceTextureSizeInGame = GetTextureSizeInGame(*TypedNodeProject->ReferenceTexture, TextureLODSettings, FirstLODAvailable);
+		}
+
+		FUintVector2 TextureSize(TypedNodeProject->TextureSizeX, TypedNodeProject->TextureSizeY);
+
+		// Max TextureSize allowed
+		const int32 MaxProjectedTextureSizeInGame = ReferenceTextureSize > 0 && ReferenceTextureSize < MaxReferenceTextureSizeInGame ? ReferenceTextureSize : MaxReferenceTextureSizeInGame;
+
+		const int32 ProjectorNodeTextureSize = FMath::Max(TextureSize.X, TextureSize.Y);
+		if (ProjectorNodeTextureSize > 0 && ProjectorNodeTextureSize > MaxProjectedTextureSizeInGame)
+		{
+			const int32 NumMips = FMath::CeilLogTwo(ProjectorNodeTextureSize) + 1;
+			const int32 MaxNumMips = FMath::CeilLogTwo(MaxProjectedTextureSizeInGame) + 1;
+
+			TextureSize.X = TextureSize.X >> (NumMips - MaxNumMips);
+			TextureSize.Y = TextureSize.Y >> (NumMips - MaxNumMips);
+		}
+
+		ImageNode->SetImageSize(TextureSize);
+
+		ImageNode->SetEnableSeamCorrection(TypedNodeProject->bEnableTextureSeamCorrection);
+		ImageNode->SetAngleFadeChannels(TypedNodeProject->bEnableAngleFadeOutForRGB, TypedNodeProject->bEnableAngleFadeOutForAlpha);
+		ImageNode->SetSamplingMethod(Invoke(
+			[](ETextureProjectSamplingMethod Method) -> mu::ESamplingMethod 
+			{
+				switch(Method)
+				{
+				case ETextureProjectSamplingMethod::Point:    return mu::ESamplingMethod::Point;
+				case ETextureProjectSamplingMethod::BiLinear: return mu::ESamplingMethod::BiLinear;
+				default: { check(false); return mu::ESamplingMethod::Point; }
+				}
+			}, 
+			TypedNodeProject->SamplingMethod)
+		);
+
+		ImageNode->SetMinFilterMethod(Invoke(
+			[](ETextureProjectMinFilterMethod Method) -> mu::EMinFilterMethod 
+			{
+				switch(Method)
+				{
+				case ETextureProjectMinFilterMethod::None:               return mu::EMinFilterMethod::None;
+				case ETextureProjectMinFilterMethod::TotalAreaHeuristic: return mu::EMinFilterMethod::TotalAreaHeuristic;
+				default: { check(false); return mu::EMinFilterMethod::None; }
+				}
+			}, 
+			TypedNodeProject->MinFilterMethod)
+		);
+
+		if (const UEdGraphPin* ConnectedPin = FollowInputPin(*TypedNodeProject->AngleFadeStartPin()))
+		{
+			mu::NodeScalarPtr ScalarNode = GenerateMutableSourceFloat(ConnectedPin, GenerationContext);
+			ImageNode->SetAngleFadeStart(ScalarNode);
+		}
+		else
+		{
+			mu::NodeScalarConstantPtr ScalarNode = new mu::NodeScalarConstant;
+			float value = FCString::Atof(*TypedNodeProject->AngleFadeStartPin()->DefaultValue);
+			ScalarNode->SetValue(value);
+			ImageNode->SetAngleFadeStart(ScalarNode);
+		}
+
+		if (const UEdGraphPin* ConnectedPin = FollowInputPin(*TypedNodeProject->AngleFadeEndPin()))
+		{
+			mu::NodeScalarPtr ScalarNode = GenerateMutableSourceFloat(ConnectedPin, GenerationContext);
+			ImageNode->SetAngleFadeEnd(ScalarNode);
+		}
+		else
+		{
+			mu::NodeScalarConstantPtr ScalarNode = new mu::NodeScalarConstant;
+			float value = FCString::Atof(*TypedNodeProject->AngleFadeEndPin()->DefaultValue);
+			ScalarNode->SetValue(value);
+			ImageNode->SetAngleFadeEnd(ScalarNode);
+		}
+
+		if (const UEdGraphPin* ConnectedPin = FollowInputPin(*TypedNodeProject->MeshPin()))
+		{
+			FMutableGraphMeshGenerationData DummyMeshData;
+			mu::NodeMeshPtr MeshNode = GenerateMutableSourceMesh(ConnectedPin, GenerationContext, DummyMeshData, false, false);
+			ImageNode->SetMesh(MeshNode);
+		}
+
+		if (const UEdGraphPin* ConnectedPin = FollowInputPin(*TypedNodeProject->MeshMaskPin()))
+		{
+			mu::NodeImagePtr MeshMaskNode = GenerateMutableSourceImage(ConnectedPin, GenerationContext, ReferenceTextureSize);
+			ImageNode->SetTargetMask(MeshMaskNode);
+		}
+
+		if (const UEdGraphPin* ConnectedPin = FollowInputPin(*TypedNodeProject->ProjectorPin()))
+		{
+			mu::NodeProjectorPtr ProjectorNode = GenerateMutableSourceProjector(ConnectedPin, GenerationContext);
+			ImageNode->SetProjector(ProjectorNode);
+		}
+
+		int TexIndex = -1;// TypedNodeProject->OutputPins.Find((UEdGraphPin*)Pin);
+
+		for (int32 i = 0; i < TypedNodeProject->GetNumOutputs(); ++i)
+		{
+			if (TypedNodeProject->OutputPins(i) == Pin)
+			{
+				TexIndex = i;
+			}
+		}
+
+		check(TexIndex >= 0 && TexIndex < TypedNodeProject->GetNumTextures());
+
+		if (const UEdGraphPin* ConnectedPin = FollowInputPin(*TypedNodeProject->TexturePins(TexIndex)))
+		{
+			mu::NodeImagePtr SourceNode = GenerateMutableSourceImage(ConnectedPin, GenerationContext, FMath::Max(TextureSize.X, TextureSize.Y));
+			ImageNode->SetImage(SourceNode);
+		}
+	}
+
+	else if (const UCustomizableObjectNodeTextureBinarise* TypedNodeTexBin = Cast<UCustomizableObjectNodeTextureBinarise>(Node))
+	{
+		mu::NodeImageBinarisePtr BinariseNode = new mu::NodeImageBinarise();
+		Result = BinariseNode;
+
+		if (const UEdGraphPin* ConnectedPin = FollowInputPin(*TypedNodeTexBin->GetBaseImagePin()))
+		{
+			mu::NodeImagePtr BaseImageNode = GenerateMutableSourceImage(ConnectedPin, GenerationContext, ReferenceTextureSize);
+			BinariseNode->SetBase(BaseImageNode);
+		}
+
+		if (const UEdGraphPin* ConnectedPin = FollowInputPin(*TypedNodeTexBin->GetThresholdPin()))
+		{
+			mu::NodeScalarPtr ThresholdNode = GenerateMutableSourceFloat(ConnectedPin, GenerationContext);
+			BinariseNode->SetThreshold(ThresholdNode);
+		}
+	}
+
+	else if (const UCustomizableObjectNodeTextureInvert* TypedNodeTexInv = Cast<UCustomizableObjectNodeTextureInvert>(Node))
+	{
+		mu::NodeImageInvertPtr InvertNode = new mu::NodeImageInvert();
+		Result = InvertNode;
+
+		if (const UEdGraphPin* ConnectedPin = FollowInputPin(*TypedNodeTexInv->GetBaseImagePin()))
+		{
+			mu::NodeImagePtr BaseImageNode = GenerateMutableSourceImage(ConnectedPin, GenerationContext, ReferenceTextureSize);
+			InvertNode->SetBase(BaseImageNode);
+		}
+	}
+
+	else if (const UCustomizableObjectNodeTextureColourMap* TypedNodeColourMap = Cast<UCustomizableObjectNodeTextureColourMap>(Node))
+	{
+		mu::NodeImageColourMapPtr ColourMapNode = new mu::NodeImageColourMap();
+
+		Result = ColourMapNode;
+
+		if (const UEdGraphPin* ConnectedPin = FollowInputPin(*TypedNodeColourMap->GetMapPin()))
+		{
+			mu::NodeImagePtr GradientImageNode = GenerateMutableSourceImage(ConnectedPin, GenerationContext, ReferenceTextureSize);
+			ColourMapNode->SetMap( GradientImageNode ); 
+		}
+
+		if (const UEdGraphPin* ConnectedPin = FollowInputPin(*TypedNodeColourMap->GetMaskPin()))
+		{
+			mu::NodeImagePtr GradientImageNode = GenerateMutableSourceImage(ConnectedPin, GenerationContext, ReferenceTextureSize);
+			ColourMapNode->SetMask( GradientImageNode ); 
+		}
+
+		if (const UEdGraphPin* ConnectedPin = FollowInputPin(*TypedNodeColourMap->GetBasePin()))
+		{
+			mu::NodeImagePtr SourceImageNode = GenerateMutableSourceImage(ConnectedPin, GenerationContext, ReferenceTextureSize);
+			ColourMapNode->SetBase( SourceImageNode );
+		}
+	}
+
+	else if (const UCustomizableObjectNodeTextureTransform* TypedNodeTransform = Cast<UCustomizableObjectNodeTextureTransform>(Node))
+	{
+		mu::NodeImageTransformPtr TransformNode = new mu::NodeImageTransform();
+		Result = TransformNode;
+		
+		if ( UEdGraphPin* BaseImagePin = FollowInputPin(*TypedNodeTransform->GetBaseImagePin()) )
+		{
+			mu::NodeImagePtr ImageNode = GenerateMutableSourceImage( BaseImagePin, GenerationContext, ReferenceTextureSize);
+			TransformNode->SetBase( ImageNode ); 
+		}
+
+		if ( UEdGraphPin* OffsetXPin = FollowInputPin(*TypedNodeTransform->GetOffsetXPin()) )
+		{
+			mu::NodeScalarPtr OffsetXNode = GenerateMutableSourceFloat(OffsetXPin, GenerationContext);
+			TransformNode->SetOffsetX( OffsetXNode ); 
+		}
+
+		if ( UEdGraphPin* OffsetYPin = FollowInputPin(*TypedNodeTransform->GetOffsetYPin()) )
+		{
+			mu::NodeScalarPtr OffsetYNode = GenerateMutableSourceFloat(OffsetYPin, GenerationContext);
+			TransformNode->SetOffsetY( OffsetYNode ); 
+		}
+
+		if ( UEdGraphPin* ScaleXPin = FollowInputPin(*TypedNodeTransform->GetScaleXPin()) )
+		{
+			mu::NodeScalarPtr ScaleXNode = GenerateMutableSourceFloat(ScaleXPin, GenerationContext);
+			TransformNode->SetScaleX( ScaleXNode ); 
+		}
+
+		if ( UEdGraphPin* ScaleYPin = FollowInputPin(*TypedNodeTransform->GetScaleYPin()) )
+		{
+			mu::NodeScalarPtr ScaleYNode = GenerateMutableSourceFloat(ScaleYPin, GenerationContext);
+			TransformNode->SetScaleY( ScaleYNode ); 
+		}
+
+		if ( UEdGraphPin* RotationPin = FollowInputPin(*TypedNodeTransform->GetRotationPin()) )
+		{
+			mu::NodeScalarPtr RotationNode = GenerateMutableSourceFloat(RotationPin, GenerationContext);
+			TransformNode->SetRotation( RotationNode ); 
+		}
+
+		TransformNode->SetAddressMode(Invoke([&]() 
+		{
+			switch (TypedNodeTransform->AddressMode)
+			{
+			case ETextureTransformAddressMode::Wrap:		 return mu::EAddressMode::Wrap;
+			case ETextureTransformAddressMode::ClampToEdge:  return mu::EAddressMode::ClampToEdge;
+			case ETextureTransformAddressMode::ClampToBlack: return mu::EAddressMode::ClampToBlack;
+			default: { check(false); return mu::EAddressMode::None; }
+			}
+		}));
+
+		FUintVector2 TextureSize(TypedNodeTransform->TextureSizeX, TypedNodeTransform->TextureSizeY);
+
+		// Calculate the max TextureSize allowed using the ReferenceTextureSize and the Reference texture from the node
+		int32 MaxReferenceTextureSizeInGame = ReferenceTextureSize;
+		if (TypedNodeTransform->ReferenceTexture)
+		{
+			GenerationContext.AddParticipatingObject(*TypedNodeTransform->ReferenceTexture);
+
+			const UTextureLODSettings& TextureLODSettings = GenerationContext.Options.TargetPlatform->GetTextureLODSettings();
+			const uint8 SurfaceLODBias = GenerationContext.Options.bUseLODAsBias ? GenerationContext.FirstLODAvailable : 0;
+			MaxReferenceTextureSizeInGame = GetTextureSizeInGame(*TypedNodeTransform->ReferenceTexture, TextureLODSettings, SurfaceLODBias);
+		}
+
+		// Max TextureSize allowed
+		const int32 MaxTransformTextureSizeInGame = ReferenceTextureSize > 0 && ReferenceTextureSize < MaxReferenceTextureSizeInGame ? ReferenceTextureSize : MaxReferenceTextureSizeInGame;
+
+		const int32 TransformNodeTextureSize = FMath::Max(TextureSize.X, TextureSize.Y);
+		if (TransformNodeTextureSize > 0 && TransformNodeTextureSize > MaxTransformTextureSizeInGame)
+		{
+			const int32 NumMips = FMath::CeilLogTwo(TransformNodeTextureSize) + 1;
+			const int32 MaxNumMips = FMath::CeilLogTwo(MaxTransformTextureSizeInGame) + 1;
+
+			TextureSize.X = TextureSize.X >> (NumMips - MaxNumMips);
+			TextureSize.Y = TextureSize.Y >> (NumMips - MaxNumMips);
+		}
+
+		TransformNode->SetKeepAspectRatio(TypedNodeTransform->bKeepAspectRatio);
+		TransformNode->SetSizeX(TextureSize.X);
+		TransformNode->SetSizeY(TextureSize.Y);
+	}
+
+	else if (const UCustomizableObjectNodeTextureSaturate* TypedNodeSaturate = Cast<UCustomizableObjectNodeTextureSaturate>(Node))
+	{
+		mu::Ptr<mu::NodeImageSaturate> SaturateNode = new mu::NodeImageSaturate();
+		Result = SaturateNode;
+	
+		if ( UEdGraphPin* BaseImagePin = FollowInputPin(*TypedNodeSaturate->GetBaseImagePin()) )
+		{
+			mu::NodeImagePtr ImageNode = GenerateMutableSourceImage(BaseImagePin, GenerationContext, ReferenceTextureSize);
+			SaturateNode->SetSource(ImageNode); 
+		}
+
+		if ( UEdGraphPin* FactorPin = FollowInputPin(*TypedNodeSaturate->GetFactorPin()))
+		{
+			mu::NodeScalarPtr FactorNode = GenerateMutableSourceFloat(FactorPin, GenerationContext);
+			SaturateNode->SetFactor(FactorNode); 
+		}
+	}
+
+	else if (const UCustomizableObjectNodePassThroughTexture* TypedNodePassThroughTex = Cast<UCustomizableObjectNodePassThroughTexture>(Node))
+	{
+		UTexture* BaseTexture = TypedNodePassThroughTex->PassThroughTexture;
+		if (BaseTexture)
+		{
+			mu::Ptr<mu::NodeImageConstant> ImageNode = new mu::NodeImageConstant();
+			Result = ImageNode;
+
+			ImageNode->SetValue(GenerateImageConstant(BaseTexture, GenerationContext, true).get());
+		}
+		else
+		{
+			GenerationContext.Compiler->CompilerLog(LOCTEXT("MissingImagePassThrough", "Missing image in pass-through texture node."), Node);
+		}
+	}
+
+	else if (const UCustomizableObjectNodePassThroughTextureSwitch* TypedNodePassThroughTextureSwitch = Cast<UCustomizableObjectNodePassThroughTextureSwitch>(Node))
+	{
+		Result = [&]()
+		{
+			const UEdGraphPin* SwitchParameter = TypedNodePassThroughTextureSwitch->SwitchParameter();
+
+			// Check Switch Parameter arity preconditions.
+			if (const UEdGraphPin* ConnectedPin = FollowInputPin(*SwitchParameter))
+			{
+				mu::NodeScalarPtr SwitchParam = GenerateMutableSourceFloat(ConnectedPin, GenerationContext);
+				// Switch Param not generated
+				if (!SwitchParam)
+				{
+					const FText Message = LOCTEXT("FailedToGenerateSwitchParam", "Could not generate switch enum parameter. Please refesh the switch node and connect an enum.");
+					GenerationContext.Compiler->CompilerLog(Message, Node);
+
+					return Result;
+				}
+
+				if (SwitchParam->GetType() != mu::NodeScalarEnumParameter::GetStaticType())
+				{
+					const FText Message = LOCTEXT("WrongSwitchParamType", "Switch parameter of incorrect type.");
+					GenerationContext.Compiler->CompilerLog(Message, Node);
+
+					return Result;
+				}
+
+				const int32 NumSwitchOptions = TypedNodePassThroughTextureSwitch->GetNumElements();
+
+				mu::NodeScalarEnumParameter* EnumParameter = static_cast<mu::NodeScalarEnumParameter*>(SwitchParam.get());
+				if (NumSwitchOptions != EnumParameter->GetValueCount())
+				{
+					const FText Message = LOCTEXT("MismatchedSwitch", "Switch enum and switch node have different number of options. Please refresh the switch node to make sure the outcomes are labeled properly.");
+					GenerationContext.Compiler->CompilerLog(Message, Node);
+				}
+
+				// TODO: Implement Mutable core pass-through switch nodes
+				mu::Ptr<mu::NodeImageSwitch> SwitchNode = new mu::NodeImageSwitch;
+				SwitchNode->SetParameter(SwitchParam);
+				SwitchNode->SetOptionCount(NumSwitchOptions);
+
+				for (int32 SelectorIndex = 0; SelectorIndex < NumSwitchOptions; ++SelectorIndex)
+				{
+					if (const UEdGraphPin* TexturePin = FollowInputPin(*TypedNodePassThroughTextureSwitch->GetElementPin(SelectorIndex)))
+					{
+						mu::Ptr<mu::NodeImage> PassThroughImage = GenerateMutableSourceImage(TexturePin, GenerationContext, ReferenceTextureSize);
+						SwitchNode->SetOption(SelectorIndex, PassThroughImage);
+					}
+					else
+					{
+						const FText Message = LOCTEXT("MissingPassThroughTexture", "Unable to generate pass-through texture switch node. Required connection not found.");
+						GenerationContext.Compiler->CompilerLog(Message, Node);
+						return Result;
+					}
+				}
+
+				Result = SwitchNode;
+				return Result;
+			}
+			else
+			{
+				GenerationContext.Compiler->CompilerLog(LOCTEXT("NoEnumParamInSwitch", "Switch nodes must have an enum switch parameter. Please connect an enum and refesh the switch node."), Node);
+				return Result;
+			}
+		}(); // invoke lambda.
+	}
+
+	// If the node is a plain colour node, generate an image out of it
+	else if (Pin->PinType.PinCategory == Schema->PC_Color)
+	{
+		mu::NodeColourPtr ColorNode = GenerateMutableSourceColor(Pin, GenerationContext);
+
+		mu::NodeImagePlainColourPtr ImageNode = new mu::NodeImagePlainColour;
+		ImageNode->SetSize(16, 16);
+		ImageNode->SetColour(ColorNode);
+		Result = ImageNode;
+	}
+
+	else if (const UCustomizableObjectNodeTable* TypedNodeTable = Cast<UCustomizableObjectNodeTable>(Node))
+	{
+		if (Pin->PinType.PinCategory == Schema->PC_Color)
+		{
+			mu::NodeColourPtr ColorNode = GenerateMutableSourceColor(Pin, GenerationContext);
+			mu::NodeImagePlainColourPtr ImageNode = new mu::NodeImagePlainColour;
+
+			ImageNode->SetSize(16, 16);
+			ImageNode->SetColour(ColorNode);
+
+			Result = ImageNode;
+		}
+		else
+		{
+			// This node will add a checker texture in case of error
+			mu::Ptr<mu::NodeImageConstant> EmptyNode = new mu::NodeImageConstant();
+			Result = EmptyNode;
+
+			bool bSuccess = true;
+
+			if (Pin->PinType.PinCategory == Schema->PC_MaterialAsset)
+			{
+				// Material pins have to skip the cache of nodes or they will return always the same column node
+				bDoNotAddToGeneratedCache = true;
+			}
+
+			UDataTable* DataTable = GetDataTable(TypedNodeTable, GenerationContext);
+
+			if (DataTable)
+			{
+				FString ColumnName = Pin->PinFriendlyName.ToString();
+				FProperty* Property = DataTable->FindTableProperty(FName(*ColumnName));
+
+				if (!Property)
+				{
+					FString Msg = FString::Printf(TEXT("Couldn't find the column [%s] in the data table's struct."), *ColumnName);
+					GenerationContext.Compiler->CompilerLog(FText::FromString(Msg), Node);
+
+					bSuccess = false;
+				}
+
+				UTexture* DefaultTexture = TypedNodeTable->GetColumnDefaultAssetByType<UTexture>(Pin);
+
+				if (bSuccess && Pin->PinType.PinCategory != Schema->PC_MaterialAsset && !DefaultTexture)
+				{
+					FString Msg = FString::Printf(TEXT("Couldn't find a default value in the data table's struct for the column [%s]. The default value is null or not a supported Texture"), *ColumnName);
+					GenerationContext.Compiler->CompilerLog(FText::FromString(Msg), Node);
+					
+					bSuccess = false;
+				}
+
+				if (bSuccess)
+				{
+					// Generating a new data table if not exists
+					mu::TablePtr Table = nullptr;
+					Table = GenerateMutableSourceTable(DataTable, TypedNodeTable, GenerationContext);
+
+					if (Table)
+					{
+						mu::Ptr<mu::NodeImageTable> ImageTableNode = new mu::NodeImageTable();
+
+						if (Pin->PinType.PinCategory == Schema->PC_MaterialAsset)
+						{
+							// Material parameters use the Data Table Column Name + Parameter id as mutable column Name to aboid duplicated names (i.e. two MI columns with the same parents but different values).
+							ColumnName = Property->GetDisplayNameText().ToString() + GenerationContext.CurrentMaterialTableParameterId;
+						}
+						else
+						{
+							// Checking if this pin texture has been used in another table node but with a different texture mode
+							const int32 ImageDataIndex = GenerationContext.GeneratedTableImages.Find({ ColumnName, Pin->PinType.PinCategory, Table });
+
+							if (ImageDataIndex != INDEX_NONE )
+							{
+								if (GenerationContext.GeneratedTableImages[ImageDataIndex].PinType != Pin->PinType.PinCategory)
+								{
+									TArray<const UObject*> Nodes;
+									Nodes.Add(TypedNodeTable);
+									Nodes.Add(GenerationContext.GeneratedTableImages[ImageDataIndex].TableNode);
+
+									FString Msg = FString::Printf(TEXT("Texture pin [%s] with different texture modes found in more than one table node. This will add multiple times the texture reseource in the final cook."), *ColumnName);
+									GenerationContext.Compiler->CompilerLog(FText::FromString(Msg), Nodes);
+								}
+							}
+							else
+							{
+								GenerationContext.GeneratedTableImages.Add({ ColumnName, Pin->PinType.PinCategory, Table, TypedNodeTable });
+							}
+
+							// Encoding the texture mode of the type to allow different texture modes from the same pin
+							if (Pin->PinType.PinCategory == Schema->PC_PassThroughImage)
+							{
+								ColumnName += "--PassThrough";
+							}
+						}
+
+						// Generating a new Texture column if not exists
+						if (Table->FindColumn(ColumnName) == INDEX_NONE)
+						{
+							int32 Dummy = -1; // TODO MTBL-1512
+							bool Dummy2 = false;
+							bSuccess = GenerateTableColumn(TypedNodeTable, Pin, Table, ColumnName, Property, Dummy, Dummy, GenerationContext.CurrentLOD, Dummy, Dummy2, GenerationContext);
+
+							if (!bSuccess)
+							{
+								FString Msg = FString::Printf(TEXT("Failed to generate the mutable table column [%s]"), *ColumnName);
+								GenerationContext.Compiler->CompilerLog(FText::FromString(Msg), Node);
+							}
+						}
+
+						if (bSuccess)
+						{
+							Result = ImageTableNode;
+
+							ImageTableNode->Table = Table;
+							ImageTableNode->ColumnName = ColumnName;
+							ImageTableNode->ParameterName = TypedNodeTable->ParameterName;
+							ImageTableNode->bNoneOption = TypedNodeTable->bAddNoneOption;
+							ImageTableNode->DefaultRowName = TypedNodeTable->DefaultRowName.ToString();
+
+							if (UTexture2D* DefaultTexture2D = Cast<UTexture2D>(DefaultTexture))
+							{
+								mu::FImageDesc ImageDesc = GenerateImageDescriptor(DefaultTexture2D);
+
+								uint32 LODBias = ReferenceTextureSize > 0 ? ComputeLODBiasForTexture(GenerationContext, *DefaultTexture2D, nullptr, ReferenceTextureSize) : 0;
+								ImageDesc.m_size[0] = ImageDesc.m_size[0] >> LODBias;
+								ImageDesc.m_size[1] = ImageDesc.m_size[1] >> LODBias;
+								
+								const uint16 MaxTextureSize = FMath::Max3(ImageDesc.m_size[0], ImageDesc.m_size[1], (uint16)1);
+								ImageTableNode->MaxTextureSize = MaxTextureSize;
+								ImageTableNode->ReferenceImageDesc = ImageDesc;
+
+								// Calculate the number of mips to tag as high res for this image.
+								int32 TotalMips = mu::Image::GetMipmapCount(ImageDesc.m_size[0], ImageDesc.m_size[1]);
+								int32 NumMipsBeyondMin = FMath::Max(0, TotalMips - GenerationContext.Options.MinDiskMips);
+								int32 HighResMipsForThisImage = FMath::Min(NumMipsBeyondMin, GenerationContext.Options.NumHighResImageMips);
+								ImageTableNode->SourceDataDescriptor.SourceHighResMips = HighResMipsForThisImage;
+							}
+						}
+					}
+					else
+					{
+						FString Msg = FString::Printf(TEXT("Couldn't generate a mutable table."));
+						GenerationContext.Compiler->CompilerLog(FText::FromString(Msg), Node);
+					}
+				}
+			}
+			else
+			{
+				GenerationContext.Compiler->CompilerLog(LOCTEXT("ImageTableError", "Couldn't find the data table of the node."), Node);
+			}
+		}
+	}
+
+	else
+	{
+		GenerationContext.Compiler->CompilerLog(LOCTEXT("UnimplementedNode", "Node type not implemented yet."), Node);
+	}
+
+	if (!bDoNotAddToGeneratedCache)
+	{
+		GenerationContext.Generated.Add(Key, FGeneratedData(Node, Result));
+		GenerationContext.GeneratedNodes.Add(Node);
+	}
+
+	if (Result)
+	{
+		Result->SetMessageContext(Node);
+	}
+
+	return Result;
+}
+
+
+#undef LOCTEXT_NAMESPACE
+
diff --git a/Engine/Plugins/Mutable/Source/CustomizableObjectEditor/Private/MuCOE/GenerateMutableSource/GenerateMutableSourceImage.h b/Engine/Plugins/Mutable/Source/CustomizableObjectEditor/Private/MuCOE/GenerateMutableSource/GenerateMutableSourceImage.h
new file mode 100644
--- /dev/null
+++ b/Engine/Plugins/Mutable/Source/CustomizableObjectEditor/Private/MuCOE/GenerateMutableSource/GenerateMutableSourceImage.h
@@ -0,0 +1,27 @@
+// Copyright Epic Games, Inc. All Rights Reserved.
+
+#pragma once
+
+
+#include "MuR/Ptr.h"
+#include "HAL/Platform.h"
+
+class FCustomizableObjectCompiler;
+class UCustomizableObjectNode;
+class UEdGraphPin;
+class UTexture2D;
+struct FMutableGraphGenerationContext;
+
+namespace mu
+{
+	class Image;
+	class NodeImage;
+}
+
+
+//mu::Ptr<mu::Image> ConvertTextureUnrealToMutable(UTexture2D* Texture, const UCustomizableObjectNode* Node, FCustomizableObjectCompiler* Compiler, bool bIsNormalComposite);
+
+mu::Ptr<mu::NodeImage> ResizeTextureByNumMips(const mu::Ptr<mu::NodeImage>& Image, int32 MipsToSkip);
+
+/** Convert a CustomizableObject Source Graph from an Image pin into a mutable source graph. */
+mu::Ptr<mu::NodeImage> GenerateMutableSourceImage(const UEdGraphPin* Pin, FMutableGraphGenerationContext& GenerationContext, int32 ReferenceTextureSize);
diff --git a/Engine/Plugins/Mutable/Source/CustomizableObjectEditor/Private/MuCOE/GenerateMutableSource/GenerateMutableSourceLayout.h b/Engine/Plugins/Mutable/Source/CustomizableObjectEditor/Private/MuCOE/GenerateMutableSource/GenerateMutableSourceLayout.h
new file mode 100644
--- /dev/null
+++ b/Engine/Plugins/Mutable/Source/CustomizableObjectEditor/Private/MuCOE/GenerateMutableSource/GenerateMutableSourceLayout.h
@@ -0,0 +1,16 @@
+// Copyright Epic Games, Inc. All Rights Reserved.
+
+#pragma once
+
+#include "MuT/NodeLayout.h"
+#include "MuCOE/CustomizableObjectLayout.h"
+
+class UEdGraphPin;
+struct FMutableGraphGenerationContext;
+
+
+extern mu::Ptr<mu::NodeLayout> CreateDefaultLayout();
+
+extern mu::Ptr<mu::NodeLayout> CreateMutableLayoutNode(FMutableGraphGenerationContext& GenerationContext, const UCustomizableObjectLayout*, bool bIgnoreLayoutWarnings, bool& bWasEmpty);
+
+extern mu::FSourceLayoutBlock ToMutable(FMutableGraphGenerationContext& GenerationContext, const FCustomizableObjectLayoutBlock&);
diff --git a/Engine/Plugins/Mutable/Source/CustomizableObjectEditor/Private/MuCOE/GenerateMutableSource/GenerateMutableSourceModifier.cpp b/Engine/Plugins/Mutable/Source/CustomizableObjectEditor/Private/MuCOE/GenerateMutableSource/GenerateMutableSourceModifier.cpp
new file mode 100644
--- /dev/null
+++ b/Engine/Plugins/Mutable/Source/CustomizableObjectEditor/Private/MuCOE/GenerateMutableSource/GenerateMutableSourceModifier.cpp
@@ -0,0 +1,611 @@
+// Copyright Epic Games, Inc. All Rights Reserved.
+
+#include "MuCOE/GenerateMutableSource/GenerateMutableSourceModifier.h"
+
+#include "Engine/StaticMesh.h"
+#include "MuCOE/CustomizableObjectCompiler.h"
+#include "MuCOE/GenerateMutableSource/GenerateMutableSourceMesh.h"
+#include "MuCOE/GenerateMutableSource/GenerateMutableSourceImage.h"
+#include "MuCOE/GenerateMutableSource/GenerateMutableSourceFloat.h"
+#include "MuCOE/GenerateMutableSource/GenerateMutableSourceSurface.h"
+#include "MuCOE/GenerateMutableSource/GenerateMutableSourceLayout.h"
+#include "MuCOE/GenerateMutableSource/GenerateMutableSourceGroupProjector.h"
+#include "MuCOE/GraphTraversal.h"
+#include "MuCOE/MutableUtils.h"
+#include "MuCOE/Nodes/CustomizableObjectNodeModifierClipDeform.h"
+#include "MuCOE/Nodes/CustomizableObjectNodeModifierClipMorph.h"
+#include "MuCOE/Nodes/CustomizableObjectNodeModifierClipWithMesh.h"
+#include "MuCOE/Nodes/CustomizableObjectNodeModifierClipWithUVMask.h"
+#include "MuCOE/Nodes/CustomizableObjectNodeModifierExtendMeshSection.h"
+#include "MuCOE/Nodes/CustomizableObjectNodeModifierEditMeshSection.h"
+#include "MuCOE/Nodes/CustomizableObjectNodeModifierMorphMeshSection.h"
+#include "MuCOE/Nodes/CustomizableObjectNodeModifierRemoveMesh.h"
+#include "MuCOE/Nodes/CustomizableObjectNodeModifierRemoveMeshBlocks.h"
+#include "MuCOE/Nodes/CustomizableObjectNodeStaticMesh.h"
+#include "MuCOE/Nodes/CustomizableObjectNodeFloatParameter.h"
+#include "MuCOE/Nodes/CustomizableObjectNodeFloatConstant.h"
+#include "MuR/Mesh.h"
+#include "MuT/NodeMeshTransform.h"
+#include "MuT/NodeModifierMeshClipDeform.h"
+#include "MuT/NodeModifierMeshClipMorphPlane.h"
+#include "MuT/NodeModifierMeshClipWithUVMask.h"
+#include "MuT/NodeModifierSurfaceEdit.h"
+#include "MuT/NodeMeshConstant.h"
+#include "MuT/NodeMeshFragment.h"
+#include "MuT/NodeMeshFormat.h"
+
+#define LOCTEXT_NAMESPACE "CustomizableObjectEditor"
+
+
+mu::Ptr<mu::NodeModifier> GenerateMutableSourceModifier(const UEdGraphPin * Pin, FMutableGraphGenerationContext & GenerationContext)
+{
+	check(Pin)
+	RETURN_ON_CYCLE(*Pin, GenerationContext)
+
+	CheckNumOutputs(*Pin, GenerationContext);
+	
+	UCustomizableObjectNode* Node = CastChecked<UCustomizableObjectNode>(Pin->GetOwningNode());
+
+	const FGeneratedKey Key(reinterpret_cast<void*>(&GenerateMutableSourceModifier), *Pin, *Node, GenerationContext, true);
+	if (const FGeneratedData* Generated = GenerationContext.Generated.Find(Key))
+	{
+		return static_cast<mu::NodeModifier*>(Generated->Node.get());
+	}
+	
+	mu::Ptr<mu::NodeModifier> Result;
+
+	bool bDoNotAddToGeneratedCache = false; // TODO Remove on MTBL-829 
+
+	
+	if (const UCustomizableObjectNodeModifierClipMorph* TypedNodeClip = Cast<UCustomizableObjectNodeModifierClipMorph>(Node))
+	{
+		const EMutableMeshConversionFlags ModifiersMeshFlags =
+			EMutableMeshConversionFlags::IgnoreSkinning |
+			EMutableMeshConversionFlags::IgnorePhysics;
+		GenerationContext.MeshGenerationFlags.Push(ModifiersMeshFlags);
+
+		mu::Ptr<mu::NodeModifierMeshClipMorphPlane> ClipNode = new mu::NodeModifierMeshClipMorphPlane();
+		Result = ClipNode;
+
+		const FVector Origin = TypedNodeClip->GetOriginWithOffset();
+		const FVector& Normal = TypedNodeClip->Normal;
+
+		ClipNode->SetPlane(FVector3f(Origin), FVector3f(Normal));
+		ClipNode->SetParams(TypedNodeClip->B, TypedNodeClip->Exponent);
+		ClipNode->SetMorphEllipse(TypedNodeClip->Radius, TypedNodeClip->Radius2, TypedNodeClip->RotationAngle);
+
+		ClipNode->SetVertexSelectionBone(GenerationContext.GetBoneUnique(TypedNodeClip->BoneName), TypedNodeClip->MaxEffectRadius);
+
+		ClipNode->MultipleTagsPolicy = TypedNodeClip->MultipleTagPolicy;
+		for (const FString& Tag : TypedNodeClip->RequiredTags)
+		{
+			 ClipNode->RequiredTags.Add(Tag);
+		}
+
+		GenerationContext.MeshGenerationFlags.Pop();
+	}
+
+	else if (const UCustomizableObjectNodeModifierClipDeform* TypedNodeClipDeform = Cast<UCustomizableObjectNodeModifierClipDeform>(Node))
+	{
+		const EMutableMeshConversionFlags ModifiersMeshFlags =
+			EMutableMeshConversionFlags::IgnoreSkinning |
+			EMutableMeshConversionFlags::IgnorePhysics;
+		GenerationContext.MeshGenerationFlags.Push(ModifiersMeshFlags);
+
+		mu::Ptr<mu::NodeModifierMeshClipDeform> ClipNode = new mu::NodeModifierMeshClipDeform();
+		Result = ClipNode;
+	
+		if (const UEdGraphPin* ConnectedPin = FollowInputPin(*TypedNodeClipDeform->ClipShapePin()))
+		{
+			FMutableGraphMeshGenerationData DummyMeshData;
+			mu::Ptr<mu::NodeMesh> ClipMesh = GenerateMutableSourceMesh(ConnectedPin, GenerationContext, DummyMeshData, false, true);
+
+			ClipNode->ClipMesh = ClipMesh;
+
+			mu::EShapeBindingMethod BindingMethod = mu::EShapeBindingMethod::ClipDeformClosestProject;
+			switch(TypedNodeClipDeform->BindingMethod)
+			{
+				case EShapeBindingMethod::ClosestProject:
+					BindingMethod = mu::EShapeBindingMethod::ClipDeformClosestProject;
+					break;
+				case EShapeBindingMethod::NormalProject:
+					BindingMethod = mu::EShapeBindingMethod::ClipDeformNormalProject;
+					break;
+				case EShapeBindingMethod::ClosestToSurface:
+					BindingMethod = mu::EShapeBindingMethod::ClipDeformClosestToSurface;
+					break;
+				default:
+					check(false);
+					break;
+			}
+
+			ClipNode->BindingMethod = BindingMethod;
+		}
+		else
+		{
+			FText ErrorMsg = LOCTEXT("ClipDeform mesh", "The clip deform node requires an input clip shape.");
+			GenerationContext.Compiler->CompilerLog(ErrorMsg, TypedNodeClipDeform, EMessageSeverity::Error);
+			Result = nullptr;
+		}
+	
+		ClipNode->MultipleTagsPolicy = TypedNodeClipDeform->MultipleTagPolicy;
+		for (const FString& Tag : TypedNodeClipDeform->RequiredTags)
+		{
+			ClipNode->RequiredTags.Add(Tag);
+		}		
+
+		GenerationContext.MeshGenerationFlags.Pop();
+	}
+
+	else if (const UCustomizableObjectNodeModifierClipWithMesh* TypedNodeClipMesh = Cast<UCustomizableObjectNodeModifierClipWithMesh>(Node))
+	{
+		const EMutableMeshConversionFlags ModifiersMeshFlags =
+			EMutableMeshConversionFlags::IgnoreSkinning |
+			EMutableMeshConversionFlags::IgnorePhysics;
+		GenerationContext.MeshGenerationFlags.Push(ModifiersMeshFlags);
+
+		// MeshClipWithMesh can be connected to multiple objects, so the compiled NodeModifierMeshClipWithMesh
+		// needs to be different for each object. If it were added to the Generated cache, all the objects would get the same.
+		bDoNotAddToGeneratedCache = true;
+
+		mu::Ptr<mu::NodeModifierMeshClipWithMesh> ClipNode = new mu::NodeModifierMeshClipWithMesh();
+		Result = ClipNode;
+
+		if (const UEdGraphPin* ConnectedPin = FollowInputPin(*TypedNodeClipMesh->ClipMeshPin()))
+		{
+			FMutableGraphMeshGenerationData DummyMeshData;
+
+			mu::NodeMeshPtr ClipMesh = GenerateMutableSourceMesh(ConnectedPin, GenerationContext, DummyMeshData, false, true);
+
+			FPinDataValue* PinData = GenerationContext.PinData.Find(ConnectedPin);
+			for (const FMeshData& MeshData : PinData->MeshesData)
+			{
+				bool bClosed = true;
+				if (const USkeletalMesh* SkeletalMesh = Cast<USkeletalMesh>(MeshData.Mesh))
+				{
+					bClosed = IsMeshClosed(SkeletalMesh, MeshData.LOD, MeshData.MaterialIndex);
+				}
+				else if (const UStaticMesh* StaticMesh = Cast<UStaticMesh>(MeshData.Mesh))
+				{
+					bClosed = IsMeshClosed(StaticMesh, MeshData.LOD, MeshData.MaterialIndex);
+				}
+				else
+				{
+					// TODO: We support the clip mesh not being constant. This message is not precise enough. It should say that it hasn't been 
+					// possible to check if the mesh is closed or not.
+					GenerationContext.Compiler->CompilerLog(LOCTEXT("UnimplementedNode", "Node type not implemented yet."), MeshData.Node);
+				}
+
+				if (!bClosed)
+				{
+					FText ErrorMsg = FText::Format(LOCTEXT("Clipping mesh", "Clipping mesh [{0}] not closed (i.e., it does not enclose a volume)."), FText::FromName(MeshData.Mesh->GetFName()));
+					GenerationContext.Compiler->CompilerLog(ErrorMsg, MeshData.Node, EMessageSeverity::Warning);
+				}
+			}
+
+			if (FMatrix Matrix = TypedNodeClipMesh->Transform.ToMatrixWithScale(); Matrix != FMatrix::Identity)
+			{
+				mu::NodeMeshTransformPtr TransformMesh = new mu::NodeMeshTransform();
+				TransformMesh->SetSource(ClipMesh.get());
+
+				TransformMesh->SetTransform(FMatrix44f(Matrix));
+				ClipMesh = TransformMesh;
+			}
+
+			ClipNode->ClipMesh = ClipMesh;
+		}
+		else
+		{
+			FText ErrorMsg = LOCTEXT("Clipping mesh missing", "The clip mesh with mesh node requires an input clip mesh.");
+			GenerationContext.Compiler->CompilerLog(ErrorMsg, TypedNodeClipMesh, EMessageSeverity::Error);
+			Result = nullptr;
+		}
+
+		ClipNode->MultipleTagsPolicy = TypedNodeClipMesh->MultipleTagPolicy;
+		for (const FString& Tag : TypedNodeClipMesh->RequiredTags)
+		{
+			 ClipNode->RequiredTags.Add(Tag);
+		}
+
+		if (TypedNodeClipMesh->CustomizableObjectToClipWith != nullptr)
+		{
+			TArray<mu::Ptr<mu::NodeModifierMeshClipWithMesh>>* ArrayDataPtr = GenerationContext.MapClipMeshNodeToMutableClipMeshNodeArray.Find(Cast<UCustomizableObjectNodeModifierClipWithMesh>(Pin->GetOwningNode()));
+
+			if (ArrayDataPtr == nullptr)
+			{
+				TArray<mu::Ptr<mu::NodeModifierMeshClipWithMesh>> ArrayData;
+				ArrayData.Add(ClipNode);
+				UCustomizableObjectNodeModifierClipWithMesh* CastedNode = Cast<UCustomizableObjectNodeModifierClipWithMesh>(Pin->GetOwningNode());
+				GenerationContext.MapClipMeshNodeToMutableClipMeshNodeArray.Add(CastedNode, ArrayData);
+			}
+			else
+			{
+				ArrayDataPtr->AddUnique(ClipNode);
+			}
+		}
+
+		GenerationContext.MeshGenerationFlags.Pop();
+	}
+
+	else if (const UCustomizableObjectNodeModifierClipWithUVMask* TypedNodeClipUVMask = Cast<UCustomizableObjectNodeModifierClipWithUVMask>(Node))
+	{
+		const EMutableMeshConversionFlags ModifiersMeshFlags =
+			EMutableMeshConversionFlags::IgnoreSkinning |
+			EMutableMeshConversionFlags::IgnorePhysics;
+		GenerationContext.MeshGenerationFlags.Push(ModifiersMeshFlags);
+
+		// This modifier can be connected to multiple objects, so the compiled node
+		// needs to be different for each object. If it were added to the Generated cache, all the objects would get the same.
+		bDoNotAddToGeneratedCache = true;
+
+		mu::Ptr<mu::NodeModifierMeshClipWithUVMask> ClipNode = new mu::NodeModifierMeshClipWithUVMask();
+		Result = ClipNode;
+
+		if (const UEdGraphPin* ConnectedPin = FollowInputPin(*TypedNodeClipUVMask->ClipMaskPin()))
+		{
+			FMutableGraphMeshGenerationData DummyMeshData;
+
+			mu::Ptr<mu::NodeImage> ClipMask = GenerateMutableSourceImage(ConnectedPin, GenerationContext, 0);
+
+			ClipNode->ClipMask = ClipMask;
+		}
+		else
+		{
+			FText ErrorMsg = LOCTEXT("ClipUVMask mesh", "The clip mesh with UV Mask node requires an input texture mask.");
+			GenerationContext.Compiler->CompilerLog(ErrorMsg, TypedNodeClipUVMask, EMessageSeverity::Error);
+			Result = nullptr;
+		}
+
+		ClipNode->LayoutIndex = TypedNodeClipUVMask->UVChannelForMask;
+
+		ClipNode->MultipleTagsPolicy = TypedNodeClipUVMask->MultipleTagPolicy;
+		for (const FString& Tag : TypedNodeClipUVMask->RequiredTags)
+		{
+			ClipNode->RequiredTags.Add(Tag);
+		}
+
+		GenerationContext.MeshGenerationFlags.Pop();
+	}
+
+	else if (UCustomizableObjectNodeModifierExtendMeshSection* TypedNodeExt = Cast<UCustomizableObjectNodeModifierExtendMeshSection>(Node))
+	{
+		const EMutableMeshConversionFlags ModifiersMeshFlags = EMutableMeshConversionFlags::None;
+		GenerationContext.MeshGenerationFlags.Push(ModifiersMeshFlags);
+
+		mu::Ptr<mu::NodeModifierSurfaceEdit> SurfNode = new mu::NodeModifierSurfaceEdit();
+		Result = SurfNode;
+
+		// TODO: This was used in the non-modifier version for group projectors. It may affect the "drop projection from LOD" feature.
+		const int32 LOD = Node->IsAffectedByLOD() ? GenerationContext.CurrentLOD : 0;
+
+		[&] // Using a lambda so control flow is easier to manage.
+		{
+			SurfNode->MultipleTagsPolicy = TypedNodeExt->MultipleTagPolicy;
+			for (const FString& Tag : TypedNodeExt->RequiredTags)
+			{
+				SurfNode->RequiredTags.Add(Tag);
+			}
+
+			mu::Ptr<mu::NodeMesh> AddMeshNode;
+			FMutableGraphMeshGenerationData MeshData;
+			if (const UEdGraphPin* ConnectedPin = FollowInputPin(*TypedNodeExt->AddMeshPin()))
+			{
+				// Flags to know which UV channels need layout
+				FLayoutGenerationFlags LayoutGenerationFlags;
+				LayoutGenerationFlags.TexturePinModes.Init(EPinMode::Mutable, TEXSTREAM_MAX_NUM_UVCHANNELS);
+
+				GenerationContext.LayoutGenerationFlags.Push(LayoutGenerationFlags);
+
+				AddMeshNode = GenerateMutableSourceMesh(ConnectedPin, GenerationContext, MeshData, true, false);
+
+				GenerationContext.LayoutGenerationFlags.Pop();
+			}
+
+			if (AddMeshNode)
+			{
+				mu::Ptr<mu::NodeMesh> MeshPtr = AddMeshNode;
+
+				const TArray<UCustomizableObjectLayout*> Layouts = TypedNodeExt->GetLayouts();
+
+				if (Layouts.Num())
+				{
+					mu::Ptr<mu::NodeMeshFragment> MeshFrag = new mu::NodeMeshFragment();
+
+					MeshFrag->SourceMesh = MeshPtr;
+					//TODO: Implement support for multiple UV channels (e.g. Add warning for vertices which have a block in a layout but not in the other)
+					MeshFrag->LayoutIndex = 0;
+
+					// For this case we don't want to create another layout: we will use the one defined in the mesh to be added since we want to add
+					// any block defined there.
+					//bool bWasEmpty = false;
+					//MeshFrag->Layout = CreateMutableLayoutNode(GenerationContext, Layouts[MeshFrag->LayoutIndex], true, bWasEmpty);
+
+					MeshPtr = MeshFrag;
+				}
+				else
+				{
+					GenerationContext.Compiler->CompilerLog(LOCTEXT("ExtendMaterialLayoutMissing", "Skeletal Mesh without Layout Node linked to an Extend Material. A 4x4 layout will be added as default layout."), Node);
+				}
+
+				mu::Ptr<mu::NodeMeshFormat> MeshFormat = new mu::NodeMeshFormat();
+				SetSurfaceFormat(GenerationContext,
+					MeshFormat->GetVertexBuffers(), MeshFormat->GetIndexBuffers(), MeshData,
+					GenerationContext.Options.CustomizableObjectNumBoneInfluences,
+					GenerationContext.Options.b16BitBoneWeightsEnabled);
+
+				MeshFormat->SetSource(MeshPtr.get());
+
+				SurfNode->MeshAdd = MeshFormat;
+			}
+
+			const int32 NumImages = TypedNodeExt->GetNumParameters(EMaterialParameterType::Texture);
+			SurfNode->Textures.SetNum(NumImages);
+			for (int32 ImageIndex = 0; ImageIndex < NumImages; ++ImageIndex)
+			{
+				mu::NodeImagePtr ImageNode;
+
+				if (!ImageNode) // If
+				{
+					const FString MaterialImageId = FGroupProjectorImageInfo::GenerateId(TypedNodeExt, ImageIndex);
+					const FGroupProjectorImageInfo* ProjectorInfo = GenerationContext.GroupProjectorLODCache.Find(MaterialImageId);
+
+					if (ProjectorInfo)
+					{
+						ensure(LOD > GenerationContext.FirstLODAvailable);
+						check(ProjectorInfo->SurfNode->Images[ImageIndex].Image == ProjectorInfo->ImageNode);
+						ImageNode = ProjectorInfo->ImageNode;
+
+						//TextureNameToProjectionResFactor.Add(ProjectorInfo->RealTextureName, ProjectorInfo->AlternateProjectionResolutionFactor);
+						//AlternateResStateName = ProjectorInfo->AlternateResStateName;
+					}
+				}
+
+				if (!ImageNode) // Else if
+				{
+					bool bShareProjectionTexturesBetweenLODs = false;
+					bool bIsGroupProjectorImage = false;
+					UTexture2D* GroupProjectionReferenceTexture = nullptr;
+					TMap<FString, float> TextureNameToProjectionResFactor;
+					FString AlternateResStateName;
+
+					ImageNode = GenerateMutableGroupProjection(LOD, ImageIndex, AddMeshNode, GenerationContext,
+						nullptr, TypedNodeExt, bShareProjectionTexturesBetweenLODs, bIsGroupProjectorImage,
+						GroupProjectionReferenceTexture, TextureNameToProjectionResFactor, AlternateResStateName);
+				}
+
+				if (!ImageNode) // Else if
+				{
+					const FNodeMaterialParameterId ImageId = TypedNodeExt->GetParameterId(EMaterialParameterType::Texture, ImageIndex);
+
+					if (TypedNodeExt->UsesImage(ImageId))
+					{
+						// TODO
+						//check(ParentMaterialNode->IsImageMutableMode(ImageIndex)); // Ensured at graph time. If it fails, something is wrong.
+
+						if (const UEdGraphPin* ConnectedPin = FollowInputPin(*TypedNodeExt->GetUsedImagePin(ImageId)))
+						{
+							// ReferenceTextureSize is used to limit the size of textures contributing to the final image.
+							const int32 ReferenceTextureSize = 0; // TODO GetBaseTextureSize(GenerationContext, TypedNodeExt, ImageIndex);
+
+							ImageNode = GenerateMutableSourceImage(ConnectedPin, GenerationContext, ReferenceTextureSize);
+						}
+					}
+				}
+
+				if (ImageNode)
+				{
+					SurfNode->Textures[ImageIndex].Extend = ImageNode;
+				}
+
+			}
+
+			for (const FString& Tag : TypedNodeExt->Tags)
+			{
+				SurfNode->EnableTags.Add(Tag);
+			}
+		}();
+
+		GenerationContext.MeshGenerationFlags.Pop();
+	}
+
+	else if (const UCustomizableObjectNodeModifierRemoveMesh* TypedNodeRem = Cast<UCustomizableObjectNodeModifierRemoveMesh>(Node))
+	{
+		const EMutableMeshConversionFlags ModifiersMeshFlags =
+			EMutableMeshConversionFlags::IgnoreSkinning |
+			EMutableMeshConversionFlags::IgnorePhysics;
+		GenerationContext.MeshGenerationFlags.Push(ModifiersMeshFlags);
+
+		mu::Ptr<mu::NodeModifierSurfaceEdit> SurfNode = new mu::NodeModifierSurfaceEdit();
+		Result = SurfNode;
+
+		SurfNode->MultipleTagsPolicy = TypedNodeRem->MultipleTagPolicy;
+		for (const FString& Tag : TypedNodeRem->RequiredTags)
+		{
+			SurfNode->RequiredTags.Add(Tag);
+		}
+
+		mu::Ptr<mu::NodeMesh> RemoveMeshNode;
+
+		if (const UEdGraphPin* ConnectedPin = FollowInputPin(*TypedNodeRem->RemoveMeshPin()))
+		{
+			FMutableGraphMeshGenerationData DummyMeshData;
+			RemoveMeshNode = GenerateMutableSourceMesh(ConnectedPin, GenerationContext, DummyMeshData, false, true);
+		}
+
+		SurfNode->MeshRemove = RemoveMeshNode;
+
+		GenerationContext.MeshGenerationFlags.Pop();
+	}
+
+	else if (const UCustomizableObjectNodeModifierRemoveMeshBlocks* TypedNodeRemBlocks = Cast<UCustomizableObjectNodeModifierRemoveMeshBlocks>(Node))
+	{
+		const EMutableMeshConversionFlags ModifiersMeshFlags =
+			EMutableMeshConversionFlags::IgnoreSkinning |
+			EMutableMeshConversionFlags::IgnorePhysics;
+		GenerationContext.MeshGenerationFlags.Push(ModifiersMeshFlags);
+
+		mu::Ptr<mu::NodeModifierMeshClipWithUVMask> ClipNode = new mu::NodeModifierMeshClipWithUVMask();
+		Result = ClipNode;
+
+		ClipNode->MultipleTagsPolicy = TypedNodeRemBlocks->MultipleTagPolicy;
+		for (const FString& Tag : TypedNodeRemBlocks->RequiredTags)
+		{
+			ClipNode->RequiredTags.Add(Tag);
+		}
+
+		bool bWasEmpty = false;
+		mu::Ptr<mu::NodeLayout> SourceLayout = CreateMutableLayoutNode(GenerationContext, TypedNodeRemBlocks->Layout, true, bWasEmpty);
+		ClipNode->ClipLayout = SourceLayout;
+		ClipNode->LayoutIndex = TypedNodeRemBlocks->ParentLayoutIndex;
+
+		GenerationContext.MeshGenerationFlags.Pop();
+	}
+
+	else if (UCustomizableObjectNodeModifierEditMeshSection* TypedNodeEdit = Cast<UCustomizableObjectNodeModifierEditMeshSection>(Node))
+	{
+		const EMutableMeshConversionFlags ModifiersMeshFlags =
+			EMutableMeshConversionFlags::IgnoreSkinning |
+			EMutableMeshConversionFlags::IgnorePhysics;
+		GenerationContext.MeshGenerationFlags.Push(ModifiersMeshFlags);
+
+		mu::Ptr<mu::NodeModifierSurfaceEdit> SurfNode = new mu::NodeModifierSurfaceEdit();
+		Result = SurfNode;
+
+		SurfNode->MultipleTagsPolicy = TypedNodeEdit->MultipleTagPolicy;
+		for (const FString& Tag : TypedNodeEdit->RequiredTags)
+		{
+			SurfNode->RequiredTags.Add(Tag);
+		}
+
+		const int32 NumImages = TypedNodeEdit->GetNumParameters(EMaterialParameterType::Texture);
+		SurfNode->Textures.SetNum(NumImages);
+		for (int32 ImageIndex = 0; ImageIndex < NumImages; ++ImageIndex)
+		{
+			const FNodeMaterialParameterId ImageId = TypedNodeEdit->GetParameterId(EMaterialParameterType::Texture, ImageIndex);
+
+			if (TypedNodeEdit->UsesImage(ImageId))
+			{
+				// TODO
+				//check(ParentMaterialNode->IsImageMutableMode(ImageIndex)); // Ensured at graph time. If it fails, something is wrong.
+
+				const UEdGraphPin* ConnectedImagePin = FollowInputPin(*TypedNodeEdit->GetUsedImagePin(ImageId));
+
+				mu::NodeModifierSurfaceEdit::FTexture& ImagePatch = SurfNode->Textures[ImageIndex];
+
+				// \todo: expose these two options?
+				ImagePatch.PatchBlendType = mu::EBlendType::BT_BLEND;
+				ImagePatch.bPatchApplyToAlpha = true;
+
+				// ReferenceTextureSize is used to limit the size of textures contributing to the final image.
+				const int32 ReferenceTextureSize = 0; //TODO GetBaseTextureSize(GenerationContext, ParentMaterialNode, ImageIndex);
+
+				ImagePatch.PatchImage = GenerateMutableSourceImage(ConnectedImagePin, GenerationContext, ReferenceTextureSize);
+
+				const UEdGraphPin* ImageMaskPin = TypedNodeEdit->GetUsedImageMaskPin(ImageId);
+				check(ImageMaskPin); // Ensured when reconstructing EditMaterial nodes. If it fails, something is wrong.
+
+				if (const UEdGraphPin* ConnectedMaskPin = FollowInputPin(*ImageMaskPin))
+				{
+					ImagePatch.PatchMask = GenerateMutableSourceImage(ConnectedMaskPin, GenerationContext, ReferenceTextureSize);
+				}
+
+				// Add the blocks to patch
+				FIntPoint GridSize = TypedNodeEdit->Layout->GetGridSize();
+				FVector2f GridSizeF = FVector2f(GridSize);
+				ImagePatch.PatchBlocks.Reserve(TypedNodeEdit->Layout->Blocks.Num());
+				for (const FCustomizableObjectLayoutBlock& LayoutBlock : TypedNodeEdit->Layout->Blocks)
+				{
+					FBox2f Rect;
+					Rect.Min = FVector2f(LayoutBlock.Min) / GridSizeF;
+					Rect.Max = FVector2f(LayoutBlock.Max) / GridSizeF;
+					ImagePatch.PatchBlocks.Add(Rect);
+				}
+			}
+		}
+
+		GenerationContext.MeshGenerationFlags.Pop();
+	}
+
+	else if (const UCustomizableObjectNodeModifierMorphMeshSection* TypedNodeMorph = Cast<UCustomizableObjectNodeModifierMorphMeshSection>(Node))
+	{
+		const EMutableMeshConversionFlags ModifiersMeshFlags =
+			EMutableMeshConversionFlags::IgnoreSkinning |
+			EMutableMeshConversionFlags::IgnorePhysics;
+		GenerationContext.MeshGenerationFlags.Push(ModifiersMeshFlags);
+
+		mu::Ptr<mu::NodeModifierSurfaceEdit> SurfNode = new mu::NodeModifierSurfaceEdit();
+		Result = SurfNode;
+
+		// This modifier needs to be applied right after the mesh constant is generated
+		SurfNode->bApplyBeforeNormalOperations = true;
+
+		SurfNode->MultipleTagsPolicy = TypedNodeMorph->MultipleTagPolicy;
+		for (const FString& Tag : TypedNodeMorph->RequiredTags)
+		{
+			SurfNode->RequiredTags.Add(Tag);
+		}
+
+		SurfNode->MeshMorph = TypedNodeMorph->MorphTargetName;
+
+		if (const UEdGraphPin* ConnectedPin = FollowInputPin(*TypedNodeMorph->FactorPin()))
+		{
+			UEdGraphNode* floatNode = ConnectedPin->GetOwningNode();
+			bool validStaticFactor = true;
+			if (const UCustomizableObjectNodeFloatParameter* floatParameterNode = Cast<UCustomizableObjectNodeFloatParameter>(floatNode))
+			{
+				if (floatParameterNode->DefaultValue < -1.0f || floatParameterNode->DefaultValue > 1.0f)
+				{
+					validStaticFactor = false;
+					FString msg = FString::Printf(TEXT("Mesh morph nodes only accept factors between -1.0 and 1.0 inclusive but the default value of the float parameter node is (%f). Factor will be ignored."), floatParameterNode->DefaultValue);
+					GenerationContext.Compiler->CompilerLog(FText::FromString(msg), Node);
+				}
+				if (floatParameterNode->ParamUIMetadata.MinimumValue < -1.0f)
+				{
+					validStaticFactor = false;
+					FString msg = FString::Printf(TEXT("Mesh morph nodes only accept factors between -1.0 and 1.0 inclusive but the minimum UI value for the input float parameter node is (%f). Factor will be ignored."), floatParameterNode->ParamUIMetadata.MinimumValue);
+					GenerationContext.Compiler->CompilerLog(FText::FromString(msg), Node);
+				}
+				if (floatParameterNode->ParamUIMetadata.MaximumValue > 1.0f)
+				{
+					validStaticFactor = false;
+					FString msg = FString::Printf(TEXT("Mesh morph nodes only accept factors between -1.0 and 1.0 inclusive but the maximum UI value for the input float parameter node is (%f). Factor will be ignored."), floatParameterNode->ParamUIMetadata.MaximumValue);
+					GenerationContext.Compiler->CompilerLog(FText::FromString(msg), Node);
+				}
+			}
+			else if (const UCustomizableObjectNodeFloatConstant* floatConstantNode = Cast<UCustomizableObjectNodeFloatConstant>(floatNode))
+			{
+				if (floatConstantNode->Value < -1.0f || floatConstantNode->Value > 1.0f)
+				{
+					validStaticFactor = false;
+					FString msg = FString::Printf(TEXT("Mesh morph nodes only accept factors between -1.0 and 1.0 inclusive but the value of the float constant node is (%f). Factor will be ignored."), floatConstantNode->Value);
+					GenerationContext.Compiler->CompilerLog(FText::FromString(msg), Node);
+				}
+			}
+
+			if (validStaticFactor)
+			{
+				mu::NodeScalarPtr FactorNode = GenerateMutableSourceFloat(ConnectedPin, GenerationContext);
+				SurfNode->MorphFactor = FactorNode;
+			}
+		}
+
+		GenerationContext.MeshGenerationFlags.Pop();
+	}
+
+	else
+	{
+		GenerationContext.Compiler->CompilerLog(LOCTEXT("UnimplementedNode", "Node type not implemented yet."), Node);
+	}
+
+	if (!bDoNotAddToGeneratedCache)
+	{
+		GenerationContext.Generated.Add(Key, FGeneratedData(Node, Result));
+	}
+	GenerationContext.GeneratedNodes.Add(Node);
+	
+
+	return Result;
+}
+
+#undef LOCTEXT_NAMESPACE
+
diff --git a/Engine/Plugins/Mutable/Source/CustomizableObjectEditor/Private/MuCOE/GenerateMutableSource/GenerateMutableSourceModifier.h b/Engine/Plugins/Mutable/Source/CustomizableObjectEditor/Private/MuCOE/GenerateMutableSource/GenerateMutableSourceModifier.h
new file mode 100644
--- /dev/null
+++ b/Engine/Plugins/Mutable/Source/CustomizableObjectEditor/Private/MuCOE/GenerateMutableSource/GenerateMutableSourceModifier.h
@@ -0,0 +1,10 @@
+// Copyright Epic Games, Inc. All Rights Reserved.
+
+#pragma once
+
+#include "MuT/NodeModifier.h"
+
+class UEdGraphPin;
+struct FMutableGraphGenerationContext;
+
+mu::Ptr<mu::NodeModifier> GenerateMutableSourceModifier(const UEdGraphPin* Pin, FMutableGraphGenerationContext& GenerationContext);
diff --git a/Engine/Plugins/Mutable/Source/CustomizableObjectEditor/Private/MuCOE/GenerateMutableSource/GenerateMutableSourceProjector.h b/Engine/Plugins/Mutable/Source/CustomizableObjectEditor/Private/MuCOE/GenerateMutableSource/GenerateMutableSourceProjector.h
new file mode 100644
--- /dev/null
+++ b/Engine/Plugins/Mutable/Source/CustomizableObjectEditor/Private/MuCOE/GenerateMutableSource/GenerateMutableSourceProjector.h
@@ -0,0 +1,11 @@
+// Copyright Epic Games, Inc. All Rights Reserved.
+
+#pragma once
+
+#include "MuT/NodeProjector.h"
+
+class UEdGraphPin;
+struct FMutableGraphGenerationContext;
+
+/** Convert a CustomizableObject Source Graph into a mutable source graph. */
+mu::NodeProjectorPtr GenerateMutableSourceProjector(const UEdGraphPin* Pin, FMutableGraphGenerationContext& GenerationContext);
