[mutable] Moved the Mutable plugin out of Experimental status into Beta.

diff --git a/Engine/Plugins/Mutable/Source/CustomizableObjectEditor/Private/MuCOE/GenerateMutableSource/GenerateMutableSourceSurface.cpp b/Engine/Plugins/Mutable/Source/CustomizableObjectEditor/Private/MuCOE/GenerateMutableSource/GenerateMutableSourceSurface.cpp
new file mode 100644
--- /dev/null
+++ b/Engine/Plugins/Mutable/Source/CustomizableObjectEditor/Private/MuCOE/GenerateMutableSource/GenerateMutableSourceSurface.cpp
@@ -0,0 +1,1201 @@
+// Copyright Epic Games, Inc. All Rights Reserved.
+
+#include "MuCOE/GenerateMutableSource/GenerateMutableSourceSurface.h"
+
+#include "Engine/SkinnedAssetCommon.h"
+#include "Engine/StaticMesh.h"
+#include "GPUSkinPublicDefs.h"
+#include "Interfaces/ITargetPlatform.h"
+#include "Interfaces/ITextureFormat.h"
+#include "Interfaces/ITextureFormatModule.h"
+#include "Interfaces/ITextureFormatManagerModule.h"
+#include "Modules/ModuleManager.h"
+#include "TextureCompressorModule.h"
+#include "Materials/MaterialInstance.h"
+#include "GPUSkinVertexFactory.h"
+#include "Rendering/SkeletalMeshLODModel.h"
+
+#include "MuCO/CustomizableObjectInstance.h"
+#include "MuCO/MutableMeshBufferUtils.h"
+#include "MuCO/UnrealConversionUtils.h"
+#include "MuCOE/CustomizableObjectCompiler.h"
+#include "MuCOE/CustomizableObjectLayout.h"
+#include "MuCOE/GenerateMutableSource/GenerateMutableSourceColor.h"
+#include "MuCOE/GenerateMutableSource/GenerateMutableSourceFloat.h"
+#include "MuCOE/GenerateMutableSource/GenerateMutableSourceGroupProjector.h"
+#include "MuCOE/GenerateMutableSource/GenerateMutableSourceImage.h"
+#include "MuCOE/GenerateMutableSource/GenerateMutableSourceMesh.h"
+#include "MuCOE/GenerateMutableSource/GenerateMutableSourceTable.h"
+#include "MuCOE/GenerateMutableSource/GenerateMutableSourceLayout.h"
+#include "MuCOE/GraphTraversal.h"
+#include "MuCOE/MutableUtils.h"
+#include "MuCOE/Nodes/CustomizableObjectNodeCopyMaterial.h"
+#include "MuCOE/Nodes/CustomizableObjectNodeFloatConstant.h"
+#include "MuCOE/Nodes/CustomizableObjectNodeFloatParameter.h"
+#include "MuCOE/Nodes/CustomizableObjectNodeMaterialVariation.h"
+#include "MuCOE/Nodes/CustomizableObjectNodeMaterialSwitch.h"
+#include "MuCOE/Nodes/CustomizableObjectNodeModifierEditMeshSection.h"
+#include "MuCOE/Nodes/CustomizableObjectNodeModifierExtendMeshSection.h"
+#include "MuCOE/Nodes/CustomizableObjectNodeModifierMorphMeshSection.h"
+#include "MuCOE/Nodes/CustomizableObjectNodeModifierRemoveMesh.h"
+#include "MuCOE/Nodes/CustomizableObjectNodeModifierRemoveMeshBlocks.h"
+#include "MuCOE/Nodes/CustomizableObjectNodeSkeletalMesh.h"
+#include "MuCOE/Nodes/CustomizableObjectNodeTable.h"
+#include "MuT/NodeImageFormat.h"
+#include "MuT/NodeImageMipmap.h"
+#include "MuT/NodeImageNormalComposite.h"
+#include "MuT/NodeImageResize.h"
+#include "MuT/NodeImageSwizzle.h"
+#include "MuT/NodeMeshConstant.h"
+#include "MuT/NodeMeshFormat.h"
+#include "MuT/NodeMeshFragment.h"
+#include "MuT/NodeScalarConstant.h"
+#include "MuT/NodeSurfaceSwitch.h"
+#include "MuT/NodeSurfaceVariation.h"
+#include "MuT/UnrealPixelFormatOverride.h"
+
+#define LOCTEXT_NAMESPACE "CustomizableObjectEditor"
+
+
+void SetSurfaceFormat( FMutableGraphGenerationContext& GenerationContext,
+					   mu::FMeshBufferSet& OutVertexBufferFormat, mu::FMeshBufferSet& OutIndexBufferFormat, const FMutableGraphMeshGenerationData& MeshData, 
+					   ECustomizableObjectNumBoneInfluences ECustomizableObjectNumBoneInfluences, bool bWith16BitWeights)
+{
+	// Limit skinning weights if necessary
+	// \todo: make it more flexible to support 3 or 5 or 1 weight, since there is support for this in 4.25
+	const int32 MutableBonesPerVertex = FGPUBaseSkinVertexFactory::UseUnlimitedBoneInfluences(MeshData.MaxNumBonesPerVertex, GenerationContext.Options.TargetPlatform) &&
+										MeshData.MaxNumBonesPerVertex < (int32)ECustomizableObjectNumBoneInfluences ?
+											MeshData.MaxNumBonesPerVertex :
+											(int32)ECustomizableObjectNumBoneInfluences;
+
+	ensure(MutableBonesPerVertex <= MAX_TOTAL_INFLUENCES);
+	
+	if (MutableBonesPerVertex != MeshData.MaxNumBonesPerVertex)
+	{
+		UE_LOG(LogMutable, Verbose, TEXT("In object [%s] Mesh bone number adjusted from %d to %d."), *GenerationContext.Object->GetName(), MeshData.MaxNumBonesPerVertex, MutableBonesPerVertex);
+	}
+
+	int MutableBufferCount = MUTABLE_VERTEXBUFFER_TEXCOORDS + 1;
+	if (MeshData.bHasVertexColors)
+	{
+		++MutableBufferCount;
+	}
+
+	if (MeshData.MaxNumBonesPerVertex > 0 && MeshData.MaxBoneIndexTypeSizeBytes > 0)
+	{
+		++MutableBufferCount;
+	}
+
+	if (MeshData.bHasRealTimeMorphs)
+	{
+		MutableBufferCount += 2;
+	}
+
+	if (MeshData.bHasClothing)
+	{
+		MutableBufferCount += 2;
+	}
+
+	MutableBufferCount += MeshData.SkinWeightProfilesSemanticIndices.Num();
+
+	OutVertexBufferFormat.SetBufferCount(MutableBufferCount);
+
+	int32 CurrentVertexBuffer = 0;
+
+	// Vertex buffer
+	MutableMeshBufferUtils::SetupVertexPositionsBuffer(CurrentVertexBuffer, OutVertexBufferFormat);
+	++CurrentVertexBuffer;
+
+	// Tangent buffer
+	MutableMeshBufferUtils::SetupTangentBuffer(CurrentVertexBuffer, OutVertexBufferFormat);
+	++CurrentVertexBuffer;
+
+	// Texture coords buffer
+	MutableMeshBufferUtils::SetupTexCoordinatesBuffer(CurrentVertexBuffer, MeshData.NumTexCoordChannels, OutVertexBufferFormat);
+	++CurrentVertexBuffer;
+
+	// Skin buffer
+	if (MeshData.MaxNumBonesPerVertex > 0 && MeshData.MaxBoneIndexTypeSizeBytes > 0)
+	{
+		const int32 MaxBoneWeightTypeSizeBytes = bWith16BitWeights ? 2 : 1;
+		MutableMeshBufferUtils::SetupSkinBuffer(CurrentVertexBuffer, MeshData.MaxBoneIndexTypeSizeBytes, MaxBoneWeightTypeSizeBytes, MutableBonesPerVertex, OutVertexBufferFormat);
+		++CurrentVertexBuffer;
+	}
+
+	// Colour buffer
+	if (MeshData.bHasVertexColors)
+	{
+		MutableMeshBufferUtils::SetupVertexColorBuffer(CurrentVertexBuffer, OutVertexBufferFormat);
+		++CurrentVertexBuffer;
+	}
+
+	// MorphTarget vertex tracking info buffers
+	if (MeshData.bHasRealTimeMorphs)
+	{
+		using namespace mu;
+		{
+			const int32 ElementSize = sizeof(uint32);
+			constexpr int32 ChannelCount = 1;
+			const EMeshBufferSemantic Semantics[ChannelCount] = { MBS_OTHER };
+			const int32 SemanticIndices[ChannelCount] = { 0 };
+			const EMeshBufferFormat Formats[ChannelCount] = { MBF_UINT32 };
+			const int32 Components[ChannelCount] = { 1 };
+			const int32 Offsets[ChannelCount] = { 0 };
+
+			OutVertexBufferFormat.SetBuffer(CurrentVertexBuffer, ElementSize, ChannelCount, Semantics, SemanticIndices, Formats, Components, Offsets);
+			++CurrentVertexBuffer;
+		}
+
+		{
+			const int32 ElementSize = sizeof(uint32);
+			constexpr int32 ChannelCount = 1;
+			const EMeshBufferSemantic Semantics[ChannelCount] = { MBS_OTHER };
+			const int32 SemanticIndices[ChannelCount] = { 1 };
+			const EMeshBufferFormat Formats[ChannelCount] = { MBF_UINT32 };
+			const int32 Components[ChannelCount] = { 1 };
+			const int32 Offsets[ChannelCount] = { 0 };
+
+			OutVertexBufferFormat.SetBuffer(CurrentVertexBuffer, ElementSize, ChannelCount, Semantics, SemanticIndices, Formats, Components, Offsets);
+			++CurrentVertexBuffer;
+		}
+	}
+
+	//Clothing Data Buffer.
+	if (MeshData.bHasClothing)
+	{
+		{
+			using namespace mu;
+			const int32 ElementSize = sizeof(int32);
+			constexpr int32 ChannelCount = 1;
+			const EMeshBufferSemantic Semantics[ChannelCount] = { MBS_OTHER };
+			const int32 SemanticIndices[ChannelCount] = { 2 };
+			const EMeshBufferFormat Formats[ChannelCount] = { MBF_INT32 };
+			const int32 Components[ChannelCount] = { 1 };
+			const int32 Offsets[ChannelCount] = { 0 };
+
+			OutVertexBufferFormat.SetBuffer(CurrentVertexBuffer, ElementSize, ChannelCount, Semantics, SemanticIndices, Formats, Components, Offsets);
+			++CurrentVertexBuffer;
+		}
+
+		{
+			using namespace mu;
+			const int32 ElementSize = sizeof(uint32);
+			constexpr int32 ChannelCount = 1;
+			const EMeshBufferSemantic Semantics[ChannelCount] = { MBS_OTHER };
+			const int32 SemanticIndices[ChannelCount] = { 3 };
+			const EMeshBufferFormat Formats[ChannelCount] = { MBF_UINT32 };
+			const int32 Components[ChannelCount] = { 1 };
+			const int32 Offsets[ChannelCount] = { 0 };
+
+			OutVertexBufferFormat.SetBuffer(CurrentVertexBuffer, ElementSize, ChannelCount, Semantics, SemanticIndices, Formats, Components, Offsets);
+			++CurrentVertexBuffer;
+		}
+	}
+
+	for (int32 ProfileSemanticIndex : MeshData.SkinWeightProfilesSemanticIndices)
+	{
+		MutableMeshBufferUtils::SetupSkinWeightProfileBuffer(CurrentVertexBuffer, MeshData.MaxBoneIndexTypeSizeBytes, 1, MutableBonesPerVertex, ProfileSemanticIndex, OutVertexBufferFormat);
+		++CurrentVertexBuffer;
+	}
+
+	// Index buffer
+	MutableMeshBufferUtils::SetupIndexBuffer(OutIndexBufferFormat);
+}
+
+
+mu::Ptr<mu::NodeSurface> GenerateMutableSourceSurface(const UEdGraphPin * Pin, FMutableGraphGenerationContext & GenerationContext)
+{
+	MUTABLE_CPUPROFILER_SCOPE(GenerateMutableSourceSurface);
+
+	check(Pin)
+	RETURN_ON_CYCLE(*Pin, GenerationContext)
+
+	CheckNumOutputs(*Pin, GenerationContext);
+
+	UCustomizableObjectNode* Node = CastChecked<UCustomizableObjectNode>(Pin->GetOwningNode());
+
+	const FGeneratedKey Key(reinterpret_cast<void*>(&GenerateMutableSourceSurface), *Pin, *Node, GenerationContext, true);
+	if (const FGeneratedData* Generated = GenerationContext.Generated.Find(Key))
+	{
+		return static_cast<mu::NodeSurface*>(Generated->Node.get());
+	}
+	
+	mu::Ptr<mu::NodeSurface> Result;
+
+	const int32 LOD = Node->IsAffectedByLOD() ? GenerationContext.CurrentLOD : 0;
+	
+	if (UCustomizableObjectNode* CustomObjNode = Cast<UCustomizableObjectNode>(Node))
+	{
+		if (CustomObjNode->IsNodeOutDatedAndNeedsRefresh())
+		{
+			CustomObjNode->SetRefreshNodeWarning();
+		}
+	}
+
+	if (UCustomizableObjectNodeMaterialBase* TypedNodeMat = Cast<UCustomizableObjectNodeMaterialBase>(Node))
+	{
+		bool bGeneratingImplicitComponent = GenerationContext.ComponentMeshOverride.get() != nullptr;
+
+		const UEdGraphPin* ConnectedMaterialPin = FollowInputPin(*TypedNodeMat->GetMeshPin());
+		// Warn when texture connections are improperly used by connecting them directly to material inputs when no layout is used
+		// TODO: delete the if clause and the warning when static meshes are operational again
+		if (ConnectedMaterialPin)
+		{
+			if (const UEdGraphPin* StaticMeshPin = FindMeshBaseSource(*ConnectedMaterialPin, true))
+			{
+				const UCustomizableObjectNode* StaticMeshNode = CastChecked<UCustomizableObjectNode>(StaticMeshPin->GetOwningNode());
+				GenerationContext.Compiler->CompilerLog(LOCTEXT("UnsupportedStaticMeshes", "Static meshes are currently not supported as material meshes"), StaticMeshNode);
+			}
+		}
+
+		mu::Ptr<mu::NodeSurfaceNew> SurfNode = new mu::NodeSurfaceNew();
+		Result = SurfNode;
+
+		// Add to the list of surfaces that could be reused between LODs for this NodeMaterial.
+		TArray<FMutableGraphGenerationContext::FSharedSurface>& SharedSurfaces = GenerationContext.SharedSurfaceIds.FindOrAdd(TypedNodeMat, {});
+		FMutableGraphGenerationContext::FSharedSurface& SharedSurface = SharedSurfaces.Add_GetRef(FMutableGraphGenerationContext::FSharedSurface(GenerationContext.CurrentLOD, SurfNode));
+		SharedSurface.bMakeUnique = !TypedNodeMat->IsReuseMaterialBetweenLODs();
+
+		int32 ReferencedMaterialsIndex = -1;
+		uint32 SurfaceMetadataUniqueHash = 0; // Value 0 is used as invalid hash.
+		if (TypedNodeMat->GetMaterial())
+		{
+			GenerationContext.AddParticipatingObject(*TypedNodeMat->GetMaterial());
+
+			ReferencedMaterialsIndex = GenerationContext.ReferencedMaterials.AddUnique(TypedNodeMat->GetMaterial());
+			if (ConnectedMaterialPin)
+			{
+				if (const UEdGraphPin* SkeletalMeshPin = FindMeshBaseSource(*ConnectedMaterialPin, false))
+				{
+					FSkeletalMaterial* SkeletalMaterial = nullptr;
+					const FSkelMeshSection* ReferenceSkelMeshSection = nullptr;
+					
+					if (const UCustomizableObjectNodeSkeletalMesh* SkeletalMeshNode = Cast<UCustomizableObjectNodeSkeletalMesh>(SkeletalMeshPin->GetOwningNode()))
+					{
+						SkeletalMaterial = SkeletalMeshNode->GetSkeletalMaterialFor(*SkeletalMeshPin);
+						ReferenceSkelMeshSection = SkeletalMeshNode->GetSkeletalMeshSectionFor(*SkeletalMeshPin);
+					}
+
+					else if (const UCustomizableObjectNodeTable* TableNode = Cast<UCustomizableObjectNodeTable>(SkeletalMeshPin->GetOwningNode()))
+					{
+						SkeletalMaterial = TableNode->GetDefaultSkeletalMaterialFor(*SkeletalMeshPin);
+						ReferenceSkelMeshSection = TableNode->GetDefaultSkeletalMeshSectionFor(*SkeletalMeshPin);
+					}
+
+					// A reference to the material slot name index will be kept in the surface metadata.
+					int32 MaterialSlotNameIndex = INDEX_NONE;
+					if (SkeletalMaterial)
+					{
+						MaterialSlotNameIndex = 
+								GenerationContext.ReferencedMaterialSlotNames.AddUnique(SkeletalMaterial->MaterialSlotName);
+					}
+
+					if (ReferenceSkelMeshSection)
+					{
+						auto HashSurfaceMetadataFunc = [](const FMutableSurfaceMetadata& Data) -> uint32
+						{
+							return CityHash32(reinterpret_cast<const char*>(&Data), sizeof(FMutableSurfaceMetadata));
+						};
+
+						auto CompareSurfaceMetadataFunc = [](const FMutableSurfaceMetadata& A, const FMutableSurfaceMetadata& B)
+						{
+							return FMemory::Memcmp(&A, &B, sizeof(FMutableSurfaceMetadata)) == 0;
+						};
+
+						FMutableSurfaceMetadata SurfaceMetadata;
+						FMemory::Memzero(SurfaceMetadata);
+						
+						SurfaceMetadata.MaterialSlotIndex = MaterialSlotNameIndex;
+						SurfaceMetadata.bCastShadow = ReferenceSkelMeshSection->bCastShadow;
+
+						SurfaceMetadataUniqueHash = Private::GenerateUniquePersistentHash(
+								SurfaceMetadata, GenerationContext.SurfaceMetadata, HashSurfaceMetadataFunc, CompareSurfaceMetadataFunc);
+
+						if (SurfaceMetadataUniqueHash != 0)
+						{
+							GenerationContext.SurfaceMetadata.FindOrAdd(SurfaceMetadataUniqueHash, SurfaceMetadata);
+						}
+						else
+						{
+							UE_LOG(LogMutable, Error, TEXT("Maximum number of surfaces reached."));
+						}
+					}
+				}
+			}
+		}
+
+		mu::Ptr<mu::NodeMesh> MeshNode;
+		
+		if (bGeneratingImplicitComponent)
+		{
+			MeshNode = GenerationContext.ComponentMeshOverride;
+			SurfNode->Mesh = MeshNode;
+
+			if (const UEdGraphPin* ConnectedPin = FollowInputPin(*TypedNodeMat->GetMeshPin()))
+			{
+				GenerationContext.Compiler->CompilerLog(LOCTEXT("MeshIgnored", "The mesh nodes connected to a material node will be ignored because it is part of an explicit mesh component."), Node);
+			}
+		}
+		else
+		{
+			if (const UEdGraphPin* ConnectedPin = FollowInputPin(*TypedNodeMat->GetMeshPin()))
+			{
+
+				// Flags to know which UV channels need layout
+				FLayoutGenerationFlags LayoutGenerationFlags;
+				
+				LayoutGenerationFlags.TexturePinModes.Init(EPinMode::Default, TEXSTREAM_MAX_NUM_UVCHANNELS);
+
+				const int32 NumImages = TypedNodeMat->GetNumParameters(EMaterialParameterType::Texture);
+				for (int32 ImageIndex = 0; ImageIndex < NumImages; ++ImageIndex)
+				{
+					if (TypedNodeMat->IsImageMutableMode(ImageIndex))
+					{
+						const int32 UVChannel = TypedNodeMat->GetImageUVLayout(ImageIndex);
+						if (LayoutGenerationFlags.TexturePinModes.IsValidIndex(UVChannel))
+						{
+							LayoutGenerationFlags.TexturePinModes[UVChannel] = EPinMode::Mutable;
+						}
+					}
+				}
+
+				GenerationContext.LayoutGenerationFlags.Push(LayoutGenerationFlags);
+
+				FMutableGraphMeshGenerationData MeshData;
+				MeshNode = GenerateMutableSourceMesh(ConnectedPin, GenerationContext, MeshData, false, false);
+
+				GenerationContext.LayoutGenerationFlags.Pop();
+
+				if (MeshNode)
+				{
+					mu::Ptr<mu::NodeMeshFormat> MeshFormatNode = new mu::NodeMeshFormat();
+					MeshFormatNode->SetSource(MeshNode.get());
+					SetSurfaceFormat(GenerationContext,
+						MeshFormatNode->GetVertexBuffers(), MeshFormatNode->GetIndexBuffers(), MeshData,
+						GenerationContext.Options.CustomizableObjectNumBoneInfluences,
+						GenerationContext.Options.b16BitBoneWeightsEnabled);
+
+					// \TODO: Make it an option?
+					MeshFormatNode->SetOptimizeBuffers(true);
+
+					MeshFormatNode->SetMessageContext(Node);
+
+					SurfNode->ExternalId = SurfaceMetadataUniqueHash;
+					SurfNode->Mesh = MeshFormatNode;
+				}
+				else
+				{
+					GenerationContext.Compiler->CompilerLog(LOCTEXT("MeshFailed", "Mesh generation failed."), Node);
+				}
+			}
+		}
+		
+		TMap<FString, float> TextureNameToProjectionResFactor;
+		FString AlternateResStateName;
+
+		bool bTableMaterialPinLinked = TypedNodeMat->GetMaterialAssetPin() && FollowInputPin(*TypedNodeMat->GetMaterialAssetPin()) != nullptr;
+		FString TableColumnName;
+
+		// Checking if we should not use the material of the table node even if it is linked to the material node
+		const UEdGraphPin* MaterialAssetConnectedPin = nullptr;
+		if (TypedNodeMat->GetMaterialAssetPin())
+		{
+			MaterialAssetConnectedPin = FollowInputPin(*TypedNodeMat->GetMaterialAssetPin());
+		}
+
+		if (MaterialAssetConnectedPin)
+		{
+			if (const UCustomizableObjectNodeTable* TypedNodeTable = Cast< UCustomizableObjectNodeTable >(MaterialAssetConnectedPin->GetOwningNode()))
+			{
+				TableColumnName = MaterialAssetConnectedPin->PinFriendlyName.ToString();
+
+				if (UMaterialInstance * TableMaterial = TypedNodeTable->GetColumnDefaultAssetByType<UMaterialInstance>(MaterialAssetConnectedPin))
+				{
+					// Checking if the reference material of the Table Node has the same parent as the material of the Material Node 
+					if (!TypedNodeMat->GetMaterial() || TableMaterial->GetMaterial() != TypedNodeMat->GetMaterial()->GetMaterial())
+					{
+						bTableMaterialPinLinked = false;
+
+						GenerationContext.Compiler->CompilerLog(LOCTEXT("DifferentParentMaterial", "The Deafult Material Instance of the Data Table must have the same Parent Material."), TypedNodeMat->GetMaterialNode());
+					}
+				}
+				else
+				{
+					FText Msg = FText::Format(LOCTEXT("DefaultValueNotFound", "Couldn't find a default value in the data table's struct for the column {0}. The default value is null or not a Material Instance."), FText::FromString(TableColumnName));
+					GenerationContext.Compiler->CompilerLog(Msg, Node);
+
+					bTableMaterialPinLinked = false;
+				}
+			}
+		}
+
+		int32 NumImages = TypedNodeMat->GetNumParameters(EMaterialParameterType::Texture);
+		SurfNode->Images.SetNum(NumImages);
+
+		if (!GenerationContext.Options.TargetPlatform || GenerationContext.Options.TargetPlatform->IsServerOnly())
+		{
+			// Don't generate the images in the server
+			NumImages = 0;
+		}
+
+		for (int32 ImageIndex = 0; ImageIndex < NumImages; ++ImageIndex)
+		{
+			const UEdGraphPin* ImagePin = TypedNodeMat->GetParameterPin(EMaterialParameterType::Texture, ImageIndex);
+
+			const bool bIsImagePinLinked = ImagePin && FollowInputPin(*ImagePin);
+
+			if (bIsImagePinLinked && !TypedNodeMat->IsImageMutableMode(ImageIndex))
+			{
+				if (const UEdGraphPin* ConnectedPin = FollowInputPin(*ImagePin))
+				{
+					// Find or add Image properties
+					const FGeneratedImagePropertiesKey PropsKey(TypedNodeMat, (uint32)ImageIndex);
+					const bool bNewImageProps = !GenerationContext.ImageProperties.Contains(PropsKey);
+
+					FGeneratedImageProperties& Props = GenerationContext.ImageProperties.FindOrAdd(PropsKey);
+					if (bNewImageProps)
+					{
+						// We don't need a reference texture or props here, but we do need the parameter name.
+						Props.TextureParameterName = TypedNodeMat->GetParameterName(EMaterialParameterType::Texture, ImageIndex).ToString();
+						Props.ImagePropertiesIndex = GenerationContext.ImageProperties.Num() - 1;
+						Props.bIsPassThrough = true;
+					}
+
+					// This is a connected pass-through texture that simply has to be passed to the core
+					mu::Ptr<mu::NodeImage> PassThroughImagePtr = GenerateMutableSourceImage(ConnectedPin, GenerationContext, 0);
+					SurfNode->Images[ImageIndex].Image = PassThroughImagePtr;
+
+					check(Props.ImagePropertiesIndex != INDEX_NONE);
+					const FString SurfNodeImageName = FString::Printf(TEXT("%d"), Props.ImagePropertiesIndex);
+					SurfNode->Images[ImageIndex].Name = SurfNodeImageName;
+					SurfNode->Images[ImageIndex].LayoutIndex = -1;
+					SurfNode->Images[ImageIndex].MaterialName = TypedNodeMat->GetMaterial()->GetName();
+					SurfNode->Images[ImageIndex].MaterialParameterName = Props.TextureParameterName;
+
+				}
+			}
+			else
+			{
+				mu::NodeImagePtr GroupProjectionImg;
+				UTexture2D* GroupProjectionReferenceTexture = nullptr;
+				const FString ImageName = TypedNodeMat->GetParameterName(EMaterialParameterType::Texture, ImageIndex).ToString();
+				const FNodeMaterialParameterId ImageId = TypedNodeMat->GetParameterId(EMaterialParameterType::Texture, ImageIndex);
+
+				FString MaterialImageId = FGroupProjectorImageInfo::GenerateId(TypedNodeMat, ImageIndex);
+				bool bShareProjectionTexturesBetweenLODs = false;
+				FGroupProjectorImageInfo* ProjectorInfo = GenerationContext.GroupProjectorLODCache.Find(MaterialImageId);
+
+				if (!ProjectorInfo) // No previous LOD of this material generated the image.
+				{
+					bool bIsGroupProjectorImage = false;
+
+					GroupProjectionImg = GenerateMutableGroupProjection(LOD, ImageIndex, MeshNode, GenerationContext,
+						TypedNodeMat, nullptr, bShareProjectionTexturesBetweenLODs, bIsGroupProjectorImage,
+						GroupProjectionReferenceTexture, TextureNameToProjectionResFactor, AlternateResStateName);
+
+					if (GroupProjectionImg.get() || TypedNodeMat->IsImageMutableMode(ImageIndex))
+					{
+						// Get the reference texture
+						UTexture2D* ReferenceTexture = nullptr;
+						{
+							//TODO(Max) UE-220247: Add support for multilayer materials
+							GenerationContext.CurrentMaterialTableParameter = ImageName;
+							GenerationContext.CurrentMaterialTableParameterId = ImageId.ParameterId.ToString();
+
+							ReferenceTexture = GroupProjectionImg.get() ? GroupProjectionReferenceTexture : nullptr;
+						
+							if (!ReferenceTexture)
+							{
+								ReferenceTexture = TypedNodeMat->GetImageReferenceTexture(ImageIndex);
+							}
+
+							// In case of group projector, don't follow the pin to find the reference texture.
+							if (!GroupProjectionImg.get() && !ReferenceTexture && ImagePin)
+							{
+								if (const UEdGraphPin* ConnectedPin = FollowInputPin(*ImagePin))
+								{
+									ReferenceTexture = FindReferenceImage(ConnectedPin, GenerationContext);
+								}
+							}
+
+							if (!ReferenceTexture && bTableMaterialPinLinked)
+							{
+								if (const UEdGraphPin* ConnectedPin = FollowInputPin(*TypedNodeMat->GetMaterialAssetPin()))
+								{
+									ReferenceTexture = FindReferenceImage(ConnectedPin, GenerationContext);
+								}
+							}	
+
+							if (!ReferenceTexture)
+							{
+								ReferenceTexture = TypedNodeMat->GetImageValue(ImageIndex);
+							}
+						}
+
+						const FGeneratedImagePropertiesKey PropsKey(TypedNodeMat, ImageIndex);
+						const bool bNewImageProps = !GenerationContext.ImageProperties.Contains(PropsKey);
+
+						FGeneratedImageProperties& Props = GenerationContext.ImageProperties.FindOrAdd(PropsKey);
+
+						if (bNewImageProps)
+						{
+							if (ReferenceTexture)
+							{
+								GenerationContext.AddParticipatingObject(*ReferenceTexture);
+
+								// Store properties for the generated images
+								Props.TextureParameterName = ImageName;
+								Props.ImagePropertiesIndex = GenerationContext.ImageProperties.Num() - 1;
+
+								Props.CompressionSettings = ReferenceTexture->CompressionSettings;
+								Props.Filter = ReferenceTexture->Filter;
+								Props.SRGB = ReferenceTexture->SRGB;
+								Props.LODBias = 0;
+								Props.MipGenSettings = ReferenceTexture->MipGenSettings;
+								Props.LODGroup = ReferenceTexture->LODGroup;
+								Props.AddressX = ReferenceTexture->AddressX;
+								Props.AddressY = ReferenceTexture->AddressY;
+								Props.bFlipGreenChannel = ReferenceTexture->bFlipGreenChannel;
+
+
+								// MaxTextureSize setting. Based on the ReferenceTexture and Platform settings.
+								const UTextureLODSettings& TextureLODSettings = GenerationContext.Options.TargetPlatform->GetTextureLODSettings();
+								Props.MaxTextureSize = GetMaxTextureSize(*ReferenceTexture, TextureLODSettings);
+
+								// ReferenceTexture source size. Textures contributing to this Image should be equal to or smaller than TextureSize. 
+								// The LOD Bias applied to the root node will be applied on top of it.
+								Props.TextureSize = (int32)FMath::Max3(ReferenceTexture->Source.GetSizeX(), ReferenceTexture->Source.GetSizeY(), 1LL);
+
+								// TODO: MTBL-1081
+								// TextureGroup::TEXTUREGROUP_UI does not support streaming. If we generate a texture that requires streaming and set this group, it will crash when initializing the resource. 
+								// If LODGroup == TEXTUREGROUP_UI, UTexture::IsPossibleToStream() will return false and UE will assume all mips are loaded, when they're not, and crash.
+								if (Props.LODGroup == TEXTUREGROUP_UI)
+								{
+									Props.LODGroup = TextureGroup::TEXTUREGROUP_Character;
+
+									FString msg = FString::Printf(TEXT("The Reference texture [%s] is using TEXTUREGROUP_UI which does not support streaming. Please set a different TEXTURE group."),
+										*ReferenceTexture->GetName(), *ImageName);
+									GenerationContext.Compiler->CompilerLog(FText::FromString(msg), Node, EMessageSeverity::Info);
+								}
+							}
+							else
+							{
+								// warning!
+								FString msg = FString::Printf(TEXT("The Reference texture for material image [%s] is not set and it couldn't be found automatically."), *ImageName);
+								GenerationContext.Compiler->CompilerLog(FText::FromString(msg), Node);
+							}
+						}
+
+						// Generate the texture nodes
+						mu::NodeImagePtr ImageNode = [&]()
+						{
+							if (TypedNodeMat->IsImageMutableMode(ImageIndex))
+							{
+								if (ImagePin)
+								{
+									if (const UEdGraphPin* ConnectedPin = FollowInputPin(*ImagePin))
+									{
+										return GenerateMutableSourceImage(ConnectedPin, GenerationContext, Props.TextureSize);
+									}
+								}
+
+								if (bTableMaterialPinLinked)
+								{
+									if (const UEdGraphPin* ConnectedPin = FollowInputPin(*TypedNodeMat->GetMaterialAssetPin()))
+									{
+										return GenerateMutableSourceImage(ConnectedPin, GenerationContext, Props.TextureSize);
+									}
+								}
+
+								// Else
+								{
+									UTexture2D* Texture2D = TypedNodeMat->GetImageValue(ImageIndex);
+
+									if (Texture2D)
+									{
+										mu::Ptr<mu::NodeImageConstant> ConstImageNode = new mu::NodeImageConstant();
+										mu::Ptr<mu::Image> ImageConstant = GenerateImageConstant(Texture2D, GenerationContext, false);
+										ConstImageNode->SetValue(ImageConstant.get());
+
+										const uint32 MipsToSkip = ComputeLODBiasForTexture(GenerationContext, *Texture2D, nullptr, Props.TextureSize);
+										mu::Ptr<mu::NodeImage> Result =  ResizeTextureByNumMips(ConstImageNode, MipsToSkip);
+
+										// Calculate the number of mips to tag as high res for this image.
+										int32 TotalMips = mu::Image::GetMipmapCount(ImageConstant->GetSizeX(), ImageConstant->GetSizeY());
+										int32 NumMipsBeyondMin = FMath::Max(0, TotalMips - int32(MipsToSkip) - GenerationContext.Options.MinDiskMips);
+										int32 HighResMipsForThisImage = FMath::Min(NumMipsBeyondMin, GenerationContext.Options.NumHighResImageMips);
+										ConstImageNode->SourceDataDescriptor.SourceHighResMips = HighResMipsForThisImage;
+
+										return Result;
+									}
+									else
+									{
+										return mu::NodeImagePtr();
+									}
+								}
+							}
+							else
+							{
+								return mu::NodeImagePtr();
+							}
+						}();
+
+						if (GroupProjectionImg.get())
+						{
+							ImageNode = GroupProjectionImg;
+						}
+
+						if (ReferenceTexture)
+						{
+							// Apply base LODBias. It will be propagated to most images.
+							const uint32 SurfaceLODBias = GenerationContext.Options.bUseLODAsBias ? GenerationContext.FirstLODAvailable : 0;
+							const uint32 BaseLODBias = ComputeLODBiasForTexture(GenerationContext, *ReferenceTexture) + SurfaceLODBias;
+							mu::NodeImagePtr LastImage = ResizeTextureByNumMips(ImageNode, BaseLODBias);
+
+							if (ReferenceTexture->MipGenSettings != TextureMipGenSettings::TMGS_NoMipmaps)
+							{
+								mu::EMipmapFilterType MipGenerationFilterType = Invoke([&]()
+								{
+									if (ReferenceTexture)
+									{
+										switch (ReferenceTexture->MipGenSettings)
+										{
+											case TextureMipGenSettings::TMGS_SimpleAverage: return mu::EMipmapFilterType::SimpleAverage;
+											case TextureMipGenSettings::TMGS_Unfiltered:    return mu::EMipmapFilterType::Unfiltered;
+											default: return mu::EMipmapFilterType::SimpleAverage;
+										}
+									}
+
+									return mu::EMipmapFilterType::SimpleAverage;
+								});
+
+
+								mu::NodeImageMipmapPtr MipmapImage = new mu::NodeImageMipmap();
+								MipmapImage->SetSource(LastImage.get());
+								MipmapImage->SetMipmapGenerationSettings(MipGenerationFilterType, mu::EAddressMode::None);
+
+								MipmapImage->SetMessageContext(Node);
+								LastImage = MipmapImage;
+							}
+
+							// Apply composite image. This needs to be computed after mipmaps generation. 	
+							if (ReferenceTexture && ReferenceTexture->GetCompositeTexture() && ReferenceTexture->CompositeTextureMode != CTM_Disabled)
+							{
+								mu::NodeImageNormalCompositePtr CompositedImage = new mu::NodeImageNormalComposite();
+								CompositedImage->SetBase(LastImage.get());
+								CompositedImage->SetPower(ReferenceTexture->CompositePower);
+
+								mu::ECompositeImageMode CompositeImageMode = [CompositeTextureMode = ReferenceTexture->CompositeTextureMode]() -> mu::ECompositeImageMode
+								{
+									switch (CompositeTextureMode)
+									{
+										case CTM_NormalRoughnessToRed: return mu::ECompositeImageMode::CIM_NormalRoughnessToRed;
+										case CTM_NormalRoughnessToGreen: return mu::ECompositeImageMode::CIM_NormalRoughnessToGreen;
+										case CTM_NormalRoughnessToBlue: return mu::ECompositeImageMode::CIM_NormalRoughnessToBlue;
+										case CTM_NormalRoughnessToAlpha: return mu::ECompositeImageMode::CIM_NormalRoughnessToAlpha;
+
+										default: return mu::ECompositeImageMode::CIM_Disabled;
+									}
+								}();
+
+								CompositedImage->SetMode(CompositeImageMode);
+
+								mu::Ptr<mu::NodeImageConstant> CompositeNormalImage = new mu::NodeImageConstant();
+
+								UTexture2D* ReferenceCompositeNormalTexture = Cast<UTexture2D>(ReferenceTexture->GetCompositeTexture());
+								if (ReferenceCompositeNormalTexture)
+								{
+									// GenerationContext.ArrayTextureUnrealToMutableTask.Add(FTextureUnrealToMutableTask(CompositeNormalImage, ReferenceCompositeNormalTexture, Node, true));
+									// TODO: The normal composite part is not propagated, so it will be unsupported. Create a task that performs the required transforms at mutable image level, and add the right operations here
+									// instead of propagating the flag and doing them on unreal-convert.
+									mu::Ptr<mu::Image> ImageConstant = GenerateImageConstant(ReferenceCompositeNormalTexture, GenerationContext, false);
+									CompositeNormalImage->SetValue(ImageConstant.get());
+
+									mu::NodeImageMipmapPtr NormalCompositeMipmapImage = new mu::NodeImageMipmap();
+									const uint32 MipsToSkip = ComputeLODBiasForTexture(GenerationContext, *ReferenceCompositeNormalTexture, ReferenceTexture);
+									NormalCompositeMipmapImage->SetSource(ResizeTextureByNumMips(CompositeNormalImage, MipsToSkip));
+									NormalCompositeMipmapImage->SetMipmapGenerationSettings(mu::EMipmapFilterType::SimpleAverage, mu::EAddressMode::None);
+
+									CompositedImage->SetNormal(NormalCompositeMipmapImage);
+
+									int32 TotalMips = mu::Image::GetMipmapCount(ImageConstant->GetSizeX(), ImageConstant->GetSizeY());
+									int32 NumMipsBeyondMin = FMath::Max(0, TotalMips - int32(MipsToSkip) - GenerationContext.Options.MinDiskMips);
+									int32 HighResMipsForThisImage = FMath::Min(NumMipsBeyondMin, GenerationContext.Options.NumHighResImageMips);
+									CompositeNormalImage->SourceDataDescriptor.SourceHighResMips = HighResMipsForThisImage;
+								}
+
+								LastImage = CompositedImage;
+							}
+
+							mu::Ptr<mu::NodeImage> FormatSource = LastImage;
+							mu::NodeImageFormatPtr FormatImage = new mu::NodeImageFormat();
+							FormatImage->SetSource(LastImage.get());
+							FormatImage->SetFormat(mu::EImageFormat::IF_RGBA_UBYTE);
+							FormatImage->SetMessageContext(Node);
+							LastImage = FormatImage;
+
+							TArray<TArray<FTextureBuildSettings>> BuildSettingsPerFormatPerLayer;
+							if (GenerationContext.Options.TargetPlatform)
+							{
+								ReferenceTexture->GetTargetPlatformBuildSettings(GenerationContext.Options.TargetPlatform, BuildSettingsPerFormatPerLayer);
+								if (BuildSettingsPerFormatPerLayer.IsEmpty())
+								{
+									const FString ReplacedImageFormatMsg = FString::Printf(TEXT("In object [%s] for platform [%s] the unsupported image format of texture [%s] is used, IF_RGBA_UBYTE will be used instead."),
+										*GenerationContext.Object->GetName(),
+										*GenerationContext.Options.TargetPlatform->PlatformName(),
+										*ReferenceTexture->GetName());
+									const FText ReplacedImageFormatText = FText::FromString(ReplacedImageFormatMsg);
+									GenerationContext.Compiler->CompilerLog(ReplacedImageFormatText, Node, EMessageSeverity::Info);
+									UE_LOG(LogMutable, Log, TEXT("%s"), *ReplacedImageFormatMsg);
+								}
+								else if (BuildSettingsPerFormatPerLayer.Num() > 1)
+								{
+									const FString ReplacedImageFormatMsg = FString::Printf(TEXT("In object [%s] for platform [%s] the image format of texture [%s] has multiple target formats. Only one will be used.."),
+										*GenerationContext.Object->GetName(),
+										*GenerationContext.Options.TargetPlatform->PlatformName(),
+										*ReferenceTexture->GetName());
+									const FText ReplacedImageFormatText = FText::FromString(ReplacedImageFormatMsg);
+									GenerationContext.Compiler->CompilerLog(ReplacedImageFormatText, Node, EMessageSeverity::Info);
+									UE_LOG(LogMutable, Log, TEXT("%s"), *ReplacedImageFormatMsg);
+								}
+							}
+
+							if (!BuildSettingsPerFormatPerLayer.IsEmpty())
+							{
+								const TArray<FTextureBuildSettings>& BuildSettingsPerLayer = BuildSettingsPerFormatPerLayer[0];
+
+								if (GenerationContext.Options.TextureCompression!=ECustomizableObjectTextureCompression::None)
+								{
+									static ITextureFormatManagerModule* TextureFormatManager = nullptr;
+									if (!TextureFormatManager)
+									{
+										TextureFormatManager = &FModuleManager::LoadModuleChecked<ITextureFormatManagerModule>("TextureFormat");
+										check(TextureFormatManager);
+									}
+									const ITextureFormat* TextureFormat = TextureFormatManager->FindTextureFormat(BuildSettingsPerLayer[0].TextureFormatName);
+									check(TextureFormat);
+									EPixelFormat UnrealTargetPlatformFormat = TextureFormat->GetEncodedPixelFormat(BuildSettingsPerLayer[0], false);
+									EPixelFormat UnrealTargetPlatformFormatAlpha = TextureFormat->GetEncodedPixelFormat(BuildSettingsPerLayer[0], true);
+
+									// \TODO: The QualityFix filter is used while the internal mutable runtime compression doesn't provide enough quality for some large block formats.
+									mu::EImageFormat MutableFormat = QualityAndPerformanceFix(UnrealToMutablePixelFormat(UnrealTargetPlatformFormat,false));
+									mu::EImageFormat MutableFormatIfAlpha = QualityAndPerformanceFix(UnrealToMutablePixelFormat(UnrealTargetPlatformFormatAlpha,true));
+
+									// Temp hack to enable RG->LA 
+									if (GenerationContext.Options.TargetPlatform)
+									{
+										bool bUseLA = GenerationContext.Options.TargetPlatform->SupportsFeature(ETargetPlatformFeatures::NormalmapLAEncodingMode);
+										if (bUseLA)
+										{
+											// See GetQualityFormat in TextureFormatASTC.cpp to understand why
+											if (UnrealTargetPlatformFormat == PF_ASTC_6x6 || UnrealTargetPlatformFormat == PF_ASTC_6x6_NORM_RG)
+											{
+												MutableFormat = mu::EImageFormat::IF_ASTC_4x4_RGBA_LDR;
+												MutableFormatIfAlpha = mu::EImageFormat::IF_ASTC_4x4_RGBA_LDR;
+
+												// Insert a channel swizzle
+												mu::Ptr<mu::NodeImageSwizzle> Swizzle = new mu::NodeImageSwizzle;
+												Swizzle->SetFormat( mu::EImageFormat::IF_RGBA_UBYTE );
+												Swizzle->SetSource(0, FormatSource);
+												Swizzle->SetSource(1, FormatSource);
+												Swizzle->SetSource(2, FormatSource);
+												Swizzle->SetSource(3, FormatSource);
+												Swizzle->SetSourceChannel(0, 0);
+												Swizzle->SetSourceChannel(1, 0);
+												Swizzle->SetSourceChannel(2, 0);
+												Swizzle->SetSourceChannel(3, 1);
+
+												FormatImage->SetSource( Swizzle.get() );
+											}
+										}
+									}
+
+									// Unsupported format: look for something generic
+									if (MutableFormat == mu::EImageFormat::IF_NONE)
+									{
+										const FString ReplacedImageFormatMsg = FString::Printf(TEXT("In object [%s] the unsupported image format %d is used, IF_RGBA_UBYTE will be used instead."), *GenerationContext.Object->GetName(), UnrealTargetPlatformFormat );
+										const FText ReplacedImageFormatText = FText::FromString(ReplacedImageFormatMsg);
+										GenerationContext.Compiler->CompilerLog(ReplacedImageFormatText, Node, EMessageSeverity::Info);
+										UE_LOG(LogMutable, Log, TEXT("%s"), *ReplacedImageFormatMsg);
+										MutableFormat = mu::EImageFormat::IF_RGBA_UBYTE;
+									}
+									if (MutableFormatIfAlpha == mu::EImageFormat::IF_NONE)
+									{
+										const FString ReplacedImageFormatMsg = FString::Printf(TEXT("In object [%s] the unsupported image format %d is used, IF_RGBA_UBYTE will be used instead."), *GenerationContext.Object->GetName(), UnrealTargetPlatformFormatAlpha);
+										const FText ReplacedImageFormatText = FText::FromString(ReplacedImageFormatMsg);
+										GenerationContext.Compiler->CompilerLog(ReplacedImageFormatText, Node, EMessageSeverity::Info);
+										UE_LOG(LogMutable, Log, TEXT("%s"), *ReplacedImageFormatMsg);
+										MutableFormatIfAlpha = mu::EImageFormat::IF_RGBA_UBYTE;
+									}
+
+									FormatImage->SetFormat(MutableFormat, MutableFormatIfAlpha);
+								}
+							}
+
+							ImageNode = LastImage;
+						}
+
+						SurfNode->Images[ImageIndex].Image = ImageNode;
+
+						check(Props.ImagePropertiesIndex != INDEX_NONE);
+						const FString SurfNodeImageName = FString::Printf(TEXT("%d"), Props.ImagePropertiesIndex);
+
+						// Encoding material layer in mutable name
+						const int32 LayerIndex = TypedNodeMat->GetParameterLayerIndex(EMaterialParameterType::Texture, ImageIndex);
+						const FString LayerEncoding = LayerIndex != INDEX_NONE ? "-MutableLayerParam:" + FString::FromInt(LayerIndex) : "";
+						
+						SurfNode->Images[ImageIndex].Name = SurfNodeImageName + LayerEncoding;
+
+						// If we are generating an implicit component (with a passthrough mesh) we don't apply any layout.
+						int32 UVLayout = -1;
+						if (!bGeneratingImplicitComponent)
+						{
+							UVLayout = TypedNodeMat->GetImageUVLayout(ImageIndex);;
+						}
+						SurfNode->Images[ImageIndex].LayoutIndex = UVLayout;
+						SurfNode->Images[ImageIndex].MaterialName = TypedNodeMat->GetMaterial()->GetName();
+						SurfNode->Images[ImageIndex].MaterialParameterName = ImageName;
+
+						if (bShareProjectionTexturesBetweenLODs && bIsGroupProjectorImage)
+						{
+							// Add to the GroupProjectorLODCache to potentially reuse this projection texture in higher LODs
+							ensure(LOD == GenerationContext.FirstLODAvailable);
+							float* AlternateProjectionResFactor = TextureNameToProjectionResFactor.Find(ImageName);
+							GenerationContext.GroupProjectorLODCache.Add(MaterialImageId,
+								FGroupProjectorImageInfo(ImageNode, ImageName, ImageName, TypedNodeMat,
+									AlternateProjectionResFactor ? *AlternateProjectionResFactor : 0.f, AlternateResStateName, SurfNode, UVLayout));
+						}
+					}
+				}
+				else
+				{
+					ensure(LOD > GenerationContext.FirstLODAvailable);
+					check(ProjectorInfo->SurfNode->Images[ImageIndex].Image == ProjectorInfo->ImageNode);
+					SurfNode->Images[ImageIndex].Image = ProjectorInfo->ImageNode;
+					SurfNode->Images[ImageIndex].Name = ProjectorInfo->TextureName;
+					SurfNode->Images[ImageIndex].LayoutIndex = ProjectorInfo->UVLayout;
+
+					TextureNameToProjectionResFactor.Add(ProjectorInfo->RealTextureName, ProjectorInfo->AlternateProjectionResolutionFactor);
+					AlternateResStateName = ProjectorInfo->AlternateResStateName;
+				}
+			}
+		}
+
+		const int32 NumVectors = TypedNodeMat->GetNumParameters(EMaterialParameterType::Vector);
+		SurfNode->Vectors.SetNum(NumVectors);
+		for (int32 VectorIndex = 0; VectorIndex < NumVectors; ++VectorIndex)
+		{
+			const UEdGraphPin* VectorPin = TypedNodeMat->GetParameterPin(EMaterialParameterType::Vector, VectorIndex);
+			bool bVectorPinConnected = VectorPin && FollowInputPin(*VectorPin);
+
+			FString VectorName = TypedNodeMat->GetParameterName(EMaterialParameterType::Vector, VectorIndex).ToString();
+			FNodeMaterialParameterId VectorId = TypedNodeMat->GetParameterId(EMaterialParameterType::Vector, VectorIndex);
+
+			if (bVectorPinConnected)
+			{				
+				if (const UEdGraphPin* ConnectedPin = FollowInputPin(*VectorPin))
+				{
+					mu::Ptr<mu::NodeColour> ColorNode = GenerateMutableSourceColor(ConnectedPin, GenerationContext);
+
+					// Encoding material layer in mutable name
+					if (const int32 LayerIndex = TypedNodeMat->GetParameterLayerIndex(EMaterialParameterType::Vector, VectorIndex); LayerIndex != INDEX_NONE)
+					{
+						VectorName += "-MutableLayerParam:" + FString::FromInt(LayerIndex);
+					}
+
+					SurfNode->Vectors[VectorIndex].Vector = ColorNode;
+					SurfNode->Vectors[VectorIndex].Name = VectorName;
+				}
+			}
+		}
+
+		const int32 NumScalar = TypedNodeMat->GetNumParameters(EMaterialParameterType::Scalar);
+		SurfNode->Scalars.SetNum(NumScalar);
+		for (int32 ScalarIndex = 0; ScalarIndex < NumScalar; ++ScalarIndex)
+		{
+			const UEdGraphPin* ScalarPin = TypedNodeMat->GetParameterPin(EMaterialParameterType::Scalar, ScalarIndex);
+			bool bScalarPinConnected = ScalarPin && FollowInputPin(*ScalarPin);
+
+			FString ScalarName = TypedNodeMat->GetParameterName(EMaterialParameterType::Scalar, ScalarIndex).ToString();
+			FNodeMaterialParameterId ScalarId = TypedNodeMat->GetParameterId(EMaterialParameterType::Scalar, ScalarIndex);
+
+			if (bScalarPinConnected)
+			{
+				if (const UEdGraphPin* ConnectedPin = FollowInputPin(*ScalarPin))
+				{
+					mu::NodeScalarPtr ScalarNode = GenerateMutableSourceFloat(ConnectedPin, GenerationContext);
+
+					// Encoding material layer in mutable name
+					if (const int32 LayerIndex = TypedNodeMat->GetParameterLayerIndex(EMaterialParameterType::Scalar, ScalarIndex); LayerIndex != INDEX_NONE)
+					{
+						ScalarName += "-MutableLayerParam:" + FString::FromInt(LayerIndex);
+					}
+
+					SurfNode->Scalars[ScalarIndex].Scalar = ScalarNode;
+					SurfNode->Scalars[ScalarIndex].Name = ScalarName;
+				}
+			}
+		}
+
+		// New method to pass the surface id as a scalar parameter
+		{
+			int32 MaterialIndex = NumScalar;
+			SurfNode->Scalars.SetNum(NumScalar + 1);
+
+			const UEdGraphPin* MaterialPin = TypedNodeMat->GetMaterialAssetPin();
+
+			//Encoding name for material material id parameter
+			FString MaterialName = "__MutableMaterialId";
+
+			if (bTableMaterialPinLinked)
+			{
+				if (const UEdGraphPin* ConnectedPin = FollowInputPin(*MaterialPin))
+				{
+					GenerationContext.CurrentMaterialTableParameterId = MaterialName;
+					mu::NodeScalarPtr ScalarNode = GenerateMutableSourceFloat(ConnectedPin, GenerationContext);
+
+					SurfNode->Scalars[MaterialIndex].Scalar = ScalarNode;
+					SurfNode->Scalars[MaterialIndex].Name = MaterialName;
+				}
+			}
+			else
+			{
+				mu::NodeScalarConstantPtr ScalarNode = new mu::NodeScalarConstant();
+				ScalarNode->SetValue(ReferencedMaterialsIndex);
+
+				SurfNode->Scalars[MaterialIndex].Scalar = ScalarNode;
+				SurfNode->Scalars[MaterialIndex].Name = MaterialName;
+			}
+		}
+		
+
+		if (const TArray<FString>* EnableTags = TypedNodeMat->GetEnableTags())
+		{
+			for (const FString& Tag : *EnableTags)
+			{
+				SurfNode->Tags.AddUnique(Tag);
+			}
+		}
+
+		TArray<mu::Ptr<mu::NodeSurfaceNew>>* ArraySurfaceNodePtr = GenerationContext.MapMaterialNodeToMutableSurfaceNodeArray.Find(TypedNodeMat->GetMaterialNode());
+		if (ArraySurfaceNodePtr == nullptr)
+		{
+			TArray<mu::Ptr<mu::NodeSurfaceNew>> ArraySurfaceNode;
+			ArraySurfaceNode.Add(SurfNode);
+			ArraySurfaceNodePtr = &GenerationContext.MapMaterialNodeToMutableSurfaceNodeArray.Add(TypedNodeMat->GetMaterialNode(), ArraySurfaceNode);
+		}
+		else
+		{
+			ArraySurfaceNodePtr->AddUnique(SurfNode);
+		}
+
+		// If an alternate resolution for a particular state is present, clone the surface node, add the image resizing and inject the surface variation node
+		if (TextureNameToProjectionResFactor.Num() > 0 && !AlternateResStateName.IsEmpty())
+		{
+			mu::Ptr<mu::NodeSurfaceNew> SurfNode2 = new mu::NodeSurfaceNew;
+
+			SurfNode2->ExternalId = SurfaceMetadataUniqueHash;
+
+			SurfNode2->Mesh = SurfNode->Mesh;
+			SurfNode2->Tags = SurfNode->Tags;
+			SurfNode2->Vectors = SurfNode->Vectors;
+			SurfNode2->Scalars = SurfNode->Scalars;
+			SurfNode2->Strings = SurfNode->Strings;
+			SurfNode2->Images = SurfNode->Images;
+
+			for (int32 ImageIndex = 0; ImageIndex < SurfNode2->Images.Num(); ++ImageIndex)
+			{
+				const FString ImageName = TypedNodeMat->GetParameterName(EMaterialParameterType::Texture, ImageIndex).ToString();
+				if (float* ResolutionFactor = TextureNameToProjectionResFactor.Find(ImageName))
+				{
+					FString MaterialImageId = FGroupProjectorImageInfo::GenerateId(TypedNodeMat, ImageIndex);
+					FGroupProjectorImageInfo* ProjectorInfo = GenerationContext.GroupProjectorLODCache.Find(MaterialImageId);
+
+					if (!ProjectorInfo || !ProjectorInfo->bIsAlternateResolutionResized)
+					{
+						mu::NodeImageResizePtr NodeImageResize = new mu::NodeImageResize;
+						NodeImageResize->SetRelative(true);
+						NodeImageResize->SetSize(*ResolutionFactor, *ResolutionFactor);
+						NodeImageResize->SetBase(SurfNode2->Images[ImageIndex].Image);
+
+						SurfNode2->Images[ImageIndex].Image = NodeImageResize;
+
+						if (ProjectorInfo)
+						{
+							ensure(LOD == GenerationContext.FirstLODAvailable);
+							ProjectorInfo->ImageResizeNode = NodeImageResize;
+							ProjectorInfo->bIsAlternateResolutionResized = true;
+						}
+					}
+					else
+					{
+						ensure(LOD > GenerationContext.FirstLODAvailable);
+						check(ProjectorInfo->bIsAlternateResolutionResized);
+						SurfNode2->Images[ImageIndex].Image = ProjectorInfo->ImageResizeNode;
+					}
+				}
+			}
+
+			ArraySurfaceNodePtr->AddUnique(SurfNode2);
+
+			mu::Ptr<mu::NodeSurfaceVariation> SurfaceVariation = new mu::NodeSurfaceVariation;
+			SurfaceVariation->Type = mu::NodeSurfaceVariation::VariationType::State;
+			SurfaceVariation->Variations.SetNum(1);
+			SurfaceVariation->Variations[0].Tag = AlternateResStateName;
+
+			SurfaceVariation->DefaultSurfaces.Add(SurfNode);
+			SurfaceVariation->Variations[0].Surfaces.Add(SurfNode2);
+
+			Result = SurfaceVariation;
+		}
+	}
+
+	else if (const UCustomizableObjectNodeMaterialVariation* TypedNodeVar = Cast<UCustomizableObjectNodeMaterialVariation>(Node))
+	{
+		mu::Ptr<mu::NodeSurfaceVariation> SurfNode = new mu::NodeSurfaceVariation();
+		Result = SurfNode;
+
+		mu::NodeSurfaceVariation::VariationType muType = mu::NodeSurfaceVariation::VariationType::Tag;
+		switch (TypedNodeVar->Type)
+		{
+		case ECustomizableObjectNodeMaterialVariationType::Tag: muType = mu::NodeSurfaceVariation::VariationType::Tag; break;
+		case ECustomizableObjectNodeMaterialVariationType::State: muType = mu::NodeSurfaceVariation::VariationType::State; break;
+		default:
+			check(false);
+			break;
+		}
+		SurfNode->Type = muType;
+
+		for (const UEdGraphPin* ConnectedPin : FollowInputPinArray(*TypedNodeVar->DefaultPin()))
+		{
+			// Is it a modifier?
+			mu::NodeSurfacePtr ChildNode = GenerateMutableSourceSurface(ConnectedPin, GenerationContext);
+			if (ChildNode)
+			{
+				SurfNode->DefaultSurfaces.Add(ChildNode);
+			}
+			else
+			{
+				GenerationContext.Compiler->CompilerLog(LOCTEXT("SurfaceFailed", "Surface generation failed."), Node);
+			}
+		}
+
+		const int32 NumVariations = TypedNodeVar->GetNumVariations();
+		SurfNode->Variations.SetNum(NumVariations);
+		for (int VariationIndex = 0; VariationIndex < NumVariations; ++VariationIndex)
+		{
+			mu::NodeSurfacePtr VariationSurfaceNode;
+
+			if (UEdGraphPin* VariationPin = TypedNodeVar->VariationPin(VariationIndex))
+			{
+				SurfNode->Variations[VariationIndex].Tag = TypedNodeVar->GetVariation(VariationIndex).Tag;
+				for (const UEdGraphPin* ConnectedPin : FollowInputPinArray(*VariationPin))
+				{
+					// Is it a modifier?
+					mu::NodeSurfacePtr ChildNode = GenerateMutableSourceSurface(ConnectedPin, GenerationContext);
+					if (ChildNode)
+					{
+						SurfNode->Variations[VariationIndex].Surfaces.Add( ChildNode );
+					}
+					else
+					{
+						GenerationContext.Compiler->CompilerLog(LOCTEXT("SurfaceModifierFailed", "Surface generation failed."), Node);
+					}
+				}
+			}
+		}
+	}
+
+	else if (const UCustomizableObjectNodeMaterialSwitch* TypedNodeSwitch = Cast<UCustomizableObjectNodeMaterialSwitch>(Node))
+	{
+		// Using a lambda so control flow is easier to manage.
+		Result = [&]()
+		{
+			const UEdGraphPin* SwitchParameter = TypedNodeSwitch->SwitchParameter();
+
+			// Check Switch Parameter arity preconditions.
+			if (const UEdGraphPin* EnumPin = FollowInputPin(*SwitchParameter))
+			{
+				mu::NodeScalarPtr SwitchParam = GenerateMutableSourceFloat(EnumPin, GenerationContext);
+
+				// Switch Param not generated
+				if (!SwitchParam)
+				{
+					// Warn about a failure.
+					if (EnumPin)
+					{
+						const FText Message = LOCTEXT("FailedToGenerateSwitchParam", "Could not generate switch enum parameter. Please refesh the switch node and connect an enum.");
+						GenerationContext.Compiler->CompilerLog(Message, Node);
+					}
+
+					return Result;
+				}
+
+				if (SwitchParam->GetType() != mu::NodeScalarEnumParameter::GetStaticType())
+				{
+					const FText Message = LOCTEXT("WrongSwitchParamType", "Switch parameter of incorrect type.");
+					GenerationContext.Compiler->CompilerLog(Message, Node);
+
+					return Result;
+				}
+
+				const int32 NumSwitchOptions = TypedNodeSwitch->GetNumElements();
+
+				mu::NodeScalarEnumParameter* EnumParameter = static_cast<mu::NodeScalarEnumParameter*>(SwitchParam.get());
+				if (NumSwitchOptions != EnumParameter->GetValueCount())
+				{
+					const FText Message = LOCTEXT("MismatchedSwitch", "Switch enum and switch node have different number of options. Please refresh the switch node to make sure the outcomes are labeled properly.");
+					GenerationContext.Compiler->CompilerLog(Message, Node);
+				}
+
+				mu::Ptr<mu::NodeSurfaceSwitch> SwitchNode = new mu::NodeSurfaceSwitch;
+				SwitchNode->Parameter = SwitchParam;
+				SwitchNode->Options.SetNum(NumSwitchOptions);
+
+				for (int32 SelectorIndex = 0; SelectorIndex < NumSwitchOptions; ++SelectorIndex)
+				{
+					if (const UEdGraphPin* ConnectedPin = FollowInputPin(*TypedNodeSwitch->GetElementPin(SelectorIndex)))
+					{
+						mu::NodeSurfacePtr ChildNode = GenerateMutableSourceSurface(ConnectedPin, GenerationContext);
+						if (ChildNode)
+						{
+							SwitchNode->Options[SelectorIndex] = ChildNode;
+						}
+						else
+						{
+							// Probably ok
+							//GenerationContext.Compiler->CompilerLog(LOCTEXT("SurfaceModifierFailed", "Surface generation failed."), Node);
+						}
+					}
+				}
+
+				Result = SwitchNode;
+				return Result;
+			}
+			else
+			{
+				GenerationContext.Compiler->CompilerLog(LOCTEXT("NoEnumParamInSwitch", "Switch nodes must have an enum switch parameter. Please connect an enum and refesh the switch node."), Node);
+				return Result;
+			}
+		}(); // invoke lambda.
+	}
+
+	else
+	{
+		GenerationContext.Compiler->CompilerLog(LOCTEXT("UnimplementedNode", "Node type not implemented yet."), Node);
+	}
+
+
+	if (Result)
+	{
+		Result->SetMessageContext(Node);
+	}
+
+	GenerationContext.Generated.Add(Key, FGeneratedData(Node, Result));
+	GenerationContext.GeneratedNodes.Add(Node);
+
+	return Result;
+}
+
+
+#undef LOCTEXT_NAMESPACE
diff --git a/Engine/Plugins/Mutable/Source/CustomizableObjectEditor/Private/MuCOE/GenerateMutableSource/GenerateMutableSourceSurface.h b/Engine/Plugins/Mutable/Source/CustomizableObjectEditor/Private/MuCOE/GenerateMutableSource/GenerateMutableSourceSurface.h
new file mode 100644
--- /dev/null
+++ b/Engine/Plugins/Mutable/Source/CustomizableObjectEditor/Private/MuCOE/GenerateMutableSource/GenerateMutableSourceSurface.h
@@ -0,0 +1,23 @@
+// Copyright Epic Games, Inc. All Rights Reserved.
+
+#pragma once
+
+#include "MuR/Ptr.h"
+#include "HAL/Platform.h"
+
+class UEdGraphPin;
+struct FMutableGraphGenerationContext;
+struct FMutableGraphMeshGenerationData;
+enum class ECustomizableObjectNumBoneInfluences:uint8;
+namespace mu
+{
+	class FMeshBufferSet;
+	class NodeSurface;
+}
+
+extern mu::Ptr<mu::NodeSurface> GenerateMutableSourceSurface(const UEdGraphPin* Pin, FMutableGraphGenerationContext & GenerationContext);
+
+extern void SetSurfaceFormat(FMutableGraphGenerationContext& GenerationContext,
+	mu::FMeshBufferSet& OutVertexBufferFormat, mu::FMeshBufferSet& OutIndexBufferFormat, const FMutableGraphMeshGenerationData& MeshData,
+	ECustomizableObjectNumBoneInfluences ECustomizableObjectNumBoneInfluences, bool bWith16BitWeights);
+
diff --git a/Engine/Plugins/Mutable/Source/CustomizableObjectEditor/Private/MuCOE/GraphTraversal.cpp b/Engine/Plugins/Mutable/Source/CustomizableObjectEditor/Private/MuCOE/GraphTraversal.cpp
new file mode 100644
--- /dev/null
+++ b/Engine/Plugins/Mutable/Source/CustomizableObjectEditor/Private/MuCOE/GraphTraversal.cpp
@@ -0,0 +1,719 @@
+// Copyright Epic Games, Inc. All Rights Reserved.
+
+#include "MuCOE/GraphTraversal.h"
+
+#include "AssetRegistry/AssetRegistryModule.h"
+#include "MuCOE/CustomizableObjectPin.h"
+#include "MuCOE/EdGraphSchema_CustomizableObject.h"
+#include "MuCOE/Nodes/CustomizableObjectNodeEnumParameter.h"
+#include "MuCOE/Nodes/CustomizableObjectNodeExposePin.h"
+#include "MuCOE/Nodes/CustomizableObjectNodeModifierExtendMeshSection.h"
+#include "MuCOE/Nodes/CustomizableObjectNodeExternalPin.h"
+#include "MuCOE/Nodes/CustomizableObjectNodeMaterial.h"
+#include "MuCOE/Nodes/CustomizableObjectNodeMaterialVariation.h"
+#include "MuCOE/Nodes/CustomizableObjectNodeMeshGeometryOperation.h"
+#include "MuCOE/Nodes/CustomizableObjectNodeMeshMorph.h"
+#include "MuCOE/Nodes/CustomizableObjectNodeMeshMorphStackApplication.h"
+#include "MuCOE/Nodes/CustomizableObjectNodeMeshMorphStackDefinition.h"
+#include "MuCOE/Nodes/CustomizableObjectNodeMeshReshape.h"
+#include "MuCOE/Nodes/CustomizableObjectNodeMeshSwitch.h"
+#include "MuCOE/Nodes/CustomizableObjectNodeMeshVariation.h"
+#include "MuCOE/Nodes/CustomizableObjectNodeObject.h"
+#include "MuCOE/Nodes/CustomizableObjectNodeSkeletalMesh.h"
+#include "MuCOE/Nodes/CustomizableObjectNodeStaticMesh.h"
+#include "MuCOE/Nodes/CustomizableObjectNodeTable.h"
+#include "MuCOE/Nodes/CustomizableObjectNodeAnimationPose.h"
+#include "MuCOE/Nodes/CustomizableObjectNodeObjectGroup.h"
+#include "MuCOE/Nodes/CustomizableObjectNodeReroute.h"
+#include "UObject/UObjectIterator.h"
+
+
+TArray<UEdGraphPin*> FollowPinArray(const UEdGraphPin& Pin, bool bIgnoreOrphan, bool* bOutCycleDetected)
+{
+	bool bCycleDetected = false;
+		
+	TArray<UEdGraphPin*> Result;
+	
+	TSet<const UEdGraphPin*> Visited;
+
+	TArray<const UEdGraphPin*> PinsToVisit;
+	PinsToVisit.Add(&Pin);
+	while (PinsToVisit.Num())
+	{
+		const UEdGraphPin& CurrentPin = *PinsToVisit.Pop();
+
+		if (!bIgnoreOrphan && IsPinOrphan(CurrentPin))
+		{
+			continue;
+		}
+
+		Visited.FindOrAdd(&CurrentPin, &bCycleDetected);
+		if (bCycleDetected)
+		{
+			continue;
+		}
+
+		for (UEdGraphPin* LinkedPin : CurrentPin.LinkedTo)
+		{
+			if (!bIgnoreOrphan && IsPinOrphan(*LinkedPin))
+			{
+				continue;
+			}
+			
+			if (const UCustomizableObjectNodeExposePin* ExposePinNode = Cast<UCustomizableObjectNodeExposePin>(LinkedPin->GetOwningNodeUnchecked()))
+			{
+				check(Pin.Direction == EGPD_Output);
+
+				for (TObjectIterator<UCustomizableObjectNodeExternalPin> It; It; ++It)
+				{
+					const UCustomizableObjectNodeExternalPin* LinkedNode = *It;
+					
+					if (IsValid(LinkedNode) &&
+						!LinkedNode->IsTemplate() &&
+						LinkedNode->GetNodeExposePin() == ExposePinNode)
+					{
+						const UEdGraphPin* ExternalPin = LinkedNode->GetExternalPin();
+						check(ExternalPin);
+						PinsToVisit.Add(ExternalPin);
+					}
+				}
+			}
+			else if (const UCustomizableObjectNodeExternalPin* ExternalPinNode = Cast<UCustomizableObjectNodeExternalPin>(LinkedPin->GetOwningNodeUnchecked()))
+			{
+				check(Pin.Direction == EGPD_Input);
+				
+				if (const UCustomizableObjectNodeExposePin* LinkedNode = ExternalPinNode->GetNodeExposePin())
+				{
+					const UEdGraphPin* ExposePin = LinkedNode->InputPin();
+					check(ExposePin);
+					PinsToVisit.Add(ExposePin);
+				}
+			}
+			else if (const UCustomizableObjectNodeReroute* NodeReroute = Cast<UCustomizableObjectNodeReroute>(LinkedPin->GetOwningNodeUnchecked()))
+			{
+				PinsToVisit.Add(Pin.Direction == EGPD_Input ? NodeReroute->GetInputPin() : NodeReroute->GetOutputPin());
+			}
+			else
+			{
+				Result.Add(LinkedPin);
+			}
+		}
+	}
+
+	if (bOutCycleDetected)
+	{
+		*bOutCycleDetected = bCycleDetected;	
+	}
+
+	return Result;
+}
+
+
+TArray<UEdGraphPin*> FollowInputPinArray(const UEdGraphPin& Pin, bool* bOutCycleDetected)
+{
+	check(Pin.Direction == EGPD_Input);
+	return FollowPinArray(Pin, true, bOutCycleDetected);
+}
+
+
+UEdGraphPin* FollowInputPin(const UEdGraphPin& Pin, bool* CycleDetected)
+{
+	TArray<UEdGraphPin*> Result = FollowInputPinArray(Pin, CycleDetected);
+	check(Result.Num() <= 1); // Use FollowPinImmersive if the pin can have more than one input.
+
+	if (!Result.IsEmpty())
+	{
+		return Result[0];
+	}
+	else
+	{
+		return nullptr;
+	}
+}
+
+
+TArray<UEdGraphPin*> FollowOutputPinArray(const UEdGraphPin& Pin, bool* bOutCycleDetected)
+{
+	check(Pin.Direction == EGPD_Output);
+	return FollowPinArray(Pin, true, bOutCycleDetected);
+}
+
+
+UEdGraphPin* FollowOutputPin(const UEdGraphPin& Pin, bool* CycleDetected)
+{
+	TArray<UEdGraphPin*> Result = FollowOutputPinArray(Pin, CycleDetected);
+	check(Result.Num() <= 1); // Use FollowPinImmersive if the pin can have more than one input.
+
+	if (!Result.IsEmpty())
+	{
+		return Result[0];
+	}
+	else
+	{
+		return nullptr;
+	}
+}
+
+
+TArray<UEdGraphPin*> ReverseFollowPinArray(const UEdGraphPin& Pin, bool bIgnoreOrphan, bool* bOutCycleDetected)
+{
+	bool bCycleDetected = false;
+		
+	TArray<UEdGraphPin*> Result;
+	
+	TSet<const UEdGraphPin*> Visited;
+
+	TArray<UEdGraphPin*> PinsToVisit;
+	PinsToVisit.Add(const_cast<UEdGraphPin*>(&Pin));
+	while (PinsToVisit.Num())
+	{
+		UEdGraphPin& CurrentPin = *PinsToVisit.Pop();
+
+		if (!bIgnoreOrphan && IsPinOrphan(CurrentPin))
+		{
+			continue;
+		}
+
+		Visited.FindOrAdd(&CurrentPin, &bCycleDetected);
+		if (bCycleDetected)
+		{
+			continue;
+		}
+		
+		if (const UCustomizableObjectNodeExposePin* ExposePinNode = Cast<UCustomizableObjectNodeExposePin>(CurrentPin.GetOwningNodeUnchecked()))
+		{
+			check(Pin.Direction == EGPD_Input);
+
+			for (TObjectIterator<UCustomizableObjectNodeExternalPin> It; It; ++It)
+			{
+				const UCustomizableObjectNodeExternalPin* LinkedNode = *It;
+
+				if (IsValid(LinkedNode) &&
+					!LinkedNode->IsTemplate() &&
+					LinkedNode->GetNodeExposePin() == ExposePinNode)
+				{
+					const UEdGraphPin* ExternalPin = LinkedNode->GetExternalPin();
+					check(ExternalPin);
+
+					for (UEdGraphPin* LinkedPin : ExternalPin->LinkedTo)
+					{
+						PinsToVisit.Add(LinkedPin);
+					}
+				}
+			}
+		}
+		else if (const UCustomizableObjectNodeExternalPin* ExternalPinNode = Cast<UCustomizableObjectNodeExternalPin>(CurrentPin.GetOwningNodeUnchecked()))
+		{
+			check(Pin.Direction == EGPD_Output);
+				
+			if (const UCustomizableObjectNodeExposePin* LinkedNode = ExternalPinNode->GetNodeExposePin())
+			{
+				const UEdGraphPin* ExposePin = LinkedNode->InputPin();
+				check(ExposePin);
+
+				for (UEdGraphPin* LinkedPin : ExposePin->LinkedTo)
+				{
+					PinsToVisit.Add(LinkedPin);
+				}
+			}
+		}
+		else if (const UCustomizableObjectNodeReroute* NodeReroute = Cast<UCustomizableObjectNodeReroute>(CurrentPin.GetOwningNodeUnchecked()))
+		{
+			UEdGraphPin* ReroutePin = Pin.Direction == EGPD_Output ? NodeReroute->GetInputPin() : NodeReroute->GetOutputPin();
+
+			for (UEdGraphPin* LinkedPin : ReroutePin->LinkedTo)
+			{
+				PinsToVisit.Add(LinkedPin);
+			}
+		}
+		else
+		{
+			if (bIgnoreOrphan || !IsPinOrphan(CurrentPin))
+			{
+				Result.Add(&CurrentPin);
+			}
+		}
+	}
+	
+	if (bOutCycleDetected)
+	{
+		*bOutCycleDetected = bCycleDetected;	
+	}
+	
+	return Result;
+}
+
+
+UCustomizableObjectNodeObject* GetRootNode(UCustomizableObject* Object, bool& bOutMultipleBaseObjectsFound)
+{
+	// Look for the base object node
+	UCustomizableObjectNodeObject* Root = nullptr;
+	TArray<UCustomizableObjectNodeObject*> ObjectNodes;
+	Object->GetPrivate()->GetSource()->GetNodesOfClass<UCustomizableObjectNodeObject>(ObjectNodes);
+
+	bOutMultipleBaseObjectsFound = false;
+
+	for (TArray<UCustomizableObjectNodeObject*>::TIterator It(ObjectNodes); It; ++It)
+	{
+		if ((*It)->bIsBase)
+		{
+			if (Root)
+			{
+				bOutMultipleBaseObjectsFound = true;
+				break;
+			}
+			else
+			{
+				Root = *It;
+			}
+		}
+	}
+
+	return Root;
+}
+
+
+bool GetParentsUntilRoot(UCustomizableObject* Object, TArray<UCustomizableObjectNodeObject*>& ArrayNodeObject, TArray<UCustomizableObject*>& ArrayCustomizableObject)
+{
+	bool MultipleBaseObjectsFound;
+	UCustomizableObjectNodeObject* Root = GetRootNode(Object, MultipleBaseObjectsFound);
+
+	bool bSuccess = true;
+
+	if (!MultipleBaseObjectsFound && (Root != nullptr))
+	{
+		if (!ArrayCustomizableObject.Contains(Object))
+		{
+			ArrayNodeObject.Add(Root);
+			ArrayCustomizableObject.Add(Object);
+		}
+		else
+		{
+			// This object has already been visted which means that there is a Cycle between Customizable Objects
+			return false;
+		}
+
+		if (Root->ParentObject != nullptr)
+		{
+			bSuccess = GetParentsUntilRoot(Root->ParentObject, ArrayNodeObject, ArrayCustomizableObject);
+		}
+	}
+
+	return bSuccess;
+}
+
+
+bool HasCandidateAsParent(UCustomizableObjectNodeObject* Node, UCustomizableObject* ParentCandidate)
+{
+	if (Node->ParentObject == ParentCandidate)
+	{
+		return true;
+	}
+
+	if (Node->ParentObject != nullptr)
+	{
+		bool MultipleBaseObjectsFound;
+		UCustomizableObjectNodeObject* ParentNodeObject = GetRootNode(Node->ParentObject, MultipleBaseObjectsFound);
+
+		if ((ParentNodeObject->ParentObject == nullptr) || MultipleBaseObjectsFound)
+		{
+			return false;
+		}
+		else
+		{
+			return HasCandidateAsParent(ParentNodeObject, ParentCandidate);
+		}
+	}
+
+	return false;
+}
+
+
+UCustomizableObject* GetFullGraphRootObject(UCustomizableObjectNodeObject* Node, TArray<UCustomizableObject*>& VisitedObjects)
+{
+	if (Node->ParentObject != nullptr)
+	{
+		VisitedObjects.Add(Node->ParentObject);
+
+		bool MultipleBaseObjectsFound;
+		UCustomizableObjectNodeObject* Root = GetRootNode(Node->ParentObject, MultipleBaseObjectsFound);
+
+		if (Root->ParentObject == nullptr)
+		{
+			if (MultipleBaseObjectsFound)
+			{
+				return nullptr;
+			}
+			else
+			{
+				return Node->ParentObject;
+			}
+		}
+		else
+		{
+			if (VisitedObjects.Contains(Root->ParentObject))
+			{
+				//There is a cycle
+				return nullptr;
+			}
+			else
+			{
+				return GetFullGraphRootObject(Root, VisitedObjects);
+			}
+		}
+	}
+
+	return nullptr;
+}
+
+
+UCustomizableObject* GetRootObject(const UCustomizableObjectNode& Node)
+{
+	return CastChecked<UCustomizableObject>(Node.GetGraph()->GetOuter());
+}
+
+
+UCustomizableObject* GetRootObject(UCustomizableObject* ChildObject)
+{
+	// Grab a node to start the search -> Get the root since it should be always present
+	bool bMultipleBaseObjectsFound = false;
+	UCustomizableObjectNodeObject* ObjectRootNode = GetRootNode(ChildObject, bMultipleBaseObjectsFound);
+
+	if (ObjectRootNode && ObjectRootNode->ParentObject)
+	{
+		TArray<UCustomizableObject*> VisitedNodes;
+		return GetFullGraphRootObject(ObjectRootNode, VisitedNodes);
+	}
+
+	// No parent object found, return input as the parent of the graph
+	// This can also mean the ObjectRootNode does not exist because it has not been opened yet (so no nodes have been generated)
+	return ChildObject;
+}
+
+
+UCustomizableObjectNodeObject* GetFullGraphRootNodeObject(UCustomizableObjectNodeObject* Node, TArray<UCustomizableObject*>& VisitedObjects)
+{
+	if (Node->ParentObject != nullptr)
+	{
+		VisitedObjects.Add(Node->ParentObject);
+
+		bool MultipleBaseObjectsFound;
+		UCustomizableObjectNodeObject* Root = GetRootNode(Node->ParentObject, MultipleBaseObjectsFound);
+
+		if (Root->ParentObject == nullptr)
+		{
+			if (MultipleBaseObjectsFound)
+			{
+				return nullptr;
+			}
+			else
+			{
+				return Root;
+			}
+		}
+		else
+		{
+			if (VisitedObjects.Contains(Root->ParentObject))
+			{
+				//There is a cycle
+				return nullptr;
+			}
+			else
+			{
+				return GetFullGraphRootNodeObject(Root, VisitedObjects);
+			}
+		}
+	}
+
+	return nullptr;
+}
+
+
+const UEdGraphPin* FindMeshBaseSource(const UEdGraphPin& Pin, const bool bOnlyLookForStaticMesh)
+{
+	check(Pin.Direction == EGPD_Output);
+	check(Pin.PinType.PinCategory == UEdGraphSchema_CustomizableObject::PC_Mesh 
+		||
+		Pin.PinType.PinCategory == UEdGraphSchema_CustomizableObject::PC_PassThroughMesh
+		||
+		Pin.PinType.PinCategory == UEdGraphSchema_CustomizableObject::PC_Material
+		||
+		Pin.PinType.PinCategory == UEdGraphSchema_CustomizableObject::PC_Modifier
+	);
+	
+	const UEdGraphNode* Node = Pin.GetOwningNode();
+	check(Node);
+
+	if (Cast<UCustomizableObjectNodeSkeletalMesh>(Node))
+	{
+		if (!bOnlyLookForStaticMesh)
+		{
+			return &Pin;
+		}
+	}
+
+	else if (Cast<UCustomizableObjectNodeStaticMesh>(Node))
+	{
+		return &Pin;
+	}
+
+	else if (const UCustomizableObjectNodeMeshGeometryOperation* TypedNodeGeom = Cast<UCustomizableObjectNodeMeshGeometryOperation>(Node))
+	{
+		if (const UEdGraphPin* ConnectedPin = FollowInputPin(*TypedNodeGeom->MeshAPin()))
+		{
+			return FindMeshBaseSource(*ConnectedPin, bOnlyLookForStaticMesh);
+		}
+	}
+
+	else if (const UCustomizableObjectNodeMeshReshape* TypedNodeReshape = Cast<UCustomizableObjectNodeMeshReshape>(Node))
+	{
+		if (const UEdGraphPin* ConnectedPin = FollowInputPin(*TypedNodeReshape->BaseMeshPin()))
+		{
+			return FindMeshBaseSource(*ConnectedPin, bOnlyLookForStaticMesh);
+		}
+	}
+
+	else if (const UCustomizableObjectNodeMeshMorph* TypedNodeMorph = Cast<UCustomizableObjectNodeMeshMorph>(Node))
+	{
+		if (const UEdGraphPin* ConnectedPin = FollowInputPin(*TypedNodeMorph->MeshPin()))
+		{
+			return FindMeshBaseSource(*ConnectedPin, bOnlyLookForStaticMesh);
+		}
+	}
+
+	else if (const UCustomizableObjectNodeMeshSwitch* TypedNodeSwitch = Cast<UCustomizableObjectNodeMeshSwitch>(Node))
+	{
+		if (const UEdGraphPin* EnumParameterPin = FollowInputPin(*TypedNodeSwitch->SwitchParameter()))
+		{
+			if (const UCustomizableObjectNodeEnumParameter* EnumNode = Cast<UCustomizableObjectNodeEnumParameter>(EnumParameterPin->GetOwningNode()))
+			{
+				if (const UEdGraphPin* DefaultPin = TypedNodeSwitch->GetElementPin(EnumNode->DefaultIndex))
+				{
+					if (const UEdGraphPin* ConnectedPin = FollowInputPin(*DefaultPin))
+					{
+						return FindMeshBaseSource(*ConnectedPin, bOnlyLookForStaticMesh);
+					}
+				}
+			}
+		}
+	}
+
+	else if (const UCustomizableObjectNodeMeshVariation* TypedNodeMeshVar = Cast<UCustomizableObjectNodeMeshVariation>(Node))
+	{
+		if (const UEdGraphPin* ConnectedPin = FollowInputPin(*TypedNodeMeshVar->DefaultPin()))
+		{
+			return FindMeshBaseSource(*ConnectedPin, bOnlyLookForStaticMesh);
+		}
+
+		for (int32 i = 0; i < TypedNodeMeshVar->GetNumVariations(); ++i)
+		{
+			if (const UEdGraphPin* ConnectedPin = FollowInputPin(*TypedNodeMeshVar->VariationPin(i)))
+			{
+				return FindMeshBaseSource(*ConnectedPin, bOnlyLookForStaticMesh);
+			}
+		}
+	}
+
+	else if (const UCustomizableObjectNodeMaterialBase* TypedNodeMat = Cast<UCustomizableObjectNodeMaterialBase>(Node))
+	{
+		if (const UEdGraphPin* ConnectedPin = FollowInputPin(*TypedNodeMat->GetMeshPin()))
+		{
+			return FindMeshBaseSource(*ConnectedPin, bOnlyLookForStaticMesh);
+		}
+	}
+
+	else if (const UCustomizableObjectNodeMaterialVariation* TypedNodeMatVar = Cast<UCustomizableObjectNodeMaterialVariation>(Node))
+	{
+		if (const UEdGraphPin* ConnectedPin = FollowInputPin(*TypedNodeMatVar->DefaultPin()))
+		{
+			return FindMeshBaseSource(*ConnectedPin, bOnlyLookForStaticMesh);
+		}
+	}
+
+	else if (const UCustomizableObjectNodeModifierExtendMeshSection* TypedNodeExtend = Cast<UCustomizableObjectNodeModifierExtendMeshSection>(Node))
+	{
+		if (const UEdGraphPin* ConnectedPin = FollowInputPin(*TypedNodeExtend->AddMeshPin()))
+		{
+			return FindMeshBaseSource(*ConnectedPin, bOnlyLookForStaticMesh);
+		}
+	}
+	
+	else if (const UCustomizableObjectNodeMeshMorphStackDefinition* TypedNodeMorphStackDef = Cast<UCustomizableObjectNodeMeshMorphStackDefinition>(Node))
+	{
+		if (const UEdGraphPin* ConnectedPin = FollowInputPin(*TypedNodeMorphStackDef->GetMeshPin()))
+		{
+			return FindMeshBaseSource(*ConnectedPin, bOnlyLookForStaticMesh);
+		}
+	}
+
+	else if (const UCustomizableObjectNodeMeshMorphStackApplication* TypedNodeMorphStackApp = Cast<UCustomizableObjectNodeMeshMorphStackApplication>(Node))
+	{
+		if (const UEdGraphPin* ConnectedPin = FollowInputPin(*TypedNodeMorphStackApp->GetMeshPin()))
+		{
+			return FindMeshBaseSource(*ConnectedPin, bOnlyLookForStaticMesh);
+		}
+	}
+
+	else if (const UCustomizableObjectNodeMeshReshape* NodeMeshReshape = Cast<UCustomizableObjectNodeMeshReshape>(Node))
+	{
+		if (const UEdGraphPin* ConnectedPin = FollowInputPin(*NodeMeshReshape->BaseMeshPin()))
+		{
+			return FindMeshBaseSource(*ConnectedPin, bOnlyLookForStaticMesh);
+		}
+	}
+	
+	else if (Cast<UCustomizableObjectNodeTable>(Node))
+	{
+		if (!bOnlyLookForStaticMesh)
+		{
+			return &Pin;
+		}
+	}
+
+	else if (const UCustomizableObjectNodeAnimationPose* NodeMeshPose = Cast<UCustomizableObjectNodeAnimationPose>(Node))
+	{
+		if (const UEdGraphPin* ConnectedPin = FollowInputPin(*NodeMeshPose->GetInputMeshPin()))
+		{
+			return FindMeshBaseSource(*ConnectedPin, bOnlyLookForStaticMesh);
+		}
+	}
+
+	else
+	{
+		unimplemented(); // Case missing.
+	}
+	
+	return nullptr;
+}
+
+
+void GetNodeGroupObjectNodeMappingImmersive(UCustomizableObject* Object, FAssetRegistryModule& AssetRegistryModule, TSet<UCustomizableObject*>& Visited, TMultiMap<FGuid, UCustomizableObjectNodeObject*>& Mapping)
+{
+	Visited.Add(Object);
+
+	TArray<FName> ArrayReferenceNames;
+	AssetRegistryModule.Get().GetReferencers(*Object->GetOuter()->GetPathName(), ArrayReferenceNames, UE::AssetRegistry::EDependencyCategory::Package, UE::AssetRegistry::EDependencyQuery::Hard);
+
+	FARFilter Filter;
+	for (const FName& ReferenceName : ArrayReferenceNames)
+	{
+		if (!ReferenceName.ToString().StartsWith(TEXT("/TempAutosave")))
+		{
+			Filter.PackageNames.Add(ReferenceName);
+		}
+	}
+
+	Filter.bIncludeOnlyOnDiskAssets = false;
+	
+	TArray<FAssetData> ArrayAssetData;
+	AssetRegistryModule.Get().GetAssets(Filter, ArrayAssetData);
+
+	for (FAssetData& AssetData : ArrayAssetData)
+	{
+		UCustomizableObject* ChildObject = Cast<UCustomizableObject>(AssetData.GetAsset());
+		if (!ChildObject)
+		{
+			continue;
+		}
+
+		if (ChildObject != Object && !ChildObject->HasAnyFlags(RF_Transient))
+		{
+			bool bMultipleBaseObjectsFound = false;
+			UCustomizableObjectNodeObject* ChildRoot = GetRootNode(ChildObject, bMultipleBaseObjectsFound);
+
+			if (ChildRoot && !bMultipleBaseObjectsFound)
+			{
+				if (ChildRoot->ParentObject == Object)
+				{
+					Mapping.Add(ChildRoot->ParentObjectGroupId, ChildRoot);
+				}
+			}
+		}
+
+		if (!Visited.Contains(ChildObject))
+		{
+			GetNodeGroupObjectNodeMappingImmersive(ChildObject, AssetRegistryModule, Visited, Mapping);
+		}
+	}
+}
+
+
+TMultiMap<FGuid, UCustomizableObjectNodeObject*> GetNodeGroupObjectNodeMapping(UCustomizableObject* Object)
+{
+	FAssetRegistryModule& AssetRegistryModule = FModuleManager::LoadModuleChecked<FAssetRegistryModule>("AssetRegistry");
+
+	TSet<UCustomizableObject*> Visited;
+	TMultiMap<FGuid, UCustomizableObjectNodeObject*> Mapping;
+
+	GetNodeGroupObjectNodeMappingImmersive(Object, AssetRegistryModule, Visited, Mapping);
+	
+	return Mapping;
+}
+
+
+void GetAllObjectsInGraph(UCustomizableObject* Object, TSet<UCustomizableObject*>& OutObjects)
+{
+	if (!Object)
+	{
+		return;
+	}
+
+	// Search the root of the CO's graph
+	UCustomizableObject* RootObject = GetRootObject(Object);
+	TMultiMap<FGuid, UCustomizableObjectNodeObject*> DummyMap;
+
+	FAssetRegistryModule& AssetRegistryModule = FModuleManager::LoadModuleChecked<FAssetRegistryModule>("AssetRegistry");
+	GetNodeGroupObjectNodeMappingImmersive(RootObject, AssetRegistryModule, OutObjects, DummyMap);
+}
+
+
+namespace GraphTraversal
+{
+	bool IsRootObject(const UCustomizableObject& Object)
+	{
+		const TObjectPtr<UEdGraph> Source = Object.GetPrivate()->GetSource();
+		if (!Source || !Source->Nodes.Num())
+		{
+			// Conservative approach.
+			return true;
+		}
+
+		TArray<UCustomizableObjectNodeObject*> ObjectNodes;
+		Source->GetNodesOfClass<UCustomizableObjectNodeObject>(ObjectNodes);
+
+		// Look for the base object node
+		const UCustomizableObjectNodeObject* Root = nullptr;
+		for (TArray<UCustomizableObjectNodeObject*>::TIterator It(ObjectNodes); It; ++It)
+		{
+			if ((*It)->bIsBase)
+			{
+				Root = *It;
+			}
+		}
+
+		return Root && !Root->ParentObject;
+	}
+}
+
+
+void NodePinConnectionListChanged(const TArray<UEdGraphPin*>& Pins)
+{
+	TMap<UEdGraphNode*, TSet<UEdGraphPin*>> SortedPins;
+	for (UEdGraphPin* Pin : Pins) 
+	{
+		if (UEdGraphNode* Node = Pin->GetOwningNodeUnchecked())
+		{
+			SortedPins.FindOrAdd(Node).Add(Pin);
+		}
+	}
+
+	for (TTuple<UEdGraphNode*, TSet<UEdGraphPin*>> Pair : SortedPins)
+	{
+		for (UEdGraphPin* ConnectedPin : Pair.Value)
+		{
+			Pair.Key->PinConnectionListChanged(ConnectedPin);
+		}
+		
+		Pair.Key->NodeConnectionListChanged();
+	}
+}
diff --git a/Engine/Plugins/Mutable/Source/CustomizableObjectEditor/Private/MuCOE/MutableUtils.cpp b/Engine/Plugins/Mutable/Source/CustomizableObjectEditor/Private/MuCOE/MutableUtils.cpp
new file mode 100644
--- /dev/null
+++ b/Engine/Plugins/Mutable/Source/CustomizableObjectEditor/Private/MuCOE/MutableUtils.cpp
@@ -0,0 +1,255 @@
+// Copyright Epic Games, Inc. All Rights Reserved.
+
+#include "MuCOE/MutableUtils.h"
+
+#include "Engine/SkeletalMesh.h"
+#include "Rendering/SkeletalMeshLODModel.h"
+#include "Engine/StaticMesh.h"
+#include "Rendering/SkeletalMeshModel.h"
+#include "Rendering/SkeletalMeshRenderData.h"
+
+
+TArray<FVector2f> GetUV(const USkeletalMesh& SkeletalMesh, const int32 LODIndex, const int32 SectionIndex, const int32 UVIndex)
+{
+	TArray<FVector2f> Result;
+
+	const FSkeletalMeshModel* ImportedModel = SkeletalMesh.GetImportedModel();
+	const FSkeletalMeshLODModel& LOD = ImportedModel->LODModels[LODIndex];
+	const FSkelMeshSection& Section = LOD.Sections[SectionIndex];
+	
+	TArray<FSoftSkinVertex> Vertices;
+	LOD.GetVertices(Vertices);
+
+	TArray<uint32> Indices;
+	SkeletalMesh.GetResourceForRendering()->LODRenderData[LODIndex].MultiSizeIndexContainer.GetIndexBuffer(Indices);
+	
+	int32 IndexIndex = Section.BaseIndex;
+	for (uint32 FaceIndex = 0; FaceIndex < Section.NumTriangles; ++FaceIndex, IndexIndex += 3)
+	{
+		Result.Add(Vertices[Indices[IndexIndex + 0]].UVs[UVIndex]);
+		Result.Add(Vertices[Indices[IndexIndex + 1]].UVs[UVIndex]);
+
+		Result.Add(Vertices[Indices[IndexIndex + 1]].UVs[UVIndex]);
+		Result.Add(Vertices[Indices[IndexIndex + 2]].UVs[UVIndex]);
+		
+		Result.Add(Vertices[Indices[IndexIndex + 2]].UVs[UVIndex]);
+		Result.Add(Vertices[Indices[IndexIndex + 0]].UVs[UVIndex]);
+	}
+
+	return Result;
+}
+
+
+TArray<FVector2f> GetUV(const UStaticMesh& StaticMesh, const int32 LODIndex, const int32 SectionIndex, const int32 UVIndex)
+{
+	TArray<FVector2f> Result;
+
+	const FStaticMeshRenderData* RenderData = StaticMesh.GetRenderData();
+	const FStaticMeshLODResources& LOD = RenderData->LODResources[LODIndex];
+	const FStaticMeshSection& Section = LOD.Sections[SectionIndex];
+	
+	const FStaticMeshVertexBuffer& VertexBuffer = LOD.VertexBuffers.StaticMeshVertexBuffer;
+	
+	FIndexArrayView Indices = LOD.IndexBuffer.GetArrayView();
+
+	int32 IndexIndex = Section.FirstIndex;
+	for (uint32 FaceIndex = 0; FaceIndex < Section.NumTriangles; ++FaceIndex, IndexIndex += 3)
+	{
+		Result.Add(VertexBuffer.GetVertexUV(Indices[IndexIndex + 0], UVIndex));
+		Result.Add(VertexBuffer.GetVertexUV(Indices[IndexIndex + 1], UVIndex));
+
+		Result.Add(VertexBuffer.GetVertexUV(Indices[IndexIndex + 1], UVIndex));
+		Result.Add(VertexBuffer.GetVertexUV(Indices[IndexIndex + 2], UVIndex));
+		
+		Result.Add(VertexBuffer.GetVertexUV(Indices[IndexIndex + 2], UVIndex));
+		Result.Add(VertexBuffer.GetVertexUV(Indices[IndexIndex + 0], UVIndex));
+	}
+	
+	return Result;
+}
+
+
+bool HasNormalizedBounds(const FVector2f& Point)
+{
+	return Point.X >= 0.0f && Point.X <= 1.0 && Point.Y >= 0.0f && Point.Y <= 1.0;
+}
+
+
+bool IsMeshClosed(const USkeletalMesh* Mesh, int LOD, int MaterialIndex)
+{
+	check(Mesh->GetImportedModel() &&
+		Mesh->GetImportedModel()->LODModels.IsValidIndex(LOD) &&
+		Mesh->GetImportedModel()->LODModels[LOD].Sections.IsValidIndex(MaterialIndex));
+
+	const float VertexCollapseEpsilonSquared = 0.0001f * 0.0001f;
+
+	const FSkeletalMeshLODModel& LODResource = Mesh->GetImportedModel()->LODModels[LOD];
+	const FSkelMeshSection& Section = LODResource.Sections[MaterialIndex];
+
+	const TArray<FSoftSkinVertex>& VertexBuffer = Section.SoftVertices;
+	const uint32 VertexStart = Section.BaseVertexIndex;
+	const uint32 VertexCount = Section.NumVertices;
+
+	// Collapse vertices
+	TArray<uint32> CollapsedVertexMap;
+	CollapsedVertexMap.SetNum(VertexCount);
+
+	auto GetCollapsedVertex = [&CollapsedVertexMap, VertexStart](uint32 VertIndex) -> uint32&
+	{
+
+		check(CollapsedVertexMap.IsValidIndex(VertIndex - VertexStart));
+		return CollapsedVertexMap[VertIndex - VertexStart];
+	};
+
+	for (uint32 V = VertexStart; V < VertexCount; ++V)
+	{
+		GetCollapsedVertex(V) = V;
+
+		for (uint32 CandidateIndex = VertexStart; CandidateIndex < V; ++CandidateIndex)
+		{
+			uint32 Collapsed = GetCollapsedVertex(CandidateIndex);
+
+			FVector3f Vec = VertexBuffer[V].Position - VertexBuffer[Collapsed].Position;
+
+			if (FVector3f::DotProduct(Vec, Vec) <= VertexCollapseEpsilonSquared)
+			{
+				GetCollapsedVertex(V) = Collapsed;
+				break;
+			}
+		}
+	}
+
+	// Count faces per edge
+	TMap<TPair<uint32, uint32>, uint32> FaceCountPerEdge;
+
+	const TArray<uint32>& IndexBuffer = LODResource.IndexBuffer;//.GetData();
+	const uint32 IndexStart = Section.BaseIndex;
+	const uint32 IndexCount = Section.NumTriangles * 3;
+
+	for (uint32 I = IndexStart; I < IndexCount; I += 3)
+	{
+		uint32 Face[3];
+		uint32 Index0 = IndexBuffer[I + 0];
+		Face[0] = GetCollapsedVertex(Index0);
+		uint32 Index1 = IndexBuffer[I + 1];
+		Face[1] = GetCollapsedVertex(Index1);
+		uint32 Index2 = IndexBuffer[I + 2];
+		Face[2] = GetCollapsedVertex(Index2);
+
+		for (int E = 0; E < 3; ++E)
+		{
+			uint32 Vertex0 = Face[E];
+			uint32 Vertex1 = Face[(E + 1) % 3];
+
+			if (Vertex0 == Vertex1)
+			{
+				return false; // Degenerated mesh
+			}
+
+			TPair<uint32, uint32> Edge;
+			Edge.Key = FGenericPlatformMath::Min(Vertex0, Vertex1);
+			Edge.Value = FGenericPlatformMath::Max(Vertex0, Vertex1);
+
+			FaceCountPerEdge.FindOrAdd(Edge)++;
+		}
+	}
+
+	// See if every edge has 2 faces
+	for (const TPair<TPair<uint32, uint32>, uint32>& Entry : FaceCountPerEdge)
+	{
+		if (Entry.Value != 2)
+		{
+			return false;
+		}
+	}
+
+	return true;
+}
+
+
+bool IsMeshClosed(const UStaticMesh* Mesh, int LOD, int MaterialIndex)
+{
+	check(Mesh->GetRenderData() &&
+		Mesh->GetRenderData()->LODResources.IsValidIndex(LOD) &&
+		Mesh->GetRenderData()->LODResources[LOD].Sections.IsValidIndex(MaterialIndex));
+
+	const float VertexCollapseEpsilonSquared = 0.0001f * 0.0001f;
+
+	const FStaticMeshLODResources& LODResource = Mesh->GetRenderData()->LODResources[LOD];
+	const FStaticMeshSection& Section = LODResource.Sections[MaterialIndex];
+
+	const uint32 VertexStart = Section.MinVertexIndex;
+	const uint32 VertexCount = Section.MaxVertexIndex - VertexStart + 1;
+
+	// Collapse vertices
+	TArray<uint32> CollapsedVertexMap;
+	CollapsedVertexMap.SetNum(VertexCount);
+
+	auto GetCollapsedVertex = [&CollapsedVertexMap, VertexStart](uint32 VertIndex) -> uint32&
+	{
+
+		check(CollapsedVertexMap.IsValidIndex(VertIndex - VertexStart));
+		return CollapsedVertexMap[VertIndex - VertexStart];
+	};
+
+	for (uint32 V = VertexStart; V < VertexCount; ++V)
+	{
+		GetCollapsedVertex(V) = V;
+
+		for (uint32 CandidateIndex = VertexStart; CandidateIndex < V; ++CandidateIndex)
+		{
+			uint32 Collapsed = GetCollapsedVertex(CandidateIndex);
+
+			FVector3f Vec = LODResource.VertexBuffers.PositionVertexBuffer.VertexPosition(V) - LODResource.VertexBuffers.PositionVertexBuffer.VertexPosition(Collapsed);
+
+			if (FVector3f::DotProduct(Vec, Vec) <= VertexCollapseEpsilonSquared)
+			{
+				GetCollapsedVertex(V) = Collapsed;
+				break;
+			}
+		}
+	}
+
+	// Count faces per edge
+	TMap<TPair<uint32, uint32>, uint32> FaceCountPerEdge;
+
+	const FIndexArrayView& IndexBuffer = LODResource.IndexBuffer.GetArrayView();
+	const uint32 IndexStart = Section.FirstIndex;
+	const uint32 IndexCount = Section.NumTriangles * 3;
+
+	for (uint32 I = IndexStart; I < IndexCount; I += 3)
+	{
+		uint32 Face[3];
+		Face[0] = GetCollapsedVertex(IndexBuffer[I + 0]);
+		Face[1] = GetCollapsedVertex(IndexBuffer[I + 1]);
+		Face[2] = GetCollapsedVertex(IndexBuffer[I + 2]);
+
+		for (int E = 0; E < 3; ++E)
+		{
+			uint32 Vertex0 = Face[E];
+			uint32 Vertex1 = Face[(E + 1) % 3];
+
+			if (Vertex0 == Vertex1)
+			{
+				return false; // Degenerated mesh
+			}
+
+			TPair<uint32, uint32> Edge;
+			Edge.Key = FGenericPlatformMath::Min(Vertex0, Vertex1);
+			Edge.Value = FGenericPlatformMath::Max(Vertex0, Vertex1);
+
+			FaceCountPerEdge.FindOrAdd(Edge)++;
+		}
+	}
+
+	// See if every edge has 2 faces
+	for (const TPair<TPair<uint32, uint32>, uint32>& Entry : FaceCountPerEdge)
+	{
+		if (Entry.Value != 2)
+		{
+			return false;
+		}
+	}
+
+	return true;
+}
diff --git a/Engine/Plugins/Mutable/Source/CustomizableObjectEditor/Private/MuCOE/Nodes/CustomizableObjectNodeColorArithmeticOp.cpp b/Engine/Plugins/Mutable/Source/CustomizableObjectEditor/Private/MuCOE/Nodes/CustomizableObjectNodeColorArithmeticOp.cpp
new file mode 100644
--- /dev/null
+++ b/Engine/Plugins/Mutable/Source/CustomizableObjectEditor/Private/MuCOE/Nodes/CustomizableObjectNodeColorArithmeticOp.cpp
@@ -0,0 +1,84 @@
+// Copyright Epic Games, Inc. All Rights Reserved.
+
+#include "MuCOE/Nodes/CustomizableObjectNodeColorArithmeticOp.h"
+
+#include "MuCOE/EdGraphSchema_CustomizableObject.h"
+
+class UCustomizableObjectNodeRemapPins;
+struct FPropertyChangedEvent;
+
+#define LOCTEXT_NAMESPACE "CustomizableObjectEditor"
+
+
+UCustomizableObjectNodeColorArithmeticOp::UCustomizableObjectNodeColorArithmeticOp()
+	: Super()
+{
+	Operation = EColorArithmeticOperation::E_Add;
+}
+
+
+void UCustomizableObjectNodeColorArithmeticOp::PostEditChangeProperty(FPropertyChangedEvent& PropertyChangedEvent)
+{
+	//UProperty* PropertyThatChanged = PropertyChangedEvent.Property;
+	//if ( PropertyThatChanged && PropertyThatChanged->GetName() == TEXT("NumLODs") )
+	{
+	}
+
+	Super::PostEditChangeProperty(PropertyChangedEvent);
+}
+
+
+void UCustomizableObjectNodeColorArithmeticOp::AllocateDefaultPins(UCustomizableObjectNodeRemapPins* RemapPins)
+{
+	const UEdGraphSchema_CustomizableObject* Schema = GetDefault<UEdGraphSchema_CustomizableObject>();
+
+	UEdGraphPin* XPin = CustomCreatePin(EGPD_Input, Schema->PC_Color, FName("A"));
+	XPin->bDefaultValueIsIgnored = true;
+	XPin->PinFriendlyName = FText();
+
+	FName PinName = TEXT("B");
+	UEdGraphPin* YPin = CustomCreatePin(EGPD_Input, Schema->PC_Color, PinName);
+	YPin->bDefaultValueIsIgnored = true;
+	YPin->PinFriendlyName = FText();
+
+	PinName = TEXT("Result");
+	UEdGraphPin* ColorPin = CustomCreatePin(EGPD_Output, Schema->PC_Color,PinName);
+	ColorPin->bDefaultValueIsIgnored = true;
+	ColorPin->PinFriendlyName = FText();
+}
+
+
+FText UCustomizableObjectNodeColorArithmeticOp::GetNodeTitle(ENodeTitleType::Type TitleType) const
+{
+	if (TitleType == ENodeTitleType::ListView)
+	{
+		return LOCTEXT("Color_Arithmetic_Operation", "Color Arithmetic Operation");
+	}
+
+	const UEnum* EnumPtr = FindObject<UEnum>(nullptr, TEXT("/Script/CustomizableObjectEditor.EColorArithmeticOperation"), true);
+
+	if (!EnumPtr)
+	{
+		return FText::FromString(FString("Color Operation"));
+	}
+
+	const int32 index = EnumPtr->GetIndexByValue((int32)Operation);
+	return EnumPtr->GetDisplayNameTextByIndex(index);
+}
+
+
+FLinearColor UCustomizableObjectNodeColorArithmeticOp::GetNodeTitleColor() const
+{
+	const UEdGraphSchema_CustomizableObject* Schema = GetDefault<UEdGraphSchema_CustomizableObject>();
+	return Schema->GetPinTypeColor(Schema->PC_Color);
+}
+
+
+FText UCustomizableObjectNodeColorArithmeticOp::GetTooltipText() const
+{
+	return LOCTEXT("Color_Arithmetic_Tooltip", "Perform an arithmetic operation between two colors on a per-component basis.");
+}
+
+
+#undef LOCTEXT_NAMESPACE
+
diff --git a/Engine/Plugins/Mutable/Source/CustomizableObjectEditor/Private/MuCOE/Nodes/CustomizableObjectNodeColorFromFloats.cpp b/Engine/Plugins/Mutable/Source/CustomizableObjectEditor/Private/MuCOE/Nodes/CustomizableObjectNodeColorFromFloats.cpp
new file mode 100644
--- /dev/null
+++ b/Engine/Plugins/Mutable/Source/CustomizableObjectEditor/Private/MuCOE/Nodes/CustomizableObjectNodeColorFromFloats.cpp
@@ -0,0 +1,67 @@
+// Copyright Epic Games, Inc. All Rights Reserved.
+
+#include "MuCOE/Nodes/CustomizableObjectNodeColorFromFloats.h"
+
+#include "MuCOE/EdGraphSchema_CustomizableObject.h"
+
+class UCustomizableObjectNodeRemapPins;
+struct FPropertyChangedEvent;
+
+#define LOCTEXT_NAMESPACE "CustomizableObjectEditor"
+
+
+UCustomizableObjectNodeColorFromFloats::UCustomizableObjectNodeColorFromFloats()
+	: Super()
+{
+
+}
+
+
+void UCustomizableObjectNodeColorFromFloats::PostEditChangeProperty(FPropertyChangedEvent& PropertyChangedEvent)
+{
+	Super::PostEditChangeProperty(PropertyChangedEvent);
+}
+
+
+void UCustomizableObjectNodeColorFromFloats::AllocateDefaultPins(UCustomizableObjectNodeRemapPins* RemapPins)
+{
+	const UEdGraphSchema_CustomizableObject* Schema = GetDefault<UEdGraphSchema_CustomizableObject>();
+
+	UEdGraphPin* RPin = CustomCreatePin(EGPD_Input, Schema->PC_Float, FName("R"));
+	RPin->bDefaultValueIsIgnored = true;
+
+	UEdGraphPin* GPin = CustomCreatePin(EGPD_Input, Schema->PC_Float, FName("G"));
+	GPin->bDefaultValueIsIgnored = true;
+
+	UEdGraphPin* BPin = CustomCreatePin(EGPD_Input, Schema->PC_Float, FName("B"));
+	BPin->bDefaultValueIsIgnored = true;
+
+	UEdGraphPin* APin = CustomCreatePin(EGPD_Input, Schema->PC_Float, FName("A"));
+	APin->bDefaultValueIsIgnored = true;
+
+	UEdGraphPin* ColorPin = CustomCreatePin(EGPD_Output, Schema->PC_Color, FName("Color"));
+	ColorPin->bDefaultValueIsIgnored = true;
+}
+
+
+FText UCustomizableObjectNodeColorFromFloats::GetNodeTitle(ENodeTitleType::Type TitleType) const
+{
+	return LOCTEXT("Color_From_Floats", "Color From Floats");
+}
+
+
+FLinearColor UCustomizableObjectNodeColorFromFloats::GetNodeTitleColor() const
+{
+	const UEdGraphSchema_CustomizableObject* Schema = GetDefault<UEdGraphSchema_CustomizableObject>();
+	return Schema->GetPinTypeColor(Schema->PC_Color);
+}
+
+
+FText UCustomizableObjectNodeColorFromFloats::GetTooltipText() const
+{
+	return LOCTEXT("Color_From_Floats_Tooltip", "Defines a color from its four RGBA components.");
+}
+
+
+#undef LOCTEXT_NAMESPACE
+
diff --git a/Engine/Plugins/Mutable/Source/CustomizableObjectEditor/Private/MuCOE/Nodes/CustomizableObjectNodeColorParameter.cpp b/Engine/Plugins/Mutable/Source/CustomizableObjectEditor/Private/MuCOE/Nodes/CustomizableObjectNodeColorParameter.cpp
new file mode 100644
--- /dev/null
+++ b/Engine/Plugins/Mutable/Source/CustomizableObjectEditor/Private/MuCOE/Nodes/CustomizableObjectNodeColorParameter.cpp
@@ -0,0 +1,61 @@
+// Copyright Epic Games, Inc. All Rights Reserved.
+
+#include "MuCOE/Nodes/CustomizableObjectNodeColorParameter.h"
+
+#include "MuCOE/EdGraphSchema_CustomizableObject.h"
+#include "MuCO/CustomizableObjectCustomVersion.h"
+
+class UCustomizableObjectNodeRemapPins;
+
+#define LOCTEXT_NAMESPACE "CustomizableObjectEditor"
+
+
+void UCustomizableObjectNodeColorParameter::AllocateDefaultPins(UCustomizableObjectNodeRemapPins* RemapPins)
+{
+	const UEdGraphSchema_CustomizableObject* Schema = GetDefault<UEdGraphSchema_CustomizableObject>();
+
+	UEdGraphPin* ValuePin = CustomCreatePin(EGPD_Output, Schema->PC_Color, FName("Value"));
+	ValuePin->bDefaultValueIsIgnored = true;
+}
+
+
+FText UCustomizableObjectNodeColorParameter::GetNodeTitle(ENodeTitleType::Type TitleType) const
+{
+	if (TitleType == ENodeTitleType::ListView || ParameterName.IsEmpty())
+	{
+		return LOCTEXT("Color_Parameter", "Color Parameter");
+	}
+	else if (TitleType == ENodeTitleType::EditableTitle)
+	{
+		return FText::Format(LOCTEXT("Color_Parameter_EditableTitle", "{0}"), FText::FromString(ParameterName));
+	}
+	else
+	{
+		return FText::Format(LOCTEXT("Color_Parameter_Title", "{0}\nColor Parameter"), FText::FromString(ParameterName));
+	}
+}
+
+
+FLinearColor UCustomizableObjectNodeColorParameter::GetNodeTitleColor() const
+{
+	const UEdGraphSchema_CustomizableObject* Schema = GetDefault<UEdGraphSchema_CustomizableObject>();
+	return Schema->GetPinTypeColor(Schema->PC_Color);
+}
+
+
+FText UCustomizableObjectNodeColorParameter::GetTooltipText() const
+{
+	return LOCTEXT("Color_Parameter_Tooltip", "Expose a runtime modifiable color parameter from the Customizable Object.");
+}
+
+
+void UCustomizableObjectNodeColorParameter::OnRenameNode(const FString& NewName)
+{
+	if (!NewName.IsEmpty())
+	{
+		ParameterName = NewName;
+	}
+}
+
+#undef LOCTEXT_NAMESPACE
+
diff --git a/Engine/Plugins/Mutable/Source/CustomizableObjectEditor/Private/MuCOE/Nodes/CustomizableObjectNodeColorSwitch.cpp b/Engine/Plugins/Mutable/Source/CustomizableObjectEditor/Private/MuCOE/Nodes/CustomizableObjectNodeColorSwitch.cpp
new file mode 100644
--- /dev/null
+++ b/Engine/Plugins/Mutable/Source/CustomizableObjectEditor/Private/MuCOE/Nodes/CustomizableObjectNodeColorSwitch.cpp
@@ -0,0 +1,29 @@
+// Copyright Epic Games, Inc. All Rights Reserved.
+
+#include "MuCOE/Nodes/CustomizableObjectNodeColorSwitch.h"
+
+#include "MuCO/CustomizableObjectCustomVersion.h"
+#include "MuCOE/EdGraphSchema_CustomizableObject.h"
+
+#define LOCTEXT_NAMESPACE "CustomizableObjectEditor"
+
+
+void UCustomizableObjectNodeColorSwitch::BackwardsCompatibleFixup(int32 CustomizableObjectCustomVersion)
+{
+	Super::BackwardsCompatibleFixup(CustomizableObjectCustomVersion);
+
+	if (CustomizableObjectCustomVersion == FCustomizableObjectCustomVersion::BugPinsSwitch)
+	{
+		OutputPinReference = FindPin(TEXT("Color"));	
+	}
+}
+
+
+FName UCustomizableObjectNodeColorSwitch::GetCategory() const
+{
+	return UEdGraphSchema_CustomizableObject::PC_Color;
+}
+
+
+#undef LOCTEXT_NAMESPACE
+
diff --git a/Engine/Plugins/Mutable/Source/CustomizableObjectEditor/Private/MuCOE/Nodes/CustomizableObjectNodeComponentMesh.cpp b/Engine/Plugins/Mutable/Source/CustomizableObjectEditor/Private/MuCOE/Nodes/CustomizableObjectNodeComponentMesh.cpp
new file mode 100644
--- /dev/null
+++ b/Engine/Plugins/Mutable/Source/CustomizableObjectEditor/Private/MuCOE/Nodes/CustomizableObjectNodeComponentMesh.cpp
@@ -0,0 +1,68 @@
+// Copyright Epic Games, Inc. All Rights Reserved.
+
+#include "MuCOE/Nodes/CustomizableObjectNodeComponentMesh.h"
+
+#include "MuCO/CustomizableObjectCustomVersion.h"
+#include "MuCOE/EdGraphSchema_CustomizableObject.h"
+#include "MuCOE/GraphTraversal.h"
+
+
+#define LOCTEXT_NAMESPACE "CustomizableObjectEditor"
+
+
+void UCustomizableObjectNodeComponentMesh::PostEditChangeProperty(FPropertyChangedEvent& PropertyChangedEvent)
+{
+	Super::PostEditChangeProperty(PropertyChangedEvent);
+
+	const FProperty* PropertyThatChanged = PropertyChangedEvent.Property;
+	if (!PropertyThatChanged)
+	{
+		return;
+	}
+
+	if (PropertyThatChanged->GetFName() == GET_MEMBER_NAME_CHECKED(UCustomizableObjectNodeComponentMesh, NumLODs))
+	{
+		LODReductionSettings.SetNum(NumLODs);
+		
+		ReconstructNode();
+	}
+}
+
+
+void UCustomizableObjectNodeComponentMesh::BackwardsCompatibleFixup(int32 CustomizableObjectCustomVersion)
+{
+	Super::BackwardsCompatibleFixup(CustomizableObjectCustomVersion);
+	
+	if (CustomizableObjectCustomVersion == FCustomizableObjectCustomVersion::ComponentsArray)
+	{
+		UCustomizableObject* Object = GetRootObject(*this);
+
+		if (FMutableMeshComponentData* Result = Object->GetPrivate()->MutableMeshComponents_DEPRECATED.FindByPredicate([&](const FMutableMeshComponentData& ComponentData)
+		{
+			return ComponentData.Name == ComponentName;
+		}))
+		{
+			ReferenceSkeletalMesh = Result->ReferenceSkeletalMesh;
+		}
+	}
+}
+
+
+FText UCustomizableObjectNodeComponentMesh::GetNodeTitle(ENodeTitleType::Type TitleType) const
+{
+	if (TitleType == ENodeTitleType::ListView)
+	{
+		return LOCTEXT("Component_Mesh", "Mesh Component");
+	}
+	else
+	{
+		FFormatNamedArguments Args;
+		Args.Add(TEXT("ComponentName"), FText::FromName(ComponentName));
+
+		return FText::Format(LOCTEXT("ComponentMesh_Title", "{ComponentName}\nMesh Component"), Args);
+	}
+}
+
+
+#undef LOCTEXT_NAMESPACE
+
