[mutable] Moved the Mutable plugin out of Experimental status into Beta.

diff --git a/Engine/Plugins/Mutable/Source/CustomizableObjectEditor/Private/MuCOE/Nodes/CustomizableObjectNodeCopyMaterial.cpp b/Engine/Plugins/Mutable/Source/CustomizableObjectEditor/Private/MuCOE/Nodes/CustomizableObjectNodeCopyMaterial.cpp
new file mode 100644
--- /dev/null
+++ b/Engine/Plugins/Mutable/Source/CustomizableObjectEditor/Private/MuCOE/Nodes/CustomizableObjectNodeCopyMaterial.cpp
@@ -0,0 +1,440 @@
+// Copyright Epic Games, Inc. All Rights Reserved.
+
+#include "MuCOE/Nodes/CustomizableObjectNodeCopyMaterial.h"
+
+#include "MuCOE/EdGraphSchema_CustomizableObject.h"
+#include "MuCOE/GraphTraversal.h"
+#include "MuCOE/Nodes/CustomizableObjectNodeExposePin.h"
+#include "MuCOE/Nodes/CustomizableObjectNodeExternalPin.h"
+#include "MuCOE/Nodes/CustomizableObjectNodeObject.h"
+#include "MuCOE/Nodes/CustomizableObjectNodeSkeletalMesh.h"
+#include "MuCOE/Nodes/CustomizableObjectNodeMaterial.h"
+#include "MuCO/CustomizableObjectCustomVersion.h"
+
+#define LOCTEXT_NAMESPACE "CustomizableObjectEditor"
+
+
+/** Mesh input pin key */
+const static FString MeshPinName = TEXT("Mesh_Input_Pin");
+
+/** Material input pin key */
+const static FString MeshSectionPinName = TEXT("MeshSection_Input_Pin");
+
+/** Material output pin key*/
+const static FString OutputPinName = TEXT("MeshSection_Output_Pin");
+
+
+FText UCustomizableObjectNodeCopyMaterial::GetNodeTitle(ENodeTitleType::Type TitleType) const
+{
+	return LOCTEXT("Copy_MeshSection", "Copy Mesh Section");
+}
+
+
+void UCustomizableObjectNodeCopyMaterial::AllocateDefaultPins(UCustomizableObjectNodeRemapPins* RemapPins)
+{
+	const UEdGraphSchema_CustomizableObject* Schema = GetDefault<UEdGraphSchema_CustomizableObject>();
+
+	// Input pins
+	UEdGraphPin* Pin = CustomCreatePin(EGPD_Input, Schema->PC_Mesh, FName(MeshPinName));
+	Pin->PinFriendlyName = LOCTEXT("MeshPin", "Mesh");
+	Pin->bDefaultValueIsIgnored = true;
+
+	Pin = CustomCreatePin(EGPD_Input, Schema->PC_Material, FName(MeshSectionPinName));
+	Pin->PinFriendlyName = LOCTEXT("BaseMeshSectionPin", "Base Mesh Section");;
+	Pin->bDefaultValueIsIgnored = true;
+
+	// Output pins
+	Pin = CustomCreatePin(EGPD_Output, Schema->PC_Material, FName(OutputPinName));
+	Pin->PinFriendlyName = LOCTEXT("MeshSectionPin", "Mesh Section");
+}
+
+
+void UCustomizableObjectNodeCopyMaterial::BackwardsCompatibleFixup(int32 CustomizableObjectCustomVersion)
+{
+	Super::BackwardsCompatibleFixup(CustomizableObjectCustomVersion);
+
+	if (CustomizableObjectCustomVersion == FCustomizableObjectCustomVersion::MaterialPinsRename)
+	{
+		UEdGraphPin* BaseMaterialPin = FindPin(TEXT("Material_Input_Pin"), EEdGraphPinDirection::EGPD_Input);
+		if (BaseMaterialPin)
+		{
+			const FString PinFriendlyName = TEXT("Base Mesh Section");
+			BaseMaterialPin->PinName = FName(MeshSectionPinName);
+			BaseMaterialPin->PinFriendlyName = FText::FromString(PinFriendlyName);
+		}
+
+		UEdGraphPin* MaterialPin = FindPin(TEXT("Material_Output_Pin"), EEdGraphPinDirection::EGPD_Output);
+		if (MaterialPin)
+		{
+			const FString PinFriendlyName = TEXT("Mesh Section");
+			MaterialPin->PinName = FName(OutputPinName);
+			MaterialPin->PinFriendlyName = FText::FromString(PinFriendlyName);
+		}
+	}
+}
+
+
+UEdGraphPin* UCustomizableObjectNodeCopyMaterial::GetMeshPin() const
+{
+	return FindPin(MeshPinName);
+}
+
+
+FPostImagePinModeChangedDelegate* UCustomizableObjectNodeCopyMaterial::GetPostImagePinModeChangedDelegate()
+{
+	if (UCustomizableObjectNodeMaterial* NodeMaterial = GetMaterialNode())
+	{
+		return NodeMaterial->GetPostImagePinModeChangedDelegate();
+	}
+	else
+	{
+		return nullptr;
+	}
+}
+
+
+TArray<UCustomizableObjectLayout*> UCustomizableObjectNodeCopyMaterial::GetLayouts() const
+{
+	if (UCustomizableObjectNodeMaterial* NodeMaterial = GetMaterialNode())
+	{
+		return NodeMaterial->GetLayouts();
+	}
+	else
+	{
+		return {};
+	}
+}
+
+
+UEdGraphPin* UCustomizableObjectNodeCopyMaterial::OutputPin() const
+{
+	UEdGraphPin* Pin = FindPin(OutputPinName, EEdGraphPinDirection::EGPD_Output);
+
+	// Legacy name
+	if (!Pin)
+	{
+		Pin = FindPin(TEXT("Material_Output_Pin"), EEdGraphPinDirection::EGPD_Output);
+	}
+
+	return Pin;
+}
+
+
+bool UCustomizableObjectNodeCopyMaterial::RealMaterialDataHasChanged() const
+{
+	if (UCustomizableObjectNodeMaterial* NodeMaterial = GetMaterialNode())
+	{
+		return NodeMaterial->RealMaterialDataHasChanged();
+	}
+	else
+	{
+		return false;
+	}
+}
+
+
+UEdGraphPin* UCustomizableObjectNodeCopyMaterial::GetMeshSectionPin() const
+{
+	UEdGraphPin* Pin = FindPin(MeshSectionPinName, EEdGraphPinDirection::EGPD_Input);
+
+	// Legacy name
+	if (!Pin)
+	{
+		Pin = FindPin(TEXT("Material_Input_Pin"), EEdGraphPinDirection::EGPD_Input);
+	}
+
+	return Pin;
+}
+
+
+UCustomizableObjectNodeSkeletalMesh* UCustomizableObjectNodeCopyMaterial::GetMeshNode() const
+{
+	UCustomizableObjectNodeSkeletalMesh* Result = nullptr;
+
+	UEdGraphPin* MaterialPin = GetMeshPin();
+	if (const UEdGraphPin* ConnectedPin = FollowInputPin(*MaterialPin))
+	{
+		const UEdGraphPin* SourceMeshPin = FindMeshBaseSource(*ConnectedPin, false);
+		if (SourceMeshPin)
+		{
+			Result = Cast<UCustomizableObjectNodeSkeletalMesh>(SourceMeshPin->GetOwningNode());
+		}
+	}
+
+	return Result;
+}
+
+
+UCustomizableObjectNodeMaterial* UCustomizableObjectNodeCopyMaterial::GetMaterialNode() const
+{
+	UCustomizableObjectNodeMaterial* Result = nullptr;
+
+	UEdGraphPin* MaterialPin = GetMeshSectionPin();
+	if (const UEdGraphPin* ConnectedPin = FollowInputPin(*MaterialPin))
+	{
+		return Cast<UCustomizableObjectNodeMaterial>(ConnectedPin->GetOwningNode());
+	}
+
+	return Result;
+}
+
+
+bool UCustomizableObjectNodeCopyMaterial::CanConnect(const UEdGraphPin* InOwnedInputPin, const UEdGraphPin* InOutputPin, bool& bOutIsOtherNodeBlocklisted, bool& bOutArePinsCompatible) const
+{
+	if (!Super::CanConnect(InOwnedInputPin, InOutputPin, bOutIsOtherNodeBlocklisted, bOutIsOtherNodeBlocklisted))
+	{
+		return false;
+	}
+
+	if (InOwnedInputPin == GetMeshSectionPin())
+	{
+		const UEdGraphNode* OuputPinOwningNode = InOutputPin->GetOwningNode();
+		return (OuputPinOwningNode->IsA(UCustomizableObjectNodeMaterial::StaticClass()) && !OuputPinOwningNode->IsA(UCustomizableObjectNodeCopyMaterial::StaticClass()))
+			|| OuputPinOwningNode->IsA(UCustomizableObjectNodeExternalPin::StaticClass());
+	}
+
+	return true;
+}
+
+
+bool UCustomizableObjectNodeCopyMaterial::ShouldBreakExistingConnections(const UEdGraphPin* InputPin, const UEdGraphPin* OutputPin) const
+{
+	return true;
+}
+
+
+bool UCustomizableObjectNodeCopyMaterial::IsNodeOutDatedAndNeedsRefresh()
+{
+	return false;
+}
+
+
+bool UCustomizableObjectNodeCopyMaterial::ProvidesCustomPinRelevancyTest() const
+{
+	return true;
+}
+
+
+bool UCustomizableObjectNodeCopyMaterial::IsPinRelevant(const UEdGraphPin* Pin) const
+{
+	const UEdGraphNode* Node = Pin->GetOwningNode();
+
+	if (Pin->Direction == EGPD_Output)
+	{
+		return (Node->IsA(UCustomizableObjectNodeMaterial::StaticClass()) && !Node->IsA(UCustomizableObjectNodeCopyMaterial::StaticClass())) ||
+			(Node->IsA(UCustomizableObjectNodeExternalPin::StaticClass()) && Pin->PinType.PinCategory == UEdGraphSchema_CustomizableObject::PC_Material) ||
+			Pin->PinType.PinCategory == UEdGraphSchema_CustomizableObject::PC_Mesh;
+	}
+	else
+	{
+		return Node->IsA(UCustomizableObjectNodeObject::StaticClass()) ||
+			(Node->IsA(UCustomizableObjectNodeExposePin::StaticClass()) && Pin->PinType.PinCategory == UEdGraphSchema_CustomizableObject::PC_Material);
+			
+	}
+}
+
+
+UMaterialInterface* UCustomizableObjectNodeCopyMaterial::GetMaterial() const
+{
+	if (UCustomizableObjectNodeMaterial* NodeMaterial = GetMaterialNode())
+	{
+		return NodeMaterial->GetMaterial();
+	}
+	else
+	{
+		return nullptr;
+	}
+}
+
+
+bool UCustomizableObjectNodeCopyMaterial::IsReuseMaterialBetweenLODs() const
+{
+	if (UCustomizableObjectNodeMaterial* NodeMaterial = GetMaterialNode())
+	{
+		return NodeMaterial->IsReuseMaterialBetweenLODs();
+	}
+	else
+	{
+		return false;
+	}
+}
+
+
+TArray<FString>* UCustomizableObjectNodeCopyMaterial::GetEnableTags()
+{
+	if (UCustomizableObjectNodeMaterial* NodeMaterial = GetMaterialNode())
+	{
+		return NodeMaterial->GetEnableTags();
+	}
+	else
+	{
+		return nullptr;
+	}
+}
+
+
+UEdGraphPin* UCustomizableObjectNodeCopyMaterial::GetMaterialAssetPin() const
+{
+	if (UCustomizableObjectNodeMaterial* NodeMaterial = GetMaterialNode())
+	{
+		return NodeMaterial->GetMaterialAssetPin();
+	}
+	else
+	{
+		return nullptr;
+	}
+}
+
+
+int32 UCustomizableObjectNodeCopyMaterial::GetNumParameters(EMaterialParameterType Type) const
+{
+	if (UCustomizableObjectNodeMaterial* NodeMaterial = GetMaterialNode())
+	{
+		return NodeMaterial->GetNumParameters(Type);
+	}
+	else
+	{
+		return 0;
+	}
+}
+
+
+FNodeMaterialParameterId UCustomizableObjectNodeCopyMaterial::GetParameterId(EMaterialParameterType Type, int32 ParameterIndex) const
+{
+	UCustomizableObjectNodeMaterial* NodeMaterial = GetMaterialNode();
+	check(NodeMaterial);
+
+	return NodeMaterial->GetParameterId(Type, ParameterIndex);
+}
+
+
+FName UCustomizableObjectNodeCopyMaterial::GetParameterName(EMaterialParameterType Type, int32 ParameterIndex) const
+{
+	UCustomizableObjectNodeMaterial* NodeMaterial = GetMaterialNode();
+	check(NodeMaterial);
+	
+	return NodeMaterial->GetParameterName(Type, ParameterIndex);
+}
+
+
+int32 UCustomizableObjectNodeCopyMaterial::GetParameterLayerIndex(EMaterialParameterType Type, int32 ParameterIndex) const
+{
+	UCustomizableObjectNodeMaterial* NodeMaterial = GetMaterialNode();
+	check(NodeMaterial);
+	
+	return NodeMaterial->GetParameterLayerIndex(Type, ParameterIndex);
+}
+
+
+FText UCustomizableObjectNodeCopyMaterial::GetParameterLayerName(EMaterialParameterType Type, int32 ParameterIndex) const
+{
+	UCustomizableObjectNodeMaterial* NodeMaterial = GetMaterialNode();
+	check(NodeMaterial);
+	
+	return NodeMaterial->GetParameterLayerName(Type, ParameterIndex);
+}
+
+
+bool UCustomizableObjectNodeCopyMaterial::HasParameter(const FNodeMaterialParameterId& ParameterId) const
+{
+	if (UCustomizableObjectNodeMaterial* NodeMaterial = GetMaterialNode())
+	{
+		return NodeMaterial->HasParameter(ParameterId);
+	}
+	else
+	{
+		return false;	
+	}
+}
+
+
+const UEdGraphPin* UCustomizableObjectNodeCopyMaterial::GetParameterPin(EMaterialParameterType Type, int32 ParameterIndex) const
+{
+	if (UCustomizableObjectNodeMaterial* NodeMaterial = GetMaterialNode())
+	{
+		return NodeMaterial->GetParameterPin(Type, ParameterIndex);
+	}
+	else
+	{
+		return nullptr;	
+	}
+}
+
+
+UEdGraphPin* UCustomizableObjectNodeCopyMaterial::GetParameterPin(const FNodeMaterialParameterId& ParameterId) const
+{
+	if (UCustomizableObjectNodeMaterial* NodeMaterial = GetMaterialNode())
+	{
+		return NodeMaterial->GetParameterPin(ParameterId);
+	}
+	else
+	{
+		return nullptr;	
+	}
+}
+
+
+bool UCustomizableObjectNodeCopyMaterial::IsImageMutableMode(int32 ImageIndex) const
+{
+	UCustomizableObjectNodeMaterial* NodeMaterial = GetMaterialNode();
+	check(NodeMaterial);
+	
+	return NodeMaterial->IsImageMutableMode(ImageIndex);
+}
+
+
+bool UCustomizableObjectNodeCopyMaterial::IsImageMutableMode(const UEdGraphPin& Pin) const
+{
+	UCustomizableObjectNodeMaterial* NodeMaterial = GetMaterialNode();
+	check(NodeMaterial);
+
+	return NodeMaterial->IsImageMutableMode(Pin);
+}
+
+
+UTexture2D* UCustomizableObjectNodeCopyMaterial::GetImageReferenceTexture(int32 ImageIndex) const
+{
+	if (UCustomizableObjectNodeMaterial* NodeMaterial = GetMaterialNode())
+	{
+		return NodeMaterial->GetImageReferenceTexture(ImageIndex);
+	}
+	else
+	{
+		return nullptr;
+	}
+}
+
+
+UTexture2D* UCustomizableObjectNodeCopyMaterial::GetImageValue(int32 ImageIndex) const
+{
+	if (UCustomizableObjectNodeMaterial* NodeMaterial = GetMaterialNode())
+	{
+		return NodeMaterial->GetImageValue(ImageIndex);
+	}
+	else
+	{
+		return nullptr;
+	}
+}
+
+
+int32 UCustomizableObjectNodeCopyMaterial::GetImageUVLayout(int32 ImageIndex) const
+{
+	if (UCustomizableObjectNodeMaterial* NodeMaterial = GetMaterialNode())
+	{
+		return NodeMaterial->GetImageUVLayout(ImageIndex);
+	}
+	else
+	{
+		return UCustomizableObjectNodeMaterialPinDataImage::UV_LAYOUT_IGNORE;
+	}
+}
+
+
+FText UCustomizableObjectNodeCopyMaterial::GetTooltipText() const
+{
+	return LOCTEXT("CopyMaterial_Tooltip", "Copies a Customizable Object material.\nDuplicates all Material node input pins and properties except for the Mesh input pin.");
+}
+
+
+#undef LOCTEXT_NAMESPACE
diff --git a/Engine/Plugins/Mutable/Source/CustomizableObjectEditor/Private/MuCOE/Nodes/CustomizableObjectNodeEnumParameter.cpp b/Engine/Plugins/Mutable/Source/CustomizableObjectEditor/Private/MuCOE/Nodes/CustomizableObjectNodeEnumParameter.cpp
new file mode 100644
--- /dev/null
+++ b/Engine/Plugins/Mutable/Source/CustomizableObjectEditor/Private/MuCOE/Nodes/CustomizableObjectNodeEnumParameter.cpp
@@ -0,0 +1,74 @@
+// Copyright Epic Games, Inc. All Rights Reserved.
+
+#include "MuCOE/Nodes/CustomizableObjectNodeEnumParameter.h"
+
+#include "MuCOE/EdGraphSchema_CustomizableObject.h"
+#include "MuCO/CustomizableObjectCustomVersion.h"
+
+class UCustomizableObjectNodeRemapPins;
+
+#define LOCTEXT_NAMESPACE "CustomizableObjectEditor"
+
+
+void UCustomizableObjectNodeEnumParameter::PostEditChangeProperty(FPropertyChangedEvent& PropertyChangedEvent)
+{
+	FProperty* PropertyThatChanged = PropertyChangedEvent.Property;
+	if ( PropertyThatChanged && PropertyThatChanged->GetName() == TEXT("Values") )
+	{
+		ReconstructNode();
+	}
+
+	Super::PostEditChangeProperty(PropertyChangedEvent);
+}
+
+
+void UCustomizableObjectNodeEnumParameter::AllocateDefaultPins(UCustomizableObjectNodeRemapPins* RemapPins)
+{
+	const UEdGraphSchema_CustomizableObject* Schema = GetDefault<UEdGraphSchema_CustomizableObject>();
+
+	UEdGraphPin* ValuePin = CustomCreatePin(EGPD_Output, Schema->PC_Enum, FName("Value"));
+	ValuePin->bDefaultValueIsIgnored = true;
+}
+
+
+FText UCustomizableObjectNodeEnumParameter::GetNodeTitle(ENodeTitleType::Type TitleType) const
+{
+	if (TitleType == ENodeTitleType::ListView || ParameterName.IsEmpty())
+	{
+		return LOCTEXT("Enum_Parameter", "Enum Parameter");
+	}
+	else if (TitleType == ENodeTitleType::EditableTitle)
+	{
+		return FText::Format(LOCTEXT("Enum_Parameter_EditableTitle", "{0}"), FText::FromString(ParameterName));
+	}
+	else
+	{
+		return FText::Format(LOCTEXT("Enum_Parameter_Title", "{0}\nEnum Parameter"), FText::FromString(ParameterName));
+	}
+}
+
+
+FLinearColor UCustomizableObjectNodeEnumParameter::GetNodeTitleColor() const
+{
+	const UEdGraphSchema_CustomizableObject* Schema = GetDefault<UEdGraphSchema_CustomizableObject>();
+	return Schema->GetPinTypeColor(Schema->PC_Enum);
+}
+
+
+FText UCustomizableObjectNodeEnumParameter::GetTooltipText() const
+{
+	return LOCTEXT("Enum_Parameter_Tooltip",
+		"Exposes and defines a parameter offering multiple choices to modify the Customizable Object.\nAlso defines a default one among them. \nIt's abstract, does not define what type those options refer to.");
+}
+
+
+void UCustomizableObjectNodeEnumParameter::OnRenameNode(const FString& NewName)
+{
+	if (!NewName.IsEmpty())
+	{
+		ParameterName = NewName;
+	}
+}
+
+#undef LOCTEXT_NAMESPACE
+
diff --git a/Engine/Plugins/Mutable/Source/CustomizableObjectEditor/Private/MuCOE/Nodes/CustomizableObjectNodeFloatParameter.cpp b/Engine/Plugins/Mutable/Source/CustomizableObjectEditor/Private/MuCOE/Nodes/CustomizableObjectNodeFloatParameter.cpp
new file mode 100644
--- /dev/null
+++ b/Engine/Plugins/Mutable/Source/CustomizableObjectEditor/Private/MuCOE/Nodes/CustomizableObjectNodeFloatParameter.cpp
@@ -0,0 +1,91 @@
+// Copyright Epic Games, Inc. All Rights Reserved.
+
+#include "MuCOE/Nodes/CustomizableObjectNodeFloatParameter.h"
+
+#include "MuCO/CustomizableObjectCustomVersion.h"
+#include "MuCOE/EdGraphSchema_CustomizableObject.h"
+#include "MuCOE/UnrealEditorPortabilityHelpers.h"
+
+class UCustomizableObjectNodeRemapPins;
+
+#define LOCTEXT_NAMESPACE "CustomizableObjectEditor"
+
+
+void UCustomizableObjectNodeFloatParameter::PostEditChangeProperty(FPropertyChangedEvent& PropertyChangedEvent)
+{
+	FProperty* PropertyThatChanged = PropertyChangedEvent.Property;
+	if ( PropertyThatChanged && (PropertyThatChanged->GetName() == TEXT("DescriptionImage") || PropertyThatChanged->GetName() == TEXT("Name")) )
+	{
+		ReconstructNode();
+	}
+
+	Super::PostEditChangeProperty(PropertyChangedEvent);
+}
+
+
+void UCustomizableObjectNodeFloatParameter::AllocateDefaultPins(UCustomizableObjectNodeRemapPins* RemapPins)
+{
+	const UEdGraphSchema_CustomizableObject* Schema = GetDefault<UEdGraphSchema_CustomizableObject>();
+
+	UEdGraphPin* ValuePin = CustomCreatePin(EGPD_Output, Schema->PC_Float, FName("Value"));
+	ValuePin->bDefaultValueIsIgnored = true;
+}
+
+
+bool UCustomizableObjectNodeFloatParameter::IsAffectedByLOD() const
+{
+	return false;
+}
+
+
+FText UCustomizableObjectNodeFloatParameter::GetNodeTitle(ENodeTitleType::Type TitleType) const
+{
+	if (TitleType == ENodeTitleType::ListView || ParameterName.IsEmpty())
+	{
+		return LOCTEXT("Float_Parameter", "Float Parameter");
+	}
+	else if (TitleType == ENodeTitleType::EditableTitle)
+	{
+		return FText::Format(LOCTEXT("Float_Parameter_EditableTitle", "{0}"), FText::FromString(ParameterName));
+	}
+	else
+	{
+		return FText::Format(LOCTEXT("Float_Parameter_Title", "{0}\nFloat Parameter"), FText::FromString(ParameterName));
+	}
+}
+
+
+FLinearColor UCustomizableObjectNodeFloatParameter::GetNodeTitleColor() const
+{
+	const UEdGraphSchema_CustomizableObject* Schema = GetDefault<UEdGraphSchema_CustomizableObject>();
+	return Schema->GetPinTypeColor(Schema->PC_Float);
+}
+
+
+FText UCustomizableObjectNodeFloatParameter::GetTooltipText() const
+{
+	return LOCTEXT("Float_Parameter_Tooltip", "Expose a numeric parameter from the Customizable Object that can be modified at runtime.");
+}
+
+
+void UCustomizableObjectNodeFloatParameter::OnRenameNode(const FString& NewName)
+{
+	if (!NewName.IsEmpty())
+	{
+		ParameterName = NewName;
+	}
+}
+
+
+void UCustomizableObjectNodeFloatParameter::BackwardsCompatibleFixup(int32 CustomizableObjectCustomVersion)
+{
+	Super::BackwardsCompatibleFixup(CustomizableObjectCustomVersion);
+
+	if (CustomizableObjectCustomVersion == FCustomizableObjectCustomVersion::RemovedParameterDecorations)
+	{
+		ReconstructNode();
+	}
+}
+
+
+#undef LOCTEXT_NAMESPACE
diff --git a/Engine/Plugins/Mutable/Source/CustomizableObjectEditor/Private/MuCOE/Nodes/CustomizableObjectNodeFloatSwitch.cpp b/Engine/Plugins/Mutable/Source/CustomizableObjectEditor/Private/MuCOE/Nodes/CustomizableObjectNodeFloatSwitch.cpp
new file mode 100644
--- /dev/null
+++ b/Engine/Plugins/Mutable/Source/CustomizableObjectEditor/Private/MuCOE/Nodes/CustomizableObjectNodeFloatSwitch.cpp
@@ -0,0 +1,29 @@
+// Copyright Epic Games, Inc. All Rights Reserved.
+
+#include "MuCOE/Nodes/CustomizableObjectNodeFloatSwitch.h"
+
+#include "MuCO/CustomizableObjectCustomVersion.h"
+#include "MuCOE/EdGraphSchema_CustomizableObject.h"
+
+#define LOCTEXT_NAMESPACE "CustomizableObjectEditor"
+
+
+void UCustomizableObjectNodeFloatSwitch::BackwardsCompatibleFixup(int32 CustomizableObjectCustomVersion)
+{
+	Super::BackwardsCompatibleFixup(CustomizableObjectCustomVersion);
+	
+	if (CustomizableObjectCustomVersion == FCustomizableObjectCustomVersion::BugPinsSwitch)
+	{
+		OutputPinReference = FindPin(TEXT("Float"));	
+	}
+}
+
+
+FName UCustomizableObjectNodeFloatSwitch::GetCategory() const
+{
+	return UEdGraphSchema_CustomizableObject::PC_Float;
+}
+
+
+#undef LOCTEXT_NAMESPACE
+
diff --git a/Engine/Plugins/Mutable/Source/CustomizableObjectEditor/Private/MuCOE/Nodes/CustomizableObjectNodeGroupProjectorParameter.cpp b/Engine/Plugins/Mutable/Source/CustomizableObjectEditor/Private/MuCOE/Nodes/CustomizableObjectNodeGroupProjectorParameter.cpp
new file mode 100644
--- /dev/null
+++ b/Engine/Plugins/Mutable/Source/CustomizableObjectEditor/Private/MuCOE/Nodes/CustomizableObjectNodeGroupProjectorParameter.cpp
@@ -0,0 +1,144 @@
+// Copyright Epic Games, Inc. All Rights Reserved.
+
+#include "MuCOE/Nodes/CustomizableObjectNodeGroupProjectorParameter.h"
+
+#include "MuCO/CustomizableObjectCustomVersion.h"
+#include "MuCOE/EdGraphSchema_CustomizableObject.h"
+#include "MuCOE/UnrealEditorPortabilityHelpers.h"
+
+class UCustomizableObjectNodeRemapPins;
+
+#define LOCTEXT_NAMESPACE "CustomizableObjectEditor"
+
+
+TArray<FGroupProjectorParameterImage> UCustomizableObjectNodeGroupProjectorParameter::GetOptionTexturesFromTable() const
+{
+	TArray<FGroupProjectorParameterImage> ArrayResult;
+
+	if (OptionTexturesDataTable == nullptr)
+	{
+		return ArrayResult;
+	}
+
+	TArray<FName> ArrayRowName = OptionTexturesDataTable->GetRowNames();
+
+	FProperty* PropertyTexturePath = OptionTexturesDataTable->FindTableProperty(DataTableTextureColumnName);
+
+	if (PropertyTexturePath == nullptr)
+	{
+		UE_LOG(LogMutable, Warning, TEXT("WARNING: No column found with texture path information to load projection textures"));
+		return ArrayResult;
+	}
+
+	int32 NameIndex = 0;
+
+	for (TMap<FName, uint8*>::TConstIterator RowIt = OptionTexturesDataTable->GetRowMap().CreateConstIterator(); RowIt; ++RowIt)
+	{
+		uint8* RowData = RowIt.Value();
+		FString PropertyValue(TEXT(""));
+		PropertyTexturePath->ExportText_InContainer(0, PropertyValue, RowData, RowData, nullptr, PPF_None);
+		UTexture2D* Texture = LoadObject<UTexture2D>(nullptr, *PropertyValue, nullptr);
+
+		if (Texture == nullptr)
+		{
+			UE_LOG(LogMutable, Warning, TEXT("WARNING: Unable to load texture %s"), *PropertyValue);
+		}
+		else
+		{
+			FGroupProjectorParameterImage GroupProjectorParameterImage;
+			GroupProjectorParameterImage.OptionName = ArrayRowName[NameIndex].ToString();
+			GroupProjectorParameterImage.OptionTexture = Texture;
+			ArrayResult.Add(GroupProjectorParameterImage);
+		}
+
+		NameIndex++;
+	}
+
+	return ArrayResult;
+}
+
+
+TArray<FGroupProjectorParameterImage> UCustomizableObjectNodeGroupProjectorParameter::GetFinalOptionTexturesNoRepeat() const
+{
+	TArray<FGroupProjectorParameterImage> ArrayDataTable = GetOptionTexturesFromTable();
+
+	for (int32 i = 0; i < OptionTextures.Num(); ++i)
+	{
+		bool AlreadyAdded = false;
+		for (int32 j = 0; j < ArrayDataTable.Num(); ++j)
+		{
+			if (OptionTextures[i].OptionName == ArrayDataTable[j].OptionName)
+			{
+				AlreadyAdded = true;
+				break;
+			}
+		}
+
+		if (!AlreadyAdded)
+		{
+			ArrayDataTable.Add(OptionTextures[i]);
+		}
+	}
+
+	return ArrayDataTable;
+}
+
+
+FText UCustomizableObjectNodeGroupProjectorParameter::GetNodeTitle(ENodeTitleType::Type TitleType) const
+{
+	if (TitleType == ENodeTitleType::ListView || ParameterName.IsEmpty())
+	{
+		return LOCTEXT("Group_Projector_Parameter", "Group Projector Parameter");
+	}
+	else if (TitleType == ENodeTitleType::EditableTitle)
+	{
+		return FText::Format(LOCTEXT("Group_Projector_Parameter_EditableTitle", "{0}"), FText::FromString(ParameterName));
+	}
+	else
+	{
+		return FText::Format(LOCTEXT("Group_Projector_Parameter_Title", "{0}\nGroup Projector Parameter"), FText::FromString(ParameterName));
+	}
+}
+
+
+FLinearColor UCustomizableObjectNodeGroupProjectorParameter::GetNodeTitleColor() const
+{
+	const UEdGraphSchema_CustomizableObject* Schema = GetDefault<UEdGraphSchema_CustomizableObject>();
+	return Schema->GetPinTypeColor(Schema->PC_GroupProjector);
+}
+
+
+FText UCustomizableObjectNodeGroupProjectorParameter::GetTooltipText() const
+{
+	return LOCTEXT("Group_Projector_Parameter_Tooltip", "Projects one or many textures to all children in the group it's connected to. It modifies only the materials that define a specific material asset texture parameter.");
+}
+
+
+void UCustomizableObjectNodeGroupProjectorParameter::BackwardsCompatibleFixup(int32 CustomizableObjectCustomVersion)
+{
+	Super::BackwardsCompatibleFixup(CustomizableObjectCustomVersion);
+	
+	if (CustomizableObjectCustomVersion == FCustomizableObjectCustomVersion::GroupProjectorPinTypeAdded)
+	{
+		if (UEdGraphPin* Pin = ProjectorPin())
+		{
+			Pin->PinType.PinCategory = UEdGraphSchema_CustomizableObject::PC_GroupProjector;
+		}
+	}
+
+	if (CustomizableObjectCustomVersion == FCustomizableObjectCustomVersion::GroupProjectorImagePinRemoved)
+	{
+		ReconstructNode();
+	}
+}
+
+
+void UCustomizableObjectNodeGroupProjectorParameter::AllocateDefaultPins(UCustomizableObjectNodeRemapPins* RemapPins)
+{
+	const UEdGraphSchema_CustomizableObject* Schema = GetDefault<UEdGraphSchema_CustomizableObject>();
+
+	CustomCreatePin(EGPD_Output, Schema->PC_GroupProjector, TEXT("Value"));	
+}
+
+
+#undef LOCTEXT_NAMESPACE
diff --git a/Engine/Plugins/Mutable/Source/CustomizableObjectEditor/Private/MuCOE/Nodes/CustomizableObjectNodeMaterial.cpp b/Engine/Plugins/Mutable/Source/CustomizableObjectEditor/Private/MuCOE/Nodes/CustomizableObjectNodeMaterial.cpp
new file mode 100644
--- /dev/null
+++ b/Engine/Plugins/Mutable/Source/CustomizableObjectEditor/Private/MuCOE/Nodes/CustomizableObjectNodeMaterial.cpp
@@ -0,0 +1,1506 @@
+// Copyright Epic Games, Inc. All Rights Reserved.
+
+#include "MuCOE/Nodes/CustomizableObjectNodeMaterial.h"
+
+#include "MaterialCachedData.h"
+#include "Materials/Material.h"
+#include "Materials/MaterialExpressionTextureCoordinate.h"
+#include "Materials/MaterialExpressionTextureSample.h"
+#include "Materials/MaterialInstance.h"
+#include "MuCOE/CustomizableObjectEditor_Deprecated.h"
+#include "MuCOE/EdGraphSchema_CustomizableObject.h"
+#include "MuCOE/GraphTraversal.h"
+#include "MuCOE/UnrealEditorPortabilityHelpers.h"
+#include "MuCOE/Nodes/CustomizableObjectNodeCopyMaterial.h"
+#include "MuCOE/Nodes/CustomizableObjectNodeSkeletalMesh.h"
+#include "MuCOE/Nodes/CustomizableObjectNodeStaticMesh.h"
+#include "MuCOE/Nodes/CustomizableObjectNodeTable.h"
+#include "MuCOE/Nodes/SCustomizableObjectNodeMaterial.h"
+#include "ObjectEditorUtils.h"
+#include "PropertyCustomizationHelpers.h"
+#include "Modules/ModuleManager.h"
+#include "MuCO/CustomizableObjectCustomVersion.h"
+#include "MuCOE/CustomizableObjectEditorLogger.h"
+#include "MuCOE/CustomizableObjectEditorUtilities.h"
+#include "MuCOE/Nodes/CustomizableObjectNodeExternalPin.h"
+
+class SGraphNode;
+class SWidget;
+class UCustomizableObjectLayout;
+class UObject;
+
+
+#define LOCTEXT_NAMESPACE "CustomizableObjectEditor"
+
+
+const TArray<EMaterialParameterType> UCustomizableObjectNodeMaterial::ParameterTypes = {
+	EMaterialParameterType::Texture,
+	EMaterialParameterType::Vector,
+	EMaterialParameterType::Scalar };
+
+
+bool UCustomizableObjectNodeMaterialRemapPinsByName::Equal(const UCustomizableObjectNode& Node, const UEdGraphPin& OldPin, const UEdGraphPin& NewPin) const
+{
+	const UCustomizableObjectNodeMaterialPinDataParameter* PinDataOldPin = Cast<UCustomizableObjectNodeMaterialPinDataParameter>(Node.GetPinData(OldPin));
+	const UCustomizableObjectNodeMaterialPinDataParameter* PinDataNewPin = Cast<UCustomizableObjectNodeMaterialPinDataParameter>(Node.GetPinData(NewPin));
+	if (PinDataOldPin && PinDataNewPin)
+	{
+		return PinDataOldPin->MaterialParameterId == PinDataNewPin->MaterialParameterId && (!OldPin.LinkedTo.Num() || OldPin.PinType == NewPin.PinType); // Pin type must match only if it was connected
+	}
+	else
+	{
+		return Super::Equal(Node, OldPin, NewPin);
+	}
+}
+
+
+void UCustomizableObjectNodeMaterialRemapPinsByName::RemapPins(const UCustomizableObjectNode& Node, const TArray<UEdGraphPin*>& OldPins, const TArray<UEdGraphPin*>& NewPins, TMap<UEdGraphPin*, UEdGraphPin*>& PinsToRemap, TArray<UEdGraphPin*>& PinsToOrphan)
+{
+	for (UEdGraphPin* OldPin : OldPins)
+	{
+		bool bFound = false;
+
+		for (UEdGraphPin* NewPin : NewPins)
+		{
+			if (Equal(Node, *OldPin, *NewPin))
+			{
+				bFound = true;
+
+				if (UEdGraphPin** Result = PinsToRemap.Find(OldPin))
+				{
+					if ((*Result)->bOrphanedPin) // The node can have a deprecated and non-deprecated pin that should remap to the same new pin. Prioritize the non-deprecated
+					{
+						*Result = NewPin;
+					}
+				}
+				else
+				{
+					PinsToRemap.Add(OldPin, NewPin);
+				}
+			}
+		}
+
+		if (!bFound && (OldPin->LinkedTo.Num() || HasSavedPinData(Node, *OldPin)))
+		{
+			PinsToOrphan.Add(OldPin);
+		}
+	}
+}
+
+
+bool UCustomizableObjectNodeMaterialRemapPinsByName::HasSavedPinData(const UCustomizableObjectNode& Node, const UEdGraphPin &Pin) const
+{
+	if (const UCustomizableObjectNodeMaterialPinDataParameter* PinData = Cast<UCustomizableObjectNodeMaterialPinDataParameter>(Node.GetPinData(Pin)))
+	{
+		return !PinData->IsDefault();
+	}
+	else
+	{
+		return false;
+	}
+}
+
+
+bool UCustomizableObjectNodeMaterialPinDataImage::IsDefault() const
+{
+	const UCustomizableObjectNodeMaterialPinDataImage* Default = Cast<UCustomizableObjectNodeMaterialPinDataImage>(GetClass()->ClassDefaultObject);
+
+	return PinMode == Default->PinMode &&
+		UVLayoutMode == Default->UVLayoutMode &&
+		ReferenceTexture == Default->ReferenceTexture;
+}
+
+
+EPinMode UCustomizableObjectNodeMaterialPinDataImage::GetPinMode() const
+{
+	return PinMode;
+}
+
+
+void UCustomizableObjectNodeMaterialPinDataImage::SetPinMode(const EPinMode InPinMode)
+{
+	FObjectEditorUtils::SetPropertyValue(this, GET_MEMBER_NAME_CHECKED(UCustomizableObjectNodeMaterialPinDataImage, PinMode), InPinMode);
+	NodeMaterial->GetPostImagePinModeChangedDelegate()->Broadcast();
+}
+
+
+void UCustomizableObjectNodeMaterialPinDataImage::Init(UCustomizableObjectNodeMaterial& InNodeMaterial)
+{
+	NodeMaterial = &InNodeMaterial;
+}
+
+
+void UCustomizableObjectNodeMaterialPinDataImage::Copy(const UCustomizableObjectNodePinData& Other)
+{
+	if (const UCustomizableObjectNodeMaterialPinDataImage* PinDataOldPin = Cast<UCustomizableObjectNodeMaterialPinDataImage>(&Other))
+	{
+		PinMode = PinDataOldPin->PinMode;
+		UVLayoutMode = PinDataOldPin->UVLayoutMode;
+		UVLayout = PinDataOldPin->UVLayout;
+		ReferenceTexture = PinDataOldPin->ReferenceTexture;
+	}
+}
+
+
+FName UCustomizableObjectNodeMaterial::GetPinName(const EMaterialParameterType Type, const int32 ParameterIndex) const
+{
+	const FString ParameterName = GetParameterName(Type, ParameterIndex).ToString();
+	return FName(GetParameterLayerIndex(Type, ParameterIndex) != INDEX_NONE ? GetParameterLayerName(Type, ParameterIndex).ToString() + " - " + ParameterName : ParameterName);
+}
+
+
+int32 UCustomizableObjectNodeMaterial::GetExpressionTextureCoordinate(UMaterial* Material, const FGuid& ImageId)
+{
+	if (const UMaterialExpressionTextureSample* TextureSample = Material->FindExpressionByGUID<UMaterialExpressionTextureSample>(ImageId))
+	{
+		if (!TextureSample->Coordinates.Expression)
+		{
+			return TextureSample->ConstCoordinate;
+		}
+		else if (const UMaterialExpressionTextureCoordinate* TextureCoords = Cast<UMaterialExpressionTextureCoordinate>(TextureSample->Coordinates.Expression))
+		{
+			return TextureCoords->CoordinateIndex;
+		}
+	}
+
+	return -1;
+}
+
+
+FName UCustomizableObjectNodeMaterial::NodePinModeToImagePinMode(const ENodePinMode NodePinMode)
+{
+	switch (NodePinMode)
+	{
+	case ENodePinMode::Mutable:
+		return UEdGraphSchema_CustomizableObject::PC_Image;
+	case ENodePinMode::Passthrough:
+		return UEdGraphSchema_CustomizableObject::PC_PassThroughImage;
+	default:
+		unimplemented();
+		return {};
+	}
+}
+
+
+FName UCustomizableObjectNodeMaterial::GetImagePinMode(const EPinMode PinMode) const
+{
+	switch (PinMode)
+	{
+	case EPinMode::Default:
+		return NodePinModeToImagePinMode(TextureParametersMode);
+	case EPinMode::Mutable:
+		return UEdGraphSchema_CustomizableObject::PC_Image;
+	case EPinMode::Passthrough:
+		return UEdGraphSchema_CustomizableObject::PC_PassThroughImage;
+	default:
+		unimplemented();
+		return {};
+	}
+}
+
+
+FName UCustomizableObjectNodeMaterial::GetImagePinMode(const UEdGraphPin& Pin) const
+{
+	return GetImagePinMode(GetPinData<UCustomizableObjectNodeMaterialPinDataImage>(Pin).GetPinMode());
+}
+
+
+int32 UCustomizableObjectNodeMaterial::GetImageUVLayoutFromMaterial(const int32 ImageIndex) const
+{
+	const FNodeMaterialParameterId ImageId = GetParameterId(EMaterialParameterType::Texture, ImageIndex);
+
+	if (const int32 TextureCoordinate = GetExpressionTextureCoordinate(Material->GetMaterial(), ImageId.ParameterId);
+		TextureCoordinate >= 0)
+	{
+		return TextureCoordinate;
+	}
+
+	FMaterialLayersFunctions Layers;
+	Material->GetMaterialLayers(Layers);
+
+	TArray<TArray<TObjectPtr<UMaterialFunctionInterface>>*> MaterialFunctionInterfaces;
+	MaterialFunctionInterfaces.SetNumUninitialized(2);
+	MaterialFunctionInterfaces[0] = &Layers.Layers;
+	MaterialFunctionInterfaces[1] = &Layers.Blends;
+
+	for (const TArray<TObjectPtr<UMaterialFunctionInterface>>* MaterialFunctionInterface : MaterialFunctionInterfaces)
+	{
+		for (const TObjectPtr<UMaterialFunctionInterface>& Layer : *MaterialFunctionInterface)
+		{
+			if (const int32 TextureCoordinate = GetExpressionTextureCoordinate(Layer->GetPreviewMaterial()->GetMaterial(), ImageId.ParameterId); TextureCoordinate >= 0)
+			{
+				return TextureCoordinate;
+			}	
+		}
+	}
+
+	return -1;
+}
+
+
+bool UCustomizableObjectNodeMaterialPinDataParameter::IsDefault() const
+{
+	return true;
+}
+
+
+/** Translates a given EPinMode to FText. */
+FText EPinModeToText(const EPinMode PinMode)
+{
+	switch (PinMode)
+	{
+	case EPinMode::Default:
+		return LOCTEXT("EPinModeDefault", "Node Defined");
+
+	case EPinMode::Mutable:
+		return LOCTEXT("EPinModeMutable", "Mutable");
+
+	case EPinMode::Passthrough:
+		return LOCTEXT("EPinModePassthrough", "Passthrough");
+
+	default:
+		check(false); // Missing case.
+		return FText();
+	}
+}
+
+
+void UCustomizableObjectNodeMaterialPinDataImage::PostEditChangeProperty(FPropertyChangedEvent& PropertyChangedEvent)
+{
+	Super::PostEditChangeProperty(PropertyChangedEvent);
+
+	if (const FProperty* PropertyThatChanged = PropertyChangedEvent.Property)
+	{
+		if (PropertyThatChanged->GetFName() == GET_MEMBER_NAME_CHECKED(UCustomizableObjectNodeMaterialPinDataImage, PinMode))
+		{
+			NodeMaterial->UCustomizableObjectNode::ReconstructNode();
+		}
+	}
+}
+
+
+bool UCustomizableObjectNodeMaterialPinDataImage::CanEditChange(const FProperty* InProperty) const
+{
+	if (InProperty->GetFName() == GET_MEMBER_NAME_CHECKED(UCustomizableObjectNodeMaterialPinDataImage, PinMode))
+	{
+		return !NodeMaterial->GetPin(*this)->LinkedTo.Num();
+	}
+	
+	return Super::CanEditChange(InProperty);
+}
+
+
+void UCustomizableObjectNodeMaterialPinDataImage::BackwardsCompatibleFixup(int32 CustomizableObjectCustomVersion)
+{
+	Super::PostLoad();
+	
+	if (CustomizableObjectCustomVersion == FCustomizableObjectCustomVersion::NodeMaterialPinDataImageDetails)
+	{
+		if (UVLayout == UV_LAYOUT_IGNORE)
+		{
+			UVLayoutMode = EUVLayoutMode::Ignore;
+			UVLayout = 0;
+		}
+		else if (UVLayout == UV_LAYOUT_DEFAULT)
+		{
+			UVLayoutMode = EUVLayoutMode::FromMaterial;
+			UVLayout = 0;
+		}
+		else
+		{
+			UVLayoutMode = EUVLayoutMode::Index;
+		}
+	}
+}
+
+
+void UCustomizableObjectNodeMaterial::AllocateDefaultPins(UCustomizableObjectNodeRemapPins* RemapPins)
+{
+	const UEdGraphSchema_CustomizableObject* Schema = GetDefault<UEdGraphSchema_CustomizableObject>();
+
+	{
+		const FString PinFriendlyName = TEXT("Mesh");
+		const FString PinName = PinFriendlyName + FString(TEXT("_Input_Pin"));
+		UEdGraphPin* MeshPin = CustomCreatePin(EGPD_Input, Schema->PC_Mesh, FName(*PinName));
+		MeshPin->PinFriendlyName = FText::FromString(PinFriendlyName);
+		MeshPin->bDefaultValueIsIgnored = true;
+	}
+	
+	{
+		const FString PinFriendlyName = TEXT("Table Material");
+		const FString PinName = PinFriendlyName + FString(TEXT("_Input_Pin"));
+		UEdGraphPin* TableMaterialPin = CustomCreatePin(EGPD_Input, Schema->PC_MaterialAsset, FName(*PinName));
+		TableMaterialPin->PinFriendlyName = FText::FromString(PinFriendlyName);
+		TableMaterialPin->bDefaultValueIsIgnored = true;
+		TableMaterialPin->PinToolTip = "Pin for a Material from a Table Node";
+	}
+
+	for (const EMaterialParameterType Type : ParameterTypes)
+	{
+		AllocateDefaultParameterPins(Type);
+	}
+
+	{
+		const FString PinFriendlyName = TEXT("Mesh Section");
+		const FString PinName = PinFriendlyName + FString(TEXT("_Output_Pin"));
+		UEdGraphPin* OutputPin = CustomCreatePin(EGPD_Output, Schema->PC_Material, FName(*PinName));
+		OutputPin->PinFriendlyName = FText::FromString(PinFriendlyName);
+	}
+}
+
+
+bool UCustomizableObjectNodeMaterial::CanPinBeHidden(const UEdGraphPin& Pin) const
+{
+	return Super::CanPinBeHidden(Pin) &&
+		Pin.Direction == EGPD_Input && 
+		Pin.PinType.PinCategory != UEdGraphSchema_CustomizableObject::PC_Mesh;
+}
+
+
+bool UCustomizableObjectNodeMaterial::HasPinViewer() const
+{
+	return true;
+}
+
+
+UCustomizableObjectNodeRemapPinsByName* UCustomizableObjectNodeMaterial::CreateRemapPinsDefault() const
+{
+	return NewObject<UCustomizableObjectNodeMaterialRemapPinsByName>();
+}
+
+
+void UCustomizableObjectNodeMaterial::BackwardsCompatibleFixup(int32 CustomizableObjectCustomVersion)
+{
+	Super::BackwardsCompatibleFixup(CustomizableObjectCustomVersion);
+	
+	if (CustomizableObjectCustomVersion == FCustomizableObjectCustomVersion::AutomaticNodeMaterial)
+	{
+		for (const FCustomizableObjectNodeMaterialImage& Image : Images_DEPRECATED)
+		{
+			const FString OldPinName = (Image.LayerIndex == -1 ? Image.Name : Image.PinName) + FString(TEXT("_Input_Image"));
+			UEdGraphPin* OldPin = FindPin(OldPinName);
+			if (!OldPin) // If we can not find a pin it means that the data was corrupted (old Image array and pins where not synchronized).
+			{
+				continue;
+			}
+
+			UCustomizableObjectNodeMaterialPinDataImage* PinData = NewObject<UCustomizableObjectNodeMaterialPinDataImage>(this);
+			PinData->ParameterId_DEPRECATED = FGuid::NewGuid();
+			PinData->ReferenceTexture = Image.ReferenceTexture;
+
+			// Find referenced Material Parameter
+			const int32 NumParameters = GetNumParameters(EMaterialParameterType::Texture);
+			for (int32 ParameterIndex = 0; ParameterIndex < NumParameters; ++ParameterIndex)
+			{
+				if (GetParameterName(EMaterialParameterType::Texture, ParameterIndex).ToString() == Image.Name)
+				{
+					PinData->ParameterId_DEPRECATED = GetParameterId(EMaterialParameterType::Texture, ParameterIndex).ParameterId;
+
+					if (Image.UVLayout == -1)
+					{
+						PinData->UVLayout = Image.UVLayout;
+					}
+					else
+					{
+						const int32 UVLayout = GetImageUVLayoutFromMaterial(ParameterIndex);
+						if (UVLayout < 0) // Could not be deduced from the Material
+						{
+							PinData->UVLayout = Image.UVLayout;						
+						}
+						else if (UVLayout == Image.UVLayout)
+						{
+							PinData->UVLayout = UV_LAYOUT_DEFAULT;							
+						}
+						else
+						{
+							PinData->UVLayout = UVLayout;
+						}					
+					}
+					
+					break;
+				}
+			}
+			
+			AddPinData(*OldPin, *PinData);
+		}
+
+		for (const FCustomizableObjectNodeMaterialVector& Vector : VectorParams_DEPRECATED)
+		{
+			const FString OldPinName = (Vector.LayerIndex == -1 ? Vector.Name : Vector.PinName) + FString(TEXT("_Input_Vector"));
+			const UEdGraphPin* OldPin = FindPin(OldPinName);
+			if (!OldPin) // If we can not find a pin it means that the data was corrupted (old ScalarParams array and pins where not synchronized).
+			{
+				continue;
+			}
+			
+			UCustomizableObjectNodeMaterialPinDataVector* PinData = NewObject<UCustomizableObjectNodeMaterialPinDataVector>(this);
+			PinData->ParameterId_DEPRECATED = FGuid::NewGuid();
+			
+			// Find referenced Material Parameter
+			const int32 NumParameters = GetNumParameters(EMaterialParameterType::Vector);
+			for (int32 ParameterIndex = 0; ParameterIndex < NumParameters; ++ParameterIndex)
+			{
+				if (GetParameterName(EMaterialParameterType::Vector, ParameterIndex).ToString() == Vector.Name)
+				{
+					PinData->ParameterId_DEPRECATED = GetParameterId(EMaterialParameterType::Vector, ParameterIndex).ParameterId;
+					break;
+				}
+			}
+			
+			AddPinData(*OldPin, *PinData);
+		}
+
+		for (const FCustomizableObjectNodeMaterialScalar& Scalar : ScalarParams_DEPRECATED)
+		{
+			const FString OldPinName = (Scalar.LayerIndex == -1 ? Scalar.Name : Scalar.PinName) + FString(TEXT("_Input_Scalar"));
+			const UEdGraphPin* OldPin = FindPin(OldPinName);
+			if (!OldPin) // If we can not find a pin it means that the data was corrupted (old ScalarParams array and pins where not synchronized).
+			{
+				continue;
+			}
+			
+			UCustomizableObjectNodeMaterialPinDataScalar* PinData = NewObject<UCustomizableObjectNodeMaterialPinDataScalar>(this);
+			PinData->ParameterId_DEPRECATED = FGuid::NewGuid();
+
+			// Find referenced Material Parameter
+			const int32 NumParameters = GetNumParameters(EMaterialParameterType::Scalar);
+			for (int32 ParameterIndex = 0; ParameterIndex < NumParameters; ++ParameterIndex)
+			{
+				if (GetParameterName(EMaterialParameterType::Scalar, ParameterIndex).ToString() == Scalar.Name)
+				{
+					PinData->ParameterId_DEPRECATED = GetParameterId(EMaterialParameterType::Scalar, ParameterIndex).ParameterId;
+					break;
+				}
+			}
+			
+			AddPinData(*OldPin, *PinData);
+		}
+
+		// Check if there are still pins which where not present in the Images, ScalarParams and ScalarParams arrays.
+		for (const UEdGraphPin* Pin : GetAllNonOrphanPins())
+		{
+			if ((Pin->PinType.PinCategory == UEdGraphSchema_CustomizableObject::PC_Image ||
+ 				 Pin->PinType.PinCategory == UEdGraphSchema_CustomizableObject::PC_Color ||
+				 Pin->PinType.PinCategory == UEdGraphSchema_CustomizableObject::PC_Float) &&
+				!GetPinData(*Pin))
+			{
+				UCustomizableObjectNodeMaterialPinDataParameter* PinData = [&](UObject* Outer) -> UCustomizableObjectNodeMaterialPinDataParameter*
+				{
+					if (Pin->PinType.PinCategory == UEdGraphSchema_CustomizableObject::PC_Image)
+					{
+						return NewObject<UCustomizableObjectNodeMaterialPinDataImage>(Outer);	
+					}
+					else if (Pin->PinType.PinCategory == UEdGraphSchema_CustomizableObject::PC_Color)
+					{
+						return NewObject<UCustomizableObjectNodeMaterialPinDataVector>(Outer);
+					}
+					else if (Pin->PinType.PinCategory == UEdGraphSchema_CustomizableObject::PC_Float)
+					{
+						return NewObject<UCustomizableObjectNodeMaterialPinDataScalar>(Outer);	
+					}
+					else
+					{
+						check(false); // Parameter type not contemplated.
+						return nullptr;
+					}
+				}(this);
+				
+				PinData->ParameterId_DEPRECATED = FGuid::NewGuid();
+				
+				AddPinData(*Pin, *PinData);
+			}
+		}
+		
+		Images_DEPRECATED.Empty();
+		VectorParams_DEPRECATED.Empty();
+		ScalarParams_DEPRECATED.Empty();
+		
+		Super::ReconstructNode(); // Super required to avoid ambiguous call compilation error.
+	}
+
+	// Fill PinsParameter.
+	if (CustomizableObjectCustomVersion == FCustomizableObjectCustomVersion::AutomaticNodeMaterialPerformanceBug) // || CustomizableObjectCustomVersion < FCustomizableObjectCustomVersion::AutomaticNodeMaterialPerformance
+	{
+		for (const UEdGraphPin* Pin : GetAllNonOrphanPins())
+		{
+			if (const UCustomizableObjectNodeMaterialPinDataParameter* PinData = Cast<UCustomizableObjectNodeMaterialPinDataParameter>(GetPinData(*Pin)))
+			{
+				PinsParameter_DEPRECATED.Add(PinData->ParameterId_DEPRECATED, FEdGraphPinReference(Pin));
+			}
+		}
+	}
+
+	if (CustomizableObjectCustomVersion == FCustomizableObjectCustomVersion::AutomaticNodeMaterialUXImprovements)
+	{
+		TextureParametersMode = bDefaultPinModeMutable_DEPRECATED ?
+			ENodePinMode::Mutable :
+			ENodePinMode::Passthrough;
+	}
+	
+	if (CustomizableObjectCustomVersion == FCustomizableObjectCustomVersion::ExtendMaterialOnlyMutableModeParameters)
+	{
+		const uint32 NumTextureParameters = GetNumParameters(EMaterialParameterType::Texture);
+		for (uint32 ImageIndex = 0; ImageIndex < NumTextureParameters; ++ImageIndex)
+		{
+			if (const UEdGraphPin* Pin = GetParameterPin(EMaterialParameterType::Texture, ImageIndex))
+			{
+				UCustomizableObjectNodeMaterialPinDataImage* PinDataImage = Cast<UCustomizableObjectNodeMaterialPinDataImage>(GetPinData(*Pin));
+				PinDataImage->NodeMaterial = this;
+			}
+		}
+	}
+	
+	if (CustomizableObjectCustomVersion == FCustomizableObjectCustomVersion::ExtendMaterialOnlyMutableModeParametersBug)
+	{
+		for (const UEdGraphPin* Pin : GetAllPins())
+		{
+			if (UCustomizableObjectNodeMaterialPinDataImage* const PinDataImage = Cast<UCustomizableObjectNodeMaterialPinDataImage>(GetPinData(*Pin)))
+			{
+				PinDataImage->NodeMaterial = this;
+			}
+		}
+	}
+	
+	if (CustomizableObjectCustomVersion == FCustomizableObjectCustomVersion::NodeMaterialAddTablePin)
+	{
+		Super::ReconstructNode();
+	}
+
+	if (CustomizableObjectCustomVersion == FCustomizableObjectCustomVersion::AddedTableMaterialSwitch)
+	{
+		UMaterialInstance* DefaultPinValue = nullptr;
+
+		if (const UEdGraphPin* MaterialAssetPin = GetMaterialAssetPin())
+		{
+			if (const UEdGraphPin* ConnectedPin = FollowInputPin(*MaterialAssetPin))
+			{
+				if (const UCustomizableObjectNodeTable* TableNode = Cast<UCustomizableObjectNodeTable>(ConnectedPin->GetOwningNode()))
+				{
+					DefaultPinValue = TableNode->GetColumnDefaultAssetByType<UMaterialInstance>(ConnectedPin);
+				}
+			}
+		}
+
+		if (DefaultPinValue && DefaultPinValue->TextureParameterValues.Num())
+		{
+			const uint32 NumTextureParameters = GetNumParameters(EMaterialParameterType::Texture);
+			for (uint32 ImageIndex = 0; ImageIndex < NumTextureParameters; ++ImageIndex)
+			{
+				if (const UEdGraphPin* ImagePin = GetParameterPin(EMaterialParameterType::Texture, ImageIndex))
+				{
+					FGuid ParameterId = GetParameterId(EMaterialParameterType::Texture, ImageIndex).ParameterId;
+
+					TArray<FMaterialParameterInfo> TextureParameterInfo;
+					TArray<FGuid> TextureGuids;
+
+					// Getting parent's texture infos
+					DefaultPinValue->GetMaterial()->GetAllTextureParameterInfo(TextureParameterInfo, TextureGuids);
+
+					int32 TextureIndex = TextureGuids.Find(ParameterId);
+
+					if (TextureIndex == INDEX_NONE)
+					{
+						continue;
+					}
+
+					FName TextureName = TextureParameterInfo[TextureIndex].Name;
+
+					// Checking if the pin's texture has been modified in the material instance
+					for (const FTextureParameterValue& Texture : DefaultPinValue->TextureParameterValues)
+					{
+						if (TextureName == Texture.ParameterInfo.Name)
+						{
+							if (UCustomizableObjectNodeMaterialPinDataImage* PinDataImage = Cast<UCustomizableObjectNodeMaterialPinDataImage>(GetPinData(*ImagePin)))
+							{
+								PinDataImage->PinMode = EPinMode::Mutable;
+							}
+						}
+					}
+				}
+			}
+		}
+	}
+
+	if (CustomizableObjectCustomVersion == FCustomizableObjectCustomVersion::NewComponentOptions)
+	{
+		MeshComponentName_DEPRECATED = FName(FString::FromInt(MeshComponentIndex_DEPRECATED));
+	}
+
+	if (CustomizableObjectCustomVersion == FCustomizableObjectCustomVersion::NodeMaterialTypedImagePins)
+	{
+		const UEdGraphSchema_CustomizableObject* Schema = GetDefault<UEdGraphSchema_CustomizableObject>();
+		
+		for (UEdGraphPin* Pin : GetAllPins())
+		{
+			if (Pin->PinType.PinCategory != UEdGraphSchema_CustomizableObject::PC_Image)
+			{
+				continue;
+			}
+			
+			if (Pin->LinkedTo.Num())
+			{
+				const UEdGraphPin* LinkedPin = Pin->LinkedTo[0];
+
+				UCustomizableObjectNodeMaterialPinDataImage& PinData = GetPinData<UCustomizableObjectNodeMaterialPinDataImage>(*Pin);
+
+				if (LinkedPin->PinType.PinCategory == UEdGraphSchema_CustomizableObject::PC_Image)
+				{
+					Pin->PinType.PinCategory = UEdGraphSchema_CustomizableObject::PC_Image;
+
+					if (PinData.PinMode == EPinMode::Default &&
+						TextureParametersMode == ENodePinMode::Passthrough)
+					{
+						PinData.PinMode = EPinMode::Mutable;
+					}
+				}
+				else if (LinkedPin->PinType.PinCategory == UEdGraphSchema_CustomizableObject::PC_PassThroughImage)
+				{
+					Pin->PinType.PinCategory = UEdGraphSchema_CustomizableObject::PC_PassThroughImage;
+
+					if (PinData.PinMode == EPinMode::Default &&
+						TextureParametersMode == ENodePinMode::Mutable)
+					{
+						PinData.PinMode = EPinMode::Passthrough;
+					}
+				}
+			}
+			else
+			{
+				if (!IsImageMutableMode(*Pin))
+				{
+					Pin->PinType.PinCategory = UEdGraphSchema_CustomizableObject::PC_PassThroughImage;
+				}				
+			}
+		}
+	}
+
+	if (CustomizableObjectCustomVersion == FCustomizableObjectCustomVersion::FixedMultilayerMaterialIds)
+	{
+		if (Material && Material->GetCachedExpressionData().bHasMaterialLayers)
+		{
+			// Needed since we can not get the layer index of repeated parameters
+			Super::ReconstructNode();
+		}
+		else
+		{
+			for (TMap<FGuid, FEdGraphPinReference>::TIterator It = PinsParameter_DEPRECATED.CreateIterator(); It; ++It)
+			{
+				PinsParameterMap.Add({ It.Key(), -1 }, It.Value());
+
+				// Move pin data id info to the new struct
+				UEdGraphPin* GraphPin = It.Value().Get();
+				if (!GraphPin)
+				{
+					continue;
+				}
+
+				UCustomizableObjectNodePinData* GenericPinData = GetPinData(*GraphPin);
+				if (!GenericPinData)
+				{
+					continue;
+				}
+
+				if (UCustomizableObjectNodeMaterialPinDataParameter* PinData = Cast<UCustomizableObjectNodeMaterialPinDataParameter>(GenericPinData))
+				{
+					PinData->MaterialParameterId.LayerIndex = INDEX_NONE;
+					PinData->MaterialParameterId.ParameterId = PinData->ParameterId_DEPRECATED;
+				}
+			}
+		}
+	}
+
+	if (CustomizableObjectCustomVersion == FCustomizableObjectCustomVersion::MaterialPinsRename)
+	{
+		UEdGraphPin* MaterialPin = FindPin(TEXT("Material_Output_Pin"), EEdGraphPinDirection::EGPD_Output);
+		if (MaterialPin)
+		{
+			const FString PinFriendlyName = TEXT("Mesh Section");
+			const FString PinName = PinFriendlyName + FString(TEXT("_Output_Pin"));
+			MaterialPin->PinName = FName(PinName);
+			MaterialPin->PinFriendlyName = FText::FromString(PinFriendlyName);
+		}
+
+	}
+}
+
+
+void UCustomizableObjectNodeMaterial::Serialize(FArchive& Ar)
+{
+	Super::Serialize(Ar);
+
+	Ar.UsingCustomVersion(FCustomizableObjectCustomVersion::GUID);
+}
+
+
+void UCustomizableObjectNodeMaterial::PostLoad()
+{
+	Super::PostLoad();
+	
+	if (Material)
+	{
+		ConditionalPostLoadReference(*Material); // Make sure the Material has been fully loaded.
+	}
+}
+
+
+void UCustomizableObjectNodeMaterial::PostEditChangeProperty(FPropertyChangedEvent& PropertyChangedEvent)
+{
+	Super::PostEditChangeProperty(PropertyChangedEvent);
+
+	if (const FProperty* PropertyThatChanged = PropertyChangedEvent.Property)
+	{
+		if (PropertyThatChanged->GetFName() == GET_MEMBER_NAME_CHECKED(UCustomizableObjectNodeMaterial, Material))
+		{
+			Super::ReconstructNode(); // Super required to avoid ambiguous call compilation error.
+		}
+		else if (PropertyThatChanged->GetFName() == GET_MEMBER_NAME_CHECKED(UCustomizableObjectNodeMaterial, TextureParametersMode))
+		{
+			Super::ReconstructNode();
+			PostImagePinModeChangedDelegate.Broadcast();
+		}
+	}
+}
+
+
+FText UCustomizableObjectNodeMaterial::GetNodeTitle(ENodeTitleType::Type TitleType) const
+{
+	if (TitleType == ENodeTitleType::ListView || !Material)
+	{
+		return LOCTEXT("Mesh Section", "Mesh Section");
+	}
+	else
+	{
+		FFormatNamedArguments Args;
+		Args.Add(TEXT("MeshSectionName"), FText::FromString(Material->GetName()));
+
+		return FText::Format(LOCTEXT("MeshSection_Title", "{MeshSectionName}\nMesh Section"), Args);
+	}
+}
+
+
+FLinearColor UCustomizableObjectNodeMaterial::GetNodeTitleColor() const
+{
+	const UEdGraphSchema_CustomizableObject* Schema = GetDefault<UEdGraphSchema_CustomizableObject>();
+	return Schema->GetPinTypeColor(Schema->PC_Material);
+}
+
+
+UEdGraphPin* UCustomizableObjectNodeMaterial::OutputPin() const
+{
+	FString PinFriendlyName = TEXT("Mesh Section");
+	FString PinName = PinFriendlyName + FString(TEXT("_Output_Pin"));
+
+	UEdGraphPin* Pin = FindPin(PinName, EEdGraphPinDirection::EGPD_Output);
+	if (!Pin)
+	{
+		Pin = FindPin(TEXT("Mesh Section"), EEdGraphPinDirection::EGPD_Output);
+	}
+
+	// Legacy name
+	if (!Pin)
+	{
+		PinFriendlyName = TEXT("Material");
+		PinName = PinFriendlyName + FString(TEXT("_Output_Pin"));
+
+		Pin = FindPin(PinName, EEdGraphPinDirection::EGPD_Output);
+		if (!Pin)
+		{
+			Pin = FindPin(TEXT("Material"), EEdGraphPinDirection::EGPD_Output);
+		}		
+	}
+
+	return Pin;
+}
+
+
+UEdGraphPin* UCustomizableObjectNodeMaterial::GetMeshPin() const
+{
+	FString PinFriendlyName = TEXT("Mesh");
+	FString PinName = PinFriendlyName + FString(TEXT("_Input_Pin"));
+
+	UEdGraphPin* Pin = FindPin(PinName);
+
+	if(Pin)
+	{
+		return Pin;
+	}
+	else
+	{
+		return FindPin(TEXT("Mesh"));
+	}
+}
+
+
+UEdGraphPin* UCustomizableObjectNodeMaterial::GetMaterialAssetPin() const
+{
+	FString PinFriendlyName = TEXT("Table Material");
+	FString PinName = PinFriendlyName + FString(TEXT("_Input_Pin"));
+
+	UEdGraphPin* Pin = FindPin(PinName);
+
+	if (Pin)
+	{
+		return Pin;
+	}
+	else
+	{
+		return FindPin(TEXT("Table Material"));
+	}
+}
+
+
+const UCustomizableObjectNodeMaterial* UCustomizableObjectNodeMaterial::GetMaterialNode() const
+{
+	return this;
+}
+
+
+bool UCustomizableObjectNodeMaterial::IsImageMutableMode(const int32 ImageIndex) const
+{
+	if (const UEdGraphPin* Pin = GetParameterPin(EMaterialParameterType::Texture, ImageIndex))
+	{
+		return IsImageMutableMode(*Pin);
+	}
+	else
+	{
+		return NodePinModeToImagePinMode(TextureParametersMode) == UEdGraphSchema_CustomizableObject::PC_Image;
+	}
+}
+
+
+bool UCustomizableObjectNodeMaterial::IsImageMutableMode(const UEdGraphPin& Pin) const
+{
+	return GetImagePinMode(Pin) == UEdGraphSchema_CustomizableObject::PC_Image; // The ImageMutableMode is stored in the PinData, not in the PinCategory
+}
+
+
+void UCustomizableObjectNodeMaterial::SetImagePinMode(UEdGraphPin& Pin, const EPinMode PinMode) const
+{
+	UCustomizableObjectNodeMaterialPinDataImage& PinData = GetPinData<UCustomizableObjectNodeMaterialPinDataImage>(Pin);
+
+	Pin.PinType.PinCategory = GetImagePinMode(PinMode);	// Change the category so that it will remap correctly when reconstructing.
+	
+	PinData.SetPinMode(PinMode); // Will trigger a reconstruct.
+}
+
+
+UTexture2D* UCustomizableObjectNodeMaterial::GetImageReferenceTexture(const int32 ImageIndex) const
+{
+	if (const UEdGraphPin* Pin = GetParameterPin(EMaterialParameterType::Texture, ImageIndex))
+	{
+		return GetPinData<UCustomizableObjectNodeMaterialPinDataImage>(*Pin).ReferenceTexture;
+	}
+	else
+	{
+		return nullptr;
+	}
+}
+
+
+UTexture2D* UCustomizableObjectNodeMaterial::GetImageValue(const int32 ImageIndex) const
+{
+	FName TextureName = GetParameterName(EMaterialParameterType::Texture, ImageIndex);
+	
+	UTexture* Texture;
+	Material->GetTextureParameterValue(TextureName, Texture);
+
+	return Cast<UTexture2D>(Texture);
+}
+
+
+TArray<UCustomizableObjectLayout*> UCustomizableObjectNodeMaterial::GetLayouts() const
+{
+	TArray<UCustomizableObjectLayout*> Result;
+
+	if (UEdGraphPin* MeshPin = GetMeshPin())
+	{
+		if (const UEdGraphPin* ConnectedPin = FollowInputPin(*MeshPin))
+		{
+			if (const UEdGraphPin* SourceMeshPin = FindMeshBaseSource(*ConnectedPin, false))
+			{
+				if (const UCustomizableObjectNodeSkeletalMesh* MeshNode = Cast<UCustomizableObjectNodeSkeletalMesh>(SourceMeshPin->GetOwningNode()))
+				{
+					Result = MeshNode->GetLayouts(*SourceMeshPin);
+				}
+				else if (const UCustomizableObjectNodeTable* TableNode = Cast<UCustomizableObjectNodeTable>(SourceMeshPin->GetOwningNode()))
+				{
+					Result = TableNode->GetLayouts(SourceMeshPin);
+				}
+			}
+		}
+	}
+
+	return Result;
+}
+
+
+int32 UCustomizableObjectNodeMaterial::GetImageUVLayout(const int32 ImageIndex) const
+{
+	if (const UEdGraphPin* Pin = GetParameterPin(EMaterialParameterType::Texture, ImageIndex))
+	{
+		const UCustomizableObjectNodeMaterialPinDataImage& PinData = GetPinData<UCustomizableObjectNodeMaterialPinDataImage>(*Pin);
+		switch(PinData.UVLayoutMode)
+		{
+		case EUVLayoutMode::FromMaterial:
+			break;
+		case EUVLayoutMode::Ignore:
+			return UCustomizableObjectNodeMaterialPinDataImage::UV_LAYOUT_IGNORE;
+		case EUVLayoutMode::Index:
+			return PinData.UVLayout;
+		default:
+			unimplemented();
+		}
+	}
+
+	const int32 UVIndex = GetImageUVLayoutFromMaterial(ImageIndex);
+	if (UVIndex == -1)
+	{
+		FCustomizableObjectEditorLogger::CreateLog(LOCTEXT("UVLayoutMaterialError", "Could not deduce the UV Layout Index from the UMaterial. Nodes connected to the Texture Sample UVs pin are supported."))
+			.Severity(EMessageSeverity::Warning)
+			.Context(*this)
+			.Log();
+		
+		return 0;
+	}
+
+	return UVIndex;
+}
+
+
+int32 UCustomizableObjectNodeMaterial::GetNumParameters(const EMaterialParameterType Type) const
+{
+	if (Material)
+	{
+		return Material->GetCachedExpressionData().GetParameterTypeEntry(Type).ParameterInfoSet.Num();
+	}
+	else
+	{
+		return 0;
+	}
+}
+
+
+FNodeMaterialParameterId UCustomizableObjectNodeMaterial::GetParameterId(const EMaterialParameterType Type, const int32 ParameterIndex) const
+{
+	const FMaterialCachedExpressionData& Data = Material->GetCachedExpressionData();
+
+	if (Data.EditorOnlyData)
+	{
+		if (Data.EditorOnlyData->EditorEntries[(int32)Type].EditorInfo.Num() != 0)
+		{
+			const FGuid ParameterId = Data.EditorOnlyData->EditorEntries[(int32)Type].EditorInfo[ParameterIndex].ExpressionGuid;
+			const int32 LayerIndex = GetParameterLayerIndex(Type, ParameterIndex);
+
+			return { ParameterId, LayerIndex };
+		}
+	}
+	
+	return FNodeMaterialParameterId();
+}
+
+
+FName UCustomizableObjectNodeMaterial::GetParameterName(const EMaterialParameterType Type, const int32 ParameterIndex) const
+{
+	check(Material);
+	
+	const FMaterialCachedParameterEntry& Entry = Material->GetCachedExpressionData().GetParameterTypeEntry(Type);
+
+	for (TSet<FMaterialParameterInfo>::TConstIterator It(Entry.ParameterInfoSet); It; ++It)
+	{
+		const int32 IteratorIndex = It.GetId().AsInteger();
+
+		if (IteratorIndex == ParameterIndex)
+		{
+			return (*It).Name;
+		}
+	}
+	
+	// The parameter should exist
+	check(false);
+
+	return FName();
+}
+
+
+int32 UCustomizableObjectNodeMaterial::GetParameterLayerIndex(const UMaterialInterface* InMaterial, const EMaterialParameterType Type, const int32 ParameterIndex)
+{
+	check(InMaterial);
+
+	const FMaterialCachedParameterEntry& Entry = InMaterial->GetCachedExpressionData().GetParameterTypeEntry(Type);
+
+	for (TSet<FMaterialParameterInfo>::TConstIterator It(Entry.ParameterInfoSet); It; ++It)
+	{
+		const int32 IteratorIndex = It.GetId().AsInteger();
+
+		if (IteratorIndex == ParameterIndex)
+		{
+			return (*It).Index;
+		}
+	}
+
+	// The parameter should exist
+	check(false);
+
+	return -1;
+}
+
+
+int32 UCustomizableObjectNodeMaterial::GetParameterLayerIndex(const EMaterialParameterType Type, const int32 ParameterIndex) const
+{
+	return GetParameterLayerIndex(Material.Get(), Type, ParameterIndex);
+}
+
+
+FText UCustomizableObjectNodeMaterial::GetParameterLayerName(const EMaterialParameterType Type, const int32 ParameterIndex) const
+{
+	check(Material)
+
+	int32 LayerIndex = GetParameterLayerIndex(Type,ParameterIndex);
+
+	FMaterialLayersFunctions LayersValue;
+	Material->GetMaterialLayers(LayersValue);
+
+
+	return LayersValue.EditorOnly.LayerNames.IsValidIndex(LayerIndex) ? LayersValue.EditorOnly.LayerNames[LayerIndex] : FText();
+}
+
+
+bool UCustomizableObjectNodeMaterial::HasParameter(const UMaterialInterface* InMaterial, const FNodeMaterialParameterId& ParameterId)
+{
+	if (!InMaterial)
+	{
+		return false;
+	}
+
+	for (const EMaterialParameterType Type : ParameterTypes)
+	{
+		const FMaterialCachedExpressionData& Data = InMaterial->GetCachedExpressionData();
+		const FMaterialCachedParameterEntry& Entry = Data.GetParameterTypeEntry(Type);
+
+		if (!Data.EditorOnlyData || Data.EditorOnlyData->EditorEntries[(int32)Type].EditorInfo.IsEmpty())
+		{
+			continue;
+		}
+
+		for (TSet<FMaterialParameterInfo>::TConstIterator It(Entry.ParameterInfoSet); It; ++It)
+		{
+			const int32 IteratorIndex = It.GetId().AsInteger();
+			
+			const FGuid& ParamGuid = Data.EditorOnlyData->EditorEntries[(int32)Type].EditorInfo[IteratorIndex].ExpressionGuid;
+			const int32 LayerIndex = GetParameterLayerIndex(InMaterial, Type, IteratorIndex);
+			const FNodeMaterialParameterId ParamId = { ParamGuid, LayerIndex };
+
+			if (ParamId == ParameterId)
+			{
+				return true;
+			}
+		}
+	}
+
+	return false;
+}
+
+
+bool UCustomizableObjectNodeMaterial::HasParameter(const FNodeMaterialParameterId& ParameterId) const
+{
+	return HasParameter( Material.Get(), ParameterId );
+}
+
+
+UEdGraphPin* UCustomizableObjectNodeMaterial::GetParameterPin(const EMaterialParameterType Type, const int32 ParameterIndex) const
+{
+	const FNodeMaterialParameterId ParameterId = GetParameterId(Type, ParameterIndex);
+	
+	return GetParameterPin(ParameterId);
+}
+
+
+UEdGraphPin* UCustomizableObjectNodeMaterial::GetParameterPin(const FNodeMaterialParameterId& ParameterId) const
+{
+	if (const FEdGraphPinReference* Result = PinsParameterMap.Find(ParameterId))
+	{
+		return Result->Get();
+	}
+	else
+	{
+		return nullptr;
+	}
+}
+
+
+bool UCustomizableObjectNodeMaterial::IsNodeOutDatedAndNeedsRefresh()
+{
+	const bool bOutdated = RealMaterialDataHasChanged();
+
+	// Remove previous compilation warnings
+	if (!bOutdated && bHasCompilerMessage)
+	{
+		RemoveWarnings();
+		GetGraph()->NotifyGraphChanged();
+	}
+
+	return bOutdated;
+}
+
+
+FString UCustomizableObjectNodeMaterial::GetRefreshMessage() const
+{
+	return "Referenced material has changed, texture channels might have been added, removed or renamed. Please refresh the node material to reflect those changes.";
+}
+
+
+TSharedPtr<IDetailsView> UCustomizableObjectNodeMaterial::CustomizePinDetails(const UEdGraphPin& Pin) const
+{
+	if (UCustomizableObjectNodeMaterialPinDataImage* PinData = Cast<UCustomizableObjectNodeMaterialPinDataImage>(GetPinData(Pin)))
+	{
+		FPropertyEditorModule& EditModule = FModuleManager::Get().GetModuleChecked<FPropertyEditorModule>("PropertyEditor");
+
+		FDetailsViewArgs DetailsViewArgs;
+		DetailsViewArgs.bAllowSearch = false;
+		DetailsViewArgs.bHideSelectionTip = true;
+		
+		const TSharedRef<IDetailsView> SettingsView = EditModule.CreateDetailView(DetailsViewArgs);
+		SettingsView->SetObject(PinData);
+		
+		return SettingsView;
+	}
+	else
+	{
+		return nullptr;
+	}
+}
+
+
+bool UCustomizableObjectNodeMaterial::CustomRemovePin(UEdGraphPin& Pin)
+{
+	for (TMap<FNodeMaterialParameterId, FEdGraphPinReference>::TIterator Iterator = PinsParameterMap.CreateIterator(); Iterator; ++Iterator)
+	{
+		if (Iterator->Value.Get() == &Pin)
+		{
+			Iterator.RemoveCurrent();
+			break;
+		}
+	}
+
+	return Super::CustomRemovePin(Pin);
+}
+
+
+void UCustomizableObjectNodeMaterial::SetMaterial(UMaterialInterface* InMaterial)
+{
+	Material = InMaterial;
+}
+
+
+UMaterialInterface* UCustomizableObjectNodeMaterial::GetMaterial() const
+{
+	return Material;	
+}
+
+
+bool UCustomizableObjectNodeMaterial::IsReuseMaterialBetweenLODs() const
+{
+	return bReuseMaterialBetweenLODs;	
+}
+
+
+TArray<FString>* UCustomizableObjectNodeMaterial::GetEnableTags()
+{
+	return &Tags;
+}
+
+
+bool UCustomizableObjectNodeMaterial::RealMaterialDataHasChanged() const
+{
+	for (const UEdGraphPin* Pin : GetAllNonOrphanPins())
+	{
+		if (const UCustomizableObjectNodeMaterialPinDataParameter* PinData = Cast<UCustomizableObjectNodeMaterialPinDataParameter>(GetPinData(*Pin)))
+		{
+			if (!HasParameter(PinData->MaterialParameterId) &&
+				(FollowInputPin(*Pin) || !PinData->IsDefault()))
+			{
+				return true;
+			}
+		}
+	}
+
+	return false;
+}
+
+
+FPostImagePinModeChangedDelegate* UCustomizableObjectNodeMaterial::GetPostImagePinModeChangedDelegate()
+{
+	return &PostImagePinModeChangedDelegate;
+}
+
+
+bool UCustomizableObjectNodeMaterial::IsPinRelevant(const UEdGraphPin* Pin) const
+{
+	const UEdGraphSchema_CustomizableObject* Schema = GetDefault<UEdGraphSchema_CustomizableObject>();
+
+	if (Pin->Direction == EEdGraphPinDirection::EGPD_Output)
+	{	
+		return Pin->PinType.PinCategory == Schema->PC_Mesh;
+	}
+	else if (Pin->Direction == EEdGraphPinDirection::EGPD_Input)
+	{	
+		return Pin->PinType.PinCategory == Schema->PC_Material;
+	}
+	else
+	{
+		return false;
+	}
+}
+
+
+FText UCustomizableObjectNodeMaterial::GetTooltipText() const
+{
+	return LOCTEXT("MeshSection_Tooltip", "Defines a Customizable Object mesh section.\nIt has a mesh section, a material assigned to it and the runtime modifiable inputs to the material asset parameters.");
+}
+
+
+TSharedPtr<SGraphNode> UCustomizableObjectNodeMaterial::CreateVisualWidget()
+{
+	return SNew(SCustomizableObjectNodeMaterial, this);
+}
+
+
+UCustomizableObjectNodeMaterialPinDataParameter* UCustomizableObjectNodeMaterial::CreatePinData(const EMaterialParameterType Type, const int32 ParameterIndex)
+{
+	UCustomizableObjectNodeMaterialPinDataParameter* PinData = nullptr;
+
+	switch (Type)
+	{
+	case EMaterialParameterType::Texture:
+		{
+			UCustomizableObjectNodeMaterialPinDataImage* PinDataImage = NewObject<UCustomizableObjectNodeMaterialPinDataImage>(this);
+			PinDataImage->Init(*this);
+		
+			PinData = PinDataImage;
+			break;
+		}
+	case EMaterialParameterType::Vector:
+		{
+			PinData = NewObject<UCustomizableObjectNodeMaterialPinDataVector>(this);	
+			break;
+		}
+
+	case EMaterialParameterType::Scalar:
+		{
+			PinData = NewObject<UCustomizableObjectNodeMaterialPinDataScalar>(this);	
+			break;
+		}
+
+	default:
+		check(false); // Parameter type not contemplated.
+	}
+
+	PinData->MaterialParameterId = GetParameterId(Type, ParameterIndex);
+
+	return PinData;
+}
+
+void UCustomizableObjectNodeMaterial::AllocateDefaultParameterPins(const EMaterialParameterType Type)
+{
+	const int32 NumParameters = GetNumParameters(Type);
+	for (int32 ParameterIndex = 0; ParameterIndex < NumParameters; ++ParameterIndex)
+	{
+		UCustomizableObjectNodeMaterialPinDataParameter* PinData = CreatePinData(Type, ParameterIndex);
+
+		const FName PinName = GetPinName(Type, ParameterIndex);
+
+		FName PinCategory;
+		switch(Type)
+		{
+		case EMaterialParameterType::Texture:
+			if (IsImageMutableMode(ParameterIndex)) // If a pin exists, we store the PinMode in its PinData.
+			{
+				PinCategory = UEdGraphSchema_CustomizableObject::PC_Image;				
+			}
+			else
+			{
+				PinCategory = UEdGraphSchema_CustomizableObject::PC_PassThroughImage;
+			}
+			break;
+
+		case EMaterialParameterType::Vector:
+			PinCategory = UEdGraphSchema_CustomizableObject::PC_Color;
+			break;
+
+		case EMaterialParameterType::Scalar:
+			PinCategory = UEdGraphSchema_CustomizableObject::PC_Float;
+			break;
+
+		default:
+			check(false); // Type not contemplated
+		}		
+		
+		UEdGraphPin* Pin = CustomCreatePin(EGPD_Input, PinCategory, PinName, PinData);
+		Pin->bHidden = true;
+		Pin->bDefaultValueIsIgnored = true;
+
+		PinsParameterMap.Add(PinData->MaterialParameterId, FEdGraphPinReference(Pin));
+	}
+}
+
+
+void UCustomizableObjectNodeMaterial::SetDefaultMaterial()
+{
+	if (const UEdGraphPin* MeshPin = GetMeshPin();
+		!Material && MeshPin)
+	{
+		if (const UEdGraphPin* LinkedMeshPin = FollowInputPin(*MeshPin))
+		{
+			UEdGraphNode* LinkedMeshNode = LinkedMeshPin->GetOwningNode();
+	
+			if (const UCustomizableObjectNodeSkeletalMesh* NodeSkeletalMesh = Cast<UCustomizableObjectNodeSkeletalMesh>(LinkedMeshNode))
+			{
+				Material = NodeSkeletalMesh->GetMaterialFor(LinkedMeshPin);
+				if (Material)
+				{
+					Super::ReconstructNode(); // Super required to avoid ambiguous call compilation error.
+				}
+			}
+			else if (const UCustomizableObjectNodeStaticMesh* NodeStaticMesh = Cast<UCustomizableObjectNodeStaticMesh>(LinkedMeshNode))
+			{
+				Material = NodeStaticMesh->GetMaterialFor(LinkedMeshPin);
+				if (Material)
+				{
+					Super::ReconstructNode();
+				}
+			}
+		}
+	}
+}
+
+
+void UCustomizableObjectNodeMaterial::PinConnectionListChanged(UEdGraphPin* Pin)
+{
+	Super::PinConnectionListChanged(Pin);
+	
+	if (Pin == GetMeshPin())
+	{
+		if (LastMeshNodeConnected.IsValid())
+		{
+			LastMeshNodeConnected->PostEditChangePropertyDelegate.RemoveDynamic(this, &UCustomizableObjectNodeMaterial::MeshPostEditChangeProperty);
+		}
+
+		if (const UEdGraphPin* ConnectedPin = FollowInputPin(*Pin))
+		{
+			UEdGraphNode* MeshNode = ConnectedPin->GetOwningNode();
+
+			if (MeshNode->IsA(UCustomizableObjectNodeStaticMesh::StaticClass()) || MeshNode->IsA(UCustomizableObjectNodeSkeletalMesh::StaticClass()))
+			{
+				SetDefaultMaterial();
+
+				LastMeshNodeConnected = Cast<UCustomizableObjectNode>(MeshNode);
+				LastMeshNodeConnected->PostEditChangePropertyDelegate.AddUniqueDynamic(this, &UCustomizableObjectNodeMaterial::MeshPostEditChangeProperty);
+			}
+		}
+	}
+	else if (Cast<UCustomizableObjectNodeMaterialPinDataImage>(GetPinData(*Pin))) // Image pin
+	{
+		// If necessary, automatically change the Pin Mode. Connected pin can never change its type.
+		if (Pin->LinkedTo.Num())
+		{
+			if (UEdGraphPin* LinkedPin = Pin->LinkedTo[0])
+			{
+				EPinMode PinMode = EPinMode::Default;
+		
+				if (LinkedPin->PinType.PinCategory == UEdGraphSchema_CustomizableObject::PC_Image)
+				{
+					PinMode = EPinMode::Mutable;
+				}
+				else if (LinkedPin->PinType.PinCategory == UEdGraphSchema_CustomizableObject::PC_PassThroughImage)
+				{
+					PinMode = EPinMode::Passthrough;
+				}
+				else
+				{
+					unimplemented()
+				}
+		
+				SetImagePinMode(*Pin, PinMode);
+			}
+		}
+	}
+}
+
+
+void UCustomizableObjectNodeMaterial::PostPasteNode()
+{
+	Super::PostPasteNode();
+	SetDefaultMaterial();
+}
+
+
+bool UCustomizableObjectNodeMaterial::CanConnect(const UEdGraphPin* InOwnedInputPin, const UEdGraphPin* InOutputPin, bool& bOutIsOtherNodeBlocklisted, bool& bOutArePinsCompatible) const
+{
+	const UEdGraphSchema_CustomizableObject* Schema = GetDefault<UEdGraphSchema_CustomizableObject>();
+
+	if (InOwnedInputPin && InOutputPin)
+	{
+		if ((InOwnedInputPin->PinType.PinCategory == Schema->PC_Image && InOutputPin->PinType.PinCategory == Schema->PC_PassThroughImage) ||
+			(InOwnedInputPin->PinType.PinCategory == Schema->PC_PassThroughImage && InOutputPin->PinType.PinCategory == Schema->PC_Image))
+		{
+			return true;
+		}
+		
+		if (InOwnedInputPin->PinType.PinCategory == Schema->PC_Mesh &&
+			InOutputPin->PinType.PinCategory == Schema->PC_PassThroughMesh)
+		{
+			return true;
+		}
+	}
+
+	return Super::CanConnect(InOwnedInputPin, InOutputPin, bOutIsOtherNodeBlocklisted, bOutArePinsCompatible);
+}
+
+
+void UCustomizableObjectNodeMaterial::MeshPostEditChangeProperty(FPostEditChangePropertyDelegateParameters& Parameters)
+{
+	if (const UEdGraphPin* MeshPin = FindPin(TEXT("Mesh")))
+	{
+		if (const UEdGraphPin* ConnectedPin = FollowInputPin(*MeshPin);
+			ConnectedPin && ConnectedPin->GetOwningNode() == Parameters.Node)
+		{
+			SetDefaultMaterial();
+		}
+		else if (UCustomizableObjectNode* MeshNode = Cast<UCustomizableObjectNode>(Parameters.Node))
+		{
+			MeshNode->PostEditChangePropertyDelegate.RemoveDynamic(this, &UCustomizableObjectNodeMaterial::MeshPostEditChangeProperty);
+		}
+	}
+}
+
+
+void UCustomizableObjectNodeMaterial::BreakExistingConnectionsPostConnection(UEdGraphPin* InputPin, UEdGraphPin* OutputPin)
+{
+	if (OutputPin->GetOwningNode() == this) // Case: NodeMaterial (OutputPin) --> (InputPin) Other Node
+	{
+		// Allow multiple NodeCopyMaterial nodes to be connected
+		if (InputPin->GetOwningNode()->IsA(UCustomizableObjectNodeCopyMaterial::StaticClass()))
+		{
+			return;
+		}
+
+		// Remove the previous NodeMaterial
+		UEdGraphPin* RemovePin = nullptr;
+
+		for (UEdGraphPin* Pin : OutputPin->LinkedTo)
+		{
+			if (Pin != InputPin && !Pin->GetOwningNode()->IsA(UCustomizableObjectNodeCopyMaterial::StaticClass())) {
+				RemovePin = Pin;
+				break;
+			}
+		}
+
+		if (RemovePin)
+		{
+			GetSchema()->BreakSinglePinLink(OutputPin, RemovePin); // Can not be called inside the range for loop.
+		}
+	}
+}
+
+
+#undef LOCTEXT_NAMESPACE
diff --git a/Engine/Plugins/Mutable/Source/CustomizableObjectEditor/Private/MuCOE/Nodes/CustomizableObjectNodeMaterialBase.cpp b/Engine/Plugins/Mutable/Source/CustomizableObjectEditor/Private/MuCOE/Nodes/CustomizableObjectNodeMaterialBase.cpp
new file mode 100644
--- /dev/null
+++ b/Engine/Plugins/Mutable/Source/CustomizableObjectEditor/Private/MuCOE/Nodes/CustomizableObjectNodeMaterialBase.cpp
@@ -0,0 +1,4 @@
+// Copyright Epic Games, Inc. All Rights Reserved.
+
+#include "MuCOE/Nodes/CustomizableObjectNodeMaterialBase.h"
+
diff --git a/Engine/Plugins/Mutable/Source/CustomizableObjectEditor/Private/MuCOE/Nodes/CustomizableObjectNodeMesh.cpp b/Engine/Plugins/Mutable/Source/CustomizableObjectEditor/Private/MuCOE/Nodes/CustomizableObjectNodeMesh.cpp
new file mode 100644
--- /dev/null
+++ b/Engine/Plugins/Mutable/Source/CustomizableObjectEditor/Private/MuCOE/Nodes/CustomizableObjectNodeMesh.cpp
@@ -0,0 +1,5 @@
+// Copyright Epic Games, Inc. All Rights Reserved.
+
+#include "MuCOE/Nodes/CustomizableObjectNodeMesh.h"
+
+
diff --git a/Engine/Plugins/Mutable/Source/CustomizableObjectEditor/Private/MuCOE/Nodes/CustomizableObjectNodeMeshClipWithMesh.cpp b/Engine/Plugins/Mutable/Source/CustomizableObjectEditor/Private/MuCOE/Nodes/CustomizableObjectNodeMeshClipWithMesh.cpp
new file mode 100644
--- /dev/null
+++ b/Engine/Plugins/Mutable/Source/CustomizableObjectEditor/Private/MuCOE/Nodes/CustomizableObjectNodeMeshClipWithMesh.cpp
@@ -0,0 +1,127 @@
+// Copyright Epic Games, Inc. All Rights Reserved.
+
+#include "MuCOE/Nodes/CustomizableObjectNodeModifierClipWithMesh.h"
+
+#include "MuCOE/CustomizableObjectGraph.h"
+#include "MuCOE/EdGraphSchema_CustomizableObject.h"
+#include "MuCOE/ICustomizableObjectEditor.h"
+#include "MuCO/CustomizableObjectCustomVersion.h"
+
+class UCustomizableObjectNodeRemapPins;
+
+#define LOCTEXT_NAMESPACE "CustomizableObjectEditor"
+
+
+UCustomizableObjectNodeModifierClipWithMesh::UCustomizableObjectNodeModifierClipWithMesh() :
+	CustomizableObjectToClipWith(nullptr)
+{
+	// We allow this node to be shared
+	Transform = FTransform::Identity;
+
+	// When initialize we don't use materials neither tags
+	bUseMaterials = false;
+	bUseTags = false;
+}
+
+
+void UCustomizableObjectNodeModifierClipWithMesh::BeginPostDuplicate(bool bDuplicateForPIE)
+{
+	Super::BeginPostDuplicate(bDuplicateForPIE);
+
+	if (UCustomizableObjectGraph* Graph = Cast<UCustomizableObjectGraph>(GetGraph()))
+	{
+		TArray<FGuid> NotifiedNodes;
+		for (const FGuid& NodeToClipWithId : ArrayMaterialNodeToClipWithID)
+		{
+			FGuid NodeToClipAlreadyUpdated = Graph->RequestNotificationForNodeIdChange(NodeToClipWithId, NodeGuid);
+
+			if (NodeToClipWithId != NodeToClipAlreadyUpdated)
+			{
+				NotifiedNodes.Add(NodeToClipAlreadyUpdated);
+			}
+		}
+
+		ArrayMaterialNodeToClipWithID = NotifiedNodes;
+	}
+}
+
+
+void UCustomizableObjectNodeModifierClipWithMesh::AllocateDefaultPins(UCustomizableObjectNodeRemapPins* RemapPins)
+{
+	const UEdGraphSchema_CustomizableObject* Schema = GetDefault<UEdGraphSchema_CustomizableObject>();
+
+	UEdGraphPin* ClipMeshPin = CustomCreatePin(EGPD_Input, Schema->PC_Mesh, FName("Clip Mesh"));
+	ClipMeshPin->bDefaultValueIsIgnored = true;
+
+	UEdGraphPin* OutputPin_p = CustomCreatePin(EGPD_Output, Schema->PC_Modifier, FName("Modifier"));
+}
+
+
+void UCustomizableObjectNodeModifierClipWithMesh::BackwardsCompatibleFixup(int32 CustomizableObjectCustomVersion)
+{
+	Super::BackwardsCompatibleFixup(CustomizableObjectCustomVersion);
+
+	const UEdGraphSchema_CustomizableObject* Schema = GetDefault<UEdGraphSchema_CustomizableObject>();
+
+	if (CustomizableObjectCustomVersion == FCustomizableObjectCustomVersion::UnifyRequiredTags)
+	{
+		RequiredTags = Tags_DEPRECATED;
+		Tags_DEPRECATED.Empty();
+	}
+}
+
+
+void UCustomizableObjectNodeModifierClipWithMesh::ReconstructNode(UCustomizableObjectNodeRemapPins* RemapPinsMode)
+{
+	Super::ReconstructNode(RemapPinsMode);
+
+	if (!CustomizableObjectToClipWith)
+	{
+		CustomizableObjectToClipWith = Cast<UCustomizableObject>(GetOutermostObject());
+	}
+}
+
+
+FText UCustomizableObjectNodeModifierClipWithMesh::GetNodeTitle(ENodeTitleType::Type TitleType) const
+{
+	return LOCTEXT("Clip_Mesh_With_Mesh", "Clip Mesh With Mesh");
+}
+
+
+void UCustomizableObjectNodeModifierClipWithMesh::PinConnectionListChanged(UEdGraphPin* Pin)
+{
+	Super::PinConnectionListChanged(Pin);
+
+	TSharedPtr<ICustomizableObjectEditor> Editor = GetGraphEditor();
+
+	if (Editor.IsValid())
+	{
+		Editor->UpdateGraphNodeProperties();
+	}
+}
+
+
+UEdGraphPin* UCustomizableObjectNodeModifierClipWithMesh::OutputPin() const
+{
+	return FindPin(TEXT("Modifier"));
+}
+
+
+UEdGraphPin* UCustomizableObjectNodeModifierClipWithMesh::ClipMeshPin() const
+{
+	return FindPin(TEXT("Clip Mesh"));
+}
+
+
+void UCustomizableObjectNodeModifierClipWithMesh::UpdateReferencedNodeId(const FGuid& NewGuid)
+{
+	ArrayMaterialNodeToClipWithID.Add(NewGuid);
+}
+
+
+FText UCustomizableObjectNodeModifierClipWithMesh::GetTooltipText() const
+{
+	return LOCTEXT("Clip_Mesh_Mesh_Tooltip", "Removes the part of a material that is completely enclosed in a mesh volume.\nIt only removes the faces that fall completely inside the cutting volume, along with the vertices and edges that define only faces that are deleted.");
+}
+
+#undef LOCTEXT_NAMESPACE
diff --git a/Engine/Plugins/Mutable/Source/CustomizableObjectEditor/Private/MuCOE/Nodes/CustomizableObjectNodeMeshGeometryOperation.cpp b/Engine/Plugins/Mutable/Source/CustomizableObjectEditor/Private/MuCOE/Nodes/CustomizableObjectNodeMeshGeometryOperation.cpp
new file mode 100644
--- /dev/null
+++ b/Engine/Plugins/Mutable/Source/CustomizableObjectEditor/Private/MuCOE/Nodes/CustomizableObjectNodeMeshGeometryOperation.cpp
@@ -0,0 +1,59 @@
+// Copyright Epic Games, Inc. All Rights Reserved.
+
+#include "MuCOE/Nodes/CustomizableObjectNodeMeshGeometryOperation.h"
+
+#include "MuCOE/EdGraphSchema_CustomizableObject.h"
+
+class UCustomizableObjectNodeRemapPins;
+
+#define LOCTEXT_NAMESPACE "CustomizableObjectEditor"
+
+
+UCustomizableObjectNodeMeshGeometryOperation::UCustomizableObjectNodeMeshGeometryOperation()
+	: Super()
+{
+}
+
+
+void UCustomizableObjectNodeMeshGeometryOperation::AllocateDefaultPins(UCustomizableObjectNodeRemapPins* RemapPins)
+{
+	const UEdGraphSchema_CustomizableObject* Schema = GetDefault<UEdGraphSchema_CustomizableObject>();
+
+	UEdGraphPin* OutputPin = CustomCreatePin(EGPD_Output, Schema->PC_Mesh, FName("Mesh"));
+	OutputPin->bDefaultValueIsIgnored = true;
+
+	UEdGraphPin* MeshPin = CustomCreatePin(EGPD_Input, Schema->PC_Mesh, FName("Mesh A"));
+	MeshPin->bDefaultValueIsIgnored = true;
+
+	MeshPin = CustomCreatePin(EGPD_Input, Schema->PC_Mesh, FName("Mesh B"));
+	MeshPin->bDefaultValueIsIgnored = true;
+
+	UEdGraphPin* FactorPin = CustomCreatePin(EGPD_Input, Schema->PC_Float, FName("Scalar A"));
+	FactorPin->bDefaultValueIsIgnored = true;
+
+	FactorPin = CustomCreatePin(EGPD_Input, Schema->PC_Float, FName("Scalar B"));
+	FactorPin->bDefaultValueIsIgnored = true;
+}
+
+
+FText UCustomizableObjectNodeMeshGeometryOperation::GetNodeTitle(ENodeTitleType::Type TitleType) const
+{
+	return LOCTEXT("Mesh_GeometryOperation", "Geometry Operation");
+}
+
+
+FLinearColor UCustomizableObjectNodeMeshGeometryOperation::GetNodeTitleColor() const
+{
+	const UEdGraphSchema_CustomizableObject* Schema = GetDefault<UEdGraphSchema_CustomizableObject>();
+	return Schema->GetPinTypeColor(Schema->PC_Mesh);
+}
+
+
+FText UCustomizableObjectNodeMeshGeometryOperation::GetTooltipText() const
+{
+	return LOCTEXT("Mesh_GeometryOperation_Tooltip", "Apply a Geometry operation on a mesh.");
+}
+
+
+
+#undef LOCTEXT_NAMESPACE
diff --git a/Engine/Plugins/Mutable/Source/CustomizableObjectEditor/Private/MuCOE/Nodes/CustomizableObjectNodeMeshMorphDetails.cpp b/Engine/Plugins/Mutable/Source/CustomizableObjectEditor/Private/MuCOE/Nodes/CustomizableObjectNodeMeshMorphDetails.cpp
new file mode 100644
--- /dev/null
+++ b/Engine/Plugins/Mutable/Source/CustomizableObjectEditor/Private/MuCOE/Nodes/CustomizableObjectNodeMeshMorphDetails.cpp
@@ -0,0 +1,121 @@
+// Copyright Epic Games, Inc. All Rights Reserved.
+
+#include "MuCOE/Nodes/CustomizableObjectNodeMeshMorphDetails.h"
+
+#include "DetailLayoutBuilder.h"
+#include "Engine/SkeletalMesh.h"
+#include "IDetailsView.h"
+#include "MuCOE/CustomizableObjectEditorUtilities.h"
+#include "MuCOE/GraphTraversal.h"
+#include "MuCOE/Nodes/CustomizableObjectNodeMeshMorph.h"
+#include "MuCOE/Nodes/CustomizableObjectNodeSkeletalMesh.h"
+#include "MuCOE/Nodes/CustomizableObjectNodeTable.h"
+#include "MuCOE/UnrealEditorPortabilityHelpers.h"
+#include "PropertyCustomizationHelpers.h"
+#include "Widgets/Input/STextComboBox.h"
+
+
+#define LOCTEXT_NAMESPACE "CustomizableObjectDetails"
+
+TSharedRef<IDetailCustomization> FCustomizableObjectNodeMeshMorphDetails::MakeInstance()
+{
+	return MakeShareable( new FCustomizableObjectNodeMeshMorphDetails );
+}
+
+void FCustomizableObjectNodeMeshMorphDetails::CustomizeDetails( IDetailLayoutBuilder& DetailBuilder )
+{
+	FCustomizableObjectNodeDetails::CustomizeDetails(DetailBuilder);
+
+	const IDetailsView* DetailsView = DetailBuilder.GetDetailsView();
+	if ( DetailsView->GetSelectedObjects().Num() )
+	{
+		Node = Cast<UCustomizableObjectNodeMeshMorph>( DetailsView->GetSelectedObjects()[0].Get() );
+	}
+
+	IDetailCategoryBuilder& BlocksCategory = DetailBuilder.EditCategory( "Morph Target" );
+	//BlocksCategory.CategoryIcon( "ActorClassIcon.CustomizableObject" );
+	
+	MorphTargetComboOptions.Empty();
+
+	bool SourceMeshFound = false;
+	if (Node)
+	{
+		if (const UEdGraphPin* ConnectedPin = FollowInputPin(*Node->MeshPin()))
+		{
+			SourceMeshFound = true;
+			
+			TSharedPtr<FString> ItemToSelect;
+
+			const UEdGraphPin* BaseSourcePin = FindMeshBaseSource(*ConnectedPin, false );
+			if ( BaseSourcePin )
+			{
+				USkeletalMesh* SkeletalMesh = nullptr;
+				if ( const UCustomizableObjectNodeSkeletalMesh* TypedNodeSkeletalMesh = Cast<UCustomizableObjectNodeSkeletalMesh>(BaseSourcePin->GetOwningNode()) )
+				{
+					SkeletalMesh = TypedNodeSkeletalMesh->SkeletalMesh;
+				}
+				else if (const UCustomizableObjectNodeTable* TypedNodeTable = Cast<UCustomizableObjectNodeTable>(BaseSourcePin->GetOwningNode()))
+				{
+					SkeletalMesh = TypedNodeTable->GetColumnDefaultAssetByType<USkeletalMesh>(BaseSourcePin);
+				}
+
+				if (SkeletalMesh)
+				{
+					for (int m = 0; m < SkeletalMesh->GetMorphTargets().Num(); ++m)
+					{
+						FString MorphName = *SkeletalMesh->GetMorphTargets()[m]->GetName();
+						MorphTargetComboOptions.Add(MakeShareable(new FString(MorphName)));
+
+						if (Node->MorphTargetName == MorphName)
+						{
+							ItemToSelect = MorphTargetComboOptions.Last();
+						}
+					}
+				}
+			}
+
+			MorphTargetComboOptions.Sort(CompareNames);
+
+			TSharedRef<IPropertyHandle> MorphTargetNameProperty = DetailBuilder.GetProperty("MorphTargetName");
+
+			BlocksCategory.AddCustomRow( LOCTEXT("MeshMorphDetails_Target", "Target") )
+			[
+				SNew(SProperty, MorphTargetNameProperty)
+				.ShouldDisplayName( false)
+				.CustomWidget()
+				[
+					SNew(SBorder)
+					.BorderImage(UE_MUTABLE_GET_BRUSH("NoBorder"))
+					.Padding(FMargin(0.0f, 0.0f, 10.0f, 0.0f))
+					[
+						SNew(STextComboBox)
+							.OptionsSource( &MorphTargetComboOptions )
+							.InitiallySelectedItem( ItemToSelect )
+							.OnSelectionChanged(this, &FCustomizableObjectNodeMeshMorphDetails::OnMorphTargetComboBoxSelectionChanged, MorphTargetNameProperty)
+					]
+				]
+			];
+		}
+	}
+	
+	if (!SourceMeshFound)
+	{
+		BlocksCategory.AddCustomRow( LOCTEXT("MeshMorphDetails_Node", "Node") )
+		[
+			SNew( STextBlock )
+			.Text( LOCTEXT( "MeshMorphDetails_NoSource", "No source mesh found." ) )			
+		];
+	}
+
+}
+
+void FCustomizableObjectNodeMeshMorphDetails::OnMorphTargetComboBoxSelectionChanged(TSharedPtr<FString> Selection, ESelectInfo::Type SelectInfo, TSharedRef<IPropertyHandle> ParentProperty)
+{	
+	if (Selection.IsValid())
+	{
+		ParentProperty->SetValue(*Selection);	
+	}
+}
+
+#undef LOCTEXT_NAMESPACE
+
diff --git a/Engine/Plugins/Mutable/Source/CustomizableObjectEditor/Private/MuCOE/Nodes/CustomizableObjectNodeMeshMorphDetails.h b/Engine/Plugins/Mutable/Source/CustomizableObjectEditor/Private/MuCOE/Nodes/CustomizableObjectNodeMeshMorphDetails.h
new file mode 100644
--- /dev/null
+++ b/Engine/Plugins/Mutable/Source/CustomizableObjectEditor/Private/MuCOE/Nodes/CustomizableObjectNodeMeshMorphDetails.h
@@ -0,0 +1,29 @@
+// Copyright Epic Games, Inc. All Rights Reserved.
+
+#pragma once
+
+#include "CustomizableObjectNodeDetails.h"
+#include "IDetailCustomization.h"
+
+namespace ESelectInfo { enum Type : int; }
+
+class FString;
+class IDetailLayoutBuilder;
+class IPropertyHandle;
+
+class FCustomizableObjectNodeMeshMorphDetails : public FCustomizableObjectNodeDetails
+{
+public:
+	/** Makes a new instance of this detail layout class for a specific detail view requesting it */
+	static TSharedRef<IDetailCustomization> MakeInstance();
+
+	/** ILayoutDetails interface */
+	virtual void CustomizeDetails( IDetailLayoutBuilder& DetailBuilder ) override;
+
+private:
+
+	class UCustomizableObjectNodeMeshMorph* Node = nullptr;
+	TArray< TSharedPtr<FString> > MorphTargetComboOptions;
+
+	void OnMorphTargetComboBoxSelectionChanged(TSharedPtr<FString> Selection, ESelectInfo::Type SelectInfo, TSharedRef<IPropertyHandle> Property);
+};
diff --git a/Engine/Plugins/Mutable/Source/CustomizableObjectEditor/Private/MuCOE/Nodes/CustomizableObjectNodeMeshMorphStackDefinition.cpp b/Engine/Plugins/Mutable/Source/CustomizableObjectEditor/Private/MuCOE/Nodes/CustomizableObjectNodeMeshMorphStackDefinition.cpp
new file mode 100644
--- /dev/null
+++ b/Engine/Plugins/Mutable/Source/CustomizableObjectEditor/Private/MuCOE/Nodes/CustomizableObjectNodeMeshMorphStackDefinition.cpp
@@ -0,0 +1,221 @@
+// Copyright Epic Games, Inc. All Rights Reserved.
+
+#include "MuCOE/Nodes/CustomizableObjectNodeMeshMorphStackDefinition.h"
+
+#include "Engine/SkeletalMesh.h"
+#include "MuCOE/EdGraphSchema_CustomizableObject.h"
+#include "MuCOE/GraphTraversal.h"
+#include "MuCOE/ICustomizableObjectEditor.h"
+#include "MuCOE/Nodes/CustomizableObjectNodeSkeletalMesh.h"
+#include "MuCOE/Nodes/CustomizableObjectNodeTable.h"
+
+class UCustomizableObjectNodeRemapPins;
+
+
+#define LOCTEXT_NAMESPACE "CustomizableObjectEditor"
+
+
+void UCustomizableObjectNodeMeshMorphStackDefinition::AllocateDefaultPins(UCustomizableObjectNodeRemapPins* RemapPins)
+{
+	const UEdGraphSchema_CustomizableObject* Schema = GetDefault<UEdGraphSchema_CustomizableObject>();
+
+	UpdateMorphList();
+	
+	// Input pins
+	UEdGraphPin* MeshPin = CustomCreatePin(EGPD_Input, Schema->PC_Mesh, FName(TEXT("Mesh")));
+	MeshPin->bDefaultValueIsIgnored = true;
+
+	for (int32 i = 0; i < MorphNames.Num(); ++i)
+	{
+		FString PinName = MorphNames[i] + FString("_Morph");
+		UEdGraphPin* Pin = CustomCreatePin(EGPD_Input, Schema->PC_Float, FName(*PinName));
+		Pin->PinFriendlyName = FText::FromString(MorphNames[i]);
+		Pin->bDefaultValueIsIgnored = true;
+	}
+
+	// Output pins
+	UEdGraphPin* StackPin = CustomCreatePin(EGPD_Output, Schema->PC_Stack, FName(TEXT("Stack")));
+	StackPin->bDefaultValueIsIgnored = true;
+}
+
+
+FText UCustomizableObjectNodeMeshMorphStackDefinition::GetNodeTitle(ENodeTitleType::Type TittleType)const
+{
+	return LOCTEXT("Mesh_Morph_Stack_Definition", "Mesh Morph Stack Definition");
+}
+
+
+FLinearColor UCustomizableObjectNodeMeshMorphStackDefinition::GetNodeTitleColor() const
+{
+	const UEdGraphSchema_CustomizableObject* Schema = GetDefault<UEdGraphSchema_CustomizableObject>();
+	return Schema->GetPinTypeColor(Schema->PC_Mesh);
+}
+
+
+FText UCustomizableObjectNodeMeshMorphStackDefinition::GetTooltipText() const
+{
+	return LOCTEXT("Morph_Stack_Definition_Tooltip","Allows to stack morphs");
+}
+
+
+void UCustomizableObjectNodeMeshMorphStackDefinition::PinConnectionListChanged(UEdGraphPin* Pin)
+{
+	Super::PinConnectionListChanged(Pin);
+
+	if (Pin->Direction == EGPD_Input && Pin->PinName == "Mesh")
+	{
+		Super::ReconstructNode();
+	}
+}
+
+
+void UCustomizableObjectNodeMeshMorphStackDefinition::ReconstructNode(UCustomizableObjectNodeRemapPins* RemapPins)
+{
+	Super::ReconstructNode(RemapPins);
+}
+
+
+void UCustomizableObjectNodeMeshMorphStackDefinition::UpdateMorphList()
+{
+	MorphNames.Empty();
+
+	UEdGraphPin* MeshPin = GetMeshPin();
+
+	if (!MeshPin)
+	{
+		return;
+	}
+
+	UEdGraphPin* OutputMeshPin = FollowInputPin(*MeshPin);
+
+	if (!OutputMeshPin)
+	{
+		return;
+	}
+
+	const UEdGraphPin* MeshNodePin = FindMeshBaseSource(*OutputMeshPin, false);
+
+	if (MeshNodePin && MeshNodePin->GetOwningNode())
+	{
+		USkeletalMesh* SkeletalMesh = nullptr;
+
+		if (const UCustomizableObjectNodeSkeletalMesh* SkeletalMeshNode = Cast< UCustomizableObjectNodeSkeletalMesh >(MeshNodePin->GetOwningNode()))
+		{
+			SkeletalMesh = SkeletalMeshNode->SkeletalMesh;
+		}
+		else if (const UCustomizableObjectNodeTable* TableNode = Cast< UCustomizableObjectNodeTable >(MeshNodePin->GetOwningNode()))
+		{
+			SkeletalMesh = TableNode->GetColumnDefaultAssetByType<USkeletalMesh>(MeshNodePin);
+		}
+
+		if (SkeletalMesh)
+		{
+			for (int32 i = 0; i < SkeletalMesh->GetMorphTargets().Num(); ++i)
+			{
+				MorphNames.Add(SkeletalMesh->GetMorphTargets()[i]->GetName());
+			}
+		}
+	}
+}
+
+
+UEdGraphPin* UCustomizableObjectNodeMeshMorphStackDefinition::GetMeshPin() const
+{
+	return FindPin(TEXT("Mesh"), EGPD_Input);
+}
+
+
+UEdGraphPin* UCustomizableObjectNodeMeshMorphStackDefinition::GetStackPin() const
+{
+	return FindPin(TEXT("Stack"), EGPD_Output);
+}
+
+
+UEdGraphPin* UCustomizableObjectNodeMeshMorphStackDefinition::GetMorphPin(int32 Index) const
+{
+	FString PinName = MorphNames[Index] + FString("_Morph");
+
+	return FindPin(PinName);
+}
+
+
+int32 UCustomizableObjectNodeMeshMorphStackDefinition::NextConnectedPin(int32 Index, TArray<FString> AvailableMorphs) const
+{
+	if (AvailableMorphs.Num())
+	{
+		for (int32 i = Index + 1; i < MorphNames.Num(); ++i)
+		{
+			FString PinName = MorphNames[i] + FString("_Morph");
+
+			UEdGraphPin* NextPin = FindPin(PinName);
+
+			// Checking if the node contains another connected pin and if the morph of this pin existis in the skeletal mesh
+			if (NextPin && FollowInputPin(*NextPin) && AvailableMorphs.Contains(MorphNames[i]))
+			{
+				return i;
+			}
+		}
+	}
+
+	return -1;
+}
+
+
+bool UCustomizableObjectNodeMeshMorphStackDefinition::IsNodeOutDatedAndNeedsRefresh()
+{
+	UEdGraphPin* MeshPin = GetMeshPin();
+
+	if (!MeshPin)
+	{
+		return false;
+	}
+
+	UEdGraphPin* OutputMeshPin = FollowInputPin(*MeshPin);
+
+	if (!OutputMeshPin)
+	{
+		if (MorphNames.Num())
+		{
+			return true;
+		}
+
+		return false;
+	}
+
+	const UEdGraphPin* MeshNodePin = FindMeshBaseSource(*OutputMeshPin, false);
+
+	if (MeshNodePin && MeshNodePin->GetOwningNode())
+	{
+		USkeletalMesh* SkeletalMesh = nullptr;
+
+		if (const UCustomizableObjectNodeSkeletalMesh* SkeletalMeshNode = Cast< UCustomizableObjectNodeSkeletalMesh >(MeshNodePin->GetOwningNode()))
+		{
+			SkeletalMesh = SkeletalMeshNode->SkeletalMesh;
+		}
+		else if (const UCustomizableObjectNodeTable* TableNode = Cast< UCustomizableObjectNodeTable >(MeshNodePin->GetOwningNode()))
+		{
+			SkeletalMesh = TableNode->GetColumnDefaultAssetByType<USkeletalMesh>(MeshNodePin);
+		}
+
+		if (SkeletalMesh)
+		{
+			if (SkeletalMesh->GetMorphTargets().Num() != MorphNames.Num())
+			{
+				return true;
+			}
+
+			for (int32 i = 0; i < SkeletalMesh->GetMorphTargets().Num(); ++i)
+			{
+				if (!MorphNames.Contains(SkeletalMesh->GetMorphTargets()[i]->GetName()))
+				{
+					return true;
+				}
+			}
+		}
+	}
+
+	return false;
+}
+
+
+#undef LOCTEXT_NAMESPACE
diff --git a/Engine/Plugins/Mutable/Source/CustomizableObjectEditor/Private/MuCOE/Nodes/CustomizableObjectNodeMeshReshape.cpp b/Engine/Plugins/Mutable/Source/CustomizableObjectEditor/Private/MuCOE/Nodes/CustomizableObjectNodeMeshReshape.cpp
new file mode 100644
--- /dev/null
+++ b/Engine/Plugins/Mutable/Source/CustomizableObjectEditor/Private/MuCOE/Nodes/CustomizableObjectNodeMeshReshape.cpp
@@ -0,0 +1,90 @@
+// Copyright Epic Games, Inc. All Rights Reserved.
+
+#include "MuCOE/Nodes/CustomizableObjectNodeMeshReshape.h"
+
+#include "MuCO/CustomizableObjectCustomVersion.h"
+#include "MuCOE/EdGraphSchema_CustomizableObject.h"
+#include "MuCOE/Nodes/CustomizableObjectNodeMeshReshapeCommon.h"
+
+class UCustomizableObjectNodeRemapPins;
+
+#define LOCTEXT_NAMESPACE "CustomizableObjectEditor"
+
+
+UCustomizableObjectNodeMeshReshape::UCustomizableObjectNodeMeshReshape()
+	: Super()
+{
+}
+
+
+void UCustomizableObjectNodeMeshReshape::AllocateDefaultPins(UCustomizableObjectNodeRemapPins* RemapPins)
+{
+	const UEdGraphSchema_CustomizableObject* Schema = GetDefault<UEdGraphSchema_CustomizableObject>();
+
+	UEdGraphPin* OutputPin = CustomCreatePin(EGPD_Output, Schema->PC_Mesh, FName("Mesh"));
+	OutputPin->bDefaultValueIsIgnored = true;
+
+	UEdGraphPin* MeshPin = CustomCreatePin(EGPD_Input, Schema->PC_Mesh, FName("Base Mesh"));
+	MeshPin->bDefaultValueIsIgnored = true;
+
+	MeshPin = CustomCreatePin(EGPD_Input, Schema->PC_Mesh, FName("Base Shape"));
+	MeshPin->bDefaultValueIsIgnored = true;
+
+	MeshPin = CustomCreatePin(EGPD_Input, Schema->PC_Mesh, FName("Target Shape"));
+	MeshPin->bDefaultValueIsIgnored = true;
+}
+
+
+FText UCustomizableObjectNodeMeshReshape::GetNodeTitle(ENodeTitleType::Type TitleType) const
+{
+	return LOCTEXT("Mesh_Reshape", "Mesh Reshape");
+}
+
+
+FLinearColor UCustomizableObjectNodeMeshReshape::GetNodeTitleColor() const
+{
+	const UEdGraphSchema_CustomizableObject* Schema = GetDefault<UEdGraphSchema_CustomizableObject>();
+	return Schema->GetPinTypeColor(Schema->PC_Mesh);
+}
+
+
+FText UCustomizableObjectNodeMeshReshape::GetTooltipText() const
+{
+	return LOCTEXT("Mesh_Reshape_Tooltip", "Apply a mesh reshape on a mesh.");
+}
+
+
+FString UCustomizableObjectNodeMeshReshape::GetRefreshMessage() const
+{
+	return "One or more bones selected to deform do not exist in the current reference mesh skeleton!";
+}
+
+
+void UCustomizableObjectNodeMeshReshape::Serialize(FArchive& Ar)
+{
+	Super::Serialize(Ar);
+
+	Ar.UsingCustomVersion(FCustomizableObjectCustomVersion::GUID);
+
+	if (Ar.CustomVer(FCustomizableObjectCustomVersion::GUID) < FCustomizableObjectCustomVersion::DeformSkeletonOptionsAdded)
+	{
+		if (bDeformAllBones_DEPRECATED)
+		{
+			SelectionMethod = EBoneDeformSelectionMethod::ALL_BUT_SELECTED;
+			BonesToDeform.Empty();
+		}
+	}
+
+	if (Ar.CustomVer(FCustomizableObjectCustomVersion::GUID) < FCustomizableObjectCustomVersion::MeshReshapeVertexColorUsageSelection)
+	{
+		if (bEnableRigidParts_DEPRECATED)
+		{
+			VertexColorUsage.R = EMeshReshapeVertexColorChannelUsage::RigidClusterId;
+			VertexColorUsage.G = EMeshReshapeVertexColorChannelUsage::RigidClusterId;
+			VertexColorUsage.B = EMeshReshapeVertexColorChannelUsage::RigidClusterId;
+			VertexColorUsage.A = EMeshReshapeVertexColorChannelUsage::RigidClusterId;
+		}
+	}
+}
+
+#undef LOCTEXT_NAMESPACE
