[mutable] Moved the Mutable plugin out of Experimental status into Beta.

diff --git a/Engine/Plugins/Mutable/Source/CustomizableObjectEditor/Private/MuCOE/Nodes/CustomizableObjectNodeMeshReshapeSelectionDetails.cpp b/Engine/Plugins/Mutable/Source/CustomizableObjectEditor/Private/MuCOE/Nodes/CustomizableObjectNodeMeshReshapeSelectionDetails.cpp
new file mode 100644
--- /dev/null
+++ b/Engine/Plugins/Mutable/Source/CustomizableObjectEditor/Private/MuCOE/Nodes/CustomizableObjectNodeMeshReshapeSelectionDetails.cpp
@@ -0,0 +1,186 @@
+// Copyright Epic Games, Inc. All Rights Reserved.
+
+#include "MuCOE/Nodes/CustomizableObjectNodeMeshReshapeSelectionDetails.h"
+
+#include "BoneSelectionWidget.h"
+#include "DetailWidgetRow.h"
+#include "Engine/SkeletalMesh.h"
+#include "MuCOE/GraphTraversal.h"
+#include "MuCOE/Nodes/CustomizableObjectNodeMeshMorph.h"
+#include "MuCOE/Nodes/CustomizableObjectNodeMeshReshape.h"
+#include "MuCOE/Nodes/CustomizableObjectNodeMeshReshapeCommon.h"
+#include "MuCOE/Nodes/CustomizableObjectNodeSkeletalMesh.h"
+#include "MuCOE/Nodes/CustomizableObjectNodeTable.h"
+
+
+#define LOCTEXT_NAMESPACE "CustomizableObjectDetails"
+
+
+TSharedRef<IPropertyTypeCustomization> FMeshReshapeBonesReferenceCustomization::MakeInstance()
+{
+	return MakeShareable(new FMeshReshapeBonesReferenceCustomization());
+}
+
+void FMeshReshapeBonesReferenceCustomization::CustomizeHeader(TSharedRef<IPropertyHandle> StructPropertyHandle, class FDetailWidgetRow& HeaderRow, IPropertyTypeCustomizationUtils& StructCustomizationUtils)
+{
+	constexpr bool bRecurse = false;
+	BoneNameProperty = StructPropertyHandle->GetChildHandle(GET_MEMBER_NAME_CHECKED(FMeshReshapeBoneReference, BoneName), bRecurse);
+
+	TSharedPtr<IPropertyHandle> Parent = StructPropertyHandle->GetParentHandle();
+	if (Parent)
+	{
+		ParentArrayProperty = Parent->AsArray();
+	}
+
+	if (BoneNameProperty->IsValidHandle())
+	{
+		HeaderRow
+		.NameContent()
+		[
+			StructPropertyHandle->CreatePropertyNameWidget()
+		]
+		.ValueContent()
+		.MaxDesiredWidth(0.0f)
+		[
+			SAssignNew(SelectionWidget, SBoneSelectionWidget)
+			.ToolTipText(StructPropertyHandle->GetToolTipText())
+			.OnBoneSelectionChanged(this, &FMeshReshapeBonesReferenceCustomization::OnBoneSelectionChanged)
+			.OnGetSelectedBone(this, &FMeshReshapeBonesReferenceCustomization::GetSelectedBone)
+			.OnGetReferenceSkeleton(this, &FMeshReshapeBonesReferenceCustomization::GetReferenceSkeleton)
+		];
+
+		DefaultColor = SelectionWidget->GetColorAndOpacity();
+	}
+}
+
+USkeletalMesh* FMeshReshapeBonesReferenceCustomization::GetSkeletalMesh(TSharedPtr<IPropertyHandle> BoneNamePropertyHandle) const 
+{
+	if (!BoneNamePropertyHandle)
+	{
+		return nullptr;
+	}
+
+	TArray<UObject*> Objects;
+	BoneNamePropertyHandle->GetOuterObjects(Objects);
+
+	USkeletalMesh* TargetSkeletalMesh = nullptr;
+	
+	for (const UObject* Object : Objects)
+	{
+		UEdGraphPin* MeshPin = nullptr;
+
+		if (const UCustomizableObjectNodeMeshMorph* NodeMorph = Cast<UCustomizableObjectNodeMeshMorph>(Object))
+		{
+			MeshPin = NodeMorph->MeshPin();
+		}
+		else if (const UCustomizableObjectNodeMeshReshape* NodeReshape = Cast<UCustomizableObjectNodeMeshReshape>(Object))
+		{
+			MeshPin = NodeReshape->BaseMeshPin();
+		}
+		else
+		{
+			checkf(false, TEXT("node not implemented"));
+		}
+
+		if (MeshPin)
+		{
+			if (const UEdGraphPin* ConnectedPin = FollowInputPin(*MeshPin))
+			{
+				if (const UEdGraphPin* SourceMeshPin = FindMeshBaseSource(*ConnectedPin, false))
+				{
+					if (const UCustomizableObjectNodeSkeletalMesh* SkeletalMeshNode = Cast<UCustomizableObjectNodeSkeletalMesh>(SourceMeshPin->GetOwningNode()))
+					{
+						return SkeletalMeshNode->SkeletalMesh;
+					}
+					else if (UCustomizableObjectNodeTable* TableNode = Cast<UCustomizableObjectNodeTable>(SourceMeshPin->GetOwningNode()))
+					{
+						return TableNode->GetColumnDefaultAssetByType<USkeletalMesh>(SourceMeshPin);
+					}
+				}
+			}
+		}
+	}
+
+	return nullptr;
+}
+
+FName FMeshReshapeBonesReferenceCustomization::GetSelectedBone(bool& bMultipleValues) const
+{
+	if (BoneNameProperty.IsValid())
+	{
+		FString OutText;
+
+		FPropertyAccess::Result Result = BoneNameProperty->GetValueAsFormattedString(OutText);
+
+		return FName(*OutText);
+	}
+
+	return FName();
+}
+
+const struct FReferenceSkeleton& FMeshReshapeBonesReferenceCustomization::GetReferenceSkeleton() const
+{
+	static FReferenceSkeleton DummySkeleton;
+
+	USkeletalMesh* SkeletalMesh = GetSkeletalMesh(BoneNameProperty);
+	return SkeletalMesh ? SkeletalMesh->GetRefSkeleton() : DummySkeleton;
+}
+
+void FMeshReshapeBonesReferenceCustomization::OnBoneSelectionChanged(FName Name)
+{
+	FMeshReshapeBoneReference NewBone;
+	NewBone.BoneName = Name;
+
+	if (!ParentArrayProperty)
+	{
+		BoneNameProperty->SetValue(Name);
+		return;
+	}
+
+	uint32 NumParentElems = 0;
+	FPropertyAccess::Result Result = ParentArrayProperty->GetNumElements(NumParentElems);
+	check(Result == FPropertyAccess::Result::Success);
+
+	// Find if already set.
+	bool BoneNameFound = false;
+	for (uint32 I = 0; I < NumParentElems; ++I)
+	{
+		TSharedPtr<IPropertyHandle> ElemProperty = ParentArrayProperty->GetElement(I);
+		check(ElemProperty);
+	
+		constexpr bool bRecurse = false;	
+		TSharedPtr<IPropertyHandle> ElemBoneNameProperty = 
+				ElemProperty->GetChildHandle(GET_MEMBER_NAME_CHECKED(FMeshReshapeBoneReference, BoneName), bRecurse);
+
+		check(ElemBoneNameProperty);
+
+		FName BoneNameValue;
+		Result = ElemBoneNameProperty->GetValue(BoneNameValue);
+
+		check(Result == FPropertyAccess::Result::Success);
+
+		if (BoneNameValue == NewBone.BoneName)
+		{
+			BoneNameFound = true;
+			break;
+		}	
+	}
+
+
+	// TODO: This behavior does not seem to be correct, if the option is repeated the value is not set
+	// only logging a warning to the console, which might not be open and could make dificult to see what 
+	// is the problem. 
+	BoneNameProperty->SetValue(Name);
+	if (BoneNameFound)
+	{
+		const FLinearColor WarningColor = FLinearColor::Yellow;
+		SelectionWidget->SetColorAndOpacity(WarningColor);
+	}
+	else
+	{
+		SelectionWidget->SetColorAndOpacity(DefaultColor);
+	}
+}
+
+#undef LOCTEXT_NAMESPACE
+
diff --git a/Engine/Plugins/Mutable/Source/CustomizableObjectEditor/Private/MuCOE/Nodes/CustomizableObjectNodeMeshReshapeSelectionDetails.h b/Engine/Plugins/Mutable/Source/CustomizableObjectEditor/Private/MuCOE/Nodes/CustomizableObjectNodeMeshReshapeSelectionDetails.h
new file mode 100644
--- /dev/null
+++ b/Engine/Plugins/Mutable/Source/CustomizableObjectEditor/Private/MuCOE/Nodes/CustomizableObjectNodeMeshReshapeSelectionDetails.h
@@ -0,0 +1,41 @@
+// Copyright Epic Games, Inc. All Rights Reserved.
+
+#pragma once
+
+#include "IPropertyTypeCustomization.h"
+
+class IPropertyHandle;
+class IPropertyHandleArray;
+class SBoneSelectionWidget;
+
+class FMeshReshapeBonesReferenceCustomization : public IPropertyTypeCustomization
+{
+public:
+	static TSharedRef<IPropertyTypeCustomization> MakeInstance();
+
+	// IPropertyTypeCustomization interface
+	void CustomizeHeader(TSharedRef<IPropertyHandle> StructPropertyHandle, class FDetailWidgetRow& HeaderRow, IPropertyTypeCustomizationUtils& StructCustomizationUtils) override;
+	void CustomizeChildren(TSharedRef<IPropertyHandle> StructPropertyHandle, class IDetailChildrenBuilder& ChildBuilder, IPropertyTypeCustomizationUtils& StructCustomizationUtils) override {};
+	
+	// Gets the Reference Skeleton from the target skeleton
+	const struct FReferenceSkeleton& GetReferenceSkeleton() const;
+	// Gets the current selected bone in the combo button
+	FName GetSelectedBone(bool& bMultipleValues) const;
+	// OnSelectionChanged callback
+	void OnBoneSelectionChanged(FName Name);
+
+private:
+
+	class USkeletalMesh* GetSkeletalMesh(TSharedPtr<IPropertyHandle> BoneNamePropertyHandle) const; 
+
+private:
+
+	// Property to change after bone has been picked
+	TSharedPtr<IPropertyHandle> BoneNameProperty;
+
+	// Parent propery, as array to identify repeated values.
+	TSharedPtr<IPropertyHandleArray> ParentArrayProperty;
+
+	TSharedPtr<SBoneSelectionWidget> SelectionWidget;
+	FLinearColor DefaultColor;
+};
diff --git a/Engine/Plugins/Mutable/Source/CustomizableObjectEditor/Private/MuCOE/Nodes/CustomizableObjectNodeMeshSwitch.cpp b/Engine/Plugins/Mutable/Source/CustomizableObjectEditor/Private/MuCOE/Nodes/CustomizableObjectNodeMeshSwitch.cpp
new file mode 100644
--- /dev/null
+++ b/Engine/Plugins/Mutable/Source/CustomizableObjectEditor/Private/MuCOE/Nodes/CustomizableObjectNodeMeshSwitch.cpp
@@ -0,0 +1,29 @@
+// Copyright Epic Games, Inc. All Rights Reserved.
+
+#include "MuCOE/Nodes/CustomizableObjectNodeMeshSwitch.h"
+
+#include "MuCO/CustomizableObjectCustomVersion.h"
+#include "MuCOE/EdGraphSchema_CustomizableObject.h"
+
+#define LOCTEXT_NAMESPACE "CustomizableObjectEditor"
+
+
+void UCustomizableObjectNodeMeshSwitch::BackwardsCompatibleFixup(int32 CustomizableObjectCustomVersion)
+{
+	Super::BackwardsCompatibleFixup(CustomizableObjectCustomVersion);
+	
+	if (CustomizableObjectCustomVersion == FCustomizableObjectCustomVersion::BugPinsSwitch)
+	{
+		OutputPinReference = FindPin(TEXT("Mesh"));	
+	}
+}
+
+
+FName UCustomizableObjectNodeMeshSwitch::GetCategory() const
+{
+	return UEdGraphSchema_CustomizableObject::PC_Mesh;
+}
+
+
+#undef LOCTEXT_NAMESPACE
+
diff --git a/Engine/Plugins/Mutable/Source/CustomizableObjectEditor/Private/MuCOE/Nodes/CustomizableObjectNodeModifierClipMorphDetails.cpp b/Engine/Plugins/Mutable/Source/CustomizableObjectEditor/Private/MuCOE/Nodes/CustomizableObjectNodeModifierClipMorphDetails.cpp
new file mode 100644
--- /dev/null
+++ b/Engine/Plugins/Mutable/Source/CustomizableObjectEditor/Private/MuCOE/Nodes/CustomizableObjectNodeModifierClipMorphDetails.cpp
@@ -0,0 +1,352 @@
+// Copyright Epic Games, Inc. All Rights Reserved.
+
+#include "MuCOE/Nodes/CustomizableObjectNodeModifierClipMorphDetails.h"
+
+#include "DetailLayoutBuilder.h"
+#include "Engine/SkeletalMesh.h"
+#include "IDetailsView.h"
+#include "MuCOE/CustomizableObjectEditorUtilities.h"
+#include "MuCOE/CustomizableObjectGraph.h"
+#include "MuCOE/GraphTraversal.h"
+#include "MuCOE/Nodes/CustomizableObjectNodeModifierClipMorph.h"
+#include "MuCOE/UnrealEditorPortabilityHelpers.h"
+#include "PropertyCustomizationHelpers.h"
+#include "MuCO/CustomizableObjectPrivate.h"
+#include "Widgets/Input/SCheckBox.h"
+#include "Widgets/Input/STextComboBox.h"
+
+class UCustomizableObjectNodeObject;
+
+
+#define LOCTEXT_NAMESPACE "MeshClipMorphDetails"
+
+
+TSharedRef<IDetailCustomization> FCustomizableObjectNodeModifierClipMorphDetails::MakeInstance()
+{
+	return MakeShareable(new FCustomizableObjectNodeModifierClipMorphDetails);
+}
+
+
+void FCustomizableObjectNodeModifierClipMorphDetails::CustomizeDetails(IDetailLayoutBuilder& DetailBuilder)
+{
+	FCustomizableObjectNodeModifierBaseDetails::CustomizeDetails(DetailBuilder);
+
+	Node = nullptr;
+	DetailBuilderPtr = &DetailBuilder;
+
+	const IDetailsView* DetailsView = DetailBuilder.GetDetailsView();
+	if (DetailsView->GetSelectedObjects().Num())
+	{
+		Node = Cast<UCustomizableObjectNodeModifierClipMorph>(DetailsView->GetSelectedObjects()[0].Get());
+	}
+
+	IDetailCategoryBuilder& BlocksCategory = DetailBuilder.EditCategory("MeshToClipAndMorph");
+	DetailBuilder.HideProperty("BoneName");
+
+	IDetailCategoryBuilder& MeshClipParametersCategory = DetailBuilder.EditCategory("MeshClipParameters");
+	DetailBuilder.HideProperty("bInvertNormal");
+
+	TSharedPtr<IPropertyHandle> ReferenceSkeletonIndexPropertyHandle = DetailBuilder.GetProperty("ReferenceSkeletonIndex");
+	const FSimpleDelegate OnReferenceSkeletonIndexChangedDelegate = 
+		FSimpleDelegate::CreateSP(this, &FCustomizableObjectNodeModifierClipMorphDetails::OnReferenceSkeletonIndexChanged);
+	ReferenceSkeletonIndexPropertyHandle->SetOnPropertyValueChanged(OnReferenceSkeletonIndexChangedDelegate);
+
+	if (Node)
+	{
+		SkeletalMesh = nullptr;
+		TSharedPtr<FString> BoneToSelect;
+
+		UCustomizableObject* CustomizableObject = Cast<UCustomizableObject>(Node->GetCustomizableObjectGraph()->GetOuter());
+
+		if (CustomizableObject)
+		{
+			if (CustomizableObject->IsChildObject())
+			{
+				bool bMultipleBaseObjects;
+				UCustomizableObjectNodeObject* RootNode = GetRootNode(CustomizableObject, bMultipleBaseObjects);
+
+				if (RootNode && !bMultipleBaseObjects)
+				{
+					TArray<UCustomizableObject*> VisitedObjects;
+					UCustomizableObject* Parent = GetFullGraphRootObject(RootNode, VisitedObjects);
+
+					if (Parent)
+					{
+						SkeletalMesh = VisitedObjects.Last()->GetRefSkeletalMesh(Node->ReferenceSkeletonIndex);
+					}
+				}
+			}
+			else
+			{
+				SkeletalMesh = CustomizableObject->GetRefSkeletalMesh(Node->ReferenceSkeletonIndex);
+			}
+		}
+
+		if (SkeletalMesh)
+		{
+			BoneComboOptions.Empty();
+				
+			for (int32 i = 0; i < SkeletalMesh->GetRefSkeleton().GetRawBoneNum(); ++i)
+			{
+				FName BoneName = SkeletalMesh->GetRefSkeleton().GetBoneName(i);
+				BoneComboOptions.Add(MakeShareable(new FString(BoneName.ToString())));
+
+				if (BoneName == Node->BoneName)
+				{
+					BoneToSelect = BoneComboOptions.Last();
+				}
+			}
+
+            BoneComboOptions.Sort(CompareNames);
+
+			// Add them to the parent combo box
+			TSharedRef<IPropertyHandle> BoneProperty = DetailBuilder.GetProperty("BoneName");
+
+			BlocksCategory.AddCustomRow(LOCTEXT("ClipMorphDetails_BoneName", "Bone Name"))
+			[
+				SNew(SProperty, BoneProperty)
+				.ShouldDisplayName(false)
+				.CustomWidget()
+				[
+					SNew(SBorder)
+					.BorderImage(UE_MUTABLE_GET_BRUSH("NoBorder"))
+					.Padding(FMargin(0.0f, 0.0f, 10.0f, 0.0f))
+					.HAlign(EHorizontalAlignment::HAlign_Fill)
+					[
+						SNew(SHorizontalBox)
+						+SHorizontalBox::Slot().HAlign(HAlign_Left).VAlign(VAlign_Center)
+						[
+							SNew(STextBlock)
+							.Text( LOCTEXT("ClipMorphDetails_BoneName", "Bone Name"))
+							.Font(IDetailLayoutBuilder::GetDetailFont())
+						]
+						+SHorizontalBox::Slot().HAlign(HAlign_Fill)
+						[
+							SNew(STextComboBox)
+							.OptionsSource(&BoneComboOptions)
+							.InitiallySelectedItem(BoneToSelect)
+							.OnSelectionChanged(this, &FCustomizableObjectNodeModifierClipMorphDetails::OnBoneComboBoxSelectionChanged, BoneProperty)
+							.Font(IDetailLayoutBuilder::GetDetailFont())
+						]
+					]
+				]
+			];
+
+			MeshClipParametersCategory.AddCustomRow(LOCTEXT("ClipMorphDetails_PlaneNormal", "Invert plane normal"))
+			[
+				SNew(SProperty, BoneProperty)
+				.ShouldDisplayName(false)
+				.CustomWidget()
+				[
+					SNew(SBorder)
+					.BorderImage(UE_MUTABLE_GET_BRUSH("NoBorder"))
+					.Padding(FMargin(0.0f, 0.0f, 10.0f, 0.0f))
+					.HAlign(EHorizontalAlignment::HAlign_Fill)
+					[
+						SNew(SHorizontalBox)
+						+SHorizontalBox::Slot().HAlign(HAlign_Left).VAlign(VAlign_Center)
+						[
+							SNew(STextBlock)
+							.Text( LOCTEXT("ClipMorphDetails_PlaneNormal", "Invert plane normal"))
+						]
+						+SHorizontalBox::Slot().HAlign(HAlign_Left)
+						[
+							SNew(SCheckBox)
+							.OnCheckStateChanged(this, &FCustomizableObjectNodeModifierClipMorphDetails::OnInvertNormalCheckboxChanged)
+							.IsChecked(this, &FCustomizableObjectNodeModifierClipMorphDetails::GetInvertNormalCheckBoxState)
+							.ToolTipText(LOCTEXT("ClipMorphDetails_InvertNormal_Tooltip", "Invert normal direction of the clip plane"))
+						]
+					]
+				]
+			];
+		}
+	}
+	else
+	{
+		BlocksCategory.AddCustomRow(LOCTEXT("Node", "Node"))
+		[
+			SNew(STextBlock)
+			.Text(LOCTEXT("ClipMorphDetails_InvertNormal_NodeNotFound", "Node not found"))
+		];
+	}
+}
+
+
+FVector FindBoneLocation(int32 BoneIndex, USkeletalMesh* SkeletalMesh)
+{
+	const TArray<FTransform>& BoneArray = SkeletalMesh->GetRefSkeleton().GetRefBonePose();
+
+	int32 ParentIndex = BoneIndex;
+	FVector Location = FVector::ZeroVector;
+
+	while (ParentIndex >= 0)
+	{
+		Location = BoneArray[ParentIndex].TransformPosition(Location);
+		ParentIndex = SkeletalMesh->GetRefSkeleton().GetParentIndex(ParentIndex);
+	}
+
+	return Location;
+}
+
+
+void FCustomizableObjectNodeModifierClipMorphDetails::OnBoneComboBoxSelectionChanged(TSharedPtr<FString> Selection, ESelectInfo::Type SelectInfo, TSharedRef<IPropertyHandle> BoneProperty)
+{
+	for (int OptionIndex = 0; OptionIndex < BoneComboOptions.Num(); ++OptionIndex)
+	{
+		if (BoneComboOptions[OptionIndex] == Selection)
+		{
+			//TArray<FVector> AffectedBoneLocations;
+
+			//if (SkeletalMesh)
+			//{
+			//	const TArray<FTransform>& BoneArray = SkeletalMesh->RefSkeleton.GetRefBonePose();
+			//	int32 BoneIndex = SkeletalMesh->RefSkeleton.FindBoneIndex(FName(**Selection));
+			//	
+			//	if (BoneIndex >= 0)
+			//	{
+			//		// Get tranform from skeleton root to selected bone
+			//		int32 ParentIndex = BoneIndex;
+			//		FTransform RootToBoneTransform = FTransform::Identity;
+
+			//		while (ParentIndex >= 0)
+			//		{
+			//			RootToBoneTransform = RootToBoneTransform * BoneArray[ParentIndex];
+			//			ParentIndex = SkeletalMesh->RefSkeleton.GetParentIndex(ParentIndex);
+			//		}
+
+			//		// Now that we have the transform from root to the selected bone, go down from this bone to the end of the bone chain hanging from it building a bounding box
+			//		int32 CurrentBone = BoneIndex;
+
+			//		FVector SelectedBoneLocationWithTransform = RootToBoneTransform.TransformPosition(FVector::ZeroVector);
+			//		FVector SelectedBoneLocationFromSelection = BoneLocations[OptionIndex];
+			//		check(FVector::DistSquared(SelectedBoneLocationWithTransform, SelectedBoneLocationFromSelection) < 0.01f);
+
+			//		while (CurrentBone >= 0)
+			//		{
+			//			FVector CurrentBoneLocation = RootToBoneTransform.TransformPosition(FVector::ZeroVector);
+			//			AffectedBoneLocations.Add(CurrentBoneLocation);
+
+			//			check(FVector::DistSquared(CurrentBoneLocation, FindBoneLocation(CurrentBone, SkeletalMesh)) < 0.01f);
+
+			//			int32 NextBone = -1;
+
+			//			for (int32 i = 0; i < SkeletalMesh->RefSkeleton.GetRawBoneNum(); ++i)
+			//			{
+			//				if (SkeletalMesh->RefSkeleton.GetParentIndex(i) == CurrentBone)
+			//				{
+			//					NextBone = i;
+			//					break;
+			//				}
+			//			}
+
+			//			CurrentBone = NextBone;
+
+			//			if (CurrentBone >= 0)
+			//			{
+			//				RootToBoneTransform = BoneArray[CurrentBone] * RootToBoneTransform;
+			//			}
+			//		}
+			//	}
+			//}
+
+			FVector Location = FVector::ZeroVector;
+			FVector Direction = FVector::ForwardVector;
+
+			if (SkeletalMesh)
+			{
+				const TArray<FTransform>& BoneArray = SkeletalMesh->GetRefSkeleton().GetRefBonePose();
+				int32 ParentIndex = SkeletalMesh->GetRefSkeleton().FindBoneIndex(FName(**Selection));
+
+				FVector ChildLocation = FVector::ForwardVector;
+
+				for (int32 i = 0; i < SkeletalMesh->GetRefSkeleton().GetRawBoneNum(); ++i)
+				{
+					if (SkeletalMesh->GetRefSkeleton().GetParentIndex(i) == ParentIndex)
+					{
+						ChildLocation = BoneArray[i].TransformPosition(FVector::ZeroVector);
+						break;
+					}
+				}
+
+				//Direction = ChildLocation;
+
+				while (ParentIndex >= 0)
+				{
+					Location = BoneArray[ParentIndex].TransformPosition(Location);
+					ChildLocation = BoneArray[ParentIndex].TransformPosition(ChildLocation);
+					//Direction = BoneArray[ParentIndex].TransformVector(Direction);
+					ParentIndex = SkeletalMesh->GetRefSkeleton().GetParentIndex(ParentIndex);
+				}
+
+				Direction = (ChildLocation - Location).GetSafeNormal();
+			}
+			 
+			BoneProperty->SetValue(*BoneComboOptions[OptionIndex].Get());
+			
+			if (Node)
+			{
+				Node->Origin = Location;
+				Node->Normal = Direction;
+			}
+
+			return;
+		}
+	}
+
+	BoneProperty->SetValue(FString());
+}
+
+
+void FCustomizableObjectNodeModifierClipMorphDetails::OnInvertNormalCheckboxChanged(ECheckBoxState CheckBoxState)
+{
+	if (Node == nullptr)
+	{
+		return;
+	}
+
+	switch (CheckBoxState)
+	{
+		case ECheckBoxState::Checked:
+		{
+			Node->bInvertNormal = true;
+			Node->Normal *= -1.0f;
+			break;
+		}
+		case ECheckBoxState::Unchecked:
+		{
+			Node->bInvertNormal = false;
+			Node->Normal *= -1.0f;
+			break;
+		}
+	}
+
+	if (Node->bLocalStartOffset)
+	{
+		Node->StartOffset.Z *= -1;
+		Node->StartOffset.X *= -1;
+	}
+}
+
+
+ECheckBoxState FCustomizableObjectNodeModifierClipMorphDetails::GetInvertNormalCheckBoxState() const
+{
+	if (Node == nullptr)
+	{
+		return ECheckBoxState::Unchecked;
+	}
+
+	return Node->bInvertNormal ? ECheckBoxState::Checked : ECheckBoxState::Unchecked;
+}
+
+
+void FCustomizableObjectNodeModifierClipMorphDetails::OnReferenceSkeletonIndexChanged()
+{
+	if (Node)
+	{
+		Node->BoneName = FName();
+	}
+	DetailBuilderPtr->ForceRefreshDetails();
+}
+
+
+#undef LOCTEXT_NAMESPACE
diff --git a/Engine/Plugins/Mutable/Source/CustomizableObjectEditor/Private/MuCOE/Nodes/CustomizableObjectNodeModifierClipMorphDetails.h b/Engine/Plugins/Mutable/Source/CustomizableObjectEditor/Private/MuCOE/Nodes/CustomizableObjectNodeModifierClipMorphDetails.h
new file mode 100644
--- /dev/null
+++ b/Engine/Plugins/Mutable/Source/CustomizableObjectEditor/Private/MuCOE/Nodes/CustomizableObjectNodeModifierClipMorphDetails.h
@@ -0,0 +1,40 @@
+// Copyright Epic Games, Inc. All Rights Reserved.
+
+#pragma once
+
+#include "CustomizableObjectNodeModifierBaseDetails.h"
+#include "IDetailCustomization.h"
+#include "UObject/ObjectMacros.h"
+
+enum class ECheckBoxState : uint8;
+namespace ESelectInfo { enum Type : int; }
+
+class FString;
+class IDetailLayoutBuilder;
+class IPropertyHandle;
+
+
+class FCustomizableObjectNodeModifierClipMorphDetails : public FCustomizableObjectNodeModifierBaseDetails
+{
+public:
+	/** Makes a new instance of this detail layout class for a specific detail view requesting it */
+	static TSharedRef<IDetailCustomization> MakeInstance();
+
+	/** ILayoutDetails interface */
+	virtual void CustomizeDetails( IDetailLayoutBuilder& DetailBuilder ) override;
+
+private:
+	class UCustomizableObjectNodeModifierClipMorph* Node;
+	class IDetailLayoutBuilder* DetailBuilderPtr = nullptr;
+	TArray< TSharedPtr<FString> > BoneComboOptions;
+
+	UPROPERTY()
+	class USkeletalMesh* SkeletalMesh;
+
+	void OnBoneComboBoxSelectionChanged(TSharedPtr<FString> Selection, ESelectInfo::Type SelectInfo, TSharedRef<IPropertyHandle> BoneProperty);
+
+	void OnInvertNormalCheckboxChanged(ECheckBoxState CheckBoxState);
+	ECheckBoxState GetInvertNormalCheckBoxState() const;
+
+	void OnReferenceSkeletonIndexChanged();
+};
diff --git a/Engine/Plugins/Mutable/Source/CustomizableObjectEditor/Private/MuCOE/Nodes/CustomizableObjectNodeModifierClipWithMeshDetails.cpp b/Engine/Plugins/Mutable/Source/CustomizableObjectEditor/Private/MuCOE/Nodes/CustomizableObjectNodeModifierClipWithMeshDetails.cpp
new file mode 100644
--- /dev/null
+++ b/Engine/Plugins/Mutable/Source/CustomizableObjectEditor/Private/MuCOE/Nodes/CustomizableObjectNodeModifierClipWithMeshDetails.cpp
@@ -0,0 +1,337 @@
+// Copyright Epic Games, Inc. All Rights Reserved.
+
+#include "MuCOE/Nodes/CustomizableObjectNodeModifierClipWithMeshDetails.h"
+
+#include "DetailLayoutBuilder.h"
+#include "IDetailsView.h"
+#include "MuCOE/CustomizableObjectEditorUtilities.h"
+#include "MuCOE/GraphTraversal.h"
+#include "MuCOE/UnrealEditorPortabilityHelpers.h"
+#include "MuCOE/Nodes/CustomizableObjectNodeMaterial.h"
+#include "MuCOE/Nodes/CustomizableObjectNodeModifierClipWithMesh.h"
+#include "MuCOE/Nodes/CustomizableObjectNodeStaticMesh.h"
+#include "PropertyCustomizationHelpers.h"
+#include "ScopedTransaction.h"
+#include "MuCO/CustomizableObjectPrivate.h"
+#include "Widgets/Input/SButton.h"
+#include "Widgets/Input/SEditableTextBox.h"
+#include "Widgets/Input/STextComboBox.h"
+
+class IPropertyHandle;
+
+
+#define LOCTEXT_NAMESPACE "CustomizableObjectNodeModifierClipWithMeshDetails"
+
+
+TSharedRef<IDetailCustomization> FCustomizableObjectNodeModifierClipWithMeshDetails::MakeInstance()
+{
+	return MakeShareable(new FCustomizableObjectNodeModifierClipWithMeshDetails);
+}
+
+
+void FCustomizableObjectNodeModifierClipWithMeshDetails::CustomizeDetails(IDetailLayoutBuilder& DetailBuilder)
+{
+	FCustomizableObjectNodeModifierBaseDetails::CustomizeDetails(DetailBuilder);
+
+	Node = nullptr;
+	DetailBuilderPtr = &DetailBuilder;
+
+	const IDetailsView* DetailsView = DetailBuilder.GetDetailsView();
+	if (DetailsView->GetSelectedObjects().Num())
+	{
+		Node = Cast<UCustomizableObjectNodeModifierClipWithMesh>(DetailsView->GetSelectedObjects()[0].Get());
+	}
+
+	IDetailCategoryBuilder& CustomizableObjectToClipCategory = DetailBuilder.EditCategory("CustomizableObjectToClip");
+
+	if(Node != nullptr)
+	{
+		SelectedCO = Node->CustomizableObjectToClipWith;
+
+		TSharedRef<IPropertyHandle> ParentObjectProperty = DetailBuilder.GetProperty("ParentObject");
+
+		// Hidding unnecessary properties
+		DetailBuilder.HideProperty("CustomizableObjectToClipWith");
+		DetailBuilder.HideProperty("ArrayMaterialNodeToClipWithID");
+
+		if (const UEdGraphPin* ConnectedPin = FollowInputPin(*Node->ClipMeshPin());
+			ConnectedPin && !Cast<UCustomizableObjectNodeStaticMesh>(ConnectedPin->GetOwningNode()))
+		{
+			DetailBuilder.HideProperty("Transform");
+		}
+
+		ArrayMaterialNodeOption.Empty();
+		ArrayMaterialNodeName.Empty();
+
+		TArray<FGuid> ArrayMaterialNodeGuid;
+		TSharedPtr<FString> SelectedMethod;
+
+		CustomizableObjectToClipCategory.AddCustomRow(LOCTEXT("FCustomizableObjectNodeModifierClipWithMeshDetails", "Blocks"))
+		[
+			SNew(SHorizontalBox)
+			+ SHorizontalBox::Slot()
+			.Padding(FMargin(0.0f, 3.0f, 0.0f, 10.0f))
+			.FillWidth(0.15f)
+			[
+				SNew(STextBlock)
+				.Text(LOCTEXT("CustomizableObjectText", "Object:"))
+			]
+			+ SHorizontalBox::Slot()
+			.FillWidth(0.85f)
+			[
+				SNew(SObjectPropertyEntryBox)
+				.AllowedClass(UCustomizableObject::StaticClass())
+				.OnObjectChanged(this, &FCustomizableObjectNodeModifierClipWithMeshDetails::ClipperCustomizableObjectSelectionChanged)
+				.ObjectPath(SelectedCO != nullptr ? SelectedCO->GetPathName() : "")
+				.ForceVolatile(true)
+			]
+		];
+
+		SelectAMethod = MakeShareable(new FString("Select a Method"));
+
+		// Filling ComboBox options
+		ClippingMethods.Empty();
+		ClippingMethods.Add(MakeShareable(new FString("Mesh Section")));
+		ClippingMethods.Add(MakeShareable(new FString("Tags")));
+		ClippingMethods.Add(MakeShareable(new FString("Mesh Section & Tags")));
+
+		// Initializing the inital value of the ComboBox
+		SelectedMethod = SetInitialClippingMethod();
+		
+		// This flag is now up to date and can be used to hide ui.
+		if (!Node->bUseTags)
+		{
+			DetailBuilder.HideProperty("Tags");
+		}
+
+		// ComboBox to select the Clipping method (material or/and tags)
+		CustomizableObjectToClipCategory.AddCustomRow(LOCTEXT("FCustomizableObjectNodeModifierClipWithMeshDetails", "Blocks"))
+		[
+			SNew(SHorizontalBox)
+			+SHorizontalBox::Slot()
+			.AutoWidth()
+			.Padding(FMargin(0.0f, 8.0f, 0.0f, 5.0f))
+			[
+				SNew(STextBlock)
+				.Text(LOCTEXT("ClippingMethod", "Clipping Method:"))
+			]
+			
+			+ SHorizontalBox::Slot()
+			.AutoWidth()
+			.Padding(FMargin(4.0f, 5.0f, 0.0f, 5.0f))
+			[
+				SNew(STextComboBox)
+				.OptionsSource(&ClippingMethods)
+				.InitiallySelectedItem(SelectedMethod)
+				.OnSelectionChanged(this, &FCustomizableObjectNodeModifierClipWithMeshDetails::OnClippingMethodComboBoxSelectionChanged)
+			]
+		];
+
+
+		// Filling materials ComboBox with all the possible materials
+		int32 NumElementBeforeAdd;
+		int32 NumElementAfterAdd;
+		TSharedPtr<FString> ItemToSelect;
+		TArray<UCustomizableObjectNodeMaterial*> ArrayMaterialNodes;
+		if (SelectedCO != nullptr)
+		{
+			SelectedCO->GetPrivate()->GetSource()->GetNodesOfClass<UCustomizableObjectNodeMaterial>(ArrayMaterialNodes);
+		}
+
+		for (const UCustomizableObjectNodeMaterial* MaterialNode : ArrayMaterialNodes)
+		{
+			if (MaterialNode->GetMaterial())
+			{
+				NumElementBeforeAdd = ArrayMaterialNodeName.Num();
+				ArrayMaterialNodeName.AddUnique(MaterialNode->GetMaterial()->GetName());
+				NumElementAfterAdd = ArrayMaterialNodeName.Num();
+
+				if (NumElementBeforeAdd != NumElementAfterAdd)
+				{
+					ArrayMaterialNodeGuid.Add(MaterialNode->NodeGuid);
+				}
+			}
+		}
+
+		// Adding a none option
+		ArrayMaterialNodeOption.Add(MakeShareable(new FString("None")));
+		ItemToSelect = ArrayMaterialNodeOption.Last();
+
+		// Initializing the Selected material if it was already selected
+		const int32 MaxIndex = ArrayMaterialNodeName.Num();
+		for (int32 i = 0; i < MaxIndex; ++i)
+		{
+			ArrayMaterialNodeOption.Add(MakeShareable(new FString(ArrayMaterialNodeName[i])));
+
+			if (Node->ArrayMaterialNodeToClipWithID.Find(ArrayMaterialNodeGuid[i]) != INDEX_NONE)
+			{
+				ItemToSelect = ArrayMaterialNodeOption.Last();
+			}
+		}
+
+        ArrayMaterialNodeOption.Sort(CompareNames);
+
+		// Widget that shows the Tags widget or/and the materials ComboBox
+		if (Node->bUseMaterials)
+		{
+			CustomizableObjectToClipCategory.AddCustomRow(LOCTEXT("FCustomizableObjectNodeModifierClipWithMeshDetails", "Blocks"))
+			[
+				SNew(SVerticalBox)
+				+SVerticalBox::Slot()
+				.AutoHeight()
+				[
+					SNew(SHorizontalBox)
+					+ SHorizontalBox::Slot()
+					.FillWidth(0.15f)
+					.Padding(0.0f, 3.0f, 0.0f, 0.0f)
+					[
+						SNew(STextBlock)
+						.Text(LOCTEXT("MaterialsText", "Material:"))
+					]
+					+ SHorizontalBox::Slot()
+					.FillWidth(0.85f)
+					.Padding(2.0f, 0.0f, 0.0f, 0.0f)
+					[
+						SNew(STextComboBox)
+						.OptionsSource(&ArrayMaterialNodeOption)
+						.InitiallySelectedItem(ItemToSelect)
+						.OnSelectionChanged(this, &FCustomizableObjectNodeModifierClipWithMeshDetails::OnMeshClipWithMeshNodeComboBoxSelectionChanged)
+					]
+				]
+			];
+		}
+	}
+	else
+	{
+		CustomizableObjectToClipCategory.AddCustomRow(LOCTEXT("Node", "Node"))
+			[
+				SNew(STextBlock)
+				.Text(LOCTEXT("Node not found", "Node not found"))
+			];
+	}
+}
+
+
+void FCustomizableObjectNodeModifierClipWithMeshDetails::ClipperCustomizableObjectSelectionChanged(const FAssetData& AssetData)
+{
+	UCustomizableObject* NewClipperCO = Cast<UCustomizableObject>(AssetData.GetAsset());
+
+	if (Node != nullptr)
+	{
+		Node->CustomizableObjectToClipWith = NewClipperCO;
+		Node->ArrayMaterialNodeToClipWithID.Empty();
+	}
+
+	if (DetailBuilderPtr)
+	{
+		DetailBuilderPtr->ForceRefreshDetails();
+	}
+}
+
+
+TSharedPtr<FString> FCustomizableObjectNodeModifierClipWithMeshDetails::SetInitialClippingMethod()
+{
+	TSharedPtr<FString> InitialMethod = SelectAMethod;
+
+	if (!Node->bUseMaterials && !Node->bUseTags)
+	{
+		if (Node->ArrayMaterialNodeToClipWithID.Num() > 0)
+		{
+			InitialMethod = ClippingMethods[0];
+			Node->bUseMaterials = true;
+		}
+
+		if (Node->RequiredTags.Num() > 0)
+		{
+			InitialMethod = ClippingMethods[1];
+			Node->bUseTags = true;
+		}
+
+		if (Node->bUseMaterials && Node->bUseTags)
+		{
+			InitialMethod = ClippingMethods[2];
+		}
+	}
+	else
+	{
+		if (Node->bUseMaterials && !Node->bUseTags)
+		{
+			InitialMethod = ClippingMethods[0];
+		}
+		else if (!Node->bUseMaterials && Node->bUseTags)
+		{
+			InitialMethod = ClippingMethods[1];
+		}
+		else
+		{
+			InitialMethod = ClippingMethods[2];
+		}
+	}
+
+	return InitialMethod;
+}
+
+
+void FCustomizableObjectNodeModifierClipWithMeshDetails::OnClippingMethodComboBoxSelectionChanged(TSharedPtr<FString> Selection, ESelectInfo::Type SelectInfo)
+{
+	if (Selection.IsValid())
+	{
+		if (*Selection == "Tags")
+		{
+			Node->bUseMaterials = false;
+			Node->bUseTags = true;
+			Node->ArrayMaterialNodeToClipWithID.Empty();
+		}
+		else if(*Selection == "Material")
+		{
+			Node->bUseMaterials = true;
+			Node->bUseTags = false;
+			Node->RequiredTags.Empty();
+		}
+		else
+		{
+			Node->bUseMaterials = true;
+			Node->bUseTags = true;
+		}
+
+		DetailBuilderPtr->ForceRefreshDetails();
+	}	
+}
+
+
+void FCustomizableObjectNodeModifierClipWithMeshDetails::OnMeshClipWithMeshNodeComboBoxSelectionChanged(TSharedPtr<FString> Selection, ESelectInfo::Type SelectInfo)
+{
+	if (Selection.IsValid() && (SelectedCO != nullptr))
+	{
+		TArray<UCustomizableObjectNodeMaterial*> ArrayMaterialNode;
+		SelectedCO->GetPrivate()->GetSource()->GetNodesOfClass<UCustomizableObjectNodeMaterial>(ArrayMaterialNode);
+
+		const FScopedTransaction Transaction(LOCTEXT("ChangedClipMeshWithMeshMaterialTransaction", "Changed Clip Mesh With Mesh Material"));
+		Node->Modify();
+		Node->ArrayMaterialNodeToClipWithID.Empty();
+		for (const UCustomizableObjectNodeMaterial* MaterialNode : ArrayMaterialNode)
+		{
+			if (MaterialNode && MaterialNode->GetMaterial() && *Selection == MaterialNode->GetMaterial()->GetName())
+			{
+				Node->ArrayMaterialNodeToClipWithID.Add(MaterialNode->NodeGuid);
+			}
+		}
+	}
+}
+
+
+FReply FCustomizableObjectNodeModifierClipWithMeshDetails::OnAddTagPressed()
+{
+	if (Node)
+	{
+		Node->RequiredTags.Add("");
+		DetailBuilderPtr->ForceRefreshDetails();
+
+		return FReply::Handled();
+	}
+
+	return FReply::Unhandled();
+}
+
+
+#undef LOCTEXT_NAMESPACE
diff --git a/Engine/Plugins/Mutable/Source/CustomizableObjectEditor/Private/MuCOE/Nodes/CustomizableObjectNodeModifierClipWithMeshDetails.h b/Engine/Plugins/Mutable/Source/CustomizableObjectEditor/Private/MuCOE/Nodes/CustomizableObjectNodeModifierClipWithMeshDetails.h
new file mode 100644
--- /dev/null
+++ b/Engine/Plugins/Mutable/Source/CustomizableObjectEditor/Private/MuCOE/Nodes/CustomizableObjectNodeModifierClipWithMeshDetails.h
@@ -0,0 +1,65 @@
+// Copyright Epic Games, Inc. All Rights Reserved.
+
+#pragma once
+
+#include "CustomizableObjectNodeModifierBaseDetails.h"
+#include "IDetailCustomization.h"
+#include "Widgets/SCompoundWidget.h"
+
+namespace ESelectInfo { enum Type : int; }
+namespace ETextCommit { enum Type : int; }
+
+class FText;
+class IDetailLayoutBuilder;
+class SHorizontalBox;
+class SVerticalBox;
+class UCustomizableObjectNodeModifierClipWithMesh;
+struct FAssetData;
+
+
+class FCustomizableObjectNodeModifierClipWithMeshDetails : public FCustomizableObjectNodeModifierBaseDetails
+{
+public:
+	/** Makes a new instance of this detail layout class for a specific detail view requesting it */
+	static TSharedRef<IDetailCustomization> MakeInstance();
+
+	/** ILayoutDetails interface */
+	virtual void CustomizeDetails( IDetailLayoutBuilder& DetailBuilder ) override;
+
+	/** Callback when new CO is selected */
+	void ClipperCustomizableObjectSelectionChanged(const FAssetData& AssetData);
+
+private:
+	/** Callback when the clipping method is selected */
+	void OnClippingMethodComboBoxSelectionChanged(TSharedPtr<FString> Selection, ESelectInfo::Type SelectInfo);
+
+	/** Callback when new material name is selected */
+	void OnMeshClipWithMeshNodeComboBoxSelectionChanged(TSharedPtr<FString> Selection, ESelectInfo::Type SelectInfo);
+
+	/** Initializes the text of the method combobox and the view to display */
+	TSharedPtr<FString> SetInitialClippingMethod();
+
+	/** The node for which details are being customized */
+	class UCustomizableObjectNodeModifierClipWithMesh* Node;
+
+	/** Pointer to the builder passed by CustomizeDetails method */
+	IDetailLayoutBuilder* DetailBuilderPtr = nullptr;
+
+	/** References to the names of all possible methods to select the material */
+	TArray<TSharedPtr<FString>> ClippingMethods;
+
+	/** References to the names of all possible selectable material nodes in SelectedCO */
+	TArray<TSharedPtr<FString>> ArrayMaterialNodeOption;
+
+	/** To cover all LODs from a single selection, only unique names of materials are considered */
+	TArray<FString> ArrayMaterialNodeName;
+
+	/** Selected CO in the details tab of the Node node, for which to select a material that will be clipped by Node node */
+	class UCustomizableObject* SelectedCO;
+
+	/** Pointer to a FString to indicate that a method must be selected */
+	TSharedPtr<FString> SelectAMethod;
+
+	/** Callback to add a tag to the arry of tags */
+	FReply OnAddTagPressed();
+};
diff --git a/Engine/Plugins/Mutable/Source/CustomizableObjectEditor/Private/MuCOE/Nodes/CustomizableObjectNodeModifierEditMeshSection.cpp b/Engine/Plugins/Mutable/Source/CustomizableObjectEditor/Private/MuCOE/Nodes/CustomizableObjectNodeModifierEditMeshSection.cpp
new file mode 100644
--- /dev/null
+++ b/Engine/Plugins/Mutable/Source/CustomizableObjectEditor/Private/MuCOE/Nodes/CustomizableObjectNodeModifierEditMeshSection.cpp
@@ -0,0 +1,312 @@
+// Copyright Epic Games, Inc. All Rights Reserved.
+
+#include "MuCOE/Nodes/CustomizableObjectNodeModifierEditMeshSection.h"
+
+#include "MuCOE/Nodes/CustomizableObjectNodeModifierExtendMeshSection.h"
+#include "MuCOE/CustomizableObjectEditor_Deprecated.h"
+#include "MuCOE/CustomizableObjectLayout.h"
+#include "MuCOE/EdGraphSchema_CustomizableObject.h"
+#include "MuCOE/CustomizableObjectPin.h"
+#include "MuCO/CustomizableObjectCustomVersion.h"
+#include "MaterialTypes.h"
+
+class FCustomizableObjectNodeParentedMaterial;
+class UCustomizableObjectNodeRemapPins;
+
+#define LOCTEXT_NAMESPACE "CustomizableObjectEditor"
+
+
+void UCustomizableObjectNodeModifierEditMeshSection::AllocateDefaultPins(UCustomizableObjectNodeRemapPins* RemapPins)
+{
+	const UEdGraphSchema_CustomizableObject* Schema = GetDefault<UEdGraphSchema_CustomizableObject>();
+	
+	const int32 NumImages = GetNumParameters(EMaterialParameterType::Texture);
+	for (int32 ImageIndex = 0; ImageIndex < NumImages; ++ImageIndex)
+	{
+		{
+			UCustomizableObjectNodeEditMaterialPinEditImageData* PinEditImageData = NewObject<UCustomizableObjectNodeEditMaterialPinEditImageData>(this);
+			PinEditImageData->ImageParamId = GetParameterId(EMaterialParameterType::Texture, ImageIndex);
+			
+			const FName ImageName = GetParameterName(EMaterialParameterType::Texture, ImageIndex);
+			UEdGraphPin* PinImage = CustomCreatePin(EGPD_Input, Schema->PC_Image, ImageName, PinEditImageData);
+			PinImage->bHidden = true;
+			PinImage->bDefaultValueIsIgnored = true;
+
+			PinsParameterMap.Add(PinEditImageData->ImageParamId, FEdGraphPinReference(PinImage));
+
+			FString PinMaskName = ImageName.ToString() + FString(TEXT(" Mask"));
+			UEdGraphPin* PinMask = CustomCreatePin(EGPD_Input, Schema->PC_Image, *PinMaskName);
+			PinMask->bHidden = true;
+			PinMask->bDefaultValueIsIgnored = true;
+			
+			PinEditImageData->PinMask = FEdGraphPinReference(PinMask);
+		}
+	}
+
+	CustomCreatePin(EGPD_Output, Schema->PC_Modifier, TEXT("Modifier"));
+}
+
+
+const UEdGraphPin* UCustomizableObjectNodeModifierEditMeshSection::GetUsedImageMaskPin(const FNodeMaterialParameterId& ImageId) const
+{
+	if (const UEdGraphPin* Pin = GetUsedImagePin(ImageId))
+	{
+		UCustomizableObjectNodeEditMaterialPinEditImageData& PinData = GetPinData<UCustomizableObjectNodeEditMaterialPinEditImageData>(*Pin);
+		return PinData.PinMask.Get();
+	}
+
+	return nullptr;
+}
+
+
+bool UCustomizableObjectNodeModifierEditMeshSection::IsSingleOutputNode() const
+{
+	return true;
+}
+
+
+bool UCustomizableObjectNodeModifierEditMeshSection::CustomRemovePin(UEdGraphPin& Pin)
+{
+	for (TMap<FNodeMaterialParameterId, FEdGraphPinReference>::TIterator It = PinsParameterMap.CreateIterator(); It; ++It)
+	{
+		if (It.Value().Get() == &Pin) // We could improve performance if FEdGraphPinReference exposed the pin id.
+		{
+			It.RemoveCurrent();
+			break;
+		}
+	}
+
+	return Super::CustomRemovePin(Pin);
+}
+
+
+bool UCustomizableObjectNodeModifierEditMeshSection::HasPinViewer() const
+{
+	return true;
+}
+
+
+void UCustomizableObjectNodeModifierEditMeshSection::BackwardsCompatibleFixup(int32 CustomizableObjectCustomVersion)
+{
+	Super::BackwardsCompatibleFixup(CustomizableObjectCustomVersion);
+
+	// Convert deprecated node index list to the node id list.
+	if (CustomizableObjectCustomVersion == FCustomizableObjectCustomVersion::PostLoadToCustomVersion
+		&& BlockIds_DEPRECATED.Num() < Blocks_DEPRECATED.Num())
+	{
+		UCustomizableObjectNodeMaterialBase* ParentMaterialNode = GetCustomizableObjectExternalNode<UCustomizableObjectNodeMaterialBase>(ParentMaterialObject_DEPRECATED.Get(), ParentMaterialNodeId_DEPRECATED);
+
+		TArray<UCustomizableObjectLayout*> Layouts = ParentMaterialNode->GetLayouts();
+
+		if (!Layouts.IsValidIndex(ParentLayoutIndex))
+		{
+			UE_LOG(LogMutable, Warning, TEXT("[%s] UCustomizableObjectNodeModifierEditMeshSection refers to an invalid texture layout index %d. Parent node has %d layouts."), 
+				*GetOutermost()->GetName(), ParentLayoutIndex, Layouts.Num());
+		}
+		else
+		{
+			UCustomizableObjectLayout* ParentLayout = Layouts[ParentLayoutIndex];
+
+			if (Cast<UCustomizableObjectNodeMaterial>(ParentMaterialNode))
+			{
+				for (int32 IndexIndex = BlockIds_DEPRECATED.Num(); IndexIndex < Blocks_DEPRECATED.Num(); ++IndexIndex)
+				{
+					int32 BlockIndex = Blocks_DEPRECATED[IndexIndex];
+					if (ParentLayout->Blocks.IsValidIndex(BlockIndex) )
+					{
+						const FGuid Id = ParentLayout->Blocks[BlockIndex].Id;
+						if (Id.IsValid())
+						{
+							BlockIds_DEPRECATED.Add(Id);
+						}
+						else
+						{
+							UE_LOG(LogMutable, Warning, TEXT("[%s] UCustomizableObjectNodeModifierEditMeshSection refers to an valid layout block %d but that block doesn't have an id."),
+								*GetOutermost()->GetName(), BlockIndex );
+						}
+					}
+					else
+					{
+						UE_LOG(LogMutable, Warning, TEXT("[%s] UCustomizableObjectNodeModifierEditMeshSection refers to an invalid layout block index %d. Parent node has %d blocks."),
+							*GetOutermost()->GetName(), BlockIndex, ParentLayout->Blocks.Num());
+					}
+				}
+			}
+		}
+	}
+	
+	// Convert deprecated node id list to absolute rect list.
+	if (CustomizableObjectCustomVersion == FCustomizableObjectCustomVersion::UseUVRects)
+	{
+		// If we are here, it means this node was loaded from a version that didn't have it's own layout.
+		check(Layout->Blocks.IsEmpty());
+
+		UCustomizableObjectNodeMaterialBase* ParentMaterialNode = GetCustomizableObjectExternalNode<UCustomizableObjectNodeMaterialBase>(ParentMaterialObject_DEPRECATED.Get(), ParentMaterialNodeId_DEPRECATED);
+
+		TArray<UCustomizableObjectLayout*> ParentLayouts = ParentMaterialNode->GetLayouts();
+
+		if (!ParentLayouts.IsValidIndex(ParentLayoutIndex))
+		{
+			UE_LOG(LogMutable, Warning, TEXT("[%s] UCustomizableObjectNodeModifierEditMeshSection refers to an invalid texture layout index %d. Parent node has %d layouts."),
+				*GetOutermost()->GetName(), ParentLayoutIndex, ParentLayouts.Num());
+		}
+		else
+		{
+			UCustomizableObjectLayout* ParentLayout = ParentLayouts[ParentLayoutIndex];
+			FIntPoint GridSize = ParentLayout->GetGridSize();
+
+			Layout->SetGridSize(GridSize);
+
+			if (Cast<UCustomizableObjectNodeMaterial>(ParentMaterialNode))
+			{
+				for (const FGuid& BlockId : BlockIds_DEPRECATED)
+				{
+					bool bSkipBlock = false;
+					for (const FCustomizableObjectLayoutBlock& ExistingBlock : Layout->Blocks)
+					{
+						if (ExistingBlock.Id == BlockId)
+						{
+							bSkipBlock = true;
+							UE_LOG(LogMutable, Log, TEXT("[%s] UCustomizableObjectNodeModifierEditMeshSection has a duplicated layout block id. One has been ignored during version upgrade."), *GetOutermost()->GetName());
+							break;
+						}
+					}
+
+					if (bSkipBlock)
+					{
+						continue;
+					}
+
+					bool bFoundInParent = false;
+					for (const FCustomizableObjectLayoutBlock& ParentBlock : ParentLayout->Blocks)
+					{
+						if (ParentBlock.Id == BlockId)
+						{
+							bFoundInParent = true;
+
+							FCustomizableObjectLayoutBlock NewBlock;
+							NewBlock = ParentBlock;
+
+							// Clear some unnecessary data.
+							NewBlock.bReduceBothAxes = false;
+							NewBlock.bReduceByTwo = false;
+							NewBlock.Priority = 0;
+
+							Layout->Blocks.Add(NewBlock);
+							break;
+						}
+					}
+
+					if (!bFoundInParent)
+					{
+						UE_LOG(LogMutable, Warning, TEXT("[%s] UCustomizableObjectNodeModifierEditMeshSection refers to and invalid layout block. It has been ignored during version upgrade."), *GetOutermost()->GetName());
+					}
+				}
+			}
+		}
+	}
+
+	if (CustomizableObjectCustomVersion == FCustomizableObjectCustomVersion::AutomaticNodeMaterial)
+	{
+		UCustomizableObjectNodeMaterialBase* ParentMaterial = GetCustomizableObjectExternalNode<UCustomizableObjectNodeMaterialBase>(ParentMaterialObject_DEPRECATED.Get(), ParentMaterialNodeId_DEPRECATED);
+
+		if (ParentMaterial)
+		{
+			for (const FCustomizableObjectNodeEditMaterialImage& Image : Images_DEPRECATED)
+			{
+				const UEdGraphPin* ImagePin = FindPin(Image.Name);
+				const UEdGraphPin* PinMask = FindPin(Image.Name + " Mask");
+				if (!ImagePin || !PinMask)
+				{
+					continue;
+				}
+				
+				UCustomizableObjectNodeEditMaterialPinEditImageData*  PinEditImageData = NewObject<UCustomizableObjectNodeEditMaterialPinEditImageData>(this);
+				PinEditImageData->ImageId_DEPRECATED = FGuid::NewGuid();
+				PinEditImageData->PinMask = PinMask;
+				
+				// Search for the Image Id the Edit pin was referring to.
+				const int32 NumImages = ParentMaterial->GetNumParameters(EMaterialParameterType::Texture);
+				for (int32 ImageIndex = 0; ImageIndex < NumImages; ++ImageIndex)
+				{
+					if (ParentMaterial->GetParameterName(EMaterialParameterType::Texture, ImageIndex).ToString() == Image.Name)
+					{
+						PinEditImageData->ImageId_DEPRECATED = ParentMaterial->GetParameterId(EMaterialParameterType::Texture, ImageIndex).ParameterId;
+						break;
+					}
+				}
+				
+				AddPinData(*ImagePin, *PinEditImageData);
+			}
+		}
+		
+		Images_DEPRECATED.Empty();
+	}
+
+	// Fill PinsParameter.
+	if (CustomizableObjectCustomVersion == FCustomizableObjectCustomVersion::AutomaticNodeMaterialPerformanceBug) // || CustomizableObjectCustomVersion < FCustomizableObjectCustomVersion::AutomaticNodeMaterialPerformance
+	{
+		for (const UEdGraphPin* Pin : GetAllNonOrphanPins())
+		{
+			if (const UCustomizableObjectNodeEditMaterialPinEditImageData* PinData = Cast<UCustomizableObjectNodeEditMaterialPinEditImageData>(GetPinData(*Pin)))
+			{
+				PinsParameter_DEPRECATED.Add(PinData->ImageId_DEPRECATED, FEdGraphPinReference(Pin));
+			}
+		}
+	}
+
+	if (CustomizableObjectCustomVersion == FCustomizableObjectCustomVersion::ConvertEditAndExtendToModifiers)
+	{
+		// Look for the parent material and set it as the modifier reference material 
+
+		PinsParameterMap = PinsParameterMap_DEPRECATED;
+		PinsParameterMap_DEPRECATED.Empty();
+
+		UCustomizableObjectNode* ParentNode = GetCustomizableObjectExternalNode<UCustomizableObjectNode>(ParentMaterialObject_DEPRECATED.Get(), ParentMaterialNodeId_DEPRECATED);
+
+		if (UCustomizableObjectNodeMaterial* MaterialParentNode = Cast<UCustomizableObjectNodeMaterial>(ParentNode))
+		{
+			ReferenceMaterial = MaterialParentNode->GetMaterial();
+		}
+		else if (UCustomizableObjectNodeModifierExtendMeshSection* ExtendParentNode = Cast<UCustomizableObjectNodeModifierExtendMeshSection>(ParentNode))
+		{
+			ReferenceMaterial = ExtendParentNode->ReferenceMaterial;
+		}
+		else
+		{
+			// Conversion failed?
+			ensure(false);
+			UE_LOG(LogMutable, Warning, TEXT("[%s] UCustomizableObjectNodeModifierExtendMeshSection version upgrade failed."), *GetOutermost()->GetName());
+		}
+
+		ReconstructNode();
+	}
+}
+
+
+FText UCustomizableObjectNodeModifierEditMeshSection::GetNodeTitle(ENodeTitleType::Type TitleType) const
+{
+	return LOCTEXT("Edit_MeshSection", "Edit Mesh Section");
+}
+
+
+FString UCustomizableObjectNodeModifierEditMeshSection::GetRefreshMessage() const
+{
+	return TEXT("Source material has changed, texture channels might have been added, removed or renamed. Please refresh the parent material node to reflect those changes.");
+}
+
+
+FText UCustomizableObjectNodeModifierEditMeshSection::GetTooltipText() const
+{
+	return LOCTEXT("Edit_Material_Tooltip", "Modify the texture parameters of an ancestor's material partially or completely.");
+}
+
+
+void UCustomizableObjectNodeModifierEditMeshSection::SetLayoutIndex(const int32 LayoutIndex)
+{
+	ParentLayoutIndex = LayoutIndex;
+}
+
+
+#undef LOCTEXT_NAMESPACE
+
diff --git a/Engine/Plugins/Mutable/Source/CustomizableObjectEditor/Private/MuCOE/Nodes/CustomizableObjectNodeModifierExtendMeshSection.cpp b/Engine/Plugins/Mutable/Source/CustomizableObjectEditor/Private/MuCOE/Nodes/CustomizableObjectNodeModifierExtendMeshSection.cpp
new file mode 100644
--- /dev/null
+++ b/Engine/Plugins/Mutable/Source/CustomizableObjectEditor/Private/MuCOE/Nodes/CustomizableObjectNodeModifierExtendMeshSection.cpp
@@ -0,0 +1,215 @@
+// Copyright Epic Games, Inc. All Rights Reserved.
+
+#include "MuCOE/Nodes/CustomizableObjectNodeModifierExtendMeshSection.h"
+
+#include "MaterialTypes.h"
+#include "MuCO/CustomizableObjectCustomVersion.h"
+#include "MuCOE/CustomizableObjectEditor_Deprecated.h"
+#include "MuCOE/CustomizableObjectGraph.h"
+#include "MuCOE/EdGraphSchema_CustomizableObject.h"
+#include "MuCOE/GraphTraversal.h"
+#include "MuCOE/Nodes/CustomizableObjectNodeCopyMaterial.h"
+#include "MuCOE/Nodes/CustomizableObjectNodeSkeletalMesh.h"
+#include "MuCOE/Nodes/CustomizableObjectNodeTable.h"
+#include "MuCOE/CustomizableObjectPin.h"
+
+class UCustomizableObjectNode;
+class UCustomizableObjectNodeRemapPins;
+
+#define LOCTEXT_NAMESPACE "CustomizableObjectEditor"
+
+
+void UCustomizableObjectNodeModifierExtendMeshSection::BackwardsCompatibleFixup(int32 CustomizableObjectCustomVersion)
+{
+	Super::BackwardsCompatibleFixup(CustomizableObjectCustomVersion);
+	
+	if (CustomizableObjectCustomVersion == FCustomizableObjectCustomVersion::ExtendMaterialRemoveImages)
+	{
+		const UCustomizableObjectNodeMaterialBase* ParentMaterial = GetCustomizableObjectExternalNode<UCustomizableObjectNodeMaterialBase>(ParentMaterialObject_DEPRECATED.Get(), ParentMaterialNodeId_DEPRECATED);
+		if (ParentMaterial)
+		{
+			for (const FCustomizableObjectNodeExtendMaterialImage& Image : Images_DEPRECATED)
+			{
+				const UEdGraphPin* ImagePin = FindPin(Image.Name);
+				if (!ImagePin)
+				{
+					continue;
+				}
+
+				FGuid ImageId = FGuid::NewGuid();
+
+				// Search for the Image Id the Extend pin was referring to.
+				const int32 NumImages = ParentMaterial->GetNumParameters(EMaterialParameterType::Texture);
+				for (int32 ImageIndex = 0; ImageIndex < NumImages; ++ImageIndex)
+				{
+					if (ParentMaterial->GetParameterName(EMaterialParameterType::Texture, ImageIndex).ToString() == Image.Name)
+					{
+						ImageId = ParentMaterial->GetParameterId(EMaterialParameterType::Texture, ImageIndex).ParameterId;
+						break;
+					}
+				}
+
+				PinsParameter_DEPRECATED.Add(ImageId, FEdGraphPinReference(ImagePin));                
+			}
+		}
+
+		Images_DEPRECATED.Empty();
+	}
+
+	if (CustomizableObjectCustomVersion == FCustomizableObjectCustomVersion::ConvertEditAndExtendToModifiers)
+	{
+		// Look for the parent material and set it as the modifier reference material 
+		PinsParameterMap = PinsParameterMap_DEPRECATED;
+		PinsParameterMap_DEPRECATED.Empty();
+
+		UCustomizableObjectNode* ParentNode = GetCustomizableObjectExternalNode<UCustomizableObjectNode>(ParentMaterialObject_DEPRECATED.Get(), ParentMaterialNodeId_DEPRECATED);
+		if (ParentNode)
+		{
+			// Add an autogenerated tag to the legacy parent, so that it can be referenced in this modifier.
+			FString NewLegacyTag = MakeNodeAutoTag(ParentNode);
+			RequiredTags.Add(NewLegacyTag);
+
+			if (UCustomizableObjectNodeMaterial* MaterialParentNode = Cast<UCustomizableObjectNodeMaterial>(ParentNode))
+			{
+				ReferenceMaterial = MaterialParentNode->GetMaterial();
+				MaterialParentNode->Tags.AddUnique(NewLegacyTag);
+			}
+			else if (UCustomizableObjectNodeModifierExtendMeshSection* ExtendParentNode = Cast<UCustomizableObjectNodeModifierExtendMeshSection>(ParentNode))
+			{
+				ReferenceMaterial = ExtendParentNode->ReferenceMaterial;
+				ExtendParentNode->Tags.AddUnique(NewLegacyTag);
+			}
+			else
+			{
+				// Conversion failed?
+				ensure(false);
+				UE_LOG(LogMutable, Warning, TEXT("[%s] UCustomizableObjectNodeModifierExtendMeshSection version upgrade failed."), *GetOutermost()->GetName());
+			}
+
+			// If the tag was added to another CO, keep track of the addition to be able to repeat it in case
+			// that CO is not re-saved but this one is.
+			UCustomizableObject* ThisNodeObject = GetRootObject(*this);
+			if (ParentMaterialObject_DEPRECATED != ThisNodeObject)
+			{
+				FLegacyTag LegacyTag;
+				LegacyTag.ParentObject = ParentMaterialObject_DEPRECATED;
+				LegacyTag.ParentNode = ParentMaterialNodeId_DEPRECATED;
+				LegacyTag.Tag = NewLegacyTag;
+				LegacyBackportsRequiredTags.AddUnique(LegacyTag);
+			}
+		}
+		else
+		{
+			// No parent, no conversion.
+			UE_LOG(LogMutable, Log, TEXT("[%s] UCustomizableObjectNodeModifierExtendMeshSection has no parent. It will not be upgraded."), *GetOutermost()->GetName());
+		}
+
+		ReconstructNode();
+	}
+}
+
+
+void UCustomizableObjectNodeModifierExtendMeshSection::AllocateDefaultPins(UCustomizableObjectNodeRemapPins* RemapPins)
+{
+	const UEdGraphSchema_CustomizableObject* Schema = GetDefault<UEdGraphSchema_CustomizableObject>();
+
+	CustomCreatePin(EGPD_Input, Schema->PC_Mesh, FName("Add Mesh") );
+
+	// Begin texture pins
+	{
+		const int32 NumImages = GetNumParameters(EMaterialParameterType::Texture);
+		for (int32 ImageIndex = 0; ImageIndex < NumImages; ++ImageIndex)
+		{
+			const FName ImageName = GetParameterName(EMaterialParameterType::Texture, ImageIndex);
+			UEdGraphPin* PinImage = CustomCreatePin(EGPD_Input, Schema->PC_Image, ImageName);
+			PinImage->bDefaultValueIsIgnored = true;
+
+			const FNodeMaterialParameterId ImageId = GetParameterId(EMaterialParameterType::Texture, ImageIndex);
+				
+			PinsParameterMap.Add(ImageId, FEdGraphPinReference(PinImage));
+		}
+	}
+	// End texture pins
+
+	CustomCreatePin(EGPD_Output, Schema->PC_Modifier, TEXT("Modifier"));
+}
+
+
+FText UCustomizableObjectNodeModifierExtendMeshSection::GetNodeTitle(ENodeTitleType::Type TitleType) const
+{
+	return LOCTEXT("Extend_MeshSection", "Extend Mesh Section");
+}
+
+
+FString UCustomizableObjectNodeModifierExtendMeshSection::GetRefreshMessage() const
+{
+	return TEXT("The Reference Material has changed, texture channels might have been added, removed or renamed. Please refresh to reflect those changes.");
+}
+
+
+UEdGraphPin* UCustomizableObjectNodeModifierExtendMeshSection::AddMeshPin() const
+{
+	return FindPin(TEXT("Add Mesh"));
+}
+
+
+TArray<UCustomizableObjectLayout*> UCustomizableObjectNodeModifierExtendMeshSection::GetLayouts() const
+{
+	TArray<UCustomizableObjectLayout*> Result;
+
+	if (UEdGraphPin* MeshPin = AddMeshPin())
+	{
+		if (const UEdGraphPin* ConnectedPin = FollowInputPin(*MeshPin))
+		{
+			if (const UEdGraphPin* SourceMeshPin = FindMeshBaseSource(*ConnectedPin, false))
+			{
+				if (const UCustomizableObjectNodeSkeletalMesh* MeshNode = Cast<UCustomizableObjectNodeSkeletalMesh>(SourceMeshPin->GetOwningNode()))
+				{
+					Result = MeshNode->GetLayouts(*SourceMeshPin);
+				}
+				else if (const UCustomizableObjectNodeTable* TableNode = Cast<UCustomizableObjectNodeTable>(SourceMeshPin->GetOwningNode()))
+				{
+					Result = TableNode->GetLayouts(SourceMeshPin);
+				}
+			}
+		}
+	}
+
+	return Result;
+}
+
+
+FText UCustomizableObjectNodeModifierExtendMeshSection::GetTooltipText() const
+{
+	return LOCTEXT("Extend_Material_Tooltip", "Extend a mesh section: add a new mesh fragment, and add its corresponding texture to the modified material texture parameters.");
+}
+
+
+bool UCustomizableObjectNodeModifierExtendMeshSection::IsSingleOutputNode() const
+{
+	return true;
+}
+
+
+bool UCustomizableObjectNodeModifierExtendMeshSection::CustomRemovePin(UEdGraphPin& Pin)
+{
+	for (TMap<FNodeMaterialParameterId, FEdGraphPinReference>::TIterator It = PinsParameterMap.CreateIterator(); It; ++It)
+	{
+		if (It.Value().Get() == &Pin) // We could improve performance if FEdGraphPinReference exposed the pin id.
+		{
+			It.RemoveCurrent();
+			break;
+		}
+	}
+
+	return Super::CustomRemovePin(Pin);
+}
+
+
+TArray<FString>* UCustomizableObjectNodeModifierExtendMeshSection::GetEnableTags() 
+{ 
+	return &Tags; 
+}
+
+
+#undef LOCTEXT_NAMESPACE
diff --git a/Engine/Plugins/Mutable/Source/CustomizableObjectEditor/Private/MuCOE/Nodes/CustomizableObjectNodeModifierMorphMeshSectionDetails.cpp b/Engine/Plugins/Mutable/Source/CustomizableObjectEditor/Private/MuCOE/Nodes/CustomizableObjectNodeModifierMorphMeshSectionDetails.cpp
new file mode 100644
--- /dev/null
+++ b/Engine/Plugins/Mutable/Source/CustomizableObjectEditor/Private/MuCOE/Nodes/CustomizableObjectNodeModifierMorphMeshSectionDetails.cpp
@@ -0,0 +1,167 @@
+// Copyright Epic Games, Inc. All Rights Reserved.
+
+#include "MuCOE/Nodes/CustomizableObjectNodeModifierMorphMeshSectionDetails.h"
+
+#include "MuCOE/CustomizableObjectEditorUtilities.h"
+#include "MuCOE/CustomizableObjectEditorStyle.h"
+#include "MuCOE/GraphTraversal.h"
+#include "MuCOE/Nodes/CustomizableObjectNodeMaterial.h"
+#include "MuCOE/Nodes/CustomizableObjectNodeModifierMorphMeshSection.h"
+#include "MuCOE/Nodes/CustomizableObjectNodeModifierExtendMeshSection.h"
+#include "MuCOE/Nodes/CustomizableObjectNodeSkeletalMesh.h"
+#include "MuCOE/UnrealEditorPortabilityHelpers.h"
+#include "PropertyCustomizationHelpers.h"
+#include "Widgets/Input/STextComboBox.h"
+#include "Widgets/Input/SSearchBox.h"
+#include "DetailLayoutBuilder.h"
+#include "IDetailsView.h"
+#include "Engine/SkeletalMesh.h"
+#include "Misc/Attribute.h"
+
+
+#define LOCTEXT_NAMESPACE "CustomizableObjectDetails"
+
+
+TSharedRef<IDetailCustomization> FCustomizableObjectNodeModifierMorphMeshSectionDetails::MakeInstance()
+{
+	return MakeShareable( new FCustomizableObjectNodeModifierMorphMeshSectionDetails );
+}
+
+
+void FCustomizableObjectNodeModifierMorphMeshSectionDetails::CustomizeDetails( IDetailLayoutBuilder& DetailBuilder )
+{
+	FCustomizableObjectNodeModifierBaseDetails::CustomizeDetails(DetailBuilder);
+
+	const IDetailsView* DetailsView = DetailBuilder.GetDetailsView();
+	if ( DetailsView->GetSelectedObjects().Num() )
+	{
+		Node = Cast<UCustomizableObjectNodeModifierMorphMeshSection>( DetailsView->GetSelectedObjects()[0].Get() );
+	}
+
+	if (!Node)
+	{
+		return;
+	}
+
+	// This property is not relevant for this node
+	DetailBuilder.HideProperty(GET_MEMBER_NAME_CHECKED(UCustomizableObjectNodeModifierWithMaterial, ReferenceMaterial), UCustomizableObjectNodeModifierWithMaterial::StaticClass());
+	
+	// Add a morph selection widget.
+	IDetailCategoryBuilder& MorphCategory = DetailBuilder.EditCategory( "Morph" );
+
+	// Scan for hint morph names
+	{
+		TArray<UCustomizableObjectNode*> CandidateNodes;
+		Node->GetPossiblyModifiedNodes(CandidateNodes);
+
+		// For now just show the first one
+		for (UCustomizableObjectNode* Candidate : CandidateNodes)
+		{
+			USkeletalMesh* SkeletalMesh = nullptr;
+
+			if (UCustomizableObjectNodeMaterialBase* MaterialNode = Cast<UCustomizableObjectNodeMaterialBase>(Candidate))
+			{
+				if (MaterialNode->OutputPin())
+				{
+					const UEdGraphPin* SourceMeshPin = FindMeshBaseSource(*MaterialNode->OutputPin(), false);
+					if (SourceMeshPin)
+					{
+						UCustomizableObjectNodeSkeletalMesh* SkeletalNode = Cast<UCustomizableObjectNodeSkeletalMesh>(SourceMeshPin->GetOwningNode());
+						if (SkeletalNode)
+						{
+							SkeletalMesh = SkeletalNode->SkeletalMesh;
+						}
+					}
+				}
+			}
+
+			else if (UCustomizableObjectNodeModifierExtendMeshSection* ExtendNode = Cast<UCustomizableObjectNodeModifierExtendMeshSection>(Candidate))
+			{
+				if (ExtendNode->OutputPin())
+				{
+					const UEdGraphPin* SourceMeshPin = FindMeshBaseSource(*ExtendNode->OutputPin(), false);
+					if (SourceMeshPin)
+					{
+						UCustomizableObjectNodeSkeletalMesh* SkeletalNode = Cast<UCustomizableObjectNodeSkeletalMesh>(SourceMeshPin->GetOwningNode());
+						if (SkeletalNode)
+						{
+							SkeletalMesh = SkeletalNode->SkeletalMesh;
+						}
+					}
+				}
+			}
+
+			if (SkeletalMesh)
+			{
+				const TArray<TObjectPtr<UMorphTarget>>& Morphs = SkeletalMesh->GetMorphTargets();
+				for (TObjectPtr<UMorphTarget> Morph : Morphs)
+				{
+					if (Morph)
+					{
+						FString MorphTargetName = Morph->GetName();
+						MorphOptionsSource.AddUnique(MakeShared<FString>(MorphTargetName));
+					}
+				}
+			}
+		}
+
+		// TODO: Add all morphs if no candidate is found? Add both in different menu sections?
+	}
+
+	TSharedRef<IPropertyHandle> MorphTargetNameProperty = DetailBuilder.GetProperty("MorphTargetName");
+	MorphCategory.AddCustomRow(LOCTEXT("MorphMaterialDetails_Target", "Target"))
+		[
+			SNew(SProperty, MorphTargetNameProperty)
+				.ShouldDisplayName(false)
+				.CustomWidget()
+				[
+					SNew(SBorder)
+						.BorderImage(UE_MUTABLE_GET_BRUSH("NoBorder"))
+						.Padding(FMargin(0.0f, 0.0f, 10.0f, 0.0f))
+						[
+							SNew(SHorizontalBox)
+
+							+ SHorizontalBox::Slot()
+							.FillWidth(10.0f)
+							.VAlign(VAlign_Center)
+							[
+								SNew(STextBlock)
+									.Text(LOCTEXT("MorphMaterialDetails_MorphTarget", "Morph Target"))
+									.Font(IDetailLayoutBuilder::GetDetailFont())
+							]
+
+							+ SHorizontalBox::Slot()
+							.FillWidth(10.0f)
+							.HAlign(HAlign_Fill)
+							.VAlign(VAlign_Center)
+							[
+								SAssignNew(this->MorphCombo , SMutableSearchComboBox)
+								.ButtonStyle(FAppStyle::Get(), "NoBorder")
+								.OptionsSource(&MorphOptionsSource)
+								.OnSelectionChanged(this, &FCustomizableObjectNodeModifierMorphMeshSectionDetails::OnMorphTargetComboBoxSelectionChanged)
+								.Content()
+								[
+									SNew(SEditableTextBox)
+									.Font(IDetailLayoutBuilder::GetDetailFont())
+									.Text_Lambda([&]() { return Node ? FText::FromString(Node->MorphTargetName) : FText();})
+									.OnTextChanged(this, &FCustomizableObjectNodeModifierMorphMeshSectionDetails::OnMorphTargetComboBoxSelectionChanged)
+								]
+							]
+						]
+				]
+		];
+}
+
+
+void FCustomizableObjectNodeModifierMorphMeshSectionDetails::OnMorphTargetComboBoxSelectionChanged(const FText& NewText)
+{
+	if (Node && Node->MorphTargetName != NewText.ToString())
+	{
+		Node->MorphTargetName = NewText.ToString();
+		Node->Modify();
+	}
+}
+
+
+
+#undef LOCTEXT_NAMESPACE
diff --git a/Engine/Plugins/Mutable/Source/CustomizableObjectEditor/Private/MuCOE/Nodes/CustomizableObjectNodeModifierRemoveMesh.cpp b/Engine/Plugins/Mutable/Source/CustomizableObjectEditor/Private/MuCOE/Nodes/CustomizableObjectNodeModifierRemoveMesh.cpp
new file mode 100644
--- /dev/null
+++ b/Engine/Plugins/Mutable/Source/CustomizableObjectEditor/Private/MuCOE/Nodes/CustomizableObjectNodeModifierRemoveMesh.cpp
@@ -0,0 +1,58 @@
+// Copyright Epic Games, Inc. All Rights Reserved.
+
+#include "MuCOE/Nodes/CustomizableObjectNodeModifierRemoveMesh.h"
+
+#include "MuCOE/EdGraphSchema_CustomizableObject.h"
+#include "MuCOE/ICustomizableObjectEditor.h"
+#include "MuCO/CustomizableObjectCustomVersion.h"
+
+class UCustomizableObjectNodeRemapPins;
+
+#define LOCTEXT_NAMESPACE "CustomizableObjectEditor"
+
+
+void UCustomizableObjectNodeModifierRemoveMesh::AllocateDefaultPins(UCustomizableObjectNodeRemapPins* RemapPins)
+{
+	const UEdGraphSchema_CustomizableObject* Schema = GetDefault<UEdGraphSchema_CustomizableObject>();
+
+	UEdGraphPin* RemoveMeshPin = CustomCreatePin(EGPD_Input, Schema->PC_Mesh, FName("Remove Mesh") );
+	RemoveMeshPin->bDefaultValueIsIgnored = true;
+
+	UEdGraphPin* OutputPin = CustomCreatePin(EGPD_Output, Schema->PC_Modifier, FName("Modifier"));
+}
+
+
+FText UCustomizableObjectNodeModifierRemoveMesh::GetNodeTitle(ENodeTitleType::Type TitleType) const
+{
+	return LOCTEXT("Remove_Mesh", "Remove Mesh");
+}
+
+
+void UCustomizableObjectNodeModifierRemoveMesh::PinConnectionListChanged(UEdGraphPin* Pin)
+{
+	Super::PinConnectionListChanged(Pin);
+
+	if (Pin == OutputPin())
+	{
+		TSharedPtr<ICustomizableObjectEditor> Editor = GetGraphEditor();
+
+		if (Editor.IsValid())
+		{
+			Editor->UpdateGraphNodeProperties();
+		}
+	}
+}
+
+FText UCustomizableObjectNodeModifierRemoveMesh::GetTooltipText() const
+{
+	return LOCTEXT("Remove_Mesh_Tooltip",
+	"Removes the faces of a material that are defined only by vertexes shared by the material and the input mesh.It also removes any vertex\nand edge that only define deleted faces, they are not left dangling. If the mesh removed covers at least all the faces included in one or\nmore layout blocs, those blocs are removed, freeing final texture layout space.");
+}
+
+bool UCustomizableObjectNodeModifierRemoveMesh::IsSingleOutputNode() const
+{
+	return true;
+}
+
+
+#undef LOCTEXT_NAMESPACE
diff --git a/Engine/Plugins/Mutable/Source/CustomizableObjectEditor/Private/MuCOE/Nodes/CustomizableObjectNodeModifierRemoveMeshBlocks.cpp b/Engine/Plugins/Mutable/Source/CustomizableObjectEditor/Private/MuCOE/Nodes/CustomizableObjectNodeModifierRemoveMeshBlocks.cpp
new file mode 100644
--- /dev/null
+++ b/Engine/Plugins/Mutable/Source/CustomizableObjectEditor/Private/MuCOE/Nodes/CustomizableObjectNodeModifierRemoveMeshBlocks.cpp
@@ -0,0 +1,179 @@
+// Copyright Epic Games, Inc. All Rights Reserved.
+
+#include "MuCOE/Nodes/CustomizableObjectNodeModifierRemoveMeshBlocks.h"
+
+#include "MuCOE/Nodes/CustomizableObjectNodeModifierExtendMeshSection.h"
+#include "MuCO/CustomizableObjectCustomVersion.h"
+#include "MuCOE/CustomizableObjectLayout.h"
+#include "MuCOE/EdGraphSchema_CustomizableObject.h"
+#include "MuCOE/ICustomizableObjectEditor.h"
+#include "MuCOE/Nodes/CustomizableObjectNodeMaterial.h"
+
+class UCustomizableObjectNodeRemapPins;
+class UEdGraphPin;
+
+#define LOCTEXT_NAMESPACE "CustomizableObjectEditor"
+
+
+void UCustomizableObjectNodeModifierRemoveMeshBlocks::BackwardsCompatibleFixup(int32 CustomizableObjectCustomVersion)
+{
+	Super::BackwardsCompatibleFixup(CustomizableObjectCustomVersion);
+	
+	// Convert deprecated node index list to the node id list.
+	if (CustomizableObjectCustomVersion == FCustomizableObjectCustomVersion::PostLoadToCustomVersion
+		&& BlockIds_DEPRECATED.Num() < Blocks_DEPRECATED.Num())
+	{
+		UCustomizableObjectNodeMaterialBase* ParentMaterialNode = GetCustomizableObjectExternalNode<UCustomizableObjectNodeMaterialBase>(ParentMaterialObject_DEPRECATED.Get(), ParentMaterialNodeId_DEPRECATED);
+		if (ParentMaterialNode)
+		{
+			TArray<UCustomizableObjectLayout*> Layouts = ParentMaterialNode->GetLayouts();
+
+			if (!Layouts.IsValidIndex(ParentLayoutIndex))
+			{
+				UE_LOG(LogMutable, Warning, TEXT("[%s] UCustomizableObjectNodeModifierRemoveMeshBlocks refers to an invalid texture layout index %d. Parent node has %d layouts."),
+					*GetOutermost()->GetName(), ParentLayoutIndex, Layouts.Num());
+			}
+			else if (UCustomizableObjectNodeMaterial* ParentMaterial = Cast<UCustomizableObjectNodeMaterial>(ParentMaterialNode))
+			{
+				UCustomizableObjectLayout* ParentLayout = Layouts[ParentLayoutIndex];
+
+				for (int IndexIndex = BlockIds_DEPRECATED.Num(); IndexIndex < Blocks_DEPRECATED.Num(); ++IndexIndex)
+				{
+					const int BlockIndex = Blocks_DEPRECATED[IndexIndex];
+					if (!ParentLayout->Blocks.IsValidIndex(BlockIndex))
+					{
+						UE_LOG(LogMutable, Warning, TEXT("[%s] UCustomizableObjectNodeModifierRemoveMeshBlocks refers to an invalid layout block index %d. Parent node has %d blocks."),
+							*GetOutermost()->GetName(), BlockIndex, ParentLayout->Blocks.Num());
+
+						continue;
+					}
+					
+					const FGuid Id = ParentLayout->Blocks[BlockIndex].Id;
+					if (!Id.IsValid())
+					{
+						UE_LOG(LogMutable, Warning, TEXT("[%s] UCustomizableObjectNodeModifierRemoveMeshBlocks refers to an valid layout block %d but that block doesn't have an id."),
+							*GetOutermost()->GetName(), BlockIndex);
+
+						continue;
+					}
+
+					BlockIds_DEPRECATED.Add(Id);
+				}
+			}
+		}
+		else
+		{
+			UE_LOG(LogMutable, Log, TEXT("[%s] UCustomizableObjectNodeModifierRemoveMeshBlocks has no parent. It will not be upgraded."), *GetOutermost()->GetName());
+		}
+	}
+
+	// Convert deprecated node id list to absolute rect list.
+	if (CustomizableObjectCustomVersion == FCustomizableObjectCustomVersion::UseUVRects)
+	{
+		// If we are here, it means this node was loaded from a version that didn't have it's own layout.
+		check(Layout->Blocks.IsEmpty());
+
+		UCustomizableObjectNodeMaterialBase* ParentMaterialNode = GetCustomizableObjectExternalNode<UCustomizableObjectNodeMaterialBase>(ParentMaterialObject_DEPRECATED.Get(), ParentMaterialNodeId_DEPRECATED);
+
+		TArray<UCustomizableObjectLayout*> ParentLayouts;
+
+		if (!ParentMaterialNode)
+		{
+			// Conversion failed?
+			ensure(false);
+			UE_LOG(LogMutable, Warning, TEXT("[%s] UCustomizableObjectNodeModifierRemoveMeshBlocks version upgrade failed."), *GetOutermost()->GetName());
+		}
+		else
+		{
+			ParentLayouts = ParentMaterialNode->GetLayouts();
+		}
+
+		if (!ParentLayouts.IsValidIndex(ParentLayoutIndex))
+		{
+			UE_LOG(LogMutable, Warning, TEXT("[%s] UCustomizableObjectNodeModifierRemoveMeshBlocks refers to an invalid texture layout index %d. Parent node has %d layouts."),
+				*GetOutermost()->GetName(), ParentLayoutIndex, ParentLayouts.Num());
+		}
+		else
+		{
+			UCustomizableObjectLayout* ParentLayout = ParentLayouts[ParentLayoutIndex];
+			FIntPoint GridSize = ParentLayout->GetGridSize();
+
+			Layout->SetGridSize(GridSize);
+
+			if (Cast<UCustomizableObjectNodeMaterial>(ParentMaterialNode))
+			{
+				for (const FGuid& BlockId : BlockIds_DEPRECATED)
+				{
+					bool bSkipBlock = false;
+					for (const FCustomizableObjectLayoutBlock& ExistingBlock : Layout->Blocks)
+					{
+						if (ExistingBlock.Id == BlockId)
+						{
+							bSkipBlock = true;
+							UE_LOG(LogMutable, Log, TEXT("[%s] UCustomizableObjectNodeModifierRemoveMeshBlocks has a duplicated layout block id. One has been ignored during version upgrade."), *GetOutermost()->GetName());
+							break;
+						}
+					}
+
+					if (bSkipBlock)
+					{
+						continue;
+					}
+
+					bool bFoundInParent = false;
+					for (const FCustomizableObjectLayoutBlock& ParentBlock : ParentLayout->Blocks)
+					{
+						if (ParentBlock.Id == BlockId)
+						{
+							bFoundInParent = true;
+
+							FCustomizableObjectLayoutBlock NewBlock;
+							NewBlock = ParentBlock;
+
+							// Clear some unnecessary data.
+							NewBlock.bReduceBothAxes = false;
+							NewBlock.bReduceByTwo = false;
+							NewBlock.Priority = 0;
+
+							Layout->Blocks.Add(NewBlock);
+							break;
+						}
+					}
+
+					if (!bFoundInParent)
+					{
+						UE_LOG(LogMutable, Warning, TEXT("[%s] UCustomizableObjectNodeModifierRemoveMeshBlocks refers to and invalid layout block. It has been ignored during version upgrade."), *GetOutermost()->GetName());
+					}
+				}
+			}
+		}
+	}
+}
+
+
+void UCustomizableObjectNodeModifierRemoveMeshBlocks::AllocateDefaultPins(UCustomizableObjectNodeRemapPins* RemapPins)
+{
+	const UEdGraphSchema_CustomizableObject* Schema = GetDefault<UEdGraphSchema_CustomizableObject>();
+
+	UEdGraphPin* OutputPin = CustomCreatePin(EGPD_Output, Schema->PC_Modifier, FName("Modifier"));
+}
+
+
+FText UCustomizableObjectNodeModifierRemoveMeshBlocks::GetNodeTitle(ENodeTitleType::Type TitleType) const
+{
+	return LOCTEXT("Remove_Mesh_Blocks", "Remove Mesh Blocks");
+}
+
+
+FText UCustomizableObjectNodeModifierRemoveMeshBlocks::GetTooltipText() const
+{
+	return LOCTEXT("Remove_Mesh_Blocks_Tooltip", "Remove all the geometry in the chosen layout blocks from a material.");
+}
+
+
+bool UCustomizableObjectNodeModifierRemoveMeshBlocks::IsSingleOutputNode() const
+{
+	return true;
+}
+
+#undef LOCTEXT_NAMESPACE
diff --git a/Engine/Plugins/Mutable/Source/CustomizableObjectEditor/Private/MuCOE/Nodes/CustomizableObjectNodeModifierRemoveMeshDetails.h b/Engine/Plugins/Mutable/Source/CustomizableObjectEditor/Private/MuCOE/Nodes/CustomizableObjectNodeModifierRemoveMeshDetails.h
new file mode 100644
--- /dev/null
+++ b/Engine/Plugins/Mutable/Source/CustomizableObjectEditor/Private/MuCOE/Nodes/CustomizableObjectNodeModifierRemoveMeshDetails.h
@@ -0,0 +1,21 @@
+// Copyright Epic Games, Inc. All Rights Reserved.
+
+#pragma once
+
+#include "CustomizableObjectNodeDetails.h"
+#include "IDetailCustomization.h"
+#include "Types/SlateEnums.h"
+#include "Misc/AssertionMacros.h"
+
+class IDetailLayoutBuilder;
+
+class FCustomizableObjectNodeModifierRemoveMeshDetails : public FCustomizableObjectNodeDetails
+{
+public:
+	/** Makes a new instance of this detail layout class for a specific detail view requesting it */
+	static TSharedRef<IDetailCustomization> MakeInstance();
+
+	/** ILayoutDetails interface */
+	virtual void CustomizeDetails(IDetailLayoutBuilder& DetailBuilder) override;
+
+};
diff --git a/Engine/Plugins/Mutable/Source/CustomizableObjectEditor/Private/MuCOE/Nodes/CustomizableObjectNodeModifierWithMaterial.cpp b/Engine/Plugins/Mutable/Source/CustomizableObjectEditor/Private/MuCOE/Nodes/CustomizableObjectNodeModifierWithMaterial.cpp
new file mode 100644
--- /dev/null
+++ b/Engine/Plugins/Mutable/Source/CustomizableObjectEditor/Private/MuCOE/Nodes/CustomizableObjectNodeModifierWithMaterial.cpp
@@ -0,0 +1,190 @@
+// Copyright Epic Games, Inc. All Rights Reserved.
+
+#include "MuCOE/Nodes/CustomizableObjectNodeModifierWithMaterial.h"
+
+#include "MuCOE/Nodes/CustomizableObjectNodeMaterial.h"
+#include "MuCOE/GraphTraversal.h"
+#include "MuCOE/CustomizableObjectPin.h"
+#include "MuCO/CustomizableObjectCustomVersion.h"
+#include "MaterialTypes.h"
+#include "MaterialCachedData.h"
+#include "Materials/MaterialExpressionTextureCoordinate.h"
+#include "Materials/MaterialExpressionTextureSample.h"
+#include "Materials/MaterialInstance.h"
+
+
+#define LOCTEXT_NAMESPACE "CustomizableObjectEditor"
+
+
+
+void UCustomizableObjectNodeModifierWithMaterial::PostEditChangeProperty(FPropertyChangedEvent& PropertyChangedEvent)
+{
+	Super::PostEditChangeProperty(PropertyChangedEvent);
+
+	if (const FProperty* PropertyThatChanged = PropertyChangedEvent.Property)
+	{
+		if (PropertyThatChanged->GetFName() == GET_MEMBER_NAME_CHECKED(UCustomizableObjectNodeModifierWithMaterial, ReferenceMaterial))
+		{
+			UCustomizableObjectNode::ReconstructNode();
+		}
+	}
+}
+
+
+bool UCustomizableObjectNodeModifierWithMaterial::UsesImage(const FNodeMaterialParameterId& ImageId) const
+{
+	if (const UEdGraphPin* Pin = GetUsedImagePin(ImageId))
+	{
+		return FollowInputPin(*Pin) != nullptr;
+	}
+	else
+	{
+		return false;
+	}
+}
+
+
+const UEdGraphPin* UCustomizableObjectNodeModifierWithMaterial::GetUsedImagePin(const FNodeMaterialParameterId& ImageId) const
+{
+	if (const FEdGraphPinReference* PinReference = PinsParameterMap.Find(ImageId))
+	{
+		if (const UEdGraphPin& Pin = *PinReference->Get(); !IsPinOrphan(Pin)) // We always have a valid pin reference. If it is nullptr, it means that something went wrong.
+		{
+			return &Pin;
+		}
+	}
+
+	return nullptr;
+}
+
+
+bool UCustomizableObjectNodeModifierWithMaterial::IsNodeOutDatedAndNeedsRefresh()
+{
+	const bool bOutdated = [&, this]()
+		{
+			for (const TTuple<FNodeMaterialParameterId, FEdGraphPinReference> Pair : PinsParameterMap)
+			{
+				const UEdGraphPin& Pin = *Pair.Value.Get();
+				if (!IsPinOrphan(Pin))
+				{
+					bool bHas = UCustomizableObjectNodeMaterial::HasParameter(ReferenceMaterial,Pair.Key);
+					if (UsesImage(Pair.Key) && !bHas)
+					{
+						return true;
+					}
+				}
+			}
+
+			return false;
+		}();
+
+	// Remove previous compilation warnings
+	if (!bOutdated && bHasCompilerMessage)
+	{
+		RemoveWarnings();
+		GetGraph()->NotifyGraphChanged();
+	}
+
+	return bOutdated;
+}
+
+
+int32 UCustomizableObjectNodeModifierWithMaterial::GetNumParameters(const EMaterialParameterType Type) const
+{
+	if (ReferenceMaterial)
+	{
+		return ReferenceMaterial->GetCachedExpressionData().GetParameterTypeEntry(Type).ParameterInfoSet.Num();
+	}
+	else
+	{
+		return 0;
+	}
+}
+
+
+FNodeMaterialParameterId UCustomizableObjectNodeModifierWithMaterial::GetParameterId(const EMaterialParameterType Type, const int32 ParameterIndex) const
+{
+	const FMaterialCachedExpressionData& Data = ReferenceMaterial->GetCachedExpressionData();
+
+	if (Data.EditorOnlyData)
+	{
+		if (Data.EditorOnlyData->EditorEntries[(int32)Type].EditorInfo.Num() != 0)
+		{
+			const FGuid ParameterId = Data.EditorOnlyData->EditorEntries[(int32)Type].EditorInfo[ParameterIndex].ExpressionGuid;
+			const int32 LayerIndex = GetParameterLayerIndex(Type, ParameterIndex);
+
+			return { ParameterId, LayerIndex };
+		}
+	}
+
+	return FNodeMaterialParameterId();
+}
+
+
+FName UCustomizableObjectNodeModifierWithMaterial::GetParameterName(const EMaterialParameterType Type, const int32 ParameterIndex) const
+{
+	check(ReferenceMaterial);
+
+	const FMaterialCachedParameterEntry& Entry = ReferenceMaterial->GetCachedExpressionData().GetParameterTypeEntry(Type);
+
+	for (TSet<FMaterialParameterInfo>::TConstIterator It(Entry.ParameterInfoSet); It; ++It)
+	{
+		const int32 IteratorIndex = It.GetId().AsInteger();
+
+		if (IteratorIndex == ParameterIndex)
+		{
+			return (*It).Name;
+		}
+	}
+
+	// The parameter should exist
+	check(false);
+
+	return FName();
+}
+
+
+int32 UCustomizableObjectNodeModifierWithMaterial::GetParameterLayerIndex(const UMaterialInterface* InMaterial, const EMaterialParameterType Type, const int32 ParameterIndex)
+{
+	check(InMaterial);
+
+	const FMaterialCachedParameterEntry& Entry = InMaterial->GetCachedExpressionData().GetParameterTypeEntry(Type);
+
+	for (TSet<FMaterialParameterInfo>::TConstIterator It(Entry.ParameterInfoSet); It; ++It)
+	{
+		const int32 IteratorIndex = It.GetId().AsInteger();
+
+		if (IteratorIndex == ParameterIndex)
+		{
+			return (*It).Index;
+		}
+	}
+
+	// The parameter should exist
+	check(false);
+
+	return -1;
+}
+
+
+int32 UCustomizableObjectNodeModifierWithMaterial::GetParameterLayerIndex(const EMaterialParameterType Type, const int32 ParameterIndex) const
+{
+	return GetParameterLayerIndex(ReferenceMaterial.Get(), Type, ParameterIndex);
+}
+
+
+FText UCustomizableObjectNodeModifierWithMaterial::GetParameterLayerName(const EMaterialParameterType Type, const int32 ParameterIndex) const
+{
+	check(ReferenceMaterial)
+
+		int32 LayerIndex = GetParameterLayerIndex(Type, ParameterIndex);
+
+	FMaterialLayersFunctions LayersValue;
+	ReferenceMaterial->GetMaterialLayers(LayersValue);
+
+
+	return LayersValue.EditorOnly.LayerNames.IsValidIndex(LayerIndex) ? LayersValue.EditorOnly.LayerNames[LayerIndex] : FText();
+}
+
+
+#undef LOCTEXT_NAMESPACE
diff --git a/Engine/Plugins/Mutable/Source/CustomizableObjectEditor/Private/MuCOE/Nodes/CustomizableObjectNodeProjectorParameter.cpp b/Engine/Plugins/Mutable/Source/CustomizableObjectEditor/Private/MuCOE/Nodes/CustomizableObjectNodeProjectorParameter.cpp
new file mode 100644
--- /dev/null
+++ b/Engine/Plugins/Mutable/Source/CustomizableObjectEditor/Private/MuCOE/Nodes/CustomizableObjectNodeProjectorParameter.cpp
@@ -0,0 +1,158 @@
+// Copyright Epic Games, Inc. All Rights Reserved.
+
+#include "MuCOE/Nodes/CustomizableObjectNodeProjectorParameter.h"
+
+#include "MuCO/CustomizableObjectCustomVersion.h"
+#include "MuCOE/EdGraphSchema_CustomizableObject.h"
+#include "MuCOE/UnrealEditorPortabilityHelpers.h"
+
+class UCustomizableObjectNodeRemapPins;
+
+#define LOCTEXT_NAMESPACE "CustomizableObjectEditor"
+
+
+void UCustomizableObjectNodeProjectorParameter::PostEditChangeProperty(FPropertyChangedEvent& PropertyChangedEvent)
+{
+	const FName PropertyName = PropertyChangedEvent.GetPropertyName();
+
+	if (PropertyName == "ProjectionAngle")
+	{
+		DefaultValue.Angle = FMath::DegreesToRadians(ProjectionAngle);
+	}
+	else if (PropertyName == "ProjectorBone")
+	{
+		DefaultValue.Position = (FVector3f)BoneComboBoxLocation;
+		DefaultValue.Direction = (FVector3f)BoneComboBoxForwardDirection;
+		DefaultValue.Up = (FVector3f)BoneComboBoxUpDirection;
+	}
+
+	Super::PostEditChangeProperty(PropertyChangedEvent);
+}
+
+
+void UCustomizableObjectNodeProjectorParameter::AllocateDefaultPins(UCustomizableObjectNodeRemapPins* RemapPins)
+{
+	const UEdGraphSchema_CustomizableObject* Schema = GetDefault<UEdGraphSchema_CustomizableObject>();
+
+	FString PinName = TEXT("Value");
+	UEdGraphPin* ValuePin = CustomCreatePin(EGPD_Output, Schema->PC_Projector, FName(*PinName));
+	ValuePin->bDefaultValueIsIgnored = true;
+}
+
+
+void UCustomizableObjectNodeProjectorParameter::BackwardsCompatibleFixup(int32 CustomizableObjectCustomVersion)
+{
+	Super::BackwardsCompatibleFixup(CustomizableObjectCustomVersion);
+	
+	if (CustomizableObjectCustomVersion == FCustomizableObjectCustomVersion::ProjectorNodesDefaultValueFix)
+	{
+		DefaultValue.ProjectionType = ProjectionType_DEPRECATED;
+	}
+}
+
+
+ECustomizableObjectProjectorType UCustomizableObjectNodeProjectorParameter::GetProjectorType() const
+{
+	return DefaultValue.ProjectionType;
+}
+
+
+FVector UCustomizableObjectNodeProjectorParameter::GetProjectorDefaultPosition() const
+{
+	return static_cast<FVector>(DefaultValue.Position);
+}
+
+
+void UCustomizableObjectNodeProjectorParameter::SetProjectorDefaultPosition(const FVector& Position)
+{
+	DefaultValue.Position = static_cast<FVector3f>(Position);
+}
+
+
+FVector UCustomizableObjectNodeProjectorParameter::GetProjectorDefaultDirection() const
+{
+	return static_cast<FVector>(DefaultValue.Direction);
+}
+
+
+void UCustomizableObjectNodeProjectorParameter::SetProjectorDefaultDirection(const FVector& Direction)
+{
+	DefaultValue.Direction = static_cast<FVector3f>(Direction);
+}
+
+
+FVector UCustomizableObjectNodeProjectorParameter::GetProjectorDefaultUp() const
+{
+	return static_cast<FVector>(DefaultValue.Up);
+}
+
+
+void UCustomizableObjectNodeProjectorParameter::SetProjectorDefaultUp(const FVector& Up)
+{
+	DefaultValue.Up = static_cast<FVector3f>(Up);
+}
+
+
+FVector UCustomizableObjectNodeProjectorParameter::GetProjectorDefaultScale() const
+{
+	return static_cast<FVector>(DefaultValue.Scale);
+}
+
+
+void UCustomizableObjectNodeProjectorParameter::SetProjectorDefaultScale(const FVector& Scale)
+{
+	DefaultValue.Scale = static_cast<FVector3f>(Scale);
+}
+
+
+float UCustomizableObjectNodeProjectorParameter::GetProjectorDefaultAngle() const
+{
+	return ProjectionAngle;
+}
+
+
+void UCustomizableObjectNodeProjectorParameter::SetProjectorDefaultAngle(float Angle)
+{
+	ProjectionAngle = Angle;
+}
+
+
+FText UCustomizableObjectNodeProjectorParameter::GetNodeTitle(ENodeTitleType::Type TitleType) const
+{
+	if (TitleType == ENodeTitleType::ListView || ParameterName.IsEmpty())
+	{
+		return LOCTEXT("Projector_Parameter", "Projector Parameter");
+	}
+	else if (TitleType == ENodeTitleType::EditableTitle)
+	{
+		return FText::Format(LOCTEXT("Projector_Parameter_EditableTitle", "{0}"), FText::FromString(ParameterName));
+	}
+	else
+	{
+		return FText::Format(LOCTEXT("Projector_Parameter_Title", "{0}\nProjector Parameter"), FText::FromString(ParameterName));
+	}
+}
+
+
+FLinearColor UCustomizableObjectNodeProjectorParameter::GetNodeTitleColor() const
+{
+	const UEdGraphSchema_CustomizableObject* Schema = GetDefault<UEdGraphSchema_CustomizableObject>();
+	return Schema->GetPinTypeColor(Schema->PC_Projector);
+}
+
+
+FText UCustomizableObjectNodeProjectorParameter::GetTooltipText() const
+{
+	return LOCTEXT("Projector_Parameter_Tooltip", "Exposes a runtime modifiable projector parameter from the Customizable Object.");
+}
+
+
+void UCustomizableObjectNodeProjectorParameter::OnRenameNode(const FString& NewName)
+{
+	if (!NewName.IsEmpty())
+	{
+		ParameterName = NewName;
+	}
+}
+
+#undef LOCTEXT_NAMESPACE
diff --git a/Engine/Plugins/Mutable/Source/CustomizableObjectEditor/Private/MuCOE/Nodes/CustomizableObjectNodeProjectorParameterDetails.h b/Engine/Plugins/Mutable/Source/CustomizableObjectEditor/Private/MuCOE/Nodes/CustomizableObjectNodeProjectorParameterDetails.h
new file mode 100644
--- /dev/null
+++ b/Engine/Plugins/Mutable/Source/CustomizableObjectEditor/Private/MuCOE/Nodes/CustomizableObjectNodeProjectorParameterDetails.h
@@ -0,0 +1,39 @@
+// Copyright Epic Games, Inc. All Rights Reserved.
+
+#pragma once
+
+#include "CustomizableObjectNodeDetails.h"
+#include "IDetailCustomization.h"
+
+class FReply;
+namespace ESelectInfo { enum Type : int; }
+
+class FString;
+class IDetailLayoutBuilder;
+class IPropertyHandle;
+
+
+class FCustomizableObjectNodeProjectorParameterDetails : public FCustomizableObjectNodeDetails
+{
+public:
+	/** Makes a new instance of this detail layout class for a specific detail view requesting it */
+	static TSharedRef<IDetailCustomization> MakeInstance();
+
+	/** ILayoutDetails interface */
+	virtual void CustomizeDetails( IDetailLayoutBuilder& DetailBuilder ) override;
+
+
+private:
+	TArray< TSharedPtr<FString> > BoneComboOptions;
+
+	class UCustomizableObjectNodeProjectorConstant* NodeConstant;
+	class UCustomizableObjectNodeProjectorParameter* NodeParameter;
+	FReply OnProjectorCopyPressed();
+	FReply OnProjectorPastePressed();
+
+	class USkeletalMesh* SkeletalMesh;
+	class IDetailLayoutBuilder* DetailBuilderPtr;
+
+	void OnBoneComboBoxSelectionChanged(TSharedPtr<FString> Selection, ESelectInfo::Type SelectInfo, TSharedRef<IPropertyHandle> BoneProperty);
+	void OnReferenceSkeletonIndexChanged();
+};
diff --git a/Engine/Plugins/Mutable/Source/CustomizableObjectEditor/Private/MuCOE/Nodes/CustomizableObjectNodeTextureColourMap.cpp b/Engine/Plugins/Mutable/Source/CustomizableObjectEditor/Private/MuCOE/Nodes/CustomizableObjectNodeTextureColourMap.cpp
new file mode 100644
--- /dev/null
+++ b/Engine/Plugins/Mutable/Source/CustomizableObjectEditor/Private/MuCOE/Nodes/CustomizableObjectNodeTextureColourMap.cpp
@@ -0,0 +1,69 @@
+// Copyright Epic Games, Inc. All Rights Reserved.
+
+#include "MuCOE/Nodes/CustomizableObjectNodeTextureColourMap.h"
+
+#include "MuCO/CustomizableObjectCustomVersion.h"
+#include "MuCOE/EdGraphSchema_CustomizableObject.h"
+
+class UCustomizableObjectNodeRemapPins;
+
+#define LOCTEXT_NAMESPACE "CustomizableObjectEditor"
+
+
+void UCustomizableObjectNodeTextureColourMap::AllocateDefaultPins(UCustomizableObjectNodeRemapPins* RemapPins)
+{
+	const UEdGraphSchema_CustomizableObject* Schema = GetDefault<UEdGraphSchema_CustomizableObject>();
+
+	FString PinName = TEXT("Texture");
+	UEdGraphPin* OutputPin = CustomCreatePin(EGPD_Output, Schema->PC_Image, FName(*PinName));
+	OutputPin->bDefaultValueIsIgnored = true;
+
+	PinName = TEXT("Base");
+	UEdGraphPin* SourcePin = CustomCreatePin(EGPD_Input, Schema->PC_Image, FName(*PinName));
+	SourcePin->bDefaultValueIsIgnored = true;
+
+	PinName = TEXT("Mask");
+	UEdGraphPin* MaskPin = CustomCreatePin(EGPD_Input, Schema->PC_Image, FName(*PinName));
+	MaskPin->bDefaultValueIsIgnored = true;
+
+	PinName = TEXT("Map");
+	UEdGraphPin* GradientPin = CustomCreatePin(EGPD_Input, Schema->PC_Image, FName(*PinName));
+	GradientPin->bDefaultValueIsIgnored = true;
+}
+
+
+void UCustomizableObjectNodeTextureColourMap::BackwardsCompatibleFixup(int32 CustomizableObjectCustomVersion)
+{
+	Super::BackwardsCompatibleFixup(CustomizableObjectCustomVersion);
+	
+	if (CustomizableObjectCustomVersion == FCustomizableObjectCustomVersion::FixPinsNamesImageToTexture2)
+	{
+		if (UEdGraphPin* TexturePin = FindPin(TEXT("Image")))
+		{
+			TexturePin->PinName = TEXT("Texture");
+			UCustomizableObjectNode::ReconstructNode();
+		}
+	}
+}
+
+
+FText UCustomizableObjectNodeTextureColourMap::GetNodeTitle(ENodeTitleType::Type TitleType) const
+{
+	return LOCTEXT("Texture_Map", "Texture Colour Map");
+}
+
+
+FLinearColor UCustomizableObjectNodeTextureColourMap::GetNodeTitleColor() const
+{
+	const UEdGraphSchema_CustomizableObject* Schema = GetDefault<UEdGraphSchema_CustomizableObject>();
+	return Schema->GetPinTypeColor(Schema->PC_Image);
+}
+
+
+FText UCustomizableObjectNodeTextureColourMap::GetTooltipText() const
+{
+	return LOCTEXT("Texture_Gradient_Sample_Tooltip", "Map colours of map using values form image.");
+}
+
+#undef LOCTEXT_NAMESPACE
+
diff --git a/Engine/Plugins/Mutable/Source/CustomizableObjectEditor/Private/MuCOE/Nodes/CustomizableObjectNodeTextureFromChannels.cpp b/Engine/Plugins/Mutable/Source/CustomizableObjectEditor/Private/MuCOE/Nodes/CustomizableObjectNodeTextureFromChannels.cpp
new file mode 100644
--- /dev/null
+++ b/Engine/Plugins/Mutable/Source/CustomizableObjectEditor/Private/MuCOE/Nodes/CustomizableObjectNodeTextureFromChannels.cpp
@@ -0,0 +1,73 @@
+// Copyright Epic Games, Inc. All Rights Reserved.
+
+#include "MuCOE/Nodes/CustomizableObjectNodeTextureFromChannels.h"
+
+#include "MuCO/CustomizableObjectCustomVersion.h"
+#include "MuCOE/EdGraphSchema_CustomizableObject.h"
+
+class UCustomizableObjectNodeRemapPins;
+
+#define LOCTEXT_NAMESPACE "CustomizableObjectEditor"
+
+
+void UCustomizableObjectNodeTextureFromChannels::AllocateDefaultPins(UCustomizableObjectNodeRemapPins* RemapPins)
+{
+	const UEdGraphSchema_CustomizableObject* Schema = GetDefault<UEdGraphSchema_CustomizableObject>();
+
+	FString PinName = TEXT("Texture");
+	UEdGraphPin* OutputPin = CustomCreatePin(EGPD_Output, Schema->PC_Image, FName(*PinName));
+	OutputPin->bDefaultValueIsIgnored = true;
+
+	PinName = TEXT("R");
+	UEdGraphPin* RPin = CustomCreatePin(EGPD_Input, Schema->PC_Image, FName(*PinName));
+	RPin->bDefaultValueIsIgnored = true;
+
+	PinName = TEXT("G");
+	UEdGraphPin* GPin = CustomCreatePin(EGPD_Input, Schema->PC_Image, FName(*PinName));
+	GPin->bDefaultValueIsIgnored = true;
+
+	PinName = TEXT("B");
+	UEdGraphPin* BPin = CustomCreatePin(EGPD_Input, Schema->PC_Image, FName(*PinName));
+	BPin->bDefaultValueIsIgnored = true;
+
+	PinName = TEXT("A");
+	UEdGraphPin* APin = CustomCreatePin(EGPD_Input, Schema->PC_Image, FName(*PinName));
+	APin->bDefaultValueIsIgnored = true;
+}
+
+
+void UCustomizableObjectNodeTextureFromChannels::BackwardsCompatibleFixup(int32 CustomizableObjectCustomVersion)
+{
+	Super::BackwardsCompatibleFixup(CustomizableObjectCustomVersion);
+	
+	if (CustomizableObjectCustomVersion == FCustomizableObjectCustomVersion::FixPinsNamesImageToTexture2)
+	{
+		if (UEdGraphPin* TexturePin = FindPin(TEXT("Image")))
+		{
+			TexturePin->PinName = TEXT("Texture");
+			UCustomizableObjectNode::ReconstructNode();
+		}
+	}
+}
+
+
+FText UCustomizableObjectNodeTextureFromChannels::GetNodeTitle(ENodeTitleType::Type TitleType) const
+{
+	return LOCTEXT("Texture_From_Channels", "Make Texture");
+}
+
+
+FLinearColor UCustomizableObjectNodeTextureFromChannels::GetNodeTitleColor() const
+{
+	const UEdGraphSchema_CustomizableObject* Schema = GetDefault<UEdGraphSchema_CustomizableObject>();
+	return Schema->GetPinTypeColor(Schema->PC_Image);
+}
+
+
+FText UCustomizableObjectNodeTextureFromChannels::GetTooltipText() const
+{
+	return LOCTEXT("Texture_From_Channels_Tooltip", "Make a colored texture with transparency from four grayscale textures that represent the values of each RGBA channel.");
+}
+
+#undef LOCTEXT_NAMESPACE
+
diff --git a/Engine/Plugins/Mutable/Source/CustomizableObjectEditor/Private/MuCOE/Nodes/CustomizableObjectNodeTextureFromColor.cpp b/Engine/Plugins/Mutable/Source/CustomizableObjectEditor/Private/MuCOE/Nodes/CustomizableObjectNodeTextureFromColor.cpp
new file mode 100644
--- /dev/null
+++ b/Engine/Plugins/Mutable/Source/CustomizableObjectEditor/Private/MuCOE/Nodes/CustomizableObjectNodeTextureFromColor.cpp
@@ -0,0 +1,60 @@
+// Copyright Epic Games, Inc. All Rights Reserved.
+
+#include "MuCOE/Nodes/CustomizableObjectNodeTextureFromColor.h"
+
+#include "MuCO/CustomizableObjectCustomVersion.h"
+#include "MuCOE/EdGraphSchema_CustomizableObject.h"
+
+class UCustomizableObjectNodeRemapPins;
+
+#define LOCTEXT_NAMESPACE "CustomizableObjectEditor"
+
+
+void UCustomizableObjectNodeTextureFromColor::AllocateDefaultPins(UCustomizableObjectNodeRemapPins* RemapPins)
+{
+	const UEdGraphSchema_CustomizableObject* Schema = GetDefault<UEdGraphSchema_CustomizableObject>();
+
+	FString PinName = TEXT("Texture");
+	UEdGraphPin* OutputPin = CustomCreatePin(EGPD_Output, Schema->PC_Image, FName(*PinName));
+	OutputPin->bDefaultValueIsIgnored = true;
+
+	PinName = TEXT("Color");
+	UEdGraphPin* RPin = CustomCreatePin(EGPD_Input, Schema->PC_Color, FName(*PinName));
+	RPin->bDefaultValueIsIgnored = true;
+}
+
+
+void UCustomizableObjectNodeTextureFromColor::BackwardsCompatibleFixup(int32 CustomizableObjectCustomVersion)
+{
+	Super::BackwardsCompatibleFixup(CustomizableObjectCustomVersion);
+	
+	if (CustomizableObjectCustomVersion == FCustomizableObjectCustomVersion::FixPinsNamesImageToTexture2)
+	{
+		if (UEdGraphPin* TexturePin = FindPin(TEXT("Image")))
+		{
+			TexturePin->PinName = TEXT("Texture");
+			UCustomizableObjectNode::ReconstructNode();
+		}
+	}
+}
+
+
+FText UCustomizableObjectNodeTextureFromColor::GetNodeTitle(ENodeTitleType::Type TitleType) const
+{
+	return LOCTEXT("Texture_From_Color", "Texture From Color");
+}
+
+
+FLinearColor UCustomizableObjectNodeTextureFromColor::GetNodeTitleColor() const
+{
+	const UEdGraphSchema_CustomizableObject* Schema = GetDefault<UEdGraphSchema_CustomizableObject>();
+	return Schema->GetPinTypeColor(Schema->PC_Image);
+}
+
+
+FText UCustomizableObjectNodeTextureFromColor::GetTooltipText() const
+{
+	return LOCTEXT("Texture_From_Color_Tooltip", "Creates a flat color texture from the color provided.");
+}
+
+#undef LOCTEXT_NAMESPACE
diff --git a/Engine/Plugins/Mutable/Source/CustomizableObjectEditor/Private/MuCOE/Nodes/CustomizableObjectNodeTextureLayer.cpp b/Engine/Plugins/Mutable/Source/CustomizableObjectEditor/Private/MuCOE/Nodes/CustomizableObjectNodeTextureLayer.cpp
new file mode 100644
--- /dev/null
+++ b/Engine/Plugins/Mutable/Source/CustomizableObjectEditor/Private/MuCOE/Nodes/CustomizableObjectNodeTextureLayer.cpp
@@ -0,0 +1,180 @@
+// Copyright Epic Games, Inc. All Rights Reserved.
+
+#include "MuCOE/Nodes/CustomizableObjectNodeTextureLayer.h"
+
+#include "MuCO/CustomizableObjectCustomVersion.h"
+#include "MuCOE/EdGraphSchema_CustomizableObject.h"
+#include "MuCOE/UnrealEditorPortabilityHelpers.h"
+
+class UCustomizableObjectNodeRemapPins;
+class UEdGraph;
+
+#define LOCTEXT_NAMESPACE "CustomizableObjectEditor"
+
+TArray<UCustomizableObjectNodeTextureLayer::TextureLayerLayersData> UCustomizableObjectNodeTextureLayer::ArrayDestroyedNodes;
+
+UCustomizableObjectNodeTextureLayer::UCustomizableObjectNodeTextureLayer()
+	: Super()
+{
+	Layers.Add(FCustomizableObjectTextureLayer());
+}
+
+
+void UCustomizableObjectNodeTextureLayer::PostEditChangeProperty(FPropertyChangedEvent& PropertyChangedEvent)
+{
+	FProperty* PropertyThatChanged = PropertyChangedEvent.Property;
+	if ( PropertyThatChanged && PropertyThatChanged->GetName() == TEXT("Layers") )
+	{
+		ReconstructNode();
+	}
+
+	Super::PostEditChangeProperty(PropertyChangedEvent);
+}
+
+
+void UCustomizableObjectNodeTextureLayer::Serialize(FArchive& Ar)
+{
+	Super::Serialize(Ar);
+
+	Ar.UsingCustomVersion(FCustomizableObjectCustomVersion::GUID);
+
+	if (Ar.CustomVer(FCustomizableObjectCustomVersion::GUID) < FCustomizableObjectCustomVersion::PinsNamesImageToTexture)
+	{
+		OutputPinReference = FEdGraphPinReference(FindPin(TEXT("Image")));
+	}
+}
+
+
+void UCustomizableObjectNodeTextureLayer::AllocateDefaultPins(UCustomizableObjectNodeRemapPins* RemapPins)
+{
+	const UEdGraphSchema_CustomizableObject* Schema = GetDefault<UEdGraphSchema_CustomizableObject>();
+
+	UEdGraphPin* OutputPin = CustomCreatePin(EGPD_Output, Schema->PC_Image, FName(TEXT("Texture")));
+	OutputPin->bDefaultValueIsIgnored = true;
+
+	OutputPinReference = FEdGraphPinReference(OutputPin);
+
+	UEdGraphPin* BasePin = CustomCreatePin(EGPD_Input, Schema->PC_Image, FName(TEXT("Base")));
+	BasePin->bDefaultValueIsIgnored = true;
+	
+	for (int32 LayerIndex = 0; LayerIndex < Layers.Num() ; ++LayerIndex )
+	{
+		FString PinName = FString::Printf( TEXT("Layer %d "), LayerIndex );
+		UEdGraphPin* LayerPin = CustomCreatePin(EGPD_Input, Schema->PC_Image, *PinName);
+		LayerPin->bDefaultValueIsIgnored = true;
+
+		PinName = FString::Printf( TEXT("Mask %d "), LayerIndex );
+		UEdGraphPin* MaskPin = CustomCreatePin(EGPD_Input, Schema->PC_Image, *PinName);
+		MaskPin->bDefaultValueIsIgnored = true;
+	}
+}
+
+
+void UCustomizableObjectNodeTextureLayer::BackwardsCompatibleFixup(int32 CustomizableObjectCustomVersion)
+{
+	Super::BackwardsCompatibleFixup(CustomizableObjectCustomVersion);
+	
+	if (CustomizableObjectCustomVersion == FCustomizableObjectCustomVersion::FixPinsNamesImageToTexture2)
+	{
+		if (UEdGraphPin* TexturePin = FindPin(TEXT("Image")))
+		{
+			TexturePin->PinName = TEXT("Texture");
+			UCustomizableObjectNode::ReconstructNode();
+		}
+	}
+}
+
+
+FText UCustomizableObjectNodeTextureLayer::GetNodeTitle(ENodeTitleType::Type TitleType) const
+{
+	return LOCTEXT("Texture_Layer", "Texture Layer");
+}
+
+
+FLinearColor UCustomizableObjectNodeTextureLayer::GetNodeTitleColor() const
+{
+	const UEdGraphSchema_CustomizableObject* Schema = GetDefault<UEdGraphSchema_CustomizableObject>();
+	return Schema->GetPinTypeColor(Schema->PC_Image);
+}
+
+
+void UCustomizableObjectNodeTextureLayer::PrepareForCopying()
+{
+	// Destroy previous copy with possible out of date information just in case
+	int32 IndexToRemove = GetIndexInArrayDestroyedNodes(this);
+
+	if (IndexToRemove != -1)
+	{
+		ArrayDestroyedNodes.RemoveAt(IndexToRemove);
+	}
+
+	TextureLayerLayersData Data;
+	Data.Name = GetName();
+	Data.Graph = GetGraph();
+	Data.Layers = Layers;
+	ArrayDestroyedNodes.Add(Data);
+}
+
+void UCustomizableObjectNodeTextureLayer::PostPasteNode()
+{
+	int32 IndexToRemove = GetIndexInArrayDestroyedNodes(this);
+
+	if (IndexToRemove != -1)
+	{
+		Layers = ArrayDestroyedNodes[IndexToRemove].Layers;
+		ArrayDestroyedNodes.RemoveAt(IndexToRemove);
+	}
+}
+
+
+bool UCustomizableObjectNodeTextureLayer::CanConnect(const UEdGraphPin* InOwnedInputPin, const UEdGraphPin* InOutputPin, bool& bOutIsOtherNodeBlocklisted, bool& bOutArePinsCompatible) const
+{
+	const UEdGraphSchema_CustomizableObject* Schema = GetDefault<UEdGraphSchema_CustomizableObject>();
+
+	// In layer pins we accept both colors and images
+	for ( int LayerIndex = 0; LayerIndex < Layers.Num(); ++LayerIndex )
+	{
+		if (LayerPin(LayerIndex) == InOwnedInputPin )
+		{
+			return (InOutputPin->PinType.PinCategory==Schema->PC_Image)
+				|| (InOutputPin->PinType.PinCategory==Schema->PC_Color);
+		}
+	}
+
+	return Super::CanConnect( InOwnedInputPin, InOutputPin, bOutIsOtherNodeBlocklisted, bOutArePinsCompatible);
+}
+
+
+UEdGraphPin* UCustomizableObjectNodeTextureLayer::OutputPin() const
+{
+	return OutputPinReference.Get();
+}
+
+
+int32 UCustomizableObjectNodeTextureLayer::GetIndexInArrayDestroyedNodes(UCustomizableObjectNodeTextureLayer* Node)
+{
+	const int32 MaxIndex = ArrayDestroyedNodes.Num();
+	FString NodeName = Node->GetName();
+	UEdGraph* NodeGraph = Node->GetGraph();
+	for (int32 i = 0; i < MaxIndex; ++i)
+	{
+		TextureLayerLayersData& DataTemp = ArrayDestroyedNodes[i];
+
+		//if ((DataTemp.Name == NodeName) && (DataTemp.Graph == NodeGraph))
+		if (DataTemp.Name == NodeName) // Copying between graphs would yeld incorredct array data, are the names of the elements unique?
+		{
+			return i;
+		}
+	}
+
+	return -1;
+}
+
+
+FText UCustomizableObjectNodeTextureLayer::GetTooltipText() const
+{
+	return LOCTEXT("Texture_Layer_Tooltip", "Combines multiple textures into one.");
+}
+
+#undef LOCTEXT_NAMESPACE
+
