[mutable] Moved the Mutable plugin out of Experimental status into Beta.

diff --git a/Engine/Plugins/Mutable/Source/CustomizableObjectEditor/Private/MuCOE/Nodes/CustomizableObjectNodeTextureParameter.cpp b/Engine/Plugins/Mutable/Source/CustomizableObjectEditor/Private/MuCOE/Nodes/CustomizableObjectNodeTextureParameter.cpp
new file mode 100644
--- /dev/null
+++ b/Engine/Plugins/Mutable/Source/CustomizableObjectEditor/Private/MuCOE/Nodes/CustomizableObjectNodeTextureParameter.cpp
@@ -0,0 +1,74 @@
+// Copyright Epic Games, Inc. All Rights Reserved.
+
+#include "MuCOE/Nodes/CustomizableObjectNodeTextureParameter.h"
+
+#include "MuCO/CustomizableObjectCustomVersion.h"
+#include "MuCOE/EdGraphSchema_CustomizableObject.h"
+
+class UCustomizableObjectNodeRemapPins;
+
+#define LOCTEXT_NAMESPACE "CustomizableObjectEditor"
+
+
+void UCustomizableObjectNodeTextureParameter::AllocateDefaultPins(UCustomizableObjectNodeRemapPins* RemapPins)
+{
+	const UEdGraphSchema_CustomizableObject* Schema = GetDefault<UEdGraphSchema_CustomizableObject>();
+
+	FString PinName = TEXT("Value");
+	UEdGraphPin* ValuePin = CustomCreatePin(EGPD_Output, Schema->PC_Image, FName(*PinName));
+	ValuePin->bDefaultValueIsIgnored = true;
+}
+
+
+FText UCustomizableObjectNodeTextureParameter::GetNodeTitle(ENodeTitleType::Type TitleType) const
+{
+	if (TitleType == ENodeTitleType::ListView || ParameterName.IsEmpty())
+	{
+		return LOCTEXT("Texture_Parameter", "Texture Parameter");
+	}
+	else if (TitleType == ENodeTitleType::EditableTitle)
+	{
+		return FText::Format(LOCTEXT("Texture_Parameter_EditableTitle", "{0}"), FText::FromString(ParameterName));
+	}
+	else
+	{
+		return FText::Format(LOCTEXT("Texture_Parameter_Title", "{0}\nTexture Parameter"), FText::FromString(ParameterName));
+	}
+}
+
+
+FLinearColor UCustomizableObjectNodeTextureParameter::GetNodeTitleColor() const
+{
+	const UEdGraphSchema_CustomizableObject* Schema = GetDefault<UEdGraphSchema_CustomizableObject>();
+	return Schema->GetPinTypeColor(Schema->PC_Image);
+}
+
+
+FText UCustomizableObjectNodeTextureParameter::GetTooltipText() const
+{
+	return LOCTEXT("Texture_Parameter_Tooltip", "Expose a runtime modifiable texture parameter from the Customizable Object.");
+}
+
+
+void UCustomizableObjectNodeTextureParameter::OnRenameNode(const FString& NewName)
+{
+	if (!NewName.IsEmpty())
+	{
+		ParameterName = NewName;
+	}
+}
+
+
+void UCustomizableObjectNodeTextureParameter::BackwardsCompatibleFixup(int32 CustomizableObjectCustomVersion)
+{
+	Super::BackwardsCompatibleFixup(CustomizableObjectCustomVersion);
+	
+	if (CustomizableObjectCustomVersion == FCustomizableObjectCustomVersion::NodeTextureParameterDefaultToReferenceValue)
+	{
+		ReferenceValue = DefaultValue;
+		DefaultValue = {};
+	}
+}
+
+#undef LOCTEXT_NAMESPACE
+
diff --git a/Engine/Plugins/Mutable/Source/CustomizableObjectEditor/Private/MuCOE/Nodes/CustomizableObjectNodeTextureSample.cpp b/Engine/Plugins/Mutable/Source/CustomizableObjectEditor/Private/MuCOE/Nodes/CustomizableObjectNodeTextureSample.cpp
new file mode 100644
--- /dev/null
+++ b/Engine/Plugins/Mutable/Source/CustomizableObjectEditor/Private/MuCOE/Nodes/CustomizableObjectNodeTextureSample.cpp
@@ -0,0 +1,72 @@
+// Copyright Epic Games, Inc. All Rights Reserved.
+
+#include "MuCOE/Nodes/CustomizableObjectNodeTextureSample.h"
+
+#include "MuCOE/EdGraphSchema_CustomizableObject.h"
+
+class UCustomizableObjectNodeRemapPins;
+struct FPropertyChangedEvent;
+
+#define LOCTEXT_NAMESPACE "CustomizableObjectEditor"
+
+
+UCustomizableObjectNodeTextureSample::UCustomizableObjectNodeTextureSample()
+	: Super()
+{
+
+}
+
+
+void UCustomizableObjectNodeTextureSample::PostEditChangeProperty(FPropertyChangedEvent& PropertyChangedEvent)
+{
+	//UProperty* PropertyThatChanged = PropertyChangedEvent.Property;
+	//if ( PropertyThatChanged && PropertyThatChanged->GetName() == TEXT("NumLODs") )
+	{
+	}
+
+	Super::PostEditChangeProperty(PropertyChangedEvent);
+}
+
+
+void UCustomizableObjectNodeTextureSample::AllocateDefaultPins(UCustomizableObjectNodeRemapPins* RemapPins)
+{
+	const UEdGraphSchema_CustomizableObject* Schema = GetDefault<UEdGraphSchema_CustomizableObject>();
+
+	FString PinName = TEXT("Color");
+	UEdGraphPin* ColorPin = CustomCreatePin(EGPD_Output, Schema->PC_Color, FName(*PinName));
+	ColorPin->bDefaultValueIsIgnored = true;
+
+	PinName = TEXT("Texture");
+	UEdGraphPin* TexturePin = CustomCreatePin(EGPD_Input, Schema->PC_Image, FName(*PinName));
+	TexturePin->bDefaultValueIsIgnored = true;
+
+	PinName = TEXT("X");
+	UEdGraphPin* XPin = CustomCreatePin(EGPD_Input, Schema->PC_Float, FName(*PinName));
+	XPin->bDefaultValueIsIgnored = true;
+
+	PinName = TEXT("Y");
+	UEdGraphPin* YPin = CustomCreatePin(EGPD_Input, Schema->PC_Float, FName(*PinName));
+	YPin->bDefaultValueIsIgnored = true;
+}
+
+
+FText UCustomizableObjectNodeTextureSample::GetNodeTitle(ENodeTitleType::Type TitleType) const
+{
+	return LOCTEXT("Sample_Texture", "Sample Texture");
+}
+
+
+FLinearColor UCustomizableObjectNodeTextureSample::GetNodeTitleColor() const
+{
+	const UEdGraphSchema_CustomizableObject* Schema = GetDefault<UEdGraphSchema_CustomizableObject>();
+	return Schema->GetPinTypeColor(Schema->PC_Color);
+}
+
+
+FText UCustomizableObjectNodeTextureSample::GetTooltipText() const
+{
+	return LOCTEXT("Texture_Sample_Tooltip","Get the color found in a texture at the targeted X and Y position (from 0.0 to 1.0, both included).");
+}
+
+#undef LOCTEXT_NAMESPACE
+
diff --git a/Engine/Plugins/Mutable/Source/CustomizableObjectEditor/Private/MuCOE/Nodes/CustomizableObjectNodeTextureToChannels.cpp b/Engine/Plugins/Mutable/Source/CustomizableObjectEditor/Private/MuCOE/Nodes/CustomizableObjectNodeTextureToChannels.cpp
new file mode 100644
--- /dev/null
+++ b/Engine/Plugins/Mutable/Source/CustomizableObjectEditor/Private/MuCOE/Nodes/CustomizableObjectNodeTextureToChannels.cpp
@@ -0,0 +1,111 @@
+// Copyright Epic Games, Inc. All Rights Reserved.
+
+#include "MuCOE/Nodes/CustomizableObjectNodeTextureToChannels.h"
+
+#include "MuCO/CustomizableObjectCustomVersion.h"
+#include "MuCOE/EdGraphSchema_CustomizableObject.h"
+
+class UCustomizableObjectNodeRemapPins;
+struct FPropertyChangedEvent;
+
+#define LOCTEXT_NAMESPACE "CustomizableObjectEditor"
+
+
+void UCustomizableObjectNodeTextureToChannels::PostEditChangeProperty(FPropertyChangedEvent& PropertyChangedEvent)
+{
+	Super::PostEditChangeProperty(PropertyChangedEvent);
+}
+
+
+void UCustomizableObjectNodeTextureToChannels::Serialize(FArchive& Ar)
+{
+	Super::Serialize(Ar);
+
+	Ar.UsingCustomVersion(FCustomizableObjectCustomVersion::GUID);
+
+	if (Ar.CustomVer(FCustomizableObjectCustomVersion::GUID) < FCustomizableObjectCustomVersion::PinsNamesImageToTexture)
+	{
+		InputPinReference = FEdGraphPinReference(FindPin(TEXT("Image")));
+	}
+}
+
+
+void UCustomizableObjectNodeTextureToChannels::AllocateDefaultPins(UCustomizableObjectNodeRemapPins* RemapPins)
+{
+	const UEdGraphSchema_CustomizableObject* Schema = GetDefault<UEdGraphSchema_CustomizableObject>();
+
+	{
+		FString PinName = TEXT("Texture");
+		UEdGraphPin* InputPin = CustomCreatePin(EGPD_Input, Schema->PC_Image, FName(*PinName));
+		InputPin->bDefaultValueIsIgnored = true;
+
+		InputPinReference = FEdGraphPinReference(InputPin);
+	}
+
+	{
+		FString PinName = TEXT("R");
+		UEdGraphPin* RPin = CustomCreatePin(EGPD_Output, Schema->PC_Image, FName(*PinName));
+		RPin->bDefaultValueIsIgnored = true;
+	}
+
+	{
+		FString PinName = TEXT("G");
+		UEdGraphPin* GPin = CustomCreatePin(EGPD_Output, Schema->PC_Image, FName(*PinName));
+		GPin->bDefaultValueIsIgnored = true;
+	}
+
+	{
+		FString PinName = TEXT("B");
+		UEdGraphPin* BPin = CustomCreatePin(EGPD_Output, Schema->PC_Image, FName(*PinName));
+		BPin->bDefaultValueIsIgnored = true;
+	}
+
+	{
+		FString PinName = TEXT("A");
+		UEdGraphPin* APin = CustomCreatePin(EGPD_Output, Schema->PC_Image, FName(*PinName));
+		APin->bDefaultValueIsIgnored = true;
+	}
+}
+
+
+void UCustomizableObjectNodeTextureToChannels::BackwardsCompatibleFixup(int32 CustomizableObjectCustomVersion)
+{
+	Super::BackwardsCompatibleFixup(CustomizableObjectCustomVersion);
+	
+	if (CustomizableObjectCustomVersion == FCustomizableObjectCustomVersion::FixPinsNamesImageToTexture2)
+	{
+		if (UEdGraphPin* TexturePin = FindPin(TEXT("Image"))) {
+			TexturePin->PinName = TEXT("Texture");
+			UCustomizableObjectNode::ReconstructNode();
+		}
+	}
+}
+
+
+UEdGraphPin* UCustomizableObjectNodeTextureToChannels::InputPin() const
+{
+	return InputPinReference.Get();
+}
+
+
+FText UCustomizableObjectNodeTextureToChannels::GetNodeTitle(ENodeTitleType::Type TitleType) const
+{
+	return LOCTEXT("Texture_To_Channels", "Break Texture");
+}
+
+
+FLinearColor UCustomizableObjectNodeTextureToChannels::GetNodeTitleColor() const
+{
+	const UEdGraphSchema_CustomizableObject* Schema = GetDefault<UEdGraphSchema_CustomizableObject>();
+	return Schema->GetPinTypeColor(Schema->PC_Image);
+}
+
+
+FText UCustomizableObjectNodeTextureToChannels::GetTooltipText() const
+{
+	return LOCTEXT("Texture_To_Channels_Tooltip", "Get the red, green, blue and alpha channels of a texture as four separate grayscale textures.");
+}
+
+
+#undef LOCTEXT_NAMESPACE
+
diff --git a/Engine/Plugins/Mutable/Source/CustomizableObjectEditor/Private/MuCOE/Nodes/CustomizableObjectNodeTextureTransform.cpp b/Engine/Plugins/Mutable/Source/CustomizableObjectEditor/Private/MuCOE/Nodes/CustomizableObjectNodeTextureTransform.cpp
new file mode 100644
--- /dev/null
+++ b/Engine/Plugins/Mutable/Source/CustomizableObjectEditor/Private/MuCOE/Nodes/CustomizableObjectNodeTextureTransform.cpp
@@ -0,0 +1,79 @@
+// Copyright Epic Games, Inc. All Rights Reserved.
+
+#include "MuCOE/Nodes/CustomizableObjectNodeTextureTransform.h"
+
+#include "MuCO/CustomizableObjectCustomVersion.h"
+#include "MuCOE/EdGraphSchema_CustomizableObject.h"
+
+class UCustomizableObjectNodeRemapPins;
+
+#define LOCTEXT_NAMESPACE "CustomizableObjectEditor"
+
+void UCustomizableObjectNodeTextureTransform::AllocateDefaultPins(UCustomizableObjectNodeRemapPins* RemapPins)
+{
+	const UEdGraphSchema_CustomizableObject* Schema = GetDefault<UEdGraphSchema_CustomizableObject>();
+
+	FString PinName = TEXT("Texture");
+	UEdGraphPin* OutputPin = CustomCreatePin(EGPD_Output, Schema->PC_Image, FName(*PinName));
+	OutputPin->bDefaultValueIsIgnored = true;
+
+	PinName = TEXT("Base Texture");
+	UEdGraphPin* ImagePin = CustomCreatePin(EGPD_Input, Schema->PC_Image, FName(*PinName));
+
+	BaseImagePinReference = FEdGraphPinReference(ImagePin);
+	
+	PinName = TEXT("Offset X");
+	UEdGraphPin* OffsetXPin = CustomCreatePin(EGPD_Input, Schema->PC_Float, FName(*PinName));
+	OffsetXPin->bDefaultValueIsIgnored = true;
+	
+	PinName = TEXT("Offset Y");
+	UEdGraphPin* OffsetYPin = CustomCreatePin(EGPD_Input, Schema->PC_Float, FName(*PinName));
+	OffsetYPin->bDefaultValueIsIgnored = true;
+
+	PinName = TEXT("Scale X");
+	UEdGraphPin* ScaleXPin = CustomCreatePin(EGPD_Input, Schema->PC_Float, FName(*PinName));
+	ScaleXPin->bDefaultValueIsIgnored = true;
+
+	PinName = TEXT("Scale Y");
+	UEdGraphPin* ScaleYPin = CustomCreatePin(EGPD_Input, Schema->PC_Float, FName(*PinName));
+	ScaleYPin->bDefaultValueIsIgnored = true;
+	
+	PinName = TEXT("Rotation");
+	UEdGraphPin* RotationPin = CustomCreatePin(EGPD_Input, Schema->PC_Float, FName(*PinName));
+	RotationPin->bDefaultValueIsIgnored = true;
+}
+
+UEdGraphPin* UCustomizableObjectNodeTextureTransform::GetBaseImagePin() const
+{
+	return BaseImagePinReference.Get();
+}
+
+FText UCustomizableObjectNodeTextureTransform::GetNodeTitle(ENodeTitleType::Type TitleType) const
+{
+	return LOCTEXT("Texture_Transform", "Texture Transform");
+}
+
+FLinearColor UCustomizableObjectNodeTextureTransform::GetNodeTitleColor() const
+{
+	const UEdGraphSchema_CustomizableObject* Schema = GetDefault<UEdGraphSchema_CustomizableObject>();
+	return Schema->GetPinTypeColor(Schema->PC_Image);
+}
+
+FText UCustomizableObjectNodeTextureTransform::GetTooltipText() const
+{
+	return LOCTEXT("Texture_Transform_Tooltip", 
+			"Applies a linear transform, rotation and scale around the center of the image plus translation, "
+			"to the content of Base Texture. Rotation is in the range [0 .. 1], 1 being full rotation, offset " 
+			"and scale are in output image normalized coordinates with origin at the center of the image. " 
+			"If Keep Aspect Ratio is set, an scaling factor preserving aspect ratio will be used as identity.");
+}
+
+void UCustomizableObjectNodeTextureTransform::Serialize(FArchive& Ar)
+{
+	Super::Serialize(Ar);
+
+	Ar.UsingCustomVersion(FCustomizableObjectCustomVersion::GUID);
+}
+
+
+#undef LOCTEXT_NAMESPACE
diff --git a/Engine/Plugins/Mutable/Source/CustomizableObjectEditor/Private/MuCOE/Nodes/CustomizableObjectNodeTextureVariation.cpp b/Engine/Plugins/Mutable/Source/CustomizableObjectEditor/Private/MuCOE/Nodes/CustomizableObjectNodeTextureVariation.cpp
new file mode 100644
--- /dev/null
+++ b/Engine/Plugins/Mutable/Source/CustomizableObjectEditor/Private/MuCOE/Nodes/CustomizableObjectNodeTextureVariation.cpp
@@ -0,0 +1,72 @@
+// Copyright Epic Games, Inc. All Rights Reserved.
+
+#include "MuCOE/Nodes/CustomizableObjectNodeTextureVariation.h"
+
+#include "MuCOE/EdGraphSchema_CustomizableObject.h"
+
+class UCustomizableObjectNodeRemapPins;
+
+#define LOCTEXT_NAMESPACE "CustomizableObjectEditor"
+
+
+UCustomizableObjectNodeTextureVariation::UCustomizableObjectNodeTextureVariation()
+	: Super()
+{
+
+}
+
+
+void UCustomizableObjectNodeTextureVariation::PostEditChangeProperty(FPropertyChangedEvent& PropertyChangedEvent)
+{
+	FProperty* PropertyThatChanged = PropertyChangedEvent.Property;
+	if ( PropertyThatChanged  )
+	{
+		ReconstructNode();
+	}
+
+	Super::PostEditChangeProperty(PropertyChangedEvent);
+}
+
+
+void UCustomizableObjectNodeTextureVariation::AllocateDefaultPins(UCustomizableObjectNodeRemapPins* RemapPins)
+{
+	const UEdGraphSchema_CustomizableObject* Schema = GetDefault<UEdGraphSchema_CustomizableObject>();
+
+	UEdGraphPin* OutputPin = CustomCreatePin(EGPD_Output, Schema->PC_Image, FName(TEXT("Texture")));
+	OutputPin->bDefaultValueIsIgnored = true;
+
+	for ( int VariationIndex = Variations.Num() - 1; VariationIndex >= 0; --VariationIndex)
+	{
+		FString PinName = FString::Printf( TEXT("Variation %d"), VariationIndex);
+		UEdGraphPin* VariationPin = CustomCreatePin(EGPD_Input, Schema->PC_Image, FName(*PinName), false);
+		VariationPin->bDefaultValueIsIgnored = true;
+
+		FString FriendlyName = FString::Printf(TEXT("Variation %d [%s]"), VariationIndex, *Variations[VariationIndex].Tag);
+		VariationPin->PinFriendlyName = FText::FromString(FriendlyName);
+	}
+
+	UEdGraphPin* DefaultVariation = CustomCreatePin(EGPD_Input, Schema->PC_Image, FName(TEXT("Default")), false);
+	DefaultVariation->bDefaultValueIsIgnored = true;
+}
+
+
+FText UCustomizableObjectNodeTextureVariation::GetNodeTitle(ENodeTitleType::Type TitleType) const
+{
+	return LOCTEXT("Texture_Variation", "Texture Variation");
+}
+
+
+FLinearColor UCustomizableObjectNodeTextureVariation::GetNodeTitleColor() const
+{
+	const UEdGraphSchema_CustomizableObject* Schema = GetDefault<UEdGraphSchema_CustomizableObject>();
+	return Schema->GetPinTypeColor(Schema->PC_Image);
+}
+
+
+FText UCustomizableObjectNodeTextureVariation::GetTooltipText() const
+{
+	return LOCTEXT("Texture_Variation_Tooltip", "Select a texture depending on what tags are active.");
+}
+
+#undef LOCTEXT_NAMESPACE
+
diff --git a/Engine/Plugins/Mutable/Source/CustomizableObjectEditor/Private/MuCOE/Nodes/SCustomizableObjectNodeMaterial.cpp b/Engine/Plugins/Mutable/Source/CustomizableObjectEditor/Private/MuCOE/Nodes/SCustomizableObjectNodeMaterial.cpp
new file mode 100644
--- /dev/null
+++ b/Engine/Plugins/Mutable/Source/CustomizableObjectEditor/Private/MuCOE/Nodes/SCustomizableObjectNodeMaterial.cpp
@@ -0,0 +1,27 @@
+// Copyright Epic Games, Inc. All Rights Reserved.
+
+#include "MuCOE/Nodes/SCustomizableObjectNodeMaterial.h"
+
+#include "MuCOE/EdGraphSchema_CustomizableObject.h"
+#include "MuCOE/Nodes/SCustomizableObjectNodeMaterialPinImage.h"
+#include "MuCOE/Nodes/CustomizableObjectNodeMaterial.h"
+
+class SGraphPin;
+
+
+void SCustomizableObjectNodeMaterial::Construct(const FArguments& InArgs, UCustomizableObjectNodeMaterial* InGraphNode)
+{
+	SCustomizableObjectNode::Construct({}, InGraphNode);
+}
+
+
+TSharedPtr<SGraphPin> SCustomizableObjectNodeMaterial::CreatePinWidget(UEdGraphPin* Pin) const
+{
+	if ((Pin->PinType.PinCategory == UEdGraphSchema_CustomizableObject::PC_Image || Pin->PinType.PinCategory == UEdGraphSchema_CustomizableObject::PC_PassThroughImage) &&
+		Pin->Direction == EGPD_Input)
+	{
+		return SNew(SCustomizableObjectNodeMaterialPinImage, Pin);
+	}
+
+	return SGraphNode::CreatePinWidget(Pin);
+}
diff --git a/Engine/Plugins/Mutable/Source/CustomizableObjectEditor/Private/MuCOE/Nodes/SCustomizableObjectNodeMaterial.h b/Engine/Plugins/Mutable/Source/CustomizableObjectEditor/Private/MuCOE/Nodes/SCustomizableObjectNodeMaterial.h
new file mode 100644
--- /dev/null
+++ b/Engine/Plugins/Mutable/Source/CustomizableObjectEditor/Private/MuCOE/Nodes/SCustomizableObjectNodeMaterial.h
@@ -0,0 +1,24 @@
+// Copyright Epic Games, Inc. All Rights Reserved.
+
+#pragma once
+
+#include "SCustomizableObjectNode.h"
+#include "SGraphNode.h"
+
+class UCustomizableObjectNodeMaterial;
+class SGraphPin;
+class UEdGraphPin;
+
+
+/** Custom widget for the Material node. */
+class SCustomizableObjectNodeMaterial : public SCustomizableObjectNode
+{
+public:
+	SLATE_BEGIN_ARGS(SCustomizableObjectNodeMaterial) {}
+	SLATE_END_ARGS();
+
+	void Construct(const FArguments& InArgs, UCustomizableObjectNodeMaterial* InGraphNode);
+
+	// SGraphNode interface
+	virtual TSharedPtr<SGraphPin> CreatePinWidget(UEdGraphPin* Pin) const override;
+};
diff --git a/Engine/Plugins/Mutable/Source/CustomizableObjectEditor/Private/MuCOE/Nodes/SCustomizableObjectNodeMaterialPinImage.cpp b/Engine/Plugins/Mutable/Source/CustomizableObjectEditor/Private/MuCOE/Nodes/SCustomizableObjectNodeMaterialPinImage.cpp
new file mode 100644
--- /dev/null
+++ b/Engine/Plugins/Mutable/Source/CustomizableObjectEditor/Private/MuCOE/Nodes/SCustomizableObjectNodeMaterialPinImage.cpp
@@ -0,0 +1,81 @@
+// Copyright Epic Games, Inc. All Rights Reserved.
+
+#include "MuCOE/Nodes/SCustomizableObjectNodeMaterialPinImage.h"
+
+#include "MuCOE/EdGraphSchema_CustomizableObject.h"
+#include "MuCOE/GraphTraversal.h"
+#include "MuCOE/Nodes/CustomizableObjectNodeMaterial.h"
+#include "MuCOE/UnrealEditorPortabilityHelpers.h"
+#include "Widgets/Input/SEditableTextBox.h"
+#include "Widgets/SToolTip.h"
+
+class IToolTip;
+class SWidget;
+struct FSlateBrush;
+
+#define LOCTEXT_NAMESPACE "CustomizableObjectEditor"
+
+
+void SCustomizableObjectNodeMaterialPinImage::Construct(const FArguments& InArgs, UEdGraphPin* InGraphPinObj)
+{
+	SCustomizableObjectNodePin::Construct(SCustomizableObjectNodePin::FArguments(), InGraphPinObj);
+
+	// Override previously defined tool tip.
+	TSharedPtr<IToolTip> TooltipWidget = SNew(SToolTip)
+		.Text(this, &SCustomizableObjectNodeMaterialPinImage::GetPinTooltipText);
+
+	SetToolTip(TooltipWidget);
+}
+
+
+TSharedRef<SWidget>	SCustomizableObjectNodeMaterialPinImage::GetDefaultValueWidget()
+{
+	LabelAndValue->SetWrapSize(TNumericLimits<float>::Max()); // Remove warping.
+
+	return SNew(SEditableTextBox)
+		.Style(FAppStyle::Get(), "Graph.EditableTextBox")
+		.Text(this, &SCustomizableObjectNodeMaterialPinImage::GetDefaultValueText)
+		.SelectAllTextWhenFocused(false)
+		.Visibility(this, &SCustomizableObjectNodeMaterialPinImage::GetDefaultValueVisibility)
+		.IsReadOnly(true)
+		.ForegroundColor(FSlateColor::UseForeground());
+}
+
+
+FText SCustomizableObjectNodeMaterialPinImage::GetPinTooltipText() const
+{
+	if (GraphPinObj->bOrphanedPin)
+	{
+		return LOCTEXT("PinModeMutableOrpahan", "Pin not disapearing due to being connected or having a property modified.");
+	}
+	
+	if (GraphPinObj->PinType.PinCategory == UEdGraphSchema_CustomizableObject::PC_Image)
+	{
+		return LOCTEXT("PinModeMutableTooltip", "Texture Parameter goes through Mutable.");
+	}
+	else
+	{
+		return LOCTEXT("PinModePassthroughTooltip", "Texture Parameter is ignored by Mutable.");
+	}
+}
+
+
+FText SCustomizableObjectNodeMaterialPinImage::GetDefaultValueText() const
+{
+	if (GraphPinObj->PinType.PinCategory == UEdGraphSchema_CustomizableObject::PC_Image)
+	{
+		return LOCTEXT("PinModeMutable", "mutable");
+	}
+	else 
+	{
+		return LOCTEXT("PinModePassthrough", "passthrough");
+	}
+}
+
+
+EVisibility SCustomizableObjectNodeMaterialPinImage::GetDefaultValueVisibility() const
+{
+	return EVisibility::Visible;
+}
+
+#undef LOCTEXT_NAMESPACE
diff --git a/Engine/Plugins/Mutable/Source/CustomizableObjectEditor/Private/MuCOE/Nodes/SCustomizableObjectNodeMaterialPinImage.h b/Engine/Plugins/Mutable/Source/CustomizableObjectEditor/Private/MuCOE/Nodes/SCustomizableObjectNodeMaterialPinImage.h
new file mode 100644
--- /dev/null
+++ b/Engine/Plugins/Mutable/Source/CustomizableObjectEditor/Private/MuCOE/Nodes/SCustomizableObjectNodeMaterialPinImage.h
@@ -0,0 +1,34 @@
+// Copyright Epic Games, Inc. All Rights Reserved.
+
+#pragma once
+
+#include "SCustomizableObjectNodePin.h"
+#include "SGraphPin.h"
+
+class SWidget;
+class UEdGraphPin;
+struct FSlateBrush;
+
+
+/** Implements the "MUTABLE" and "PASSTHROUGH" text next to the pin name. */
+class SCustomizableObjectNodeMaterialPinImage : public SCustomizableObjectNodePin
+{
+public:
+
+	SLATE_BEGIN_ARGS(SCustomizableObjectNodeMaterialPinImage) {}
+	SLATE_END_ARGS()
+
+	void Construct(const FArguments& InArgs, UEdGraphPin* InGraphPinObj);
+
+	TSharedRef<SWidget>	GetDefaultValueWidget() override;
+
+private:
+	/** Return pin state text. */
+	FText GetDefaultValueText() const;
+
+	/** Return true if the pin state is visible. */
+	EVisibility GetDefaultValueVisibility() const;
+
+	/** Return pin tool tip. */
+	FText GetPinTooltipText() const;
+};
diff --git a/Engine/Plugins/Mutable/Source/CustomizableObjectEditor/Private/MuCOE/SCustomizableObjectEditorPerformanceReport.cpp b/Engine/Plugins/Mutable/Source/CustomizableObjectEditor/Private/MuCOE/SCustomizableObjectEditorPerformanceReport.cpp
new file mode 100644
--- /dev/null
+++ b/Engine/Plugins/Mutable/Source/CustomizableObjectEditor/Private/MuCOE/SCustomizableObjectEditorPerformanceReport.cpp
@@ -0,0 +1,1474 @@
+// Copyright Epic Games, Inc. All Rights Reserved.
+
+
+#include "MuCOE/SCustomizableObjectEditorPerformanceReport.h"
+
+#include "AssetThumbnail.h"
+#include "DetailsViewArgs.h"
+#include "Editor.h"
+#include "Framework/MultiBox/MultiBoxBuilder.h"
+#include "Framework/Views/TableViewMetadata.h"
+#include "Modules/ModuleManager.h"
+#include "MuCO/CustomizableObject.h"
+#include "MuCO/CustomizableObjectInstance.h"
+#include "MuCOE/CustomizableObjectEditorStyle.h"
+#include "MuCOE/SCustomizableObjectEditorTextureAnalyzer.h"
+#include "MuCOE/UnrealEditorPortabilityHelpers.h"
+#include "PropertyCustomizationHelpers.h"
+#include "MuCO/CustomizableObjectInstancePrivate.h"
+#include "MuCO/CustomizableObjectPrivate.h"
+#include "Widgets/Input/SComboButton.h"
+#include "Widgets/Input/SNumericEntryBox.h"
+#include "Widgets/Layout/SScrollBox.h"
+#include "Widgets/Views/SListView.h"
+
+#include UE_INLINE_GENERATED_CPP_BY_NAME(SCustomizableObjectEditorPerformanceReport)
+
+
+class ITableRow;
+class STableViewBase;
+class SWidget;
+struct FGeometry;
+struct FSlateBrush;
+
+
+#define LOCTEXT_NAMESPACE "CustomizableObjectEditor"
+
+static const FName ColumnID_PR_LOD_Label("LOD");
+static const FName ColumnID_PR_InstanceName_Label("InstanceName");
+static const FName ColumnID_PR_TexturesTotalSize_Label("TexturesTotalSize");
+static const FName ColumnID_PR_NumberOfTriangles_Label("NumberOfTriangles");
+static const FName ColumnID_PR_Name_Label("ParameterName");
+static const FName ColumnID_PR_UpdateTime_Label("UpdateTimeState");
+static const FName ColumnID_PR_Delta_Label("WorstCaseDelta");
+static const FName ColumnID_PR_Average_Label("Average");
+
+void UWorstCasePerformanceReportInstance::SetInstance(UCustomizableObjectInstance* UpdatedInstance)
+{
+	if (UpdatedInstance)
+	{
+		if (!WorstCaseInstance)
+		{
+			/* Uncommment next line if we need the report to have the valid time entries for state parameters expanded by default */
+			//SubVisibility = EVisibility::Visible;
+		}
+		WorstCaseInstance = UpdatedInstance->Clone();
+	}
+}
+
+void SCustomizableObjecEditorPerformanceReport::Construct(const FArguments& InArgs)
+{
+	ReportSubject = InArgs._CustomizableObject;
+	StopTests();
+	ResetSavedData();
+	ResetPerformanceReportOptions();
+
+	SAssignNew(PerformanceReportTextureAnalyzer, SCustomizableObjecEditorTextureAnalyzer)
+		.CustomizableObjectEditor(nullptr)
+		.CustomizableObjectInstanceEditor(nullptr);
+
+	FPropertyEditorModule& PropPlugin = FModuleManager::LoadModuleChecked<FPropertyEditorModule>("PropertyEditor");
+	FDetailsViewArgs DetailsViewArgs;
+	DetailsViewArgs.NotifyHook = this;
+	DetailsViewArgs.NameAreaSettings = FDetailsViewArgs::ENameAreaSettings::HideNameArea;
+	DetailsViewArgs.bAllowSearch = false;
+	DetailsViewArgs.bShowObjectLabel = false;
+	DetailsViewArgs.bAllowMultipleTopLevelObjects = true;
+
+	float WidgetSize = 128.0f;
+	int ThumbnailSize = 128;
+
+	ChildSlot
+	[
+		SNew(SSplitter)
+
+		// Performance report status and dynamic results (left section)
+		+ SSplitter::Slot()
+		[
+			SNew(SBox)
+			.Padding(5.0f)
+			.MinDesiredWidth(300)
+			[
+				SNew(SVerticalBox)
+
+				// Button with settings dropdown
+				+ SVerticalBox::Slot()
+				.HAlign(HAlign_Center)
+				.AutoHeight()
+				.Padding(5.0f)
+				[
+					SNew(SComboButton)
+					.OnGetMenuContent(this, &SCustomizableObjecEditorPerformanceReport::GenerateOptionsMenuContent)
+					.IsEnabled(ReportSubject != nullptr)
+					.ButtonContent()
+					[
+						SNew(SButton)
+						.Text(LOCTEXT("GeneratePerformanceReport", "Generate Performance Report"))
+						.OnClicked(this, &SCustomizableObjecEditorPerformanceReport::RunPerformanceReport)
+						.ButtonStyle(FAppStyle::Get(), "FlatButton.Success")
+						.ForegroundColor(FLinearColor::White)
+					]
+				]
+
+				+ SVerticalBox::Slot()
+				.HAlign(HAlign_Center)
+				.AutoHeight()
+				.Padding(5.0f)
+				[
+					SNew(SButton)
+					.Text(LOCTEXT("GetWorstTimeInstance", "Select the currently slowest instance"))
+					.OnClicked(this, &SCustomizableObjecEditorPerformanceReport::SelectSlowestInstance)
+					.ButtonStyle(FAppStyle::Get(), "FlatButton.Info")
+					.IsEnabled(this, &SCustomizableObjecEditorPerformanceReport::IsGetWorstInstanceButtonEnabled)
+					// This button target changes fast, I think it's better if it triggers on down instead of on up.
+					.ClickMethod(EButtonClickMethod::MouseDown)
+					.TouchMethod(EButtonTouchMethod::Down)
+					.PressMethod(EButtonPressMethod::ButtonPress)
+				]
+
+				+ SVerticalBox::Slot()
+				.HAlign(HAlign_Center)
+				.AutoHeight()
+				.Padding(5.0f)
+				[
+					SAssignNew(StatusText, STextBlock)
+				]
+
+				// Info on Texture Size worst cases
+				+ SVerticalBox::Slot()
+				.HAlign(HAlign_Left)
+				.AutoHeight()
+				.Padding(5.0f, 10.0f, 5.0f, 5.0f)
+				[
+					SNew(STextBlock)
+					.Text(LOCTEXT("MutableLargestTextures", "Instances with largest mutable textures:"))
+				]
+				+ SVerticalBox::Slot()
+				.AutoHeight()
+				[
+					SNew(SBorder)
+					.Padding(5.0f)
+					[
+						SAssignNew(TextureWorstCasesView, SVerticalBox)
+						.Visibility_Lambda([this]() { return WorstTextureInstances.Num() > 0 && WorstTextureInstances[0] && WorstTextureInstances[0]->WorstCaseInstance ? EVisibility::Visible : EVisibility::Collapsed; })
+					]
+				]
+
+				// Info on Geometry Amount worst cases
+				+ SVerticalBox::Slot()
+				.HAlign(HAlign_Left)
+				.AutoHeight()
+				.Padding(5.0f, 20.0f, 5.0f, 5.0f)
+				[
+					SNew(STextBlock)
+					.Text(LOCTEXT("MutableMostGeometry", "Instances with most triangles:"))
+				]
+				+ SVerticalBox::Slot()
+				.AutoHeight()
+				[
+					SNew(SBorder)
+					.Padding(5.0f)
+					.Visibility_Lambda([this]() { return WorstGeometryInstances.Num() > 0 && WorstGeometryInstances[0] && WorstGeometryInstances[0]->WorstCaseInstance ? EVisibility::Visible : EVisibility::Collapsed; })
+					[
+						SAssignNew(GeometryWorstCasesView, SVerticalBox)
+					]
+				]
+
+				// Info on Update Time worst cases detailed by state and parameter
+				+ SVerticalBox::Slot()
+				.HAlign(HAlign_Left)
+				.AutoHeight()
+				.Padding(5.0f, 20.0f, 5.0f, 5.0f)
+				[
+					SNew(STextBlock)
+					.Text(LOCTEXT("MutableLongestUpdateTimesDetailed", "Instances with longest updates by state and parameter:"))
+				]
+				+ SVerticalBox::Slot()
+				.FillHeight(1.0)
+				[
+					SNew(SBorder)
+					.Padding(5.0f)
+					.Visibility_Lambda([this]() { return WorstTimeStateInstances.Num() > 0
+						&& WorstTimeStateInstances[0].Num() > 0
+						&& WorstTimeStateInstances[0][0]
+						&& WorstTimeStateInstances[0][0]->WorstCaseInstance
+						? EVisibility::Visible : EVisibility::Collapsed; })
+					[
+						SAssignNew(TimeWorstCasesDetailedView, SScrollBox)
+					]
+				]
+			]
+		]
+
+		// Currently selected worst instance and its texture analyzer (right section)
+		+ SSplitter::Slot()
+		[
+			BuildAnalyzer()
+		]
+	];
+
+	BuildTextureWorstCasesView();
+	BuildGeometryWorstCasesView();
+	BuildTimeWorstCasesDetailedView();
+	StatusText->SetText(LOCTEXT("PerformanceReportStatusNone","Generate Performance Report to collect test data."));
+
+	HelperCallback = NewObject<UPerformanceReportHelper>(GetTransientPackage());
+	HelperCallback->PerformanceReport = this;
+}
+
+SCustomizableObjecEditorPerformanceReport::~SCustomizableObjecEditorPerformanceReport()
+{
+	StopTests();
+	
+	if (CurrentReportWorstCase && CurrentReportWorstCase->WorstCaseInstance)
+	{
+		CurrentReportWorstCase->WorstCaseInstance->UpdatedDelegate.RemoveDynamic(HelperCallback, &UPerformanceReportHelper::DelegatedCallback);
+	}
+	HelperCallback = nullptr;
+	WorstTimeInstance = nullptr;
+	CurrentReportWorstCase = nullptr;
+	CurrentReportInstance = nullptr;
+}
+
+TSharedRef<SWidget> SCustomizableObjecEditorPerformanceReport::GenerateOptionsMenuContent()
+{
+	const bool bShouldCloseWindowAfterMenuSelection = false;
+	FMenuBuilder MenuBuilder(false, nullptr);
+
+	// settings
+	MenuBuilder.BeginSection("ResetOptions", LOCTEXT("MutableResetPerformanceReportOptions", "Reset Options"));
+	{
+		SAssignNew(PerformanceReportResetOptionsEntry, SButton)
+			.Text(LOCTEXT("MutableGeneratePerformanceReport", "Reset Performance Report Options"))
+			.OnClicked(this, &SCustomizableObjecEditorPerformanceReport::ResetPerformanceReportOptions);
+		MenuBuilder.AddWidget(PerformanceReportResetOptionsEntry.ToSharedRef(), FText());
+	}
+	MenuBuilder.EndSection();
+
+	MenuBuilder.BeginSection("StressTestParameters", LOCTEXT("MutablePerformanceReportStressTestParameters", "Stress Test Parameters"));
+	{
+		SAssignNew(StressTestOptionInstanceCountEntry, SNumericEntryBox<int32>)
+			.AllowSpin(true)
+			.MinValue(0)
+			.MaxValue(1000000)
+			.MinSliderValue(1)
+			.MaxSliderValue(1000)
+			.Value(this, &SCustomizableObjecEditorPerformanceReport::GetStressTestOptionInstanceCount)
+			.OnValueChanged(this, &SCustomizableObjecEditorPerformanceReport::OnStressTestOptionInstanceCountChanged);
+		MenuBuilder.AddWidget(StressTestOptionInstanceCountEntry.ToSharedRef(), LOCTEXT("MutableStressTestInstanceCount", "Instance Count"));
+
+		SAssignNew(StressTestOptionCreateInstanceTimeEntry, SNumericEntryBox<int32>)
+			.AllowSpin(true)
+			.MinValue(1)
+			.MaxValue(1000000)
+			.MinSliderValue(1)
+			.MaxSliderValue(1000)
+			.Value(this, &SCustomizableObjecEditorPerformanceReport::GetStressTestOptionCreateInstanceTime)
+			.OnValueChanged(this, &SCustomizableObjecEditorPerformanceReport::OnStressTestOptionCreateInstanceTimeChanged);
+		MenuBuilder.AddWidget(StressTestOptionCreateInstanceTimeEntry.ToSharedRef(), LOCTEXT("MutableStressTestCreateInstanceTime", "Instance Create Time (ms)"));
+
+		SAssignNew(StressTestOptionCreateInstanceTimeVarEntry, SNumericEntryBox<int32>)
+			.AllowSpin(true)
+			.MinValue(1)
+			.MaxValue(1000000)
+			.MinSliderValue(1)
+			.MaxSliderValue(10000)
+			.Value(this, &SCustomizableObjecEditorPerformanceReport::GetStressTestOptionCreateInstanceTimeVar)
+			.OnValueChanged(this, &SCustomizableObjecEditorPerformanceReport::OnStressTestOptionCreateInstanceTimeVarChanged);
+		MenuBuilder.AddWidget(StressTestOptionCreateInstanceTimeVarEntry.ToSharedRef(), LOCTEXT("MutableStressTestCreateInstanceTimeVar", "Instance Create Time variation (ms)"));
+
+		SAssignNew(StressTestOptionInstanceUpdateTimeEntry, SNumericEntryBox<int32>)
+			.AllowSpin(true)
+			.MinValue(0)
+			.MaxValue(1000000)
+			.MinSliderValue(1)
+			.MaxSliderValue(10000)
+			.Value(this, &SCustomizableObjecEditorPerformanceReport::GetStressTestOptionInstanceUpdateTime)
+			.OnValueChanged(this, &SCustomizableObjecEditorPerformanceReport::OnStressTestOptionInstanceUpdateTimeChanged);
+		MenuBuilder.AddWidget(StressTestOptionInstanceUpdateTimeEntry.ToSharedRef(), LOCTEXT("MutableStressTestInstanceUpdateTime", "Instance Update Time (ms)"));
+
+		SAssignNew(StressTestOptionInstanceUpdateTimeVarEntry, SNumericEntryBox<int32>)
+			.AllowSpin(true)
+			.MinValue(1)
+			.MaxValue(1000000)
+			.MinSliderValue(1)
+			.MaxSliderValue(10000)
+			.Value(this, &SCustomizableObjecEditorPerformanceReport::GetStressTestOptionInstanceUpdateTimeVar)
+			.OnValueChanged(this, &SCustomizableObjecEditorPerformanceReport::OnStressTestOptionInstanceUpdateTimeVarChanged);
+		MenuBuilder.AddWidget(StressTestOptionInstanceUpdateTimeVarEntry.ToSharedRef(), LOCTEXT("MutableStressTestInstanceUpdateTimeVar", "Instance Update Time variation (ms)"));
+
+		SAssignNew(StressTestOptionInstanceLifeTimeEntry, SNumericEntryBox<int32>)
+			.AllowSpin(true)
+			.MinValue(0)
+			.MaxValue(1000000)
+			.MinSliderValue(1)
+			.MaxSliderValue(10000)
+			.Value(this, &SCustomizableObjecEditorPerformanceReport::GetStressTestOptionInstanceLifeTime)
+			.OnValueChanged(this, &SCustomizableObjecEditorPerformanceReport::OnStressTestOptionInstanceLifeTimeChanged);
+		MenuBuilder.AddWidget(StressTestOptionInstanceLifeTimeEntry.ToSharedRef(), LOCTEXT("MutableStressTestInstanceLifeTime", "Instance Life Time (ms)"));
+
+		SAssignNew(StressTestOptionInstanceLifeTimeVarEntry, SNumericEntryBox<int32>)
+			.AllowSpin(true)
+			.MinValue(1)
+			.MaxValue(1000000)
+			.MinSliderValue(1)
+			.MaxSliderValue(10000)
+			.Value(this, &SCustomizableObjecEditorPerformanceReport::GetStressTestOptionInstanceLifeTimeVar)
+			.OnValueChanged(this, &SCustomizableObjecEditorPerformanceReport::OnStressTestOptionInstanceLifeTimeVarChanged);
+		MenuBuilder.AddWidget(StressTestOptionInstanceLifeTimeVarEntry.ToSharedRef(), LOCTEXT("MutableStressTestInstanceLifeTimeMs", "Instance Life Time variation (ms)"));
+	}
+	MenuBuilder.EndSection();
+
+	return MenuBuilder.MakeWidget();
+}
+
+void SCustomizableObjecEditorPerformanceReport::AddReferencedObjects(FReferenceCollector& Collector)
+{
+	Collector.AddReferencedObject(ReportSubject);
+	Collector.AddReferencedObject(HelperCallback);
+	Collector.AddReferencedObject(WorstTimeInstance);
+	Collector.AddReferencedObject(CurrentReportWorstCase);
+	Collector.AddReferencedObject(CurrentReportInstance);
+	Collector.AddReferencedObjects(WorstTextureInstances);
+	Collector.AddReferencedObjects(WorstGeometryInstances);
+	for (auto& Arr : WorstTimeStateInstances)
+	{
+		Collector.AddReferencedObjects(Arr);
+	}
+	for (auto& ArrArr : WorstTimeParameterInstances)
+	{
+		for (auto& Arr : ArrArr)
+		{
+			Collector.AddReferencedObjects(Arr);
+		}
+	}
+}
+
+void SCustomizableObjecEditorPerformanceReport::Tick(const FGeometry& AllottedGeometry, const double InCurrentTime, const float InDeltaTime)
+{
+	SCompoundWidget::Tick(AllottedGeometry, InCurrentTime, InDeltaTime);
+}
+
+void SCustomizableObjecEditorPerformanceReport::StopTests()
+{
+	if (StressTest.IsValid() && StressTest->GetTestInCourse())
+	{
+		StressTest->WorstFacesDelegate.Unbind();
+		StressTest->WorstTextureDelegate.Unbind();
+		StressTest->StressTestEnded.Unbind();
+		StressTest->FinishTest();
+		StressTest.Reset();
+	}
+
+	if (StateTest.IsValid())
+	{
+		StateTest->WorstTimeDelegate.Unbind();
+		StateTest->StateTestEnded.Unbind();
+		StateTest->FinishTest();
+		StateTest.Reset();
+	}
+}
+
+FReply SCustomizableObjecEditorPerformanceReport::RunPerformanceReport()
+{
+	return RunPerformanceReport(nullptr);
+}
+
+FReply SCustomizableObjecEditorPerformanceReport::RunPerformanceReport(UCustomizableObject* ObjectToTest)
+{
+	if (ObjectToTest)
+	{
+		ReportSubject = ObjectToTest;
+	}
+
+	if (ReportSubject)
+	{	
+		StopTests();
+
+		if (CurrentReportInstance || CurrentReportWorstCase)
+		{
+			CurrentReportInstance = nullptr;
+			CurrentReportWorstCase = nullptr;
+			UpdateSelectedInstanceData();
+		}
+
+		ResetSavedData();
+
+		StressTest = MakeShareable(new FRunningStressTest);
+		StressTest->SetCustomizableObject(ReportSubject);
+		StressTest->SetVerifyInstancesThreshold(5.0f);
+		StressTest->SetCreateInstanceTimeMs(StressTestOptionCreateInstanceTime);
+		StressTest->SetCreateInstanceTimeMsVar(StressTestOptionCreateInstanceTimeVar);
+		StressTest->SetInstanceLifeTimeMs(StressTestOptionInstanceLifeTime);
+		StressTest->SetInstanceLifeTimeMsVar(StressTestOptionInstanceLifeTimeVar);
+		StressTest->SetInstanceUpdateTimeMs(StressTestOptionInstanceUpdateTime);
+		StressTest->SetInstanceUpdateTimeMsVar(StressTestOptionInstanceUpdateTimeVar);
+		StressTest->SetPendingInstanceCount(StressTestOptionInstanceCount);
+		StressTest->SetNextInstanceTimeMs(0);
+		StressTest->WorstFacesDelegate.BindRaw(this, &SCustomizableObjecEditorPerformanceReport::WorstFacesFound);
+		StressTest->WorstTextureDelegate.BindRaw(this, &SCustomizableObjecEditorPerformanceReport::WorstTextureFound);
+		StressTest->StressTestEnded.BindRaw(this, &SCustomizableObjecEditorPerformanceReport::StressTestEnded);
+		StressTest->InitMeasureTest();
+
+		StatusText->SetText(LOCTEXT("PerformanceReportStatusStressRunning", "Running Stress test."));
+	}
+	else
+	{
+		StatusText->SetText(LOCTEXT("PerformanceReportStatusNoObject", "Generate Performance Report found no valid Customizable Object."));
+	}
+
+	return FReply::Handled();
+}
+
+FReply SCustomizableObjecEditorPerformanceReport::SelectSlowestInstance()
+{
+	check(WorstTimeInstance);
+	UWorstCasePerformanceReportInstance* SlowestInstanceInStructure = nullptr;
+	if (WorstTimeInstance)
+	{
+		switch (WorstTimeInstance->WorstCaseType)
+		{
+		case UWorstCasePerformanceReportInstance::EPerformanceReportWorstCaseType::SLOWEST_UPDATE_ENTER_STATE:
+		{
+			SlowestInstanceInStructure = GetWorstCaseStructureFromIndexes(WorstTimeInstance->WorstCaseType, WorstTimeInstance->LOD, WorstTimeInstance->LongestUpdateTimeStateIndex);
+			break;
+		}
+		case UWorstCasePerformanceReportInstance::EPerformanceReportWorstCaseType::SLOWEST_UPDATE_PARAM_IN_RELEVANT_STATE:
+		{
+			if (UWorstCasePerformanceReportInstance* SlowestInstanceInStructureParent = GetWorstCaseStructureFromIndexes(
+				UWorstCasePerformanceReportInstance::EPerformanceReportWorstCaseType::SLOWEST_UPDATE_ENTER_STATE,
+				WorstTimeInstance->LOD,
+				WorstTimeInstance->LongestUpdateTimeStateIndex))
+			{
+				SlowestInstanceInStructureParent->SubVisibility = EVisibility::Visible;
+			}
+			SlowestInstanceInStructure = GetWorstCaseStructureFromIndexes(WorstTimeInstance->WorstCaseType, WorstTimeInstance->LOD, WorstTimeInstance->LongestUpdateTimeStateIndex, WorstTimeInstance->LongestUpdateTimeParameterIndexInCO);
+			break;
+		}
+		default:
+		{
+			break;
+		}
+		}
+	}
+	if (SlowestInstanceInStructure)
+	{
+		SelectInstance(SlowestInstanceInStructure);
+	}
+	return FReply::Handled();
+}
+
+bool SCustomizableObjecEditorPerformanceReport::IsGetWorstInstanceButtonEnabled() const
+{
+	return WorstTimeInstance && WorstTimeInstance != CurrentReportWorstCase;
+}
+
+FReply SCustomizableObjecEditorPerformanceReport::ResetPerformanceReportOptions()
+{
+	StressTestOptionInstanceCount = 30;
+	StressTestOptionCreateInstanceTime = 200;
+	StressTestOptionCreateInstanceTimeVar = 200;
+	StressTestOptionInstanceLifeTime = 5000;
+	StressTestOptionInstanceLifeTimeVar = 5000;
+	StressTestOptionInstanceUpdateTime = 500;
+	StressTestOptionInstanceUpdateTimeVar = 500;
+	return FReply::Handled();
+}
+
+TOptional<int32> SCustomizableObjecEditorPerformanceReport::GetStressTestOptionInstanceCount() const
+{
+	return StressTestOptionInstanceCount;
+}
+
+void SCustomizableObjecEditorPerformanceReport::OnStressTestOptionInstanceCountChanged(int Value)
+{
+	StressTestOptionInstanceCount = Value;
+}
+
+TOptional<int32> SCustomizableObjecEditorPerformanceReport::GetStressTestOptionCreateInstanceTime() const
+{
+	return StressTestOptionCreateInstanceTime;
+}
+
+void SCustomizableObjecEditorPerformanceReport::OnStressTestOptionCreateInstanceTimeChanged(int Value)
+{
+	StressTestOptionCreateInstanceTime = Value;
+}
+
+TOptional<int32> SCustomizableObjecEditorPerformanceReport::GetStressTestOptionCreateInstanceTimeVar() const
+{
+	return StressTestOptionCreateInstanceTimeVar;
+}
+
+void SCustomizableObjecEditorPerformanceReport::OnStressTestOptionCreateInstanceTimeVarChanged(int Value)
+{
+	StressTestOptionCreateInstanceTimeVar = Value;
+}
+
+TOptional<int32> SCustomizableObjecEditorPerformanceReport::GetStressTestOptionInstanceUpdateTime() const
+{
+	return StressTestOptionInstanceUpdateTime;
+}
+
+void SCustomizableObjecEditorPerformanceReport::OnStressTestOptionInstanceUpdateTimeChanged(int Value)
+{
+	StressTestOptionInstanceUpdateTime = Value;
+}
+
+TOptional<int32> SCustomizableObjecEditorPerformanceReport::GetStressTestOptionInstanceUpdateTimeVar() const
+{
+	return StressTestOptionInstanceUpdateTimeVar;
+}
+
+void SCustomizableObjecEditorPerformanceReport::OnStressTestOptionInstanceUpdateTimeVarChanged(int Value)
+{
+	StressTestOptionInstanceUpdateTimeVar = Value;
+}
+
+TOptional<int32> SCustomizableObjecEditorPerformanceReport::GetStressTestOptionInstanceLifeTime() const
+{
+	return StressTestOptionInstanceLifeTime;
+}
+
+void SCustomizableObjecEditorPerformanceReport::OnStressTestOptionInstanceLifeTimeChanged(int Value)
+{
+	StressTestOptionInstanceLifeTime = Value;
+}
+
+TOptional<int32> SCustomizableObjecEditorPerformanceReport::GetStressTestOptionInstanceLifeTimeVar() const
+{
+	return StressTestOptionInstanceLifeTimeVar;
+}
+
+void SCustomizableObjecEditorPerformanceReport::OnStressTestOptionInstanceLifeTimeVarChanged(int Value)
+{
+	StressTestOptionInstanceLifeTimeVar = Value;
+}
+
+void SCustomizableObjecEditorPerformanceReport::ResetSavedData()
+{
+	if (ReportSubject)
+	{
+		NumLOD = ReportSubject->GetNumLODs();
+		NumStates = ReportSubject->GetStateCount();
+
+
+		WorstTextureInstances.Empty(NumLOD);
+		WorstTextureInstances.SetNum(NumLOD);
+		for (uint32 i = 0; i < NumLOD; ++i)
+		{
+			WorstTextureInstances[i] = NewObject<UWorstCasePerformanceReportInstance>(GetTransientPackage(), NAME_None, RF_Transient);
+			WorstTextureInstances[i]->LOD = i;
+			WorstTextureInstances[i]->WorstCaseType = UWorstCasePerformanceReportInstance::EPerformanceReportWorstCaseType::LARGEST_TEXTURE_SIZE;
+		}
+
+		WorstGeometryInstances.Empty(NumLOD);
+		WorstGeometryInstances.SetNum(NumLOD);
+		for (uint32 i = 0; i < NumLOD; ++i)
+		{
+			WorstGeometryInstances[i] = NewObject<UWorstCasePerformanceReportInstance>(GetTransientPackage(), NAME_None, RF_Transient);
+			WorstGeometryInstances[i]->LOD = i;
+			WorstGeometryInstances[i]->WorstCaseType = UWorstCasePerformanceReportInstance::EPerformanceReportWorstCaseType::MOST_TRIANGLES;
+		}
+
+		WorstTimeInstance = nullptr;
+
+		// TArray<TArray<UWorstCasePerformanceReportInstance*>> WorstTimeStateInstances;	// By LOD then by state index
+		WorstTimeStateInstances.Empty(NumLOD);
+		WorstTimeStateInstances.SetNum(NumLOD);
+		for (uint32 i = 0; i < NumLOD; ++i)
+		{
+			WorstTimeStateInstances[i].Empty(NumStates);
+			WorstTimeStateInstances[i].SetNum(NumStates);
+			for (uint32 j = 0; j < NumStates; ++j)
+			{
+				WorstTimeStateInstances[i][j] = NewObject<UWorstCasePerformanceReportInstance>(GetTransientPackage(), NAME_None, RF_Transient);
+				WorstTimeStateInstances[i][j]->LOD = i;
+				WorstTimeStateInstances[i][j]->WorstCaseType = UWorstCasePerformanceReportInstance::EPerformanceReportWorstCaseType::SLOWEST_UPDATE_ENTER_STATE;
+			}
+		}
+
+		// TArray<TArray<TArray<UWorstCasePerformanceReportInstance*>>> WorstTimeParameterInstances;	// By LOD then by state index then by parameter index
+		WorstTimeParameterInstances.Empty(NumLOD);
+		WorstTimeParameterInstances.SetNum(NumLOD);
+		for (uint32 i = 0; i < NumLOD; ++i)
+		{
+			WorstTimeParameterInstances[i].Empty(NumStates);
+			WorstTimeParameterInstances[i].SetNum(NumStates);
+			for (uint32 j = 0; j < NumStates; ++j)
+			{
+				WorstTimeParameterInstances[i][j].Reset();
+			}
+		}
+	}
+}
+
+void SCustomizableObjecEditorPerformanceReport::WorstFacesFound(uint32 CurrentLODNumFaces, uint32 LOD, UCustomizableObjectInstance* Instance, MeasuredData* data)
+{
+	// If any of there checks fails, we are receiving updates from a different CO than this was initialized with.
+	check(WorstGeometryInstances.Num() > (int32)LOD);
+	check(WorstGeometryInstances[LOD]);
+	check(WorstGeometryInstances[LOD]->LOD == LOD);
+	check(WorstGeometryInstances[LOD]->WorstCaseType == UWorstCasePerformanceReportInstance::EPerformanceReportWorstCaseType::MOST_TRIANGLES);
+
+	bool bSomethingUpdated = false;
+	UWorstCasePerformanceReportInstance* WorstCase = WorstGeometryInstances[LOD];
+	if ((int32)CurrentLODNumFaces > WorstGeometryInstances[LOD]->NumFaces)
+	{
+		WorstCase->SetInstance(Instance);
+		WorstCase->NumFaces = CurrentLODNumFaces;
+		bSomethingUpdated = true;
+	}
+
+	if (data)
+	{
+		WorstCase->MeasuredDataForThisWorstCase = *data;
+		bSomethingUpdated = true;
+	}
+
+	if (bSomethingUpdated)
+	{
+		BuildGeometryWorstCasesView();
+	}
+}
+
+void SCustomizableObjecEditorPerformanceReport::WorstTextureFound(MaterialInfoMap MapMaterialParam, uint32 CurrentLODTextureSize, uint32 LOD, UCustomizableObjectInstance* Instance, MeasuredData* data)
+{
+	// If any of there checks fails, we are receiving updates from a different CO than this was initialized with.
+	check(WorstTextureInstances.Num() > (int32)LOD);
+	check(WorstTextureInstances[LOD]);
+	check(WorstTextureInstances[LOD]->LOD == LOD);
+	check(WorstTextureInstances[LOD]->WorstCaseType == UWorstCasePerformanceReportInstance::EPerformanceReportWorstCaseType::LARGEST_TEXTURE_SIZE);
+
+	bool bSomethingUpdated = false;
+	UWorstCasePerformanceReportInstance* WorstCase = WorstTextureInstances[LOD];
+	if ((int32)CurrentLODTextureSize > WorstTextureInstances[LOD]->TextureBytes)
+	{
+		WorstCase->SetInstance(Instance);
+		WorstCase->TextureBytes = CurrentLODTextureSize;
+		bSomethingUpdated = true;
+	}
+
+	if (data)
+	{
+		WorstCase->MeasuredDataForThisWorstCase = *data;
+		bSomethingUpdated = true;
+	}
+
+	if (bSomethingUpdated)
+	{
+		BuildTextureWorstCasesView();
+	}
+}
+
+void SCustomizableObjecEditorPerformanceReport::WorstTimeFound(float CurrentInstanceGenerationTime, int32 LongesTimeStateIndex, int32 LongesTimeParameterIndexInCO, uint32 LOD, class UCustomizableObjectInstance* Instance, MeasuredData* data)
+{
+	UWorstCasePerformanceReportInstance* WorstCase = nullptr;
+	bool bIsWorstCaseState = LongesTimeParameterIndexInCO == INDEX_NONE;
+	bool IsNewPArameter = false;
+
+	if (bIsWorstCaseState)
+	{
+		// If any of there checks fails, we are receiving updates from a different CO than this was initialized with.
+		check(WorstTimeStateInstances.Num() > (int32)LOD);
+		check(WorstTimeStateInstances[LOD].Num() > LongesTimeStateIndex);
+		check(WorstTimeStateInstances[LOD][LongesTimeStateIndex]);
+		check(WorstTimeStateInstances[LOD][LongesTimeStateIndex]->LOD == LOD);
+		check(WorstTimeStateInstances[LOD][LongesTimeStateIndex]->WorstCaseType == UWorstCasePerformanceReportInstance::EPerformanceReportWorstCaseType::SLOWEST_UPDATE_ENTER_STATE);
+
+		if (CurrentInstanceGenerationTime > WorstTimeStateInstances[LOD][LongesTimeStateIndex]->UpdateTime)
+		{
+			WorstCase = WorstTimeStateInstances[LOD][LongesTimeStateIndex];
+		}
+	}
+	else
+	{
+		// If any of there checks fails, we are receiving updates from a different CO than this was initialized with.
+		check(WorstTimeParameterInstances.Num() > (int32)LOD);
+		check(WorstTimeParameterInstances[LOD].Num() > LongesTimeStateIndex);
+		UWorstCasePerformanceReportInstance* WorstCaseParameterTime = nullptr;
+
+		for (auto& WCP : WorstTimeParameterInstances[LOD][LongesTimeStateIndex])
+		{
+			if (LongesTimeParameterIndexInCO == WCP->LongestUpdateTimeParameterIndexInCO)
+			{
+				WorstCaseParameterTime = WCP;
+				// If any of there checks fails, we are receiving updates from a different CO than this was initialized with.
+				check(WorstCaseParameterTime);
+				check(WorstCaseParameterTime->LOD == LOD);
+				check(WorstCaseParameterTime->WorstCaseType == UWorstCasePerformanceReportInstance::EPerformanceReportWorstCaseType::SLOWEST_UPDATE_PARAM_IN_RELEVANT_STATE);
+				break;
+			}
+		}
+
+		if (!WorstCaseParameterTime)
+		{
+			WorstCaseParameterTime = NewObject<UWorstCasePerformanceReportInstance>(GetTransientPackage(), NAME_None, RF_Transient);
+			WorstCaseParameterTime->LOD = LOD;
+			WorstCaseParameterTime->WorstCaseType = UWorstCasePerformanceReportInstance::EPerformanceReportWorstCaseType::SLOWEST_UPDATE_PARAM_IN_RELEVANT_STATE;
+			WorstTimeParameterInstances[LOD][LongesTimeStateIndex].Push(WorstCaseParameterTime);
+			IsNewPArameter = true;
+		}
+
+		if (WorstCaseParameterTime && CurrentInstanceGenerationTime > WorstCaseParameterTime->UpdateTime)
+		{
+			WorstCase = WorstCaseParameterTime;
+		}
+	}
+
+	if (WorstCase)
+	{
+		WorstCase->SetInstance(Instance);
+		if (WorstCase->WorstCaseInstance)
+		{
+			WorstCase->WorstCaseInstance->GetPrivate()->SetState(LongesTimeStateIndex);
+		}
+		WorstCase->LongestUpdateTimeStateIndex = LongesTimeStateIndex;
+		WorstCase->LongestUpdateTimeParameterIndexInCO = LongesTimeParameterIndexInCO;
+		WorstCase->UpdateTime = CurrentInstanceGenerationTime;
+		if (data)
+		{
+			WorstCase->MeasuredDataForThisWorstCase = *data;
+		}
+		UpdateWorstTimeInstance(WorstCase);
+		BuildTimeWorstCasesDetailedView();
+	}
+}
+
+void SCustomizableObjecEditorPerformanceReport::UpdateWorstTimeInstance(const UWorstCasePerformanceReportInstance* EvenWorseTime)
+{
+	if (!WorstTimeInstance || (EvenWorseTime && (EvenWorseTime->UpdateTime > WorstTimeInstance->UpdateTime)))
+	{
+		WorstTimeInstance = EvenWorseTime;
+	}
+}
+
+void SCustomizableObjecEditorPerformanceReport::StressTestEnded()
+{
+	if (ReportSubject)
+	{
+		StateTest = MakeShareable(new FRuntimeTest);
+		StateTest->WorstTimeDelegate.BindRaw(this, &SCustomizableObjecEditorPerformanceReport::WorstTimeFound);
+		StateTest->StateTestEnded.BindRaw(this, &SCustomizableObjecEditorPerformanceReport::StateTestEnded);
+		StateTest->StartTest(ReportSubject);
+		StatusText->SetText(LOCTEXT("PerformanceReportStatusStateRunning", "Running State test."));
+	}
+}
+
+void SCustomizableObjecEditorPerformanceReport::StateTestEnded()
+{
+	StatusText->SetText(LOCTEXT("PerformanceReportStatusFinished", "Tests completed."));
+}
+
+UWorstCasePerformanceReportInstance* SCustomizableObjecEditorPerformanceReport::GetWorstCaseStructureFromIndexes(UWorstCasePerformanceReportInstance::EPerformanceReportWorstCaseType WorstCaseType, uint32 LOD, int32 StateID, int32 ParamID) const
+{
+	if (LOD >= 0 && LOD < NumLOD)
+	{
+		switch (WorstCaseType)
+		{
+		case UWorstCasePerformanceReportInstance::EPerformanceReportWorstCaseType::LARGEST_TEXTURE_SIZE:
+		{
+			check((uint32)WorstTextureInstances.Num() > LOD);
+			return WorstTextureInstances[LOD];
+		}
+		case UWorstCasePerformanceReportInstance::EPerformanceReportWorstCaseType::MOST_TRIANGLES:
+		{
+			check((uint32)WorstGeometryInstances.Num() > LOD);
+			return WorstGeometryInstances[LOD];
+		}
+		case UWorstCasePerformanceReportInstance::EPerformanceReportWorstCaseType::SLOWEST_UPDATE_ENTER_STATE:
+		{
+			if (StateID != INDEX_NONE)
+			{
+				check((uint32)WorstTimeStateInstances.Num() > LOD);
+				check(StateID >= 0);
+				check(WorstTimeStateInstances[LOD].Num() > StateID);
+				return WorstTimeStateInstances[LOD][StateID];
+			}
+		}
+		case UWorstCasePerformanceReportInstance::EPerformanceReportWorstCaseType::SLOWEST_UPDATE_PARAM_IN_RELEVANT_STATE:
+		{
+			if (StateID != INDEX_NONE && ParamID != INDEX_NONE)
+			{
+				check((uint32)WorstTimeParameterInstances.Num() > LOD);
+				check(StateID >= 0);
+				check(WorstTimeParameterInstances[LOD].Num() > StateID);
+				const int32 NumParams = WorstTimeParameterInstances[LOD][StateID].Num();
+				for (int i = 0; i < NumParams; ++i)
+				{
+					if (WorstTimeParameterInstances[LOD][StateID][i] && WorstTimeParameterInstances[LOD][StateID][i]->LongestUpdateTimeParameterIndexInCO == ParamID)
+					{
+						return WorstTimeParameterInstances[LOD][StateID][i];
+					}
+				}
+			}
+		}
+		}
+	}
+	return nullptr;
+}
+
+TSharedRef<SWidget> SCustomizableObjecEditorPerformanceReport::BuildAnalyzer()
+{
+	AssetThumbnailPool = MakeShareable(new FAssetThumbnailPool(32));
+
+	TSharedRef<SVerticalBox> ResultVerticalBox = SNew(SVerticalBox)
+		.Visibility_Lambda([this]() { return CurrentReportInstance ? EVisibility::Visible : EVisibility::Collapsed; });
+
+	ResultVerticalBox->AddSlot()
+	.HAlign(HAlign_Left)
+	.AutoHeight()
+	[
+		SNew(SBorder)
+		.Padding(5.0f)
+		[
+			SNew(SVerticalBox)
+
+			+ SVerticalBox::Slot()
+			.HAlign(HAlign_Left)
+			.AutoHeight()
+			.Padding(5.0f)
+			[
+				SNew(SHorizontalBox)
+
+				+ SHorizontalBox::Slot()
+				.AutoWidth()
+				[
+					SNew(STextBlock)
+					.Text(LOCTEXT("SelectTextureAnalyzerInstance", "Selected Instance: "))
+				]
+
+				+ SHorizontalBox::Slot()
+				.AutoWidth()
+				[
+					SAssignNew(SelectedTextureAnalyzerText, STextBlock)
+				]
+			]
+
+			+ SVerticalBox::Slot()
+			.HAlign(HAlign_Left)
+			.AutoHeight()
+			[
+				SNew(SHorizontalBox)
+
+				+ SHorizontalBox::Slot()
+				.AutoWidth()
+				.Padding(5.0f)
+				[
+					SAssignNew(InstanceThumbnail, SObjectPropertyEntryBox)
+					.AllowedClass(UCustomizableObjectInstance::StaticClass())
+					.ObjectPath(this, &SCustomizableObjecEditorPerformanceReport::OnGetInstancePath)
+					.AllowClear(false)
+					.DisplayUseSelected(false)
+					.DisplayBrowse(false)
+					.EnableContentPicker(false)
+					.DisplayThumbnail(true)
+					.ThumbnailPool(AssetThumbnailPool)
+				]
+
+				+ SHorizontalBox::Slot()
+				.AutoWidth()
+				.VAlign(VAlign_Center)
+				.Padding(5.0f)
+				[
+					SAssignNew(PerformanceReportRefreshSelectedInstance, SButton)
+					.ButtonStyle(FAppStyle::Get(), "FlatButton.Danger")
+					.IsEnabled(this, &SCustomizableObjecEditorPerformanceReport::IsSelectedInstanceOutdated)
+					.OnClicked(this, &SCustomizableObjecEditorPerformanceReport::RefreshSelectedInstance)
+					[
+						SNew(STextBlock)
+						.Text(LOCTEXT("RefreshSelectedInstanceButton", "Refresh"))
+					]
+				]
+			]
+		]
+	];
+
+	ResultVerticalBox->AddSlot()
+	.HAlign(HAlign_Fill)
+	[
+		SNew(SBorder)
+		.Padding(5.0f)
+		[
+			SNew(SScrollBox)
+			+ SScrollBox::Slot()
+			[
+				PerformanceReportTextureAnalyzer.ToSharedRef()
+			]
+		]
+	];
+
+	UpdateSelectedInstanceData();
+
+	return ResultVerticalBox;
+}
+
+void SCustomizableObjecEditorPerformanceReport::UpdateSelectedInstanceData()
+{
+	// Refresh Identification text
+	if (!SelectedTextureAnalyzerText)
+	{
+		return;
+	}
+
+	FText Result = LOCTEXT("NoInstanceCurrentlySelected", "Nonee currently selected. Click on one of the rows on the left to view it's details.");
+	const FNumberFormattingOptions BytesFormat = FNumberFormattingOptions().SetMaximumFractionalDigits(2).SetMinimumFractionalDigits(2);
+	const FNumberFormattingOptions TimeFormat = FNumberFormattingOptions().SetMaximumFractionalDigits(2).SetMinimumFractionalDigits(2);
+	const FNumberFormattingOptions PolygonsFormat = FNumberFormattingOptions().SetMaximumFractionalDigits(0);
+	const FNumberFormattingOptions PercentageFormat = FNumberFormattingOptions().SetMaximumFractionalDigits(0).SetAlwaysSign(true);
+
+	if (CurrentReportInstance && CurrentReportWorstCase)
+	{
+		const FText TextLOD = FText::AsNumber(CurrentReportWorstCase->LOD, &FNumberFormattingOptions::DefaultNoGrouping());
+		FText TextReason;
+		FText TextAmount;
+		switch (CurrentReportWorstCase->WorstCaseType)
+		{
+		case UWorstCasePerformanceReportInstance::EPerformanceReportWorstCaseType::LARGEST_TEXTURE_SIZE:
+		{
+			TextReason = FText::FromString("Largest mutable textures size");
+			const float Bytes = CurrentReportWorstCase->TextureBytes;
+			const float Value = (Bytes > 1048576) ? float(Bytes) / 1048576.0f : float(Bytes) / 1024.0f;
+
+			TextAmount = FText::Format(LOCTEXT("MutablePerformanceTextureSize", "{0} {1}"), FText::AsNumber(Value, &BytesFormat), FText::FromString((Bytes > 1048576) ? "MB" : "KB"));
+			break;
+		}
+		case UWorstCasePerformanceReportInstance::EPerformanceReportWorstCaseType::MOST_TRIANGLES:
+		{
+			TextReason = FText::FromString("Most triangles");
+			TextAmount = FText::AsNumber(CurrentReportWorstCase->NumFaces, &PolygonsFormat);
+			break;
+		}
+		case UWorstCasePerformanceReportInstance::EPerformanceReportWorstCaseType::SLOWEST_UPDATE_ENTER_STATE:
+		{
+			TextReason = FText::FromString("Slowest update when entering the state");
+			TextAmount = FText::AsNumber(CurrentReportWorstCase->UpdateTime, &TimeFormat);
+			break;
+		}
+		case UWorstCasePerformanceReportInstance::EPerformanceReportWorstCaseType::SLOWEST_UPDATE_PARAM_IN_RELEVANT_STATE:
+		{
+			TextReason = FText::FromString("Slowest update when changing parameter");
+			TextAmount = FText::AsNumber(CurrentReportWorstCase->UpdateTime, &TimeFormat);
+			break;
+		}
+		default:
+		{
+			break;
+		}
+		}
+
+		Result = FText::Format(LOCTEXT("PerformanceReportReason", "{0} for LOD {1} ({2})"), TextReason, TextLOD, TextAmount);
+	}
+
+	SelectedTextureAnalyzerText->SetText(Result);
+
+	// Refresh Texture Analyzer
+	if (CurrentReportInstance && CurrentReportInstance->IsValidLowLevelFast())
+	{
+		// If we have at least one skeletal mesh in the instance, we know it's generated (at least once)
+		if (CurrentReportInstance->HasAnySkeletalMesh())
+		{
+			PerformanceReportTextureAnalyzer->RefreshTextureAnalyzerTable(CurrentReportInstance);
+		}
+		else
+		{
+			CurrentReportInstance->UpdatedDelegate.AddUniqueDynamic(HelperCallback, &UPerformanceReportHelper::DelegatedCallback);
+			CurrentReportInstance->UpdateSkeletalMeshAsync(true, true);
+		}
+	}
+}
+
+void SCustomizableObjecEditorPerformanceReport::BuildTextureWorstCasesView()
+{
+	const FNumberFormattingOptions BytesFormat = FNumberFormattingOptions().SetMaximumFractionalDigits(2).SetMinimumFractionalDigits(2);
+
+	TextureWorstCasesView->ClearChildren();
+
+	TextureWorstCasesView->AddSlot()
+	.HAlign(HAlign_Fill)
+	.AutoHeight()
+	[
+		SNew(SListView<UWorstCasePerformanceReportInstance*>)
+		.ListItemsSource(&ObjectPtrDecay(WorstTextureInstances))
+		.OnGenerateRow(this, &SCustomizableObjecEditorPerformanceReport::GenerateRow)
+		.OnMouseButtonClick(this, &SCustomizableObjecEditorPerformanceReport::SelectInstance)
+		.SelectionMode(ESelectionMode::Single)
+		.IsFocusable(true)
+		.HeaderRow
+		(
+			SNew(SHeaderRow)
+			+ SHeaderRow::Column(ColumnID_PR_LOD_Label)
+			.DefaultLabel(LOCTEXT("LOD", "LOD"))
+			.FillWidth(0.2)
+
+			+SHeaderRow::Column(ColumnID_PR_TexturesTotalSize_Label)
+			.DefaultLabel(LOCTEXT("TexturesTotalSize", "Textures total size"))
+			.FillWidth(0.2)
+
+			+ SHeaderRow::Column(ColumnID_PR_Delta_Label)
+			.DefaultLabel(LOCTEXT("Delta", "Difference"))
+			.FillWidth(0.1)
+
+			+ SHeaderRow::Column(ColumnID_PR_Average_Label)
+			.DefaultLabel(LOCTEXT("Average", "Average"))
+			.FillWidth(0.1)
+		)
+	];
+}
+
+void SCustomizableObjecEditorPerformanceReport::BuildGeometryWorstCasesView()
+{
+	GeometryWorstCasesView->ClearChildren();
+
+	GeometryWorstCasesView->AddSlot()
+	.HAlign(HAlign_Fill)
+	.AutoHeight()
+	[
+		SNew(SListView<UWorstCasePerformanceReportInstance*>)
+		.ListItemsSource(&ObjectPtrDecay(WorstGeometryInstances))
+		.OnGenerateRow(this, &SCustomizableObjecEditorPerformanceReport::GenerateRow)
+		.OnMouseButtonClick(this, &SCustomizableObjecEditorPerformanceReport::SelectInstance)
+		.SelectionMode(ESelectionMode::Single)
+		.IsFocusable(true)
+		.HeaderRow
+		(
+			SNew(SHeaderRow)
+			+ SHeaderRow::Column(ColumnID_PR_LOD_Label)
+			.DefaultLabel(LOCTEXT("LOD", "LOD"))
+			.FillWidth(0.2)
+
+			+ SHeaderRow::Column(ColumnID_PR_NumberOfTriangles_Label)
+			.DefaultLabel(LOCTEXT("NumberOfTriangles", "Number of triangles"))
+			.FillWidth(0.2)
+
+			+ SHeaderRow::Column(ColumnID_PR_Delta_Label)
+			.DefaultLabel(LOCTEXT("Delta", "Difference"))
+			.FillWidth(0.1)
+
+			+ SHeaderRow::Column(ColumnID_PR_Average_Label)
+			.DefaultLabel(LOCTEXT("Average", "Average"))
+			.FillWidth(0.1)
+		)
+	];
+}
+
+void SCustomizableObjecEditorPerformanceReport::BuildTimeWorstCasesDetailedView()
+{
+	const FNumberFormattingOptions TimeFormat = FNumberFormattingOptions().SetMaximumFractionalDigits(2).SetMinimumFractionalDigits(2);
+
+	TimeWorstCasesDetailedView->ClearChildren();
+
+	for (int32 LOD = 0; LOD < WorstTimeStateInstances.Num(); ++LOD)
+	{
+		// Currently, states test only returns LOD 0 data
+		if (LOD != 0)
+		{
+			continue;
+		}
+
+		const FText TextLOD = FText::AsNumber(LOD, &FNumberFormattingOptions::DefaultNoGrouping());
+
+		TimeWorstCasesDetailedView->AddSlot()
+		.HAlign(HAlign_Left)
+		.Padding(5.0f)
+		[
+			SNew(STextBlock)
+			.Text(FText::Format(LOCTEXT("MutableLODNTitle", "LOD {0}:"), TextLOD))
+		];
+
+		TimeWorstCasesDetailedView->AddSlot()
+		.HAlign(HAlign_Fill)
+		[
+			SNew(SListView<UWorstCasePerformanceReportInstance*>)
+			.ListItemsSource(&ObjectPtrDecay(WorstTimeStateInstances[LOD]))
+			.OnGenerateRow(this, &SCustomizableObjecEditorPerformanceReport::GenerateRow)
+			.OnMouseButtonClick(this, &SCustomizableObjecEditorPerformanceReport::SelectInstance)
+			.SelectionMode(ESelectionMode::Single)
+			.IsFocusable(true)
+			.HeaderRow
+			(
+				SAssignNew(TimeWorstCasesViewHeader, SHeaderRow)
+				+ SHeaderRow::Column(ColumnID_PR_Name_Label)
+				.DefaultLabel(LOCTEXT("Name", "Name"))
+				.FillWidth(0.2)
+
+				+ SHeaderRow::Column(ColumnID_PR_UpdateTime_Label)
+				.DefaultLabel(LOCTEXT("StateEnterUpdateTime", "Update time (ms)"))
+				.FillWidth(0.2)
+
+				+ SHeaderRow::Column(ColumnID_PR_Delta_Label)
+				.DefaultLabel(LOCTEXT("Delta", "Difference"))
+				.FillWidth(0.1)
+
+				+ SHeaderRow::Column(ColumnID_PR_Average_Label)
+				.DefaultLabel(LOCTEXT("Average", "Average"))
+				.FillWidth(0.1)
+			)
+		];
+	}
+
+	TimeWorstCasesDetailedView->Invalidate(EInvalidateWidget::Layout);
+}
+
+void SCustomizableObjecEditorPerformanceReport::OnOpenAssetEditor(UObject* ObjectToEdit)
+{
+	if (ObjectToEdit)
+	{
+		GEditor->EditObject(ObjectToEdit);
+	}
+}
+
+void SCustomizableObjecEditorPerformanceReport::ChangeFocusedReportInstance(UWorstCasePerformanceReportInstance* InReportInstance)
+{
+	CurrentReportWorstCase = InReportInstance;
+}
+
+FString SCustomizableObjecEditorPerformanceReport::OnGetInstancePath() const
+{
+	if (CurrentReportInstance && CurrentReportInstance->IsValidLowLevelFast(false))
+	{
+		return CurrentReportInstance->GetPathName();
+	}
+
+	return FString();
+}
+
+bool SCustomizableObjecEditorPerformanceReport::IsSelectedInstanceOutdated() const
+{
+	if (CurrentReportWorstCase)
+	{
+		return CurrentReportWorstCase->WorstCaseInstance != CurrentReportInstance;
+		// // Alternative method to checking if a selected instance is outdated, if in the future, CurrentReportInstnce and CurrentReportWorstCase->WorstCaseInstance are kept synchronized
+		//const UWorstCasePerformanceReportInstance* InstanceInData = GetWorstCaseStructureFromIndexes(CurrentReportWorstCase->WorstCaseType, CurrentReportWorstCase->LOD, CurrentReportWorstCase->LongestUpdateTimeStateIndex, CurrentReportWorstCase->LongestUpdateTimeParameterIndexInCO);
+		//if (InstanceInData)
+		//{
+		//	switch (CurrentReportWorstCase->WorstCaseType)
+		//	{
+		//	case UWorstCasePerformanceReportInstance::EPerformanceReportWorstCaseType::LARGEST_TEXTURE_SIZE:
+		//	{
+		//		if (InstanceInData->TextureBytes > CurrentReportWorstCase->TextureBytes)
+		//		{
+		//			return true;
+		//		}
+		//		break;
+		//	}
+		//	case UWorstCasePerformanceReportInstance::EPerformanceReportWorstCaseType::MOST_TRIANGLES:
+		//	{
+		//		if (InstanceInData->NumFaces > CurrentReportWorstCase->NumFaces)
+		//		{
+		//			return true;
+		//		}
+		//		break;
+		//	}
+		//	case UWorstCasePerformanceReportInstance::EPerformanceReportWorstCaseType::SLOWEST_UPDATE_ENTER_STATE:
+		//	case UWorstCasePerformanceReportInstance::EPerformanceReportWorstCaseType::SLOWEST_UPDATE_PARAM_IN_RELEVANT_STATE:
+		//	{
+		//		if (InstanceInData->UpdateTime > CurrentReportWorstCase->UpdateTime)
+		//		{
+		//			return true;
+		//		}
+		//		break;
+		//	}
+		//	}
+		//}
+	}
+	return false;
+}
+
+FReply SCustomizableObjecEditorPerformanceReport::RefreshSelectedInstance()
+{
+	CurrentReportInstance = CurrentReportWorstCase->WorstCaseInstance;
+	UpdateSelectedInstanceData();
+
+	return FReply::Handled();
+}
+
+TSharedRef<ITableRow> SCustomizableObjecEditorPerformanceReport::GenerateRow(UWorstCasePerformanceReportInstance* ReportedInstance, const TSharedRef<STableViewBase>& OwnerTable)
+{
+	return SNew(SPerformanceReportWorstCaseRow, OwnerTable).ReportInstance(ReportedInstance).PerformanceReport(this).Style(&FCustomizableObjectEditorStyle::Get().GetWidgetStyle<FTableRowStyle>("CustomizableObjectPerformanceReport.Row"));
+}
+
+const TArray<UWorstCasePerformanceReportInstance*>* SCustomizableObjecEditorPerformanceReport::GetWorstTimeParameters(uint32 LOD, uint32 StateIdx) const
+{
+	check((uint32)WorstTimeParameterInstances.Num() > LOD);
+	//check((uint32)WorstTimeParameterInstances[LOD].Num() > StateIdx);
+
+	if ((uint32)WorstTimeParameterInstances.Num() > LOD && (uint32)WorstTimeParameterInstances[LOD].Num() > StateIdx)
+	{
+		return &ObjectPtrDecay(WorstTimeParameterInstances[LOD][StateIdx]);
+	}
+
+	return nullptr;
+}
+
+void SCustomizableObjecEditorPerformanceReport::SelectInstance(UWorstCasePerformanceReportInstance* ClickedInstance)
+{
+	if (CurrentReportInstance)
+	{
+		CurrentReportInstance->UpdatedDelegate.RemoveDynamic(HelperCallback, &UPerformanceReportHelper::DelegatedCallback);
+	}
+	CurrentReportWorstCase = ClickedInstance;
+	CurrentReportInstance = ClickedInstance->WorstCaseInstance;
+	UpdateSelectedInstanceData();
+
+	BuildTextureWorstCasesView();
+	BuildGeometryWorstCasesView();
+	BuildTimeWorstCasesDetailedView();
+}
+
+void SCustomizableObjecEditorPerformanceReport::SelectInstanceFinishedUpdate(UCustomizableObjectInstance* UpdatedInstance)
+{
+	if (CurrentReportInstance && UpdatedInstance == CurrentReportInstance)
+	{
+		PerformanceReportTextureAnalyzer->RefreshTextureAnalyzerTable(CurrentReportInstance);
+	}
+}
+
+void SPerformanceReportWorstCaseRow::Construct(const FArguments& InArgs, const TSharedRef<STableViewBase>& InOwnerTableView)
+{
+	ReportInstance = InArgs._ReportInstance;
+	PerformanceReport = InArgs._PerformanceReport;
+	SMutableExpandableTableRow<UWorstCasePerformanceReportInstance*>::Construct(FSuperRowType::FArguments().Style(InArgs._Style), InOwnerTableView);
+	Style = InArgs._Style;
+}
+
+void SPerformanceReportWorstCaseRow::AddReferencedObjects(FReferenceCollector& Collector)
+{
+	Collector.AddReferencedObject(ReportInstance);
+	Collector.AddReferencedObjects(SubWidgetData);
+}
+
+TSharedRef<SWidget> SPerformanceReportWorstCaseRow::GenerateWidgetForColumn(const FName& ColumnName)
+{
+	check(ReportInstance);
+
+	const FNumberFormattingOptions BytesFormat = FNumberFormattingOptions().SetMaximumFractionalDigits(2).SetMinimumFractionalDigits(2);
+	const FNumberFormattingOptions TimeFormat = FNumberFormattingOptions().SetMaximumFractionalDigits(2).SetMinimumFractionalDigits(2);
+	const FNumberFormattingOptions PolygonsFormat = FNumberFormattingOptions().SetMaximumFractionalDigits(0);
+	const FNumberFormattingOptions PercentageFormat = FNumberFormattingOptions().SetMaximumFractionalDigits(0).SetAlwaysSign(true);
+	const bool bIsSelected = PerformanceReport && ReportInstance == PerformanceReport->CurrentReportWorstCase;
+	TSharedRef<SOverlay> ResultWidget = SNew(SOverlay);
+
+	if (bIsSelected)
+	{
+		ResultWidget->AddSlot()
+		[
+			SNew(SBorder)
+			.BorderImage(this, &SPerformanceReportWorstCaseRow::GetSelectedRowHighlightColor)
+		];
+	}
+
+	if (ReportInstance->WorstCaseInstance)
+	{
+		if (ColumnName == ColumnID_PR_Delta_Label)
+		{
+			const float Mean = ReportInstance->MeasuredDataForThisWorstCase.MeanTime;
+			if (Mean)
+			{
+				float CurrentMaximum = 0.0f;
+
+				switch (ReportInstance->WorstCaseType)
+				{
+				case UWorstCasePerformanceReportInstance::EPerformanceReportWorstCaseType::LARGEST_TEXTURE_SIZE:
+				{
+					CurrentMaximum = (float)ReportInstance->TextureBytes;
+					break;
+				}
+				case UWorstCasePerformanceReportInstance::EPerformanceReportWorstCaseType::MOST_TRIANGLES:
+				{
+					CurrentMaximum = (float)ReportInstance->NumFaces;
+					break;
+				}
+				case UWorstCasePerformanceReportInstance::EPerformanceReportWorstCaseType::SLOWEST_UPDATE_ENTER_STATE:
+				case UWorstCasePerformanceReportInstance::EPerformanceReportWorstCaseType::SLOWEST_UPDATE_PARAM_IN_RELEVANT_STATE:
+				{
+					CurrentMaximum = ReportInstance->UpdateTime;
+					break;
+				}
+				default:
+				{
+					break;
+				}
+				}
+
+				FText Percentage = FText::AsNumber((CurrentMaximum * 100.0f / Mean) - 100.0f, &PercentageFormat);
+
+				ResultWidget->AddSlot()
+				[
+					SNew(STextBlock)
+					.Text(FText::Format(LOCTEXT("MutablePerformanceDeltaFromAverage", "{0}%"), Percentage))
+				];
+			}
+		}
+		else if (ColumnName == ColumnID_PR_Average_Label)
+		{
+			const float Mean = ReportInstance->MeasuredDataForThisWorstCase.MeanTime;
+			if (Mean)
+			{
+				switch (ReportInstance->WorstCaseType)
+				{
+				case UWorstCasePerformanceReportInstance::EPerformanceReportWorstCaseType::LARGEST_TEXTURE_SIZE:
+				{
+					const float Value = (Mean > 1048576) ? float(Mean) / 1048576.0f : float(Mean) / 1024.0f;
+
+					ResultWidget->AddSlot()
+					[
+						SNew(STextBlock).Text(FText::Format(LOCTEXT("MutablePerformanceTextureSize", "{0} {1}"), FText::AsNumber(Value, &BytesFormat), FText::FromString((Mean > 1048576) ? "MB" : "KB")))
+					];
+					break;
+				}
+				case UWorstCasePerformanceReportInstance::EPerformanceReportWorstCaseType::MOST_TRIANGLES:
+				{
+					ResultWidget->AddSlot()
+					[
+						SNew(STextBlock).Text(FText::AsNumber(Mean, &PolygonsFormat))
+					];
+					break;
+				}
+				case UWorstCasePerformanceReportInstance::EPerformanceReportWorstCaseType::SLOWEST_UPDATE_ENTER_STATE:
+				case UWorstCasePerformanceReportInstance::EPerformanceReportWorstCaseType::SLOWEST_UPDATE_PARAM_IN_RELEVANT_STATE:
+				{
+					ResultWidget->AddSlot()
+					[
+						SNew(STextBlock).Text(FText::AsNumber(Mean, &TimeFormat))
+					];
+					break;
+				}
+				default:
+				{
+					break;
+				}
+				}
+			}
+		}
+		else if (ColumnName == ColumnID_PR_Name_Label)
+		{
+			if (const UCustomizableObject* CustomizableObject = ReportInstance->WorstCaseInstance->GetCustomizableObject();
+				CustomizableObject &&
+				!CustomizableObject->GetPrivate()->IsLocked())
+			{
+				const bool bIsDetailParam = ReportInstance->WorstCaseType == UWorstCasePerformanceReportInstance::EPerformanceReportWorstCaseType::SLOWEST_UPDATE_PARAM_IN_RELEVANT_STATE;
+				const FText TextState = ReportInstance->WorstCaseType == UWorstCasePerformanceReportInstance::EPerformanceReportWorstCaseType::SLOWEST_UPDATE_ENTER_STATE ?
+					FText::FromString("State: " + CustomizableObject->GetStateName(ReportInstance->LongestUpdateTimeStateIndex))
+					: FText::FromString((bIsDetailParam ? "		" : "") + CustomizableObject->GetParameterName(ReportInstance->LongestUpdateTimeParameterIndexInCO));
+
+				ResultWidget->AddSlot()
+				[
+					SNew(STextBlock).Text(TextState)
+				];
+			}
+		}
+		else if (ColumnName == ColumnID_PR_UpdateTime_Label)
+		{
+			ResultWidget->AddSlot()
+			[
+				SNew(STextBlock).Text(FText::AsNumber(ReportInstance->UpdateTime, &TimeFormat))
+			];
+		}
+		else if (ColumnName == ColumnID_PR_LOD_Label)
+		{
+			const FText TextLOD = FText::AsNumber(ReportInstance->LOD, &FNumberFormattingOptions::DefaultNoGrouping());
+
+			ResultWidget->AddSlot()
+			[
+				SNew(STextBlock).Text(FText::Format(LOCTEXT("MutablePerformanceLOD", "{0}"), TextLOD))
+			];
+		}
+		else if (ColumnName == ColumnID_PR_TexturesTotalSize_Label)
+		{
+			const float Bytes = ReportInstance->TextureBytes;
+			const float Value = (Bytes > 1048576) ? float(Bytes) / 1048576.0f : float(Bytes) / 1024.0f;
+
+			ResultWidget->AddSlot()
+			[
+				SNew(STextBlock).Text(FText::Format(LOCTEXT("MutablePerformanceTextureSize", "{0} {1}"), FText::AsNumber(Value, &BytesFormat), FText::FromString((Bytes > 1048576) ? "MB" : "KB")))
+			];
+		}
+		else if (ColumnName == ColumnID_PR_NumberOfTriangles_Label)
+		{
+			ResultWidget->AddSlot()
+			[
+				SNew(STextBlock).Text(FText::AsNumber(ReportInstance->NumFaces, &PolygonsFormat))
+			];
+		}
+		else if (ColumnName == ColumnID_PR_InstanceName_Label)
+		{
+			const FText InstanceName = FText::FromString(ReportInstance->GetName());
+
+			ResultWidget->AddSlot()
+			[
+				SNew(STextBlock).Text(InstanceName)
+			];
+		}
+	}
+
+	return ResultWidget;
+}
+
+TSharedPtr<SWidget> SPerformanceReportWorstCaseRow::GenerateAdditionalWidgetForRow()
+{
+	if (ReportInstance && ReportInstance->WorstCaseType == UWorstCasePerformanceReportInstance::EPerformanceReportWorstCaseType::SLOWEST_UPDATE_ENTER_STATE)
+	{
+		const TArray<UWorstCasePerformanceReportInstance*>* WorstTimeParameters = PerformanceReport->GetWorstTimeParameters(ReportInstance->LOD, ReportInstance->LongestUpdateTimeStateIndex);
+
+		if (WorstTimeParameters)
+		{
+			SubWidgetData.Empty(WorstTimeParameters->Num());
+			for (UWorstCasePerformanceReportInstance* elem : *WorstTimeParameters)
+			{
+				SubWidgetData.Push(elem);
+			}
+			return SNew(SListView<UWorstCasePerformanceReportInstance*>)
+				.ListItemsSource(&ObjectPtrDecay(SubWidgetData))
+				.OnGenerateRow(this, &SPerformanceReportWorstCaseRow::GenerateStyledRow)
+				.OnMouseButtonClick(PerformanceReport, &SCustomizableObjecEditorPerformanceReport::SelectInstance)
+				.SelectionMode(ESelectionMode::Single)
+				.IsFocusable(true)
+				.HeaderRow(PerformanceReport->TimeWorstCasesViewHeader);
+		}
+	}
+
+	SubWidgetData.Empty();
+	return nullptr;
+}
+
+EVisibility SPerformanceReportWorstCaseRow::GetAdditionalWidgetDefaultVisibility() const
+{
+	return (ReportInstance) ? ReportInstance->SubVisibility : EVisibility::Collapsed;
+}
+
+void SPerformanceReportWorstCaseRow::SetAdditionalWidgetVisibility(const EVisibility InVisibility)
+{
+	if (ReportInstance)
+	{
+		ReportInstance->SubVisibility = InVisibility;
+	}
+	SMutableExpandableTableRow::SetAdditionalWidgetVisibility(InVisibility);
+}
+
+const FSlateBrush* SPerformanceReportWorstCaseRow::GetSelectedRowHighlightColor() const
+{
+	const bool bIsSelected = PerformanceReport && ReportInstance == PerformanceReport->CurrentReportWorstCase;
+	const bool bIsUpToDate = bIsSelected && ReportInstance && ReportInstance->WorstCaseInstance == PerformanceReport->CurrentReportInstance;
+	return bIsUpToDate? UE_MUTABLE_GET_BRUSH("Brushes.Highlight") : UE_MUTABLE_GET_BRUSH("Brushes.AccentGray");
+}
+
+TSharedRef<ITableRow> SPerformanceReportWorstCaseRow::GenerateStyledRow(UWorstCasePerformanceReportInstance* ReportedInstance, const TSharedRef<STableViewBase>& OwnerTable)
+{
+	const bool bEvenEntryIndex = (GetIndexInList() % 2 == 0);
+	return SNew(SPerformanceReportWorstCaseRow, OwnerTable).ReportInstance(ReportedInstance).PerformanceReport(PerformanceReport).Style(
+		bEvenEntryIndex ?
+		&FCustomizableObjectEditorStyle::Get().GetWidgetStyle<FTableRowStyle>("CustomizableObjectPerformanceReportAllEven.Row")
+		: &FCustomizableObjectEditorStyle::Get().GetWidgetStyle<FTableRowStyle>("CustomizableObjectPerformanceReportAllOdd.Row")
+	);
+}
+
+void UPerformanceReportHelper::DelegatedCallback(UCustomizableObjectInstance* Instance)
+{
+	if (PerformanceReport)
+	{
+		PerformanceReport->SelectInstanceFinishedUpdate(Instance);
+	}
+}
+
+#undef LOCTEXT_NAMESPACE
diff --git a/Engine/Plugins/Mutable/Source/CustomizableObjectEditor/Private/MuCOE/SCustomizableObjectEditorPerformanceReport.h b/Engine/Plugins/Mutable/Source/CustomizableObjectEditor/Private/MuCOE/SCustomizableObjectEditorPerformanceReport.h
new file mode 100644
--- /dev/null
+++ b/Engine/Plugins/Mutable/Source/CustomizableObjectEditor/Private/MuCOE/SCustomizableObjectEditorPerformanceReport.h
@@ -0,0 +1,284 @@
+// Copyright Epic Games, Inc. All Rights Reserved.
+
+#pragma once
+
+#include "Misc/NotifyHook.h"
+#include "MuCOE/StatePerformanceTesting.h"
+#include "MuCOE/StressTest.h"
+#include "MuCOE/Widgets/SMutableExpandableTableRow.h"
+
+#include "Widgets/Views/SListView.h"
+#include "Widgets/Views/ITableRow.h"
+#include "SCustomizableObjectEditorPerformanceReport.generated.h"
+
+template <typename NumericType> class SNumericEntryBox;
+
+class ITableRow;
+class STableViewBase;
+class SWidget;
+class UCustomizableObject;
+class UCustomizableObjectInstance;
+struct FFrame;
+struct FGeometry;
+struct FSlateBrush;
+struct FTableRowStyle;
+
+
+// Instances performing worst in some aspect
+UCLASS()
+class UWorstCasePerformanceReportInstance : public UObject
+{
+public:
+	GENERATED_BODY()
+
+	enum EPerformanceReportWorstCaseType {
+		LARGEST_TEXTURE_SIZE,
+		MOST_TRIANGLES,
+		SLOWEST_UPDATE_ENTER_STATE,
+		SLOWEST_UPDATE_PARAM_IN_RELEVANT_STATE
+	};
+
+	// The instance that presented this worst case.
+	UPROPERTY(VisibleAnywhere, AssetRegistrySearchable, Transient, NoClear, Category = "Stats|PerformanceReport", meta = (DisplayName = "Instance", ColumnWidth = "90")/*"PerformanceReport|BiggestTextureSize"*/)
+	TObjectPtr<class UCustomizableObjectInstance> WorstCaseInstance = nullptr;
+
+	// What worst case does this represent
+	EPerformanceReportWorstCaseType WorstCaseType = LARGEST_TEXTURE_SIZE;
+
+	// LOD number where this instance generated the worst case  (currently only relevant when LARGEST_TEXTURE_SIZE or MOST_TRIANGLES)
+	UPROPERTY(VisibleAnywhere, AssetRegistrySearchable, Transient, NoClear, Category = "Stats|PerformanceReport", meta = (DisplayName = "LOD", ColumnWidth = "10", DisplayRight = "true")/*"PerformanceReport|BiggestTextureSize"*/)
+	int32 LOD = INDEX_NONE;
+
+	// Total texture memory size that this instance had when the worst case was recorded (currently only relevant when LARGEST_TEXTURE_SIZE)
+	UPROPERTY(VisibleAnywhere, AssetRegistrySearchable, Transient, NoClear, Category = "Stats|PerformanceReport|Texture", meta = (DisplayName = "Textures Total Size", ColumnWidth = "60", DisplayRight = "true"))
+	int32 TextureBytes = -1;
+
+	// Amount of triangles this instance had when the worst case was recorded (currently only relevant when MOST_TRIANGLES)
+	UPROPERTY(VisibleAnywhere, AssetRegistrySearchable, Transient, NoClear, Category = "Stats|PerformanceReport|Geometry", meta = (DisplayName = "Number of Triangles", ColumnWidth = "60", DisplayRight = "true"))
+	int32 NumFaces = -1;
+
+	// How long it took this instance to update when the worst case was recorded (currently only relevant when SLOWEST_UPDATE_PARAM_IN_RELEVANT_STATE or SLOWEST_UPDATE_PARAM_IN_RELEVANT_STATE)
+	UPROPERTY(VisibleAnywhere, AssetRegistrySearchable, Transient, NoClear, Category = "Stats|PerformanceReport|Time", meta = (DisplayName = "Update Time", ColumnWidth = "60", DisplayRight = "true"))
+	float UpdateTime = -1.0;
+
+	// State CO index where this instance took most cycles to update (currently only relevant when SLOWEST_UPDATE_ENTER_STATE or SLOWEST_UPDATE_PARAM_IN_RELEVANT_STATE)
+	UPROPERTY(VisibleAnywhere, AssetRegistrySearchable, Transient, NoClear, Category = "Stats|PerformanceReport|Time", meta = (DisplayName = "State ID", ColumnWidth = "60", DisplayRight = "true"))
+	int32 LongestUpdateTimeStateIndex = INDEX_NONE;
+
+	// Parameter CO index where this instance took most cycles to update (currently only relevant when SLOWEST_UPDATE_PARAM_IN_RELEVANT_STATE)
+	UPROPERTY(VisibleAnywhere, AssetRegistrySearchable, Transient, NoClear, Category = "Stats|PerformanceReport|Time", meta = (DisplayName = "Parameter ID", ColumnWidth = "60", DisplayRight = "true"))
+	int32 LongestUpdateTimeParameterIndexInCO = INDEX_NONE;
+
+	// Measured data for the current worst case, the type depends on which is the current WorstCaseType
+	MeasuredData MeasuredDataForThisWorstCase;
+
+	EVisibility SubVisibility = EVisibility::Collapsed;
+
+	// TODO: Obtain all data in all cases
+	//MaterialInfoMap MapMaterialParam;
+
+	void SetInstance(class UCustomizableObjectInstance* UpdatedInstance);
+};
+
+
+UCLASS()
+class UPerformanceReportHelper : public UObject
+{
+public:
+	GENERATED_BODY()
+	// Standard constructor
+	UPerformanceReportHelper() {}
+
+	// Method to assign for the callback
+	UFUNCTION()
+	void DelegatedCallback(UCustomizableObjectInstance* Instance);
+
+	// Properties to be updated once the instance has been updated
+	class SCustomizableObjecEditorPerformanceReport* PerformanceReport = nullptr;
+};
+
+
+// Editor GUI that executes all relevant testing of a CO and presents and prints results
+class SCustomizableObjecEditorPerformanceReport : public SCompoundWidget, public FNotifyHook, public FGCObject
+{
+public:
+
+	SLATE_BEGIN_ARGS(SCustomizableObjecEditorPerformanceReport){}
+		SLATE_ARGUMENT(UCustomizableObject *,CustomizableObject)
+	SLATE_END_ARGS()
+
+	void Construct(const FArguments& InArgs);
+	virtual ~SCustomizableObjecEditorPerformanceReport();
+
+	// FSerializableObject interface
+	void AddReferencedObjects(FReferenceCollector& Collector) override;
+	virtual FString GetReferencerName() const override
+	{
+		return TEXT("SCustomizableObjecEditorPerformanceReport");
+	}
+
+	void Tick(const FGeometry& AllottedGeometry, const double InCurrentTime, const float InDeltaTime) override;
+
+	void StopTests();
+
+	// Force the re-start of the performance report
+	FReply RunPerformanceReport();
+	// Force the re-start of the performance report and provides the instance to be tested
+	FReply RunPerformanceReport(UCustomizableObject* ObjectToTest);
+	// Select and view the details of the slowest instance found up to this moment
+	FReply SelectSlowestInstance();
+
+	bool IsGetWorstInstanceButtonEnabled() const;
+
+	FReply ResetPerformanceReportOptions();
+
+private:
+
+	// Customizable object whose performance is being analyzed
+	TObjectPtr<UCustomizableObject> ReportSubject;
+
+	// Number of LODs that the Customizable Object has
+	uint32 NumLOD;
+
+	// Number of States that the Customizable Object has
+	uint32 NumStates;
+
+	// Currently focused report case and instance
+	TObjectPtr<UWorstCasePerformanceReportInstance> CurrentReportWorstCase = nullptr;
+	TObjectPtr<class UCustomizableObjectInstance> CurrentReportInstance = nullptr;
+
+	// Stress tester ================================================================================================================================================================
+	TSharedPtr<FRunningStressTest> StressTest;
+
+	int StressTestOptionInstanceCount;
+	TSharedPtr<class SNumericEntryBox<int32>> StressTestOptionInstanceCountEntry;
+	TOptional<int32> GetStressTestOptionInstanceCount() const;
+	void OnStressTestOptionInstanceCountChanged(int32 Value);
+
+	int StressTestOptionCreateInstanceTime;
+	TSharedPtr<class SNumericEntryBox<int32>> StressTestOptionCreateInstanceTimeEntry;
+	TOptional<int32> GetStressTestOptionCreateInstanceTime() const;
+	void OnStressTestOptionCreateInstanceTimeChanged(int32 Value);
+
+	int StressTestOptionCreateInstanceTimeVar;
+	TSharedPtr<class SNumericEntryBox<int32>> StressTestOptionCreateInstanceTimeVarEntry;
+	TOptional<int32> GetStressTestOptionCreateInstanceTimeVar() const;
+	void OnStressTestOptionCreateInstanceTimeVarChanged(int32 Value);
+
+	int StressTestOptionInstanceUpdateTime;
+	TSharedPtr<class SNumericEntryBox<int32>> StressTestOptionInstanceUpdateTimeEntry;
+	TOptional<int32> GetStressTestOptionInstanceUpdateTime() const;
+	void OnStressTestOptionInstanceUpdateTimeChanged(int32 Value);
+
+	int StressTestOptionInstanceUpdateTimeVar;
+	TSharedPtr<class SNumericEntryBox<int32>> StressTestOptionInstanceUpdateTimeVarEntry;
+	TOptional<int32> GetStressTestOptionInstanceUpdateTimeVar() const;
+	void OnStressTestOptionInstanceUpdateTimeVarChanged(int32 Value);
+
+	int StressTestOptionInstanceLifeTime;
+	TSharedPtr<class SNumericEntryBox<int32>> StressTestOptionInstanceLifeTimeEntry;
+	TOptional<int32> GetStressTestOptionInstanceLifeTime() const;
+	void OnStressTestOptionInstanceLifeTimeChanged(int32 Value);
+
+	int StressTestOptionInstanceLifeTimeVar;
+	TSharedPtr<class SNumericEntryBox<int32>> StressTestOptionInstanceLifeTimeVarEntry;
+	TOptional<int32> GetStressTestOptionInstanceLifeTimeVar() const;
+	void OnStressTestOptionInstanceLifeTimeVarChanged(int32 Value);
+
+	// State tester =================================================================================================================================================================
+	TSharedPtr<class FTestingCustomizableObject> StateTest;
+
+	// Texture Analyzer =============================================================================================================================================================
+	TSharedPtr<class SCustomizableObjecEditorTextureAnalyzer> PerformanceReportTextureAnalyzer;
+
+	// Worst cases ==================================================================================================================================================================
+	TArray<TObjectPtr<UWorstCasePerformanceReportInstance>> WorstTextureInstances;	// By LOD
+	TArray<TObjectPtr<UWorstCasePerformanceReportInstance>> WorstGeometryInstances;	// By LOD
+	TObjectPtr<const UWorstCasePerformanceReportInstance> WorstTimeInstance;	// Slowest instance in any test that checks for time for this report
+	TArray<TArray<TObjectPtr<UWorstCasePerformanceReportInstance>>> WorstTimeStateInstances;	// By LOD then by state index
+	TArray<TArray<TArray<TObjectPtr<UWorstCasePerformanceReportInstance>>>> WorstTimeParameterInstances;	// By LOD then by state index then all parameters relevant for that state (not ordered)
+	void ResetSavedData();
+	void WorstFacesFound(uint32 CurrentLODNumFaces, uint32 LOD, class UCustomizableObjectInstance* Instance, MeasuredData* data);
+	void WorstTextureFound(MaterialInfoMap MapMaterialParam, uint32 CurrentLODTextureSize, uint32 LOD, class UCustomizableObjectInstance* Instance, MeasuredData* data);
+	void WorstTimeFound(float CurrentInstanceGenerationTime, int32 LongesTimeStateIndex, int32 LongesTimeParameterIndexInCO, uint32 LOD, class UCustomizableObjectInstance* Instance, MeasuredData* data);
+	void UpdateWorstTimeInstance(const UWorstCasePerformanceReportInstance* EvenWorseTime);
+	void StressTestEnded();
+	void StateTestEnded();
+	UWorstCasePerformanceReportInstance* GetWorstCaseStructureFromIndexes(UWorstCasePerformanceReportInstance::EPerformanceReportWorstCaseType WorstCaseType, uint32 LOD, int32 StateID = INDEX_NONE, int32 ParamID = INDEX_NONE) const;
+
+	// UI ==========================================================================================================================================================================
+	TSharedRef<SWidget> GenerateOptionsMenuContent();
+	TSharedRef<SWidget> BuildAnalyzer();
+	void UpdateSelectedInstanceData();
+	void BuildTextureWorstCasesView();
+	void BuildGeometryWorstCasesView();
+	void BuildTimeWorstCasesDetailedView();
+	void OnOpenAssetEditor(UObject* ObjectToEdit);
+	void ChangeFocusedReportInstance(UWorstCasePerformanceReportInstance* InReportInstance);
+	FString OnGetInstancePath() const;
+	bool IsSelectedInstanceOutdated() const;
+	FReply RefreshSelectedInstance();
+	TSharedRef<ITableRow> GenerateRow(UWorstCasePerformanceReportInstance* ReportedInstance, const TSharedRef<STableViewBase>& OwnerTable);
+	// This array pointer is only valid until the array of arrays WorstTimeParameterInstances is modified. Use this only to access the actual instance pointers, but don't keep the array pointer for later use.
+	const TArray<UWorstCasePerformanceReportInstance*>* GetWorstTimeParameters(uint32 LOD, uint32 StateIdx) const;
+	void SelectInstance(UWorstCasePerformanceReportInstance* ClickedInstance);
+	void SelectInstanceFinishedUpdate(class UCustomizableObjectInstance* ClickedInstance);
+
+	TSharedPtr<class SHeaderRow> TimeWorstCasesViewHeader;
+	TSharedPtr<class SVerticalBox> TextureWorstCasesView;
+	TSharedPtr<class SVerticalBox> GeometryWorstCasesView;
+	TSharedPtr<class SVerticalBox> TimeWorstCasesView;
+	TSharedPtr<class SScrollBox> TimeWorstCasesDetailedView;
+	TSharedPtr<class STextBlock> StatusText;
+	TSharedPtr<class STextBlock> SelectedTextureAnalyzerText;
+	TSharedPtr<class SButton> PerformanceReportResetOptionsEntry;
+	TSharedPtr<class SObjectPropertyEntryBox> InstanceThumbnail;
+	TSharedPtr<class FAssetThumbnailPool> AssetThumbnailPool;
+	TSharedPtr<class FAssetThumbnail> AssetThumbnail;
+	TSharedPtr<class SButton> PerformanceReportRefreshSelectedInstance;
+
+	/** UObject class to be able to use the UCustomizableObjectInstance::UpdatedDelegate update callback */
+	UPROPERTY()
+	TObjectPtr<UPerformanceReportHelper> HelperCallback;
+
+	friend class UPerformanceReportHelper;
+	friend class SPerformanceReportWorstCaseRow;
+};
+
+
+class SPerformanceReportWorstCaseRow : public SMutableExpandableTableRow<UWorstCasePerformanceReportInstance*>, public FGCObject
+{
+public:
+	SLATE_BEGIN_ARGS(SPerformanceReportWorstCaseRow) {}
+		SLATE_ARGUMENT(UWorstCasePerformanceReportInstance*, ReportInstance)
+		SLATE_ARGUMENT(SCustomizableObjecEditorPerformanceReport*, PerformanceReport)
+		SLATE_ARGUMENT(const FTableRowStyle*, Style)
+	SLATE_END_ARGS()
+
+	void Construct(const FArguments& InArgs, const TSharedRef<STableViewBase>& InOwnerTableView);
+
+	// FSerializableObject interface
+	void AddReferencedObjects(FReferenceCollector& Collector) override;
+	virtual FString GetReferencerName() const override
+	{
+		return TEXT("SPerformanceReportWorstCaseRow");
+	}
+
+private:
+	// Report instance we represent
+	TObjectPtr<UWorstCasePerformanceReportInstance> ReportInstance;
+
+	// Performance report we are from. Necessary to generate the parameter worst time case widgets when the row is a state worst case
+	SCustomizableObjecEditorPerformanceReport* PerformanceReport;
+
+	// Local copy of additional widget data
+	TArray<TObjectPtr<UWorstCasePerformanceReportInstance>> SubWidgetData;
+
+	virtual TSharedRef<SWidget> GenerateWidgetForColumn(const FName& ColumnName) override;
+	virtual TSharedPtr<SWidget> GenerateAdditionalWidgetForRow() override;
+	virtual EVisibility GetAdditionalWidgetDefaultVisibility() const override;
+	virtual void SetAdditionalWidgetVisibility(const EVisibility Visibility) override;
+	const FSlateBrush* GetSelectedRowHighlightColor() const;
+
+	TSharedRef<ITableRow> GenerateStyledRow(UWorstCasePerformanceReportInstance* ReportedInstance, const TSharedRef<STableViewBase>& OwnerTable);
+};
