[mutable] Moved the Mutable plugin out of Experimental status into Beta.

diff --git a/Engine/Plugins/Mutable/Source/CustomizableObjectEditor/Private/MuCOE/SCustomizableObjectProperties.cpp b/Engine/Plugins/Mutable/Source/CustomizableObjectEditor/Private/MuCOE/SCustomizableObjectProperties.cpp
new file mode 100644
--- /dev/null
+++ b/Engine/Plugins/Mutable/Source/CustomizableObjectEditor/Private/MuCOE/SCustomizableObjectProperties.cpp
@@ -0,0 +1,443 @@
+// Copyright Epic Games, Inc. All Rights Reserved.
+
+#include "MuCOE/SCustomizableObjectProperties.h"
+
+#include "Framework/Views/TableViewMetadata.h"
+#include "Widgets/Input/SComboButton.h"
+#include "Widgets/Views/STreeView.h"
+
+class ITableRow;
+class SWidget;
+class UCustomizableObject;
+
+#define LOCTEXT_NAMESPACE "SCustomizableObjectProperties"
+
+//DEFINE_LOG_CATEGORY_STATIC(LogEditorCustomizableObjectProperties, Log, All);
+
+class FStatePropertiesNode
+{
+public:
+
+	FStatePropertiesNode()
+	{
+	}
+
+	FStatePropertiesNode( const FString& InName )
+	{
+		Name = MakeShareable( new FString(InName) );
+	}
+
+
+	TSharedPtr<FString> Name;
+
+	TSharedPtr<FString> GetName()
+	{
+		return Name;
+	}
+
+
+	TArray<TSharedPtr<FStatePropertiesNode>> GetChildrenList()
+	{
+		TArray<TSharedPtr<FStatePropertiesNode>> ChildrenList;
+
+		return ChildrenList;
+	}
+
+};
+
+
+/** The item used for visualizing the class in the tree. */
+class SStatePropertiesItem : public STableRow< TSharedPtr<FString> >
+{
+public:
+	
+	SLATE_BEGIN_ARGS( SStatePropertiesItem )
+		: _StateName()
+		, _OwnerTableView()
+		, _HighlightText(NULL)
+		, _TextColor(FLinearColor(1.0f, 1.0f, 1.0f, 1.0f))
+		{}
+
+		/** The classname this item contains. */
+		SLATE_ARGUMENT( TSharedPtr<FString>, StateName )
+		/** The table that owns this item. */
+		SLATE_ARGUMENT( TSharedPtr<STableViewBase>, OwnerTableView )
+		/** The text this item should highlight, if any. */
+		SLATE_ARGUMENT( const FString*, HighlightText )
+		/** The color text this item will use. */
+		SLATE_ARGUMENT( FSlateColor, TextColor )
+		/** The node this item is associated with. */
+		SLATE_ARGUMENT( TSharedPtr<FStatePropertiesNode>, AssociatedNode)
+
+	SLATE_END_ARGS()
+
+	/**
+		* Construct the widget
+		*
+		* @param InArgs   A declaration from which to construct the widget
+		*/
+	void Construct( const FArguments& InArgs )
+	{
+		StateName = InArgs._StateName;
+		AssociatedNode = InArgs._AssociatedNode;
+
+		this->ChildSlot
+		[
+			SNew(SHorizontalBox)
+				
+			+SHorizontalBox::Slot()
+				[
+					SNew( SExpanderArrow, SharedThis(this) )
+				]
+
+			+SHorizontalBox::Slot()
+				.FillWidth(1.0f)
+				.Padding( 0.0f, 3.0f, 6.0f, 3.0f )
+				.VAlign(VAlign_Center)
+				[
+					SNew( STextBlock )
+						.Text( FText::FromString(*StateName.Get() ) )
+						//.HighlightString(*InArgs._HighlightText)
+						//.ColorAndOpacity( this, &SClassItem::GetTextColor)
+				]
+
+			+SHorizontalBox::Slot()
+				.HAlign(HAlign_Right)
+				.VAlign(VAlign_Center)
+				.Padding( 0.0f, 0.0f, 6.0f, 0.0f )
+				[
+					SNew( SComboButton )
+						.ContentPadding(FMargin(2.0f))
+						//.OnGetMenuContent(this, &SClassItem::GenerateDropDown)
+				]
+		];
+		
+		TextColor = InArgs._TextColor;
+
+		STableRow< TSharedPtr<FString> >::ConstructInternal(
+			STableRow::FArguments()
+				.ShowSelection(true),
+			InArgs._OwnerTableView.ToSharedRef() );	
+	}
+
+private:
+
+	/**
+		* Generates the drop down menu for the item.
+		*
+		* @return		The drop down menu widget.
+		*/
+	/*
+	TSharedRef<SWidget> GenerateDropDown()
+	{
+		// Empty list of commands.
+		TSharedPtr< FUICommandList > Commands;
+
+		const bool bShouldCloseWindowAfterMenuSelection = true;	// Set the menu to automatically close when the user commits to a choice
+		FMenuBuilder MenuBuilder( bShouldCloseWindowAfterMenuSelection, Commands);
+		{ 
+			if (UClass* Class = AssociatedNode->Class.Get())
+			{
+				bool bIsBlueprint(false);
+				bool bHasBlueprint(false);
+
+				ClassViewer::Helpers::GetClassInfo(Class, bIsBlueprint, bHasBlueprint);
+
+				bHasBlueprint = AssociatedNode->Blueprint.IsValid();
+
+				if (bIsBlueprint)
+				{
+					FUIAction Action( FExecuteAction::CreateStatic( &ClassViewer::Helpers::CreateBlueprint, Class ) );
+					MenuBuilder.AddMenuEntry(TEXT("Create Blueprint..."), TEXT("Creates a new Blueprint from the current one."), NAME_None, Action);
+				}
+
+				if (bHasBlueprint)
+				{
+					MenuBuilder.AddMenuSeparator();
+
+					FUIAction Action( FExecuteAction::CreateStatic( &ClassViewer::Helpers::OpenBlueprintTool, ClassViewer::Helpers::GetBlueprint(Class) ) );
+					MenuBuilder.AddMenuEntry(TEXT("Open Blueprint Editor"), TEXT("Open the Blueprint in the editor."), NAME_None, Action);
+				}
+
+				if (bHasBlueprint)
+				{
+					MenuBuilder.AddMenuSeparator();
+
+					FUIAction Action( FExecuteAction::CreateStatic( &ClassViewer::Helpers::FindInContentBrowser, ClassViewer::Helpers::GetBlueprint(Class), Class ) );
+					MenuBuilder.AddMenuEntry(TEXT("Find In Content Browser"), TEXT("Find in Content Browser"), NAME_None, Action);
+				}
+			}
+		}
+
+		return MenuBuilder.MakeWidget();
+	}
+	*/
+
+	/** Returns the text color for the item based on if it is selected or not. */
+	FSlateColor GetTextColor() const
+	{
+		const TSharedPtr< ITypedTableView< TSharedPtr<FString> > > OwnerWidget = OwnerTablePtr.Pin();
+		const TSharedPtr<FString>* MyItem = OwnerWidget->Private_ItemFromWidget( this );
+		const bool bIsSelected = OwnerWidget->Private_IsItemSelected( *MyItem );
+
+		if(bIsSelected)
+		{
+			return FSlateColor::UseForeground();
+		}
+
+		return TextColor;
+	}
+
+private:
+
+	/** The class name for which this item is associated with. */
+	TSharedPtr<FString> StateName;
+
+	/** The text color for this item. */
+	FSlateColor TextColor;
+
+	/** The Class Viewer Node this item is associated with. */
+	TSharedPtr< FStatePropertiesNode > AssociatedNode;
+
+};
+
+
+void SCustomizableObjectProperties::Construct(const FArguments& InArgs )
+{
+	bDisablePopulate = false;
+	bSaveExpansionStates = true;
+
+	Object = InArgs._Object;
+
+	this->ChildSlot
+	[
+		SNew(SVerticalBox)
+
+		+SVerticalBox::Slot()
+		.FillHeight(1.0f)
+		[
+			SNew(SVerticalBox)
+
+			+SVerticalBox::Slot()
+			.FillHeight(1.0f)
+			[
+				SAssignNew(StatesTree, STreeView<TSharedPtr< FStatePropertiesNode > >)
+					//.Visibility(EVisibility::Collapsed)
+					.SelectionMode(ESelectionMode::Single)
+					.TreeItemsSource( &RootTreeItems )
+					// Called to child items for any given parent item
+					.OnGetChildren( this, &SCustomizableObjectProperties::OnGetChildrenForTree )
+					// Generates the actual widget for a tree item
+					.OnGenerateRow( this, &SCustomizableObjectProperties::OnGenerateRow ) 
+
+					// Generates the right click menu.
+					//.OnContextMenuOpening(this, &SClassViewer::BuildMenuWidget)
+
+					// Find out when the user selects something in the tree
+					//.OnSelectionChanged( this, &SClassViewer::OnClassViewerSelectionChanged )
+
+					.HeaderRow
+					(
+						SNew(SHeaderRow)
+						.Visibility(EVisibility::Collapsed)
+						+ SHeaderRow::Column(TEXT("State"))
+						.DefaultLabel(NSLOCTEXT("CustomizableObjectDetails","State","State"))
+					)
+			]
+		]
+	];
+
+	
+
+	// Register to have Populate called when doing a Hot Reload or when a Blueprint is compiled.
+	//GEditor->OnHotReload().AddSP(this, &SCustomizableObjectProperties::Populate);
+	//GEditor->OnBlueprintCompiled().AddSP(this, &SCustomizableObjectProperties::Populate);
+
+	Populate();
+}
+
+void SCustomizableObjectProperties::SetObject( UCustomizableObject* InObject )
+{
+	Object = InObject;
+	Populate();
+}
+
+TSharedRef<SWidget> SCustomizableObjectProperties::GetContent()
+{
+	return SharedThis( this );
+}
+
+SCustomizableObjectProperties::~SCustomizableObjectProperties()
+{
+	// Unregister to have Populate called when doing a Hot Reload or when a Blueprint is compiled.
+	//GEditor->OnHotReload().RemoveAll( this );
+	//GEditor->OnBlueprintCompiled().RemoveAll( this );
+	//GEditor->OnClassPackageLoadedOrUnloaded().RemoveAll( this );
+}
+
+void SCustomizableObjectProperties::OnGetChildrenForTree( TSharedPtr<FStatePropertiesNode> InParent, TArray< TSharedPtr< FStatePropertiesNode > >& OutChildren )
+{
+	// Simply return the children, it's already setup.
+	OutChildren = InParent->GetChildrenList();
+}
+
+
+TSharedRef< ITableRow > SCustomizableObjectProperties::OnGenerateRow( TSharedPtr<FStatePropertiesNode> Item, const TSharedRef< STableViewBase >& OwnerTable )
+{	
+	// If the item was accepted by the filter, leave it bright, otherwise dim it.
+	float AlphaValue = 1.0f;
+	TSharedRef< ITableRow > ReturnRow = SNew(SStatePropertiesItem)
+		.StateName(Item->GetName())
+		.TextColor(FLinearColor(1.0f, 1.0f, 1.0f, AlphaValue))
+		.AssociatedNode(Item)
+		.OwnerTableView( OwnerTable )
+		;
+
+	// Expand the item if needed.
+	bool* bIsExpanded = ExpansionStateMap.Find( Item->GetName() );
+	if( bIsExpanded && *bIsExpanded )
+	{
+		StatesTree->SetItemExpansion( Item, *bIsExpanded );
+	}
+
+	return ReturnRow;
+}
+
+
+void SCustomizableObjectProperties::SetAllExpansionStates(bool bInExpansionState)
+{
+	// Go through all the items in the root of the tree and recursively visit their children to set every item in the tree.
+	for(int32 ChildIndex = 0; ChildIndex < RootTreeItems.Num(); ChildIndex++)
+	{
+		SetAllExpansionStates_Helper( RootTreeItems[ChildIndex], bInExpansionState );
+	}
+}
+
+void SCustomizableObjectProperties::SetAllExpansionStates_Helper(TSharedPtr< FStatePropertiesNode > InNode, bool bInExpansionState)
+{
+	StatesTree->SetItemExpansion(InNode, bInExpansionState);
+
+	// Recursively go through the children.
+	for(int32 ChildIndex = 0; ChildIndex < InNode->GetChildrenList().Num(); ChildIndex++)
+	{
+		SetAllExpansionStates_Helper( InNode->GetChildrenList()[ChildIndex], bInExpansionState );
+	}
+}
+
+void SCustomizableObjectProperties::ResetExpansionStates()
+{
+	// Empty the map of states, do not save the states, and refresh the tree with a call to Populate so it will build the tree fresh.
+	ExpansionStateMap.Empty();
+	bSaveExpansionStates = false;
+	Populate();
+}
+
+
+void SCustomizableObjectProperties::MapExpansionStatesInTree( TSharedPtr<FStatePropertiesNode> InItem )
+{
+	ExpansionStateMap.Add( InItem->GetName(), StatesTree->IsItemExpanded( InItem ) );
+
+	// Map out all the children, this will be done recursively.
+	for( int32 ChildIdx(0); ChildIdx < InItem->GetChildrenList().Num(); ++ChildIdx )
+	{
+		MapExpansionStatesInTree( InItem->GetChildrenList()[ChildIdx] );
+	}
+}
+
+void SCustomizableObjectProperties::SetExpansionStatesInTree( TSharedPtr<FStatePropertiesNode> InItem )
+{
+	bool* bIsExpanded = ExpansionStateMap.Find( InItem->GetName() );
+	if( bIsExpanded )
+	{
+		StatesTree->SetItemExpansion( InItem, *bIsExpanded );
+
+		// No reason to set expansion states if the parent is not expanded, it does not seem to do anything.
+		if( *bIsExpanded )
+		{
+			for( int32 ChildIdx(0); ChildIdx < InItem->GetChildrenList().Num(); ++ChildIdx )
+			{
+				SetExpansionStatesInTree( InItem->GetChildrenList()[ChildIdx] );
+			}
+		}
+	}
+	else
+	{
+		// Default to no expansion.
+		StatesTree->SetItemExpansion( InItem, false );
+	}
+}
+
+void SCustomizableObjectProperties::Populate()
+{
+	// Populating the tree/list is disabled during loads to stop needless re-population.
+	if(bDisablePopulate)
+	{
+		return;
+	}
+
+	// If showing a class tree, we may need to save expansion states.
+	if( StatesTree->GetVisibility() == EVisibility::Visible )
+	{
+		if( bSaveExpansionStates )
+		{
+			for( int32 ChildIdx(0); ChildIdx < RootTreeItems.Num(); ++ChildIdx )
+			{
+				// Check if the item is actually expanded or if it's only expanded because it is root level.
+				bool* bIsExpanded = ExpansionStateMap.Find( RootTreeItems[ChildIdx]->GetName() );
+				if( (bIsExpanded && !*bIsExpanded) || !bIsExpanded)
+				{
+					StatesTree->SetItemExpansion( RootTreeItems[ChildIdx], false );
+				}
+
+				// Recursively map out the expansion state of the tree-node.
+				MapExpansionStatesInTree( RootTreeItems[ChildIdx] );
+			}
+		}
+
+		// This is set to false before the call to populate when it is not desired.
+		bSaveExpansionStates = true;
+	}
+
+	// Empty the tree out so it can be redone.
+	RootTreeItems.Empty();
+
+	// Check if we will restore expansion states, we will not if there is filtering happening.
+	const bool bRestoreExpansionState = true;
+
+	if ( Object )
+	{
+		/*
+		for (int o=0; o<Object->States.Num(); ++o)
+		{
+			TSharedPtr<FStatePropertiesNode> node = MakeShareable( new FStatePropertiesNode );
+			node->Name = MakeShareable( new FString(Object->States[o].Name) );
+			RootTreeItems.Add( node );
+
+			if( bRestoreExpansionState )
+			{
+				SetExpansionStatesInTree( node );
+			}
+		}
+		*/
+	}
+
+	// All root level items get expanded automatically.
+	//StatesTree->SetItemExpansion(RootNode, true);
+
+	// Now that new items are in the tree, we need to request a refresh.
+	StatesTree->RequestTreeRefresh();
+}
+
+TSharedPtr<FStatePropertiesNode> SCustomizableObjectProperties::CreateNoneOption()
+{
+	TSharedPtr<FStatePropertiesNode> NoneItem = MakeShareable( new FStatePropertiesNode("None") );
+
+	// The item "passes" the filter so it does not appear grayed out.
+	//NoneItem->bPassesFilter = true;
+
+	return NoneItem;
+}
+
+
+#undef LOCTEXT_NAMESPACE
+
diff --git a/Engine/Plugins/Mutable/Source/CustomizableObjectEditor/Private/MuCOE/SCustomizableObjectProperties.h b/Engine/Plugins/Mutable/Source/CustomizableObjectEditor/Private/MuCOE/SCustomizableObjectProperties.h
new file mode 100644
--- /dev/null
+++ b/Engine/Plugins/Mutable/Source/CustomizableObjectEditor/Private/MuCOE/SCustomizableObjectProperties.h
@@ -0,0 +1,115 @@
+// Copyright Epic Games, Inc. All Rights Reserved.
+
+#pragma once
+
+
+#include "Widgets/SCompoundWidget.h"
+
+class STableViewBase;
+namespace ESelectInfo { enum Type : int; }
+template <typename ItemType> class STreeView;
+
+class FString;
+class ITableRow;
+class SWidget;
+class UCustomizableObject;
+
+
+//////////////////////////////////////////////////////////////////////////
+// SClassViewer
+
+class FStatePropertiesNode;
+
+class SCustomizableObjectProperties : public SCompoundWidget
+{
+public:
+	SLATE_BEGIN_ARGS( SCustomizableObjectProperties )
+		: _Object((UCustomizableObject*)NULL)
+	{
+	}
+
+	SLATE_ARGUMENT( UCustomizableObject*, Object )
+
+	SLATE_END_ARGS()
+
+	/**
+	 * Construct the widget
+	 *
+	 * @param	InArgs			A declaration from which to construct the widget
+	 */
+	void Construct(const FArguments& InArgs);
+
+	/** Gets the widget contents of the app */
+	virtual TSharedRef<SWidget> GetContent();
+
+	virtual ~SCustomizableObjectProperties();
+
+	/** */
+	void SetObject( UCustomizableObject* InObject );
+
+private:
+	/** Retrieves the children for the input node.
+	 *	@param InParent				The parent node to retrieve the children from.
+	 *	@param OutChildren			List of children for the parent node.
+	 *
+	 */
+	void OnGetChildrenForTree( TSharedPtr<FStatePropertiesNode> InParent, TArray< TSharedPtr< FStatePropertiesNode > >& OutChildren );
+
+	/** Creates the row widget when called by Slate when an item appears on the tree. */
+	TSharedRef< ITableRow > OnGenerateRow( TSharedPtr<FStatePropertiesNode> Item, const TSharedRef< STableViewBase >& OwnerTable );
+
+	/** Called by Slate when an item is selected from the tree/list. */
+	void OnClassViewerSelectionChanged( TSharedPtr<FStatePropertiesNode> Item, ESelectInfo::Type SelectInfo );
+
+	/** 
+	 *	Sets all expansion states in the tree.
+	 *
+	 *	@param bInExpansionState			The expansion state to set the tree to.
+	 */
+	void SetAllExpansionStates(bool bInExpansionState);
+
+	/**
+	 *	A helper function to recursively set the tree.
+	 *
+	 *	@param	InNode						The current node in the tree.
+	 *	@param	bInExpansionState			The expansion state to set the tree to.
+	 */
+	void SetAllExpansionStates_Helper(TSharedPtr< FStatePropertiesNode > InNode, bool bInExpansionState);
+
+	/** Resets the expansion states on the tree to default. */
+	void ResetExpansionStates();
+
+	/** Recursive function to map the expansion states of items in the tree. 
+	 *	@param InItem		The current item to examine the expansion state of.
+	 */
+	void MapExpansionStatesInTree( TSharedPtr<FStatePropertiesNode> InItem );
+
+	/** Recursive function to set the expansion states of items in the tree. 
+	 *	@param InItem		The current item to set the expansion state of.
+	 */
+	void SetExpansionStatesInTree( TSharedPtr<FStatePropertiesNode> InItem );
+
+	/** Populates the tree with items based on the current filter. */
+	void Populate();
+
+	/** Creates a "None" option for the tree/list. */
+	TSharedPtr<FStatePropertiesNode> CreateNoneOption();
+
+private:
+
+	UCustomizableObject* Object;
+
+	TSharedPtr<STreeView<TSharedPtr< FStatePropertiesNode > >> StatesTree;
+
+	/** The root items to be displayed in the state tree. */
+	TArray< TSharedPtr< FStatePropertiesNode > > RootTreeItems;
+
+	/** Turns off the re-population for a single call. */
+	bool bDisablePopulate;
+
+	/** true if expansions states should be saved when compiling. */
+	bool bSaveExpansionStates;
+
+	/** The map holding the expansion state map for the tree. */
+	TMap< TSharedPtr< FString >, bool > ExpansionStateMap;
+};
diff --git a/Engine/Plugins/Mutable/Source/CustomizableObjectEditor/Private/MuCOE/SCustomizableObjectSystemVersion.h b/Engine/Plugins/Mutable/Source/CustomizableObjectEditor/Private/MuCOE/SCustomizableObjectSystemVersion.h
new file mode 100644
--- /dev/null
+++ b/Engine/Plugins/Mutable/Source/CustomizableObjectEditor/Private/MuCOE/SCustomizableObjectSystemVersion.h
@@ -0,0 +1,30 @@
+// Copyright Epic Games, Inc. All Rights Reserved.
+
+#pragma once
+
+#include "Widgets/SCompoundWidget.h"
+
+struct FGeometry;
+
+
+class CUSTOMIZABLEOBJECTEDITOR_API SCustomizableObjectSystemVersion : public SCompoundWidget
+{
+public:
+
+	SLATE_BEGIN_ARGS(SCustomizableObjectSystemVersion)
+		{}
+	SLATE_END_ARGS()
+
+	/** */
+	void Construct( const FArguments& InArgs );
+
+public:
+
+	// SWidget interface
+	void Tick( const FGeometry& AllottedGeometry, const double InCurrentTime, const float InDeltaTime ) override;
+
+private:
+
+
+};
+
