[mutable] Moved the Mutable plugin out of Experimental status into Beta.

diff --git a/Engine/Plugins/Mutable/Source/CustomizableObjectEditor/Private/MuCOE/SMutableCodeViewer.cpp b/Engine/Plugins/Mutable/Source/CustomizableObjectEditor/Private/MuCOE/SMutableCodeViewer.cpp
new file mode 100644
--- /dev/null
+++ b/Engine/Plugins/Mutable/Source/CustomizableObjectEditor/Private/MuCOE/SMutableCodeViewer.cpp
@@ -0,0 +1,3088 @@
+// Copyright Epic Games, Inc. All Rights Reserved.
+
+#include "MuCOE/SMutableCodeViewer.h"
+
+#include "DesktopPlatformModule.h"
+#include "EditorDirectories.h"
+#include "IDesktopPlatform.h"
+#include "Framework/MultiBox/MultiBoxBuilder.h"
+#include "Framework/Views/TableViewMetadata.h"
+#include "Misc/Paths.h"
+#include "MuCO/CustomizableObject.h"
+#include "MuCO/UnrealToMutableTextureConversionUtils.h"
+#include "MuCOE/SMutableBoolViewer.h"
+#include "MuCOE/SMutableColorViewer.h"
+#include "MuCOE/SMutableConstantsWidget.h"
+#include "MuCOE/SMutableCurveViewer.h"
+#include "MuCOE/SMutableImageViewer.h"
+#include "MuCOE/SMutableIntViewer.h"
+#include "MuCOE/SMutableLayoutViewer.h"
+#include "MuCOE/SMutableMeshViewer.h"
+#include "MuCOE/SMutableInstanceViewer.h"
+#include "MuCOE/SMutableParametersWidget.h"
+#include "MuCOE/SMutableProjectorViewer.h"
+#include "MuCOE/SMutableScalarViewer.h"
+#include "MuCOE/SMutableSkeletonViewer.h"
+#include "MuCOE/SMutableStringViewer.h"
+#include "MuCOE/UnrealEditorPortabilityHelpers.h"
+#include "MuR/SystemPrivate.h"
+#include "MuT/ErrorLog.h"
+#include "MuT/TypeInfo.h"
+#include "Widgets/SNullWidget.h"
+#include "Widgets/Colors/SColorBlock.h"
+#include "Widgets/Input/SButton.h"
+#include "Widgets/Input/SComboBox.h"
+#include "Widgets/Input/SNumericEntryBox.h"
+#include "Widgets/Layout/SScrollBox.h"
+#include "Widgets/Views/STreeView.h"
+#include "MuCOE/Widgets/MutableExpanderArrow.h"
+
+
+#include "Widgets/Input/SSearchBox.h"
+#include "Internationalization/Regex.h"
+
+class FExtender;
+class FReferenceCollector;
+class FUICommandList;
+class SWidget;
+namespace mu { struct FProjector; }
+namespace mu { struct FShape; }
+struct FGeometry;
+struct FSlateBrush;
+
+
+#define LOCTEXT_NAMESPACE "SMutableDebugger"
+
+
+namespace MutableCodeTreeViewColumns
+{
+	static const FName OperationsColumnID("Operations");
+	static const FName AdditionalDataColumnID("Flags");
+};
+
+/**
+ * Mutable tree row used to display the operations held on the Mutable model object. 
+ */
+class SMutableCodeTreeRow final : public SMultiColumnTableRow<TSharedPtr<FMutableCodeTreeElement>>
+{
+public:
+
+	void Construct(const FArguments& Args, const TSharedRef<STableViewBase>& InOwnerTableView, const TSharedPtr<FMutableCodeTreeElement>& InRowItem)
+	{
+		RowItem = InRowItem;
+		
+		SMultiColumnTableRow< TSharedPtr<FMutableCodeTreeElement> >::Construct(
+			STableRow::FArguments()
+			.ShowSelection(true)
+			, InOwnerTableView
+		);
+	}
+
+	FLinearColor OnGetExtraDataBoxColor() const
+	{
+		if (RowItem->bIsDynamicResource)
+		{
+			return DynamicResourceBoxColor;
+		}
+		else if (RowItem->bIsStateConstant)
+		{
+			return StateConstantBoxColor;
+		}
+		else
+		{
+			return ExtraDataBackgroundBoxDefaultColor;
+		}
+	}
+
+	FText OnGetExtraDataText() const
+	{
+		// DEBUG :Uncomment the next line in order to debug the current state being used by the element
+		// return FText::FromString(FString::FromInt(RowItem->GetStateIndex()));
+		
+		if (RowItem->bIsDynamicResource)
+		{
+			return  DynamicResourceText;
+		}
+		else if (RowItem->bIsStateConstant)
+		{
+			return StateConstantText;
+		}
+		else
+		{
+			return FText::FromString(FString(""));
+		}
+	}
+
+	/** Depending on the state of the row returns one color or another to be used by the highlighting system */
+	FLinearColor GetHighlightColor() const
+	{
+		if (bShouldBeHiglighted)
+		{
+			if (RowItem->DuplicatedOf)
+			{
+				return HighlightedDuplicatedBoxColor;
+			}
+			else
+			{
+				return HighlightedUniqueRowBoxColor;
+			}
+		}
+
+		return HighlightBoxDefaultColor;
+	}
+	
+	/** Method intended with the generation of the wanted objects for each column*/
+	virtual TSharedRef<SWidget> GenerateWidgetForColumn(const FName& ColumnName) override
+	{
+		// Primary column showing the name of the operation and tye type
+		if (ColumnName == MutableCodeTreeViewColumns::OperationsColumnID)
+		{
+			// Prepare a ui container for all the UI objects required by this row element
+			TSharedRef<SHorizontalBox> RowContainer = SNew(SHorizontalBox)
+				
+			// First coll showing operation name and type
+			+ SHorizontalBox::Slot()
+			.HAlign(EHorizontalAlignment::HAlign_Fill)
+			.AutoWidth()
+			[
+				SNew(SOverlay)
+
+				+ SOverlay::Slot()
+				[
+					SAssignNew(this->HighlightingColorBox, SColorBlock)
+					.Color(this, &SMutableCodeTreeRow::GetHighlightColor)
+				]
+
+				+ SOverlay::Slot()
+				[
+					SNew(SHorizontalBox)
+					+ SHorizontalBox::Slot()
+					.AutoWidth()
+					[
+						SNew(SMutableExpanderArrow, SharedThis(this))
+					]
+					
+					+ SHorizontalBox::Slot()
+					[
+						SNew(STextBlock)
+						.Text(FText::FromString(RowItem->MainLabel))
+						.ColorAndOpacity(RowItem->LabelColor)
+					]
+				]
+			];
+
+			return RowContainer;
+		}
+
+		// Second column showing some extra data related with the operation being displayed
+		if (ColumnName == MutableCodeTreeViewColumns::AdditionalDataColumnID)
+		{
+			TSharedRef<SHorizontalBox> RowContainer =  SNew(SHorizontalBox)
+			
+			+ SHorizontalBox::Slot()
+			.HAlign(EHorizontalAlignment::HAlign_Left)
+			.AutoWidth()
+			[
+				SNew(SHorizontalBox)
+
+				+SHorizontalBox::Slot()
+				.MaxWidth(4.0f)
+				[
+					SNew(SColorBlock)
+					.Color(this,&SMutableCodeTreeRow::OnGetExtraDataBoxColor)
+				]
+
+				+ SHorizontalBox::Slot()
+				.Padding(4,1)
+				.AutoWidth()
+				[
+					SNew(STextBlock)
+					.Text(this,&SMutableCodeTreeRow::OnGetExtraDataText)
+				]
+			];
+			
+			return RowContainer;
+		}
+
+		// Invalid column name so no widget will be produced 
+		return SNullWidget::NullWidget;
+	}
+	
+	/** Marks the row to be highlighted */
+	void Highlight()
+	{
+		bShouldBeHiglighted = true;
+	}
+
+	/** Resets the highlighting status */
+	void ResetHighlight()
+	{
+		bShouldBeHiglighted = false;
+	}
+
+	/** Returns a reference to the Element this row is representing */
+	TSharedPtr<FMutableCodeTreeElement>& GetItem()
+	{
+		return RowItem;
+	}
+
+private:
+
+	/** Pointer to the element that did spawn this row */
+	TSharedPtr<FMutableCodeTreeElement> RowItem = nullptr;
+
+	/** Transparent color */
+	const FLinearColor TransparentColor = FLinearColor(0,0,0,0);
+
+	/*
+	 * Operation Highlighting 
+	 */
+	
+	/** Custom Widget used to display a color. Used as the background of the text on the row to serve as highlighting Visual Element*/
+	TSharedPtr<SColorBlock> HighlightingColorBox = nullptr;
+	
+	/** The color used to highlight the row if duplicated from another row */
+	const FLinearColor HighlightedDuplicatedBoxColor = FLinearColor(1, 1, 1, 0.15);
+
+	/** The color used to highlight elements that are originals (not duplicates)  */
+	const FLinearColor HighlightedUniqueRowBoxColor = FLinearColor(1, 1, 1, 0.28);
+
+	/** Default color used when the row is not highlighted */
+	const FLinearColor HighlightBoxDefaultColor = TransparentColor;
+
+	/*
+	 * Extra data objects
+	 */
+
+	// Text used to set the width of the color area in front of the extra data
+	const FText EmptyText = FText(INVTEXT(" "));
+	
+	/** String printed on the UI when the operation is shown to be dynamic resource */
+	const FText DynamicResourceText = FText::FromString(FString("dyn"));
+
+	/** String printed on the UI when the operation is shown to be state constant */
+	const FText StateConstantText = FText::FromString(FString("const"));
+	
+	/** Color used on the extra data column when no extra data is shown */
+	const FLinearColor ExtraDataBackgroundBoxDefaultColor =  TransparentColor;
+
+	/** Color shown on the extra data column when the resource is found to be Dynamic */
+	const FLinearColor DynamicResourceBoxColor = FLinearColor(0,0,1,0.8);
+
+	/** Color shown on the extra data column when the resource is found to be State Constant */
+	const FLinearColor StateConstantBoxColor = FLinearColor(1,0,0,0.8);
+
+	bool bShouldBeHiglighted = false;
+};
+
+
+void SMutableCodeViewer::AddReferencedObjects(FReferenceCollector& Collector)
+{
+	// Add UObjects here if we own any at some point
+	//Collector.AddReferencedObject(CustomizableObject);
+}
+
+
+FString SMutableCodeViewer::GetReferencerName() const
+{
+	return TEXT("SMutableCodeViewer");
+}
+
+void SMutableCodeViewer::ClearSelectedTreeRow() const
+{
+	check(TreeView);
+	TreeView->ClearSelection();
+}
+
+void SMutableCodeViewer::SetCurrentModel(const TSharedPtr<mu::Model, ESPMode::ThreadSafe>& InMutableModel, const TArray<TSoftObjectPtr<const UTexture>>& InReferencedTextures)
+{
+	MutableModel = InMutableModel;
+	ReferencedTextures = InReferencedTextures;
+	PreviewParameters = mu::Model::NewParameters(MutableModel);
+
+	RootNodes.Empty();
+	RootNodeAddresses.Empty();
+	ItemCache.Empty();
+	MainItemPerOp.Empty();
+	TreeElements.Empty();
+	ExpandedElements.Empty();
+	FoundModelOperationTypeElements.Empty();
+	ModelOperationTypes.Empty();
+	ModelOperationTypeNames.Empty();
+
+	// Reset navigation by type / constant resource
+	NavigationElements.Empty();
+	NavigationIndex = -1;
+
+	// Reset navigation by string
+	NameBasedNavigationElements.Empty();
+	StringNavigationIndex = -1;
+
+	// Generate all elements before starting the tree UI so we have a deterministic set of unique and duplicated elements
+	GenerateAllTreeElements();
+
+	// Setup Navigation system
+	{
+		// Store the addresses of the root nodes so they can be used by operation search methods
+		CacheRootNodeAddresses();
+
+		// Cache the operation types that are present on the model
+		CacheOperationTypesPresentOnModel();
+
+		// Get an array of mutable types as an array of FStrings for the UI
+		GenerateNavigationOpTypeStrings();
+
+		// Generate list elements for the found operation types so we are able to search over them on our type dropdown
+		GenerateNavigationDropdownElements();
+
+		// Check we did find types (witch should always happen in a normal run) and select the NONE option as the default value
+		check(FoundModelOperationTypeElements.Num())
+		CurrentlySelectedOperationTypeElement = NoneOperationEntry;
+	}
+}
+
+
+void SMutableCodeViewer::Construct(const FArguments& InArgs, const TSharedPtr<mu::Model, ESPMode::ThreadSafe>& InMutableModel,
+	const TArray<TSoftObjectPtr<const UTexture>>& InReferencedTextures )
+{
+	// Min width allowed for the column. Needed to avoid having issues with the constants space being to small
+	// and then getting too tall on the y axis crashing the UI drawer.
+	constexpr float MinParametersCollWidth = 200;
+
+	SetCurrentModel(InMutableModel, InReferencedTextures);
+	
+	FToolBarBuilder ToolbarBuilder(TSharedPtr<const FUICommandList>(), FMultiBoxCustomization::None, TSharedPtr<FExtender>(), true);
+	ToolbarBuilder.SetLabelVisibility(EVisibility::Visible);
+	ToolbarBuilder.SetStyle(&FAppStyle::Get(), "SlimToolBar");
+
+	ToolbarBuilder.AddWidget(
+		SNew(STextBlock)
+		.Text(FText::FromString(InArgs._DataTag)));
+
+	TSharedRef<SScrollBar> TreeVertScrollBar =
+		SNew(SScrollBar).
+		Orientation(EOrientation::Orient_Vertical).
+		AlwaysShowScrollbar(false);
+	
+	ChildSlot
+	[
+		SNew(SVerticalBox)
+		+ SVerticalBox::Slot()
+		.AutoHeight()
+		.VAlign(VAlign_Center)
+		[
+			ToolbarBuilder.MakeWidget()
+		]
+		+ SVerticalBox::Slot()
+		.VAlign(VAlign_Fill)
+		.Padding(5,2)
+		[
+			SNew(SSplitter)
+			.Orientation(EOrientation::Orient_Horizontal)
+			
+			+ SSplitter::Slot()
+			.Value(0.35f)
+			.MinSize(520)
+			.Resizable(true)
+			[
+				SNew(SVerticalBox)
+
+				// Search box for tree operations
+				+ SVerticalBox::Slot()
+				.AutoHeight()
+				.HAlign(HAlign_Left)
+				[
+					SNew(SHorizontalBox)
+				
+					// Search by name
+					+ SHorizontalBox::Slot()
+					.AutoWidth()
+					[
+						SNew(SHorizontalBox)
+				
+						+ SHorizontalBox::Slot()
+						.AutoWidth()
+						.VAlign(VAlign_Center)
+						[
+							SNew(STextBlock)
+							.Text(LOCTEXT("SelectedOperationByStringLabel","Search Operation by String :"))
+						]
+					
+						+ SHorizontalBox::Slot()
+						.MaxWidth(250)
+						.VAlign(VAlign_Center)
+						[
+							SNew(SSearchBox)
+							.HintText(LOCTEXT("OperationToSearchHintText","Search OP"))
+							.SearchResultData(this,&SMutableCodeViewer::SearchResultsData)
+							.OnSearch(this, &SMutableCodeViewer::OnTreeStringSearch)
+							.OnTextChanged(this, &SMutableCodeViewer::OnTreeSearchTextChanged)
+							.OnTextCommitted(this, &SMutableCodeViewer::OnTreeSearchTextCommitted)
+						]
+					]
+					
+					
+					// Regex control for search by name
+					+ SHorizontalBox::Slot()
+					.AutoWidth()
+					.Padding(4,2)
+					[
+						SNew(SHorizontalBox)
+				
+						+ SHorizontalBox::Slot()
+						.AutoWidth()
+						.VAlign(VAlign_Center)
+						[
+							SNew(STextBlock)
+							.Text(LOCTEXT("OperationToSearchRegexLabel","Is RegEx?"))
+						]
+				
+						+ SHorizontalBox::Slot()
+						.AutoWidth()
+						.VAlign(VAlign_Center)
+						[
+							SNew(SCheckBox)
+							.OnCheckStateChanged(this,&SMutableCodeViewer::OnRegexToggleChanged)
+						]
+					]	
+				]
+				
+				// Operation type filtering slot
+				+ SVerticalBox::Slot()
+				.AutoHeight()
+				.HAlign(HAlign_Left)
+				[
+					// Box containing navigation elements
+					SNew(SVerticalBox)
+				
+					+ SVerticalBox::Slot()
+					.Padding(2,4)
+					.AutoHeight()
+					[
+						SNew(SHorizontalBox)
+				
+						+ SHorizontalBox::Slot()
+						.AutoWidth()
+						[
+							SNew(SHorizontalBox)
+				
+							+ SHorizontalBox::Slot()
+							.AutoWidth()
+							.VAlign(VAlign_Center)
+							[
+								SNew(STextBlock)
+								.Text(LOCTEXT("SelectedOperationTypeLabel","Search Operation Type :"))
+							]
+				
+							// ComboBox used to select one or another Op_Type for tree navigation purposes
+							+SHorizontalBox::Slot()
+							.AutoWidth()
+							.VAlign(VAlign_Center)
+							[
+								SAssignNew(TargetedTypeSelector,SComboBox<TSharedPtr<const FMutableOperationElement>>)
+								.OptionsSource(&FoundModelOperationTypeElements)
+								.InitiallySelectedItem(CurrentlySelectedOperationTypeElement)
+								[
+									SNew(STextBlock)
+									.Text(this,&SMutableCodeViewer::GetCurrentNavigationOpTypeText)
+									.ColorAndOpacity(this,&SMutableCodeViewer::GetCurrentNavigationOpTypeColor)
+								]
+								.OnGenerateWidget(this,&SMutableCodeViewer::OnGenerateOpNavigationDropDownWidget)
+								.OnSelectionChanged(this,&SMutableCodeViewer::OnNavigationSelectedOperationChanged)
+							]
+						]
+				
+						+ SHorizontalBox::Slot()
+						.Padding(4,0)
+						.AutoWidth()
+						.VAlign(VAlign_Center)
+						[
+							SNew(SButton)
+							.Text(LOCTEXT("GoToPreviousOperationButton"," < "))
+							.OnClicked(this,&SMutableCodeViewer::OnGoToPreviousOperationButtonPressed)
+							.IsEnabled(this,&SMutableCodeViewer::CanInteractWithPreviousOperationButton)
+						]
+				
+						+ SHorizontalBox::Slot()
+						.Padding(4,0)
+						.AutoWidth()
+						.VAlign(VAlign_Center)
+						[
+							SNew(STextBlock)
+							.Text(this,&SMutableCodeViewer::OnPrintNavigableObjectAddressesCount)
+							.Justification(ETextJustify::Center)
+						]
+				
+						+ SHorizontalBox::Slot()
+						.Padding(4,0)
+						.AutoWidth()
+						.VAlign(VAlign_Center)
+						[
+							SNew(SButton)
+							.Text(LOCTEXT("GoToNextOperationButton"," > "))
+							.OnClicked(this,&SMutableCodeViewer::OnGoToNextOperationButtonPressed)
+							.IsEnabled(this,&SMutableCodeViewer::CanInteractWithNextOperationButton)
+						]
+					]
+				]
+				
+				// Tree operations slot
+				+ SVerticalBox::Slot()
+				.FillHeight(1.0f)
+				[
+					SNew(SBorder)
+					.BorderImage(UE_MUTABLE_GET_BRUSH("ToolPanel.GroupBorder"))
+					.Padding(FMargin(4.0f, 4.0f))
+					[
+						SNew(SHorizontalBox)
+
+						+ SHorizontalBox::Slot()
+						.FillContentWidth(1)
+						[
+							SNew(SScrollBox)		
+							.Orientation(EOrientation::Orient_Horizontal)
+							.ConsumeMouseWheel(EConsumeMouseWheel::Never)
+							
+							+ SScrollBox::Slot()
+							.HAlign(HAlign_Fill)
+							.FillContentSize(1)
+							[
+								SAssignNew(TreeView, STreeView<TSharedPtr<FMutableCodeTreeElement>>)
+								.TreeItemsSource(&RootNodes)
+								.OnGenerateRow(this, &SMutableCodeViewer::GenerateRowForNodeTree)
+								.OnRowReleased(this, &SMutableCodeViewer::OnRowReleased)
+								.OnGetChildren(this, &SMutableCodeViewer::GetChildrenForInfo)
+								.OnSelectionChanged(this, &SMutableCodeViewer::OnSelectionChanged)
+								.OnSetExpansionRecursive(this, &SMutableCodeViewer::TreeExpandRecursive)
+								.OnContextMenuOpening(this, &SMutableCodeViewer::OnTreeContextMenuOpening)
+								.OnExpansionChanged(this, &SMutableCodeViewer::OnExpansionChanged)
+								.SelectionMode(ESelectionMode::Single)
+								.ExternalScrollbar(TreeVertScrollBar)
+								.HeaderRow
+								(
+									SNew(SHeaderRow)
+									.ResizeMode(ESplitterResizeMode::Fill)
+
+									+ SHeaderRow::Column(MutableCodeTreeViewColumns::OperationsColumnID)
+										.DefaultLabel(LOCTEXT("Operation", "Operation"))
+
+									+ SHeaderRow::Column(MutableCodeTreeViewColumns::AdditionalDataColumnID)
+										.DefaultLabel(LOCTEXT("OperationFlags", "Flags"))
+										.FixedWidth(50.0f)
+								)
+							]
+						]
+
+						+ SHorizontalBox::Slot()
+						.AutoWidth()
+						[
+							TreeVertScrollBar
+						]
+					]
+				]
+				
+			]
+			+ SSplitter::Slot()
+			.Value(0.75f)
+			[
+				SNew(SSplitter)
+				.Orientation(EOrientation::Orient_Horizontal)
+				
+				+ SSplitter::Slot()
+				.Value(0.28f)
+				.MinSize(MinParametersCollWidth)
+				[
+					// Splitter managing both parameter and constant panels
+					SNew(SSplitter)
+					.Orientation(EOrientation::Orient_Vertical)
+					+SSplitter::Slot()
+					[
+						SNew(SVerticalBox)
+
+						+ SVerticalBox::Slot()
+						.AutoHeight()
+						[
+							SNew(SHorizontalBox)
+
+							+ SHorizontalBox::Slot()
+							.AutoWidth()
+							[
+								SNew(STextBlock)
+								.Text(LOCTEXT("SkipMipsLabel", "Skip mips on generate :"))
+								.Visibility(this, &SMutableCodeViewer::IsMipSkipVisible)
+							]
+							+ SHorizontalBox::Slot()
+							[
+								SNew(SNumericEntryBox<int32>)
+								.Visibility(this, &SMutableCodeViewer::IsMipSkipVisible)
+								.AllowSpin(true)
+								.MinValue(0)
+								.MaxValue(16)
+								.MinSliderValue(0)
+								.MaxSliderValue(16)
+								.Value(this, &SMutableCodeViewer::GetCurrentMipSkip)
+								.OnValueChanged(this, &SMutableCodeViewer::OnCurrentMipSkipChanged)
+							]
+						]
+
+						+ SVerticalBox::Slot()
+						[
+							SNew(SScrollBox)
+							+ SScrollBox::Slot()
+							[
+								SAssignNew(ParametersWidget, SMutableParametersWidget)
+								.OnParametersValueChanged(this, &SMutableCodeViewer::OnPreviewParameterValueChanged)
+								.Parameters(PreviewParameters)
+							]
+						]
+					]
+
+					+ SSplitter::Slot()
+					[
+						// Generate a new Constants panel to show the data stored on the current mutable program
+						SAssignNew(ConstantsWidget, SMutableConstantsWidget,
+							&(MutableModel->GetPrivate()->m_program),
+							SharedThis(this))
+					]
+				]
+				+ SSplitter::Slot()
+				.Value(0.72f)
+				[
+					SAssignNew(PreviewBorder, SBorder)
+					.BorderImage(UE_MUTABLE_GET_BRUSH("ToolPanel.GroupBorder"))
+					.Padding(FMargin(4.0f, 4.0f))					
+				]
+			]
+		]
+	];
+	
+	// Set the tree expanded by default
+	// It does not recalculate states since the expansion of the instance will NOT expand duplicates witch means the widget position
+	// of the children of duplicated (or the original of an operation with duplicates) will not change.
+	TreeExpandInstance();
+
+	// Enable the recalculation of states once the tree has already been initially expanded since now we do not control
+	// how the user is point to interact with the view.
+	bShouldRecalculateStates = true;
+	// Now, on expansion or contraction the states will get recalculated
+} 
+
+
+EVisibility SMutableCodeViewer::IsMipSkipVisible() const 
+{ 
+	return bSelectedOperationIsImage ? EVisibility::Visible : EVisibility::Hidden;
+}
+
+
+TOptional<int32> SMutableCodeViewer::GetCurrentMipSkip() const 
+{ 
+	return MipsToSkip; 
+}
+
+
+void SMutableCodeViewer::OnCurrentMipSkipChanged(int32 NewValue)
+{
+	MipsToSkip = NewValue;
+	bIsPreviewPendingUpdate = true;
+}
+
+#pragma region CodeTree operation name search
+
+
+void SMutableCodeViewer::OnRegexToggleChanged(ECheckBoxState CheckBoxState)
+{
+	const bool bPreChangeValue = bIsSearchStringRegularExpression;
+	bIsSearchStringRegularExpression = CheckBoxState == ECheckBoxState::Checked ? true : false;
+	
+	if (bPreChangeValue != bIsSearchStringRegularExpression)
+	{
+		CacheOperationsMatchingStringPattern();
+		GoToNextOperation();
+	}
+}
+
+void SMutableCodeViewer::OnTreeStringSearch(SSearchBox::SearchDirection SearchDirection)
+{
+	if (SearchDirection==SSearchBox::SearchDirection::Next)
+	{
+		GoToNextOperation();
+	}
+	else
+	{
+		GoToPreviousOperation();
+	}
+}
+
+void SMutableCodeViewer::GoToNextOperation()
+{
+	// Contingency : Prevent a second scroll operation from being performed if still we do not have the first target in view
+	if (bWasScrollToTargetRequested)
+	{
+		return;
+	}
+	
+	if (NameBasedNavigationElements.Num())
+	{
+		const int32 PreviousIndex = StringNavigationIndex;
+
+		// Focus on next target
+		StringNavigationIndex = StringNavigationIndex >= NameBasedNavigationElements.Num() - 1
+										   ? 0
+										   : StringNavigationIndex + 1;
+
+		if (StringNavigationIndex != PreviousIndex)
+		{
+			FocusViewOnNavigationTarget(NameBasedNavigationElements[StringNavigationIndex]);
+		}
+	}
+}
+
+
+void SMutableCodeViewer::GoToPreviousOperation()
+{
+	// Contingency : Prevent a second scroll operation from being performed if still we do not have the first target in view
+	if (bWasScrollToTargetRequested)
+	{
+		return;
+	}
+	
+	if (NameBasedNavigationElements.Num())
+	{
+		const int32 PreviousIndex = StringNavigationIndex;
+		
+		// Focus on previous target
+		StringNavigationIndex = StringNavigationIndex <= 0 ?  NameBasedNavigationElements.Num() -1 : StringNavigationIndex -1;
+
+		if (PreviousIndex != StringNavigationIndex)
+		{
+			FocusViewOnNavigationTarget(NameBasedNavigationElements[StringNavigationIndex]);
+		}
+	}
+}
+
+void SMutableCodeViewer::GoToTargetOperation(const int32& InTargetIndex)
+{
+	if (InTargetIndex == StringNavigationIndex)
+	{
+		return;
+	}
+	
+	if (NameBasedNavigationElements.Num() && InTargetIndex > 0 && InTargetIndex <= NameBasedNavigationElements.Num()-1)
+	{
+		// Focus on the target index
+		StringNavigationIndex = InTargetIndex;
+		FocusViewOnNavigationTarget(NameBasedNavigationElements[StringNavigationIndex]);
+	}
+}
+
+
+void SMutableCodeViewer::OnTreeSearchTextChanged(const FText& InUpdatedText)
+{
+	SearchString = InUpdatedText.ToString();
+}
+
+
+TOptional<SSearchBox::FSearchResultData> SMutableCodeViewer::SearchResultsData() const
+{
+	if (NameBasedNavigationElements.Num() == 0)
+	{
+		return TOptional<SSearchBox::FSearchResultData>();
+	}
+	return TOptional<SSearchBox::FSearchResultData>({ NameBasedNavigationElements.Num(), StringNavigationIndex + 1});
+}
+
+
+void SMutableCodeViewer::OnTreeSearchTextCommitted(const FText& InUpdatedText, ETextCommit::Type TextCommitType)
+{
+	if (TextCommitType == ETextCommit::OnEnter)
+	{
+		check (InUpdatedText.ToString() == SearchString);
+		CacheOperationsMatchingStringPattern();	
+		GoToNextOperation();
+	}
+}
+
+
+void SMutableCodeViewer::CacheOperationsMatchingStringPattern()
+{
+	check(MutableModel);
+	check(RootNodeAddresses.Num());
+	
+	if ( LastSearchedString == SearchString &&
+		bWasLastSearchRegEx == bIsSearchStringRegularExpression && 
+		LastSearchedModel == MutableModel )
+	{
+		// Do not perform a search again since the context has not changed
+		return;
+	} 
+	
+	if (!SearchString.IsEmpty())
+	{
+		UE_LOG(LogMutable,Display,TEXT("Starting string search with target string ""\"%s""\" "),*SearchString);
+	
+		// Object containing all data required by the search operation to be able to be called recursively
+		FElementsSearchCache SearchPayload;
+		// Initialize the Search Payload with the root node addresses. This way the search will use them as the root nodes where
+		// to start searching
+		SearchPayload.SetupRootBatch(RootNodeAddresses);
+
+		const mu::FProgram& Program = MutableModel->GetPrivate()->m_program;
+		GetOperationsMatchingStringPattern(SearchString,bIsSearchStringRegularExpression,SearchPayload, Program);
+	
+		// Dump the located resources array onto the navigation array
+		NameBasedNavigationElements = MoveTemp(SearchPayload.FoundElements);
+		SortElementsByTreeIndex(NameBasedNavigationElements);
+		
+		UE_LOG(LogMutable, Display, TEXT("Operations found with matching pattern ""\"%s""\" is  %i"), *SearchString, NameBasedNavigationElements.Num());
+	}
+	else
+	{
+		NameBasedNavigationElements.Reset();
+	}
+
+	// Reset the search index
+	StringNavigationIndex = -1;
+
+	// Keep track of what context was used to perform the search to avoid doing it again if the context has not changed
+	LastSearchedString = SearchString;
+	bWasLastSearchRegEx = bIsSearchStringRegularExpression;
+	LastSearchedModel = MutableModel;
+}
+
+
+void SMutableCodeViewer::GetOperationsMatchingStringPattern(const FString& InStringPattern,const bool bIsRegularExpression ,FElementsSearchCache& SearchPayload,const mu::FProgram& InProgram)
+{
+	// next batch of addresses to be explored 
+	TArray<FItemCacheKey> NextBatchAddressesData;
+	
+	for (int32 ParentIndex = 0; ParentIndex < SearchPayload.BatchData.Num(); ParentIndex++)
+	{	
+		const FItemCacheKey CacheKey = SearchPayload.BatchData[ParentIndex];
+		const FString OperationDescriptiveText = GetOperationDescriptiveText(CacheKey);
+		
+		bool bMatchesPattern = false;
+		if (!bIsRegularExpression)
+		{
+			// Check if the provided text is contained over the element identification text
+			bMatchesPattern = OperationDescriptiveText.Contains(InStringPattern);
+		}
+		else
+		{
+			FRegexPattern Pattern{InStringPattern};
+			FRegexMatcher RegexMatcher{Pattern,OperationDescriptiveText};
+			bMatchesPattern = RegexMatcher.FindNext();
+		}
+		
+		// Get one of the previous run "children" and treat as a parent to get it's children and process them
+		const mu::OP::ADDRESS& ParentAddress = SearchPayload.BatchData[ParentIndex].Child;
+		
+		if (bMatchesPattern)
+		{
+			SearchPayload.AddToFoundElements(ParentAddress,ParentIndex,ItemCache);
+		}
+		
+		// Get all NON PROCESSED the children of this operation to later be able to process them (on next recursive call)
+		SearchPayload.CacheChildrenOfAddressIfNotProcessed(ParentAddress, InProgram, NextBatchAddressesData);
+	}
+
+	// At this point all the addresses to be computed on the next batch have already been set and will be computed on
+	// the next recursive call
+	
+	// Explore children if found 
+	if (NextBatchAddressesData.Num())
+	{
+		// Cache next batch data so the next invocations is able to locate the provided addresses on the itemsCache
+		SearchPayload.BatchData = MoveTemp(NextBatchAddressesData);
+		
+		GetOperationsMatchingStringPattern(InStringPattern,bIsRegularExpression, SearchPayload, InProgram);
+	}
+}
+
+
+FString SMutableCodeViewer::GetOperationDescriptiveText(const FItemCacheKey& InItemCacheKey)
+{
+	FString OperationDescriptiveText;
+		
+	if (const TSharedPtr<FMutableCodeTreeElement>* Element = ItemCache.Find(InItemCacheKey))
+	{
+		OperationDescriptiveText = Element->Get()->MainLabel;
+		check (!OperationDescriptiveText.IsEmpty());
+	}
+	
+	return OperationDescriptiveText;
+}
+
+#pragma endregion 
+
+
+#pragma region CodeTree operation search
+
+FText SMutableCodeViewer::GetCurrentNavigationOpTypeText() const
+{
+	check (CurrentlySelectedOperationTypeElement);
+	
+	return CurrentlySelectedOperationTypeElement->OperationTypeText;
+}
+
+FSlateColor SMutableCodeViewer::GetCurrentNavigationOpTypeColor() const
+{
+	check (CurrentlySelectedOperationTypeElement);
+
+	return CurrentlySelectedOperationTypeElement->OperationTextColor;
+}
+
+void SMutableCodeViewer::GenerateNavigationDropdownElements()
+{
+	const int32 OperationTypesCount = ModelOperationTypes.Num();
+	
+	// It must have at least one type, if not may be because we are running this before filling ModelOperationTypes
+	FoundModelOperationTypeElements.Empty(OperationTypesCount);
+	
+	for	(int32 OperationTypeIndex = 0; OperationTypeIndex < OperationTypesCount;  OperationTypeIndex++)
+	{
+		// Get the type as a string to be able to print it on the UI
+		const FText OperationTypeName = FText::FromString(ModelOperationTypeNames[OperationTypeIndex]);
+
+		const mu::OP_TYPE RepresentedType = ModelOperationTypes[OperationTypeIndex].Key;
+		const uint32 OperationTypeInstancesCount = ModelOperationTypes[OperationTypeIndex].Value;
+		
+		// Get the Color to be used on the text that will represent the operation on the dropdown
+		const FSlateColor OperationColor = ColorPerComputationalCost[StaticCast<uint8>(GetOperationTypeComputationalCost(RepresentedType))];
+		
+		// Generate an element to be used by the ComboBox handling the selection of the type to be used during navigation
+		TSharedPtr<FMutableOperationElement> OperationElement = MakeShared<FMutableOperationElement>(RepresentedType, OperationTypeName,OperationTypeInstancesCount,OperationColor);
+		FoundModelOperationTypeElements.Add(OperationElement);
+	}
+
+	// Add an entry for the NONE type of operation
+	{
+		const FText EntryName = FText::FromString("NONE");
+		const FSlateColor EntryColor = ColorPerComputationalCost[StaticCast<uint8>(EOperationComputationalCost::Standard)];
+		NoneOperationEntry = MakeShared<FMutableOperationElement>(mu::OP_TYPE::NONE,EntryName,0,EntryColor);
+
+		// @warn While not visible this element must be part of the collection for the ComboBox to be able to work
+		// properly
+		FoundModelOperationTypeElements.Add(NoneOperationEntry);
+	}
+	
+	// Add an extra operation type that will represent the constant resource based navigation type
+	{
+		const FText EntryName = FText::FromString("Selected Constant");	
+		const FSlateColor EntryColor = FSlateColor(FLinearColor(0.35f ,0.35f,1.0f,1));
+		ConstantBasedNavigationEntry = MakeShared<FMutableOperationElement>(mu::OP_TYPE::NONE,EntryName,0,EntryColor);
+		
+		// @warn While not visible this element must be part of the collection for the ComboBox to be able to work
+		// properly
+		FoundModelOperationTypeElements.Add(ConstantBasedNavigationEntry);
+	}
+}
+
+TSharedRef<SWidget> SMutableCodeViewer::OnGenerateOpNavigationDropDownWidget(
+	TSharedPtr<const FMutableOperationElement> MutableOperationElement) const
+{
+	TSharedRef<STextBlock> NewSlateObject = SNew(STextBlock)
+		.Text(MutableOperationElement->OperationTypeText)
+		.ColorAndOpacity(MutableOperationElement->OperationTextColor);
+
+	// Set the visibility type for the UI object (currently will be hidden for the NONE type)
+	NewSlateObject->SetVisibility(MutableOperationElement->GetEntryVisibility());
+	
+	return NewSlateObject;
+}
+
+void SMutableCodeViewer::OnNavigationSelectedOperationChanged(
+	TSharedPtr<const FMutableOperationElement, ESPMode::ThreadSafe> MutableOperationElement, ESelectInfo::Type Arg)
+{
+	check (MutableOperationElement.IsValid());
+
+	// Cache the currently selected operation set on the UI by the user
+	const mu::OP_TYPE NewOperationType = MutableOperationElement->OperationType;
+	OperationTypeToSearch = NewOperationType;
+	CurrentlySelectedOperationTypeElement = MutableOperationElement;
+	
+	// Only do the internal work if the type is one that makes sense searching
+	if (OperationTypeToSearch != mu::OP_TYPE::NONE)
+	{
+		// Locate all operations on the mutable operations tree (not the visual one) that do share the same operation type
+		// as the one selected. This will fill the array with the elements we should be looking for during the navigation operation
+		CacheAddressesOfOperationsOfType();
+	}
+	// None can be set by the user or be an indication that we are navigating over constant related operations
+	// todo: Separate both operations in some way on the UI to avoid complications in the code and in the UI's UX
+	else
+	{
+		// Clear all the elements on the navigation addresses 
+		NavigationElements.Empty();
+	}
+	
+}
+
+
+void SMutableCodeViewer::GenerateNavigationOpTypeStrings()
+{
+	// Grab only the names from the operation types located during the caching of operation types of the model
+	for (const TTuple<mu::OP_TYPE, uint32>& LocatedOperationType : ModelOperationTypes)
+	{
+		// Find the name of the Operation type
+		const uint16 OperationIndex = static_cast<uint16>(LocatedOperationType.Key);
+		const TCHAR* OpName = mu::s_opNames[OperationIndex];
+
+		// Remove trailing whitespaces adding noise and messing up concatenations with other strings
+		FString OperationNameString{OpName};
+		OperationNameString.RemoveSpacesInline();
+
+		// Save the name
+		ModelOperationTypeNames.Add( OperationNameString);
+	}
+}
+
+void SMutableCodeViewer::OnSelectedOperationTypeFromTree()
+{
+	// We require to have only 1 element selected to avoid having inconsistencies during operation
+	check(TreeView->GetNumItemsSelected() == 1);
+	
+	const TSharedPtr<FMutableCodeTreeElement> ReferenceOperationElement = TreeView->GetSelectedItems()[0];
+	
+	const mu::OP_TYPE OperationType =
+		MutableModel->GetPrivate()->m_program.GetOpType(ReferenceOperationElement->MutableOperation);
+
+	// Find the operation type directly in our array of operation elements (from the drop down)
+	const TSharedPtr<const FMutableOperationElement>* RepresentativeElement = FoundModelOperationTypeElements.
+		FindByPredicate([OperationType](const TSharedPtr<const FMutableOperationElement> Other)
+		{
+			return Other->OperationType == OperationType;
+		});
+
+	// Ensure an element was found. Failing the next check would mean that we are not caching all the types present on
+	// the current operation's tree
+	check(RepresentativeElement != nullptr);
+	
+	// Set the type operation type to be looking for -> Will invoke OnOptionTypeSelectionChanged
+	TargetedTypeSelector->SetSelectedItem(*RepresentativeElement);
+	
+	// Reset the navigation index so it starts from scratch
+	NavigationIndex = -1;
+}
+
+
+void SMutableCodeViewer::SortElementsByTreeIndex(TArray<TSharedPtr<FMutableCodeTreeElement>>& InElementsArrayToSort)
+{
+	// Sort the array from lower index to bigger index (0 , 1 , 2 ...)
+	InElementsArrayToSort.Sort([](const TSharedPtr<FMutableCodeTreeElement> A , const TSharedPtr<FMutableCodeTreeElement> B)
+	{
+		return A->IndexOnTree < B->IndexOnTree;
+	});
+}
+
+
+void SMutableCodeViewer::CacheAddressesOfOperationsOfType()
+{
+	// Clear previous data
+	NavigationElements.Empty();
+	check(RootNodeAddresses.Num());
+
+	// Object containing all data required by the search operation to be able to be called recursively
+	FElementsSearchCache SearchPayload;
+	// Initialize the Search Payload with the root node addresses. This way the search will use them as the root nodes where
+	// to start searching
+	SearchPayload.SetupRootBatch(RootNodeAddresses);
+	
+	// Main update procedure run for the targeted state and the targeted parameter values
+	const mu::FProgram& Program = MutableModel->GetPrivate()->m_program;
+	GetOperationsOfType(OperationTypeToSearch,SearchPayload, Program);
+	
+	if (!SearchPayload.FoundElements.IsEmpty())
+	{
+		// Cache the navigation addresses so we are able to navigate over them
+		NavigationElements = MoveTemp(SearchPayload.FoundElements);
+		SortElementsByTreeIndex(NavigationElements);
+		
+		// Reset the navigation index
+		NavigationIndex = -1;
+	}
+}
+
+void SMutableCodeViewer::GetOperationsOfType(const mu::OP_TYPE& TargetOperationType,
+                                             FElementsSearchCache& InSearchPayload,
+                                             const mu::FProgram& InProgram)
+{
+	// next batch of addresses to be explored 
+	TArray<FItemCacheKey> NextBatchAddressesData;
+	
+	for	(int32 ParentIndex = 0 ; ParentIndex < InSearchPayload.BatchData.Num(); ParentIndex++)
+	{
+		// Get one of the previous run "children" and treat as a parent to get it's children and process them
+		const mu::OP::ADDRESS& CurrentAddress = InSearchPayload.BatchData[ParentIndex].Child;
+		
+		// Cache if same data type and we share the same address (means this op is pointing at the provided resource)
+		// It will cache duplicated entries
+		if ( InProgram.GetOpType(CurrentAddress) == TargetOperationType)
+		{
+			// Since this element is of the type we are looking for then cache it on InSearchPayload.FoundElements
+			InSearchPayload.AddToFoundElements(CurrentAddress,ParentIndex,ItemCache);
+		}
+		
+		// Get all NON PROCESSED the children of this operation to later be able to process them (on next recursive call)
+		InSearchPayload.CacheChildrenOfAddressIfNotProcessed(CurrentAddress, InProgram, NextBatchAddressesData);
+	}
+
+	// Explore children if found 
+	if (NextBatchAddressesData.Num())
+	{
+		// Cache next batch data so the next invocations are able to locate the provided addresses on the itemsCache
+		InSearchPayload.BatchData = MoveTemp(NextBatchAddressesData);
+		
+		// Process the children of this object
+		GetOperationsOfType(TargetOperationType, InSearchPayload, InProgram);
+	}
+ }
+
+
+void SMutableCodeViewer::CacheOperationTypesPresentOnModel()
+{
+	check(MutableModel)
+
+	// Initialize NodeOperationTypes with empty TPair for each possible mutable operation type
+	{
+		constexpr uint32 OperationTypesCount = static_cast<uint16>(mu::OP_TYPE::COUNT);
+		ModelOperationTypes.Empty(OperationTypesCount);
+		for (uint32 Index = 0; Index < OperationTypesCount ; Index++)
+		{
+			mu::OP_TYPE TargetType = StaticCast<mu::OP_TYPE>(Index);
+			ModelOperationTypes.Add(TPair<mu::OP_TYPE,uint32>{TargetType,0});
+		}
+	}
+	
+	// Locate all operation types found on the provided model program data structure and count how many instances of each
+	// there are
+	{
+		// Get the types and the amount of instances of each unique operation on the mutable model
+		const mu::FProgram& Program = MutableModel->GetPrivate()->m_program;
+		const uint32 ProgramAddressesCount = Program.m_opAddress.Num();
+		
+		// Ensure first operation type is NONE since we are skipping it due to it having to have that type
+		check (Program.GetOpType(Program.m_opAddress[0]) == mu::OP_TYPE::NONE);
+
+		// Iterate over the addresses of the program and count how many instances each type has.
+		for (uint32 ProgramAddressesIndex = 1 ; ProgramAddressesIndex < ProgramAddressesCount; ProgramAddressesIndex++)
+		{
+			// Locate what is the position (index) of the operation type of the address on our collection of types found until now
+			const mu::OP_TYPE OperationType = Program.GetOpType(ProgramAddressesIndex);
+			
+			// Increase the counter for this operation type
+			const uint16 TypeAsInteger = StaticCast<uint16>(OperationType);
+			ModelOperationTypes[TypeAsInteger].Value++;
+		}
+	}
+	
+	// Remove all operation types that do have no operations present on the model
+	{
+		ModelOperationTypes.RemoveAll(
+		[](const TPair<mu::OP_TYPE, uint32>& Current)
+			{
+				return Current.Value == 0;
+			});
+	}
+	
+	// Sort the contents of the array of mutable operation types alphabetically
+	ModelOperationTypes.StableSort([&](const TPair<mu::OP_TYPE,uint32>& A, const TPair<mu::OP_TYPE,uint32>& B)
+	{
+		// Find the name
+		FString AString;
+		{
+			const uint16 OperationIndex = static_cast<uint16>(A.Key);
+			const TCHAR* OpName = mu::s_opNames[OperationIndex];
+			AString = FString(OpName);
+		}
+		
+		// Find out the name of the first element
+		FString BString;
+		{
+			const uint16 OperationIndex = static_cast<uint16>(B.Key);
+			const TCHAR* OpName = mu::s_opNames[OperationIndex];
+			BString = FString(OpName);
+		}
+		
+		// Then the name of the second element
+		return AString < BString;
+	});	
+	
+	// ModelOperationTypes is now an array with all the types found on the operations tree in alphabetical order
+}
+
+
+
+FText SMutableCodeViewer::OnPrintNavigableObjectAddressesCount() const
+{
+	FString OutputString = "";
+	if (const int32 NavigationElementsCount = NavigationElements.Num())
+	{
+		// Show the index if the index showing adds information
+		if (NavigationIndex >= 0)
+		{
+			OutputString.Append( FString::FromInt( NavigationIndex+1));
+			OutputString.Append(" / ");
+		}
+		
+		OutputString.Append( FString::FromInt(NavigationElementsCount));
+
+		// Format : 1 / 12 or 12
+	}
+	
+	// Depending on the amount of navigable objects (addresses, not actual elements) display the amount there are
+	return FText::FromString(OutputString);
+}
+
+
+bool SMutableCodeViewer::CanInteractWithPreviousOperationButton() const
+{
+	// Only navigable if there are more than 0 elements to traverse and we are not scrolling
+	return NavigationElements.Num() > 0 && NavigationIndex > 0 && (!bWasScrollToTargetRequested && !bWasUniqueExpansionInvokedForNavigation);
+}
+
+bool SMutableCodeViewer::CanInteractWithNextOperationButton() const
+{
+	// Only navigable if there are more than 0 elements to traverse and we are not scrolling
+	return NavigationElements.Num() > 0 && NavigationIndex < NavigationElements.Num() -1 && (!bWasScrollToTargetRequested && !bWasUniqueExpansionInvokedForNavigation);
+}
+
+
+FReply SMutableCodeViewer::OnGoToPreviousOperationButtonPressed()
+{
+	// Focus on previous target
+	NavigationIndex = NavigationIndex<=0 ? 0 : NavigationIndex - 1;
+	FocusViewOnNavigationTarget(NavigationElements[NavigationIndex]);
+	
+	return FReply::Handled();
+}
+
+FReply SMutableCodeViewer::OnGoToNextOperationButtonPressed()
+{
+	// Focus on next target
+	NavigationIndex = NavigationIndex>=NavigationElements.Num() -1 ? NavigationElements.Num() -1 : NavigationIndex + 1;
+	FocusViewOnNavigationTarget(NavigationElements[NavigationIndex]);
+	
+	return FReply::Handled();
+}
+
+void SMutableCodeViewer::FocusViewOnNavigationTarget(TSharedPtr<FMutableCodeTreeElement> InTargetElement)
+{
+	// Stage 1 : Expand all tree so all navigable elements get to be reachable
+	if (!bWasUniqueExpansionInvokedForNavigation && !bWasScrollToTargetRequested)
+	{
+		TreeExpandUnique();
+		bWasUniqueExpansionInvokedForNavigation = true;
+		
+		// Cache the current navigation target so after the update we can focus it 
+		ToFocusElement = InTargetElement;
+		
+		// Early exit, this method will get called again later after tree update
+		return;		
+	}
+	
+	// Stage 2 : Try to get to the targeted element. if not visible scroll into view
+	check (InTargetElement.IsValid());
+	
+	// If required scroll to the area where we know the element is going to be in view
+	// a way to ensure this happens is by calling 
+	if (TreeView->IsItemVisible(InTargetElement))
+	{
+		// Stage 3-b : Select the element we have provided since now is sure to be in view
+		
+		// This line selects the element with at the same time updates the UI to show the row representing this element selected
+		TreeView->SetSelection(InTargetElement);
+		ToFocusElement.Reset();							// We have focused the target so we no longer need to keep a reference to it
+		
+		// Done!
+		// We have the element in view and we have selected it!
+	}
+	else
+	{
+		// Stage 3-a (optional) : Ask for the provided element to be scrolled into view.
+		
+		// Failing this check would mean we have performed a scroll but we are still not able to view the element
+		check (!bWasScrollToTargetRequested);
+		
+		// Request the tree to show us the target element we want to get focused
+		TreeView->RequestScrollIntoView(InTargetElement);
+		
+		// Read this variable after the update and then select the object (easy at this point)
+		// You may want to just call again this method after refresh since the element will be on view
+		bWasScrollToTargetRequested = true;
+
+		// Early exit, this method will get called again later after tree update once the scroll has been completed
+		return;
+	}
+
+	// Reset the control flag so we do not expand all tree again if not required
+	bWasUniqueExpansionInvokedForNavigation = false;
+	bWasScrollToTargetRequested = false;
+}
+
+#pragma endregion 
+
+#pragma region Operation Cost Color Hints
+
+void SMutableCodeViewer::GenerateAllTreeElements()
+{
+	// By generating all tree elements prior to usage we are able to :
+	//	- Compute the index of each one to aid on navigation
+	//	- Remove non-deterministic assignation of the "Duplicated" state of elements. It was due to user interaction with the tree
+	//	Only unique elements, their children and duplicated elements will be generated. Children of duplicates will
+	// 	be ignored due to how we handle them when expanding and contracting elements (OnExpansionChanged)
+
+	// Generate all root nodes
+	const mu::Model::Private* ModelPrivate = MutableModel->GetPrivate();
+	const mu::FProgram& Program = ModelPrivate->m_program;
+	const uint32 StateCount = ModelPrivate->m_program.m_states.Num();
+	for ( uint32 StateIndex=0; StateIndex<StateCount; ++StateIndex )
+	{
+		const mu::FProgram::FState& State = ModelPrivate->m_program.m_states[StateIndex];
+		FString Caption = FString::Printf( TEXT("state [%s]"), *State.Name );
+
+		const FSlateColor LabelColor = ColorPerComputationalCost[StaticCast<uint8>(GetOperationTypeComputationalCost(
+			ModelPrivate->m_program.GetOpType(State.m_root)))];
+
+		RootNodes.Add(MakeShareable(new FMutableCodeTreeElement(ItemCache.Num(),StateIndex, MutableModel, State.m_root, Caption,LabelColor)));
+
+		// Iterate over each root node and generate all the elements in a human readable pattern (Z Pattern)
+
+		constexpr  mu::OP::ADDRESS CommonParent = 0;
+		const FItemCacheKey Key = { CommonParent, State.m_root, StateIndex };
+
+		// Add the element to the cache so we keep the indices straight.
+		ItemCache.Add(Key, RootNodes.Last());
+		GenerateElementRecursive(StateIndex,State.m_root,Program);
+	}
+	
+}
+
+void SMutableCodeViewer::GenerateElementRecursive(const int32& InStateIndex, mu::OP::ADDRESS InParentAddress,  const mu::FProgram& InProgram)
+{
+	// This will be used to add operations
+	uint32 ChildIndex = 0;
+	auto AddOpFunc = [this, InParentAddress, &InProgram, &ChildIndex, &InStateIndex](mu::OP::ADDRESS ChildAddress, const FString& Caption)
+		{
+			if (ChildAddress)
+			{
+				const FItemCacheKey Key = { InParentAddress, ChildAddress, ChildIndex };
+				const TSharedPtr<FMutableCodeTreeElement>* CachedItem = ItemCache.Find(Key);
+
+				// Will this ever really hit?
+				check(!CachedItem);
+
+				if (!CachedItem)
+				{
+					// Locate the "original" tree element
+					const TSharedPtr<FMutableCodeTreeElement>* MainItemPtr = MainItemPerOp.Find(ChildAddress);
+
+					const FSlateColor LabelColor = ColorPerComputationalCost[StaticCast<uint8>(GetOperationTypeComputationalCost(InProgram.GetOpType(ChildAddress)))];
+
+					const TSharedPtr<FMutableCodeTreeElement> Item = MakeShareable(new FMutableCodeTreeElement(ItemCache.Num(),InStateIndex, MutableModel, ChildAddress, Caption, LabelColor, MainItemPtr));
+
+					// Cache this element for later access
+					ItemCache.Add(Key, Item);
+
+					// It is not a duplicated of another one, then we can continue searching
+					if (!MainItemPtr)
+					{
+						MainItemPerOp.Add(ChildAddress, Item);
+
+						GenerateElementRecursive(InStateIndex,ChildAddress, InProgram);
+					}
+				}
+			}
+			++ChildIndex;
+		};
+
+	
+	// For some specific parent operation types we create more detailed subtrees.
+	bool bUseGeneric = false;
+	const mu::OP_TYPE ParentOperationType = InProgram.GetOpType(InParentAddress);
+	switch (ParentOperationType)
+	{
+	case mu::OP_TYPE::IM_SWITCH:
+	case mu::OP_TYPE::LA_SWITCH:
+	case mu::OP_TYPE::ME_SWITCH:
+	case mu::OP_TYPE::CO_SWITCH:
+	case mu::OP_TYPE::SC_SWITCH:
+	case mu::OP_TYPE::NU_SWITCH:
+	case mu::OP_TYPE::IN_SWITCH:
+	case mu::OP_TYPE::ED_SWITCH:
+	{
+		const uint8* OpData = InProgram.GetOpArgsPointer(InParentAddress);
+
+		mu::OP::ADDRESS VarAddress;
+		FMemory::Memcpy(&VarAddress, OpData, sizeof(mu::OP::ADDRESS));
+		OpData += sizeof(mu::OP::ADDRESS);
+		AddOpFunc(VarAddress, TEXT("var "));
+
+		mu::OP::ADDRESS DefAddress;
+		FMemory::Memcpy(&DefAddress, OpData, sizeof(mu::OP::ADDRESS));
+		OpData += sizeof(mu::OP::ADDRESS);
+		AddOpFunc(DefAddress, TEXT("def "));
+
+		uint32 CaseCount;
+		FMemory::Memcpy(&CaseCount, OpData, sizeof(uint32));
+		OpData += sizeof(uint32);
+
+		for (uint32 C = 0; C < CaseCount; ++C)
+		{
+			int32 Condition;
+			FMemory::Memcpy(&Condition, OpData, sizeof(int32));
+			OpData += sizeof(int32);
+
+			mu::OP::ADDRESS At;
+			FMemory::Memcpy(&At, OpData, sizeof(mu::OP::ADDRESS));
+			OpData += sizeof(mu::OP::ADDRESS);
+
+			// This conditional is necessary to exactly match the generic op behaviour (see ForEachReference in Operations.cpp)
+			if (At)
+			{
+				FString Caption = FString::Printf(TEXT("case %d "), Condition);
+				AddOpFunc(At, Caption);
+			}
+		}
+
+		break;
+	}
+
+	case mu::OP_TYPE::IM_SWIZZLE:
+	{
+		mu::OP::ImageSwizzleArgs Args = InProgram.GetOpArgs<mu::OP::ImageSwizzleArgs>(InParentAddress);
+		for (int32 Channel = 0; Channel < 4; ++Channel)
+		{
+			FString Caption = FString::Printf(TEXT("%d is %d from "), Channel, Args.sourceChannels[Channel]);
+			AddOpFunc(Args.sources[Channel], Caption);
+		}
+		break;
+	}
+
+	case mu::OP_TYPE::CO_SWIZZLE:
+	{
+		mu::OP::ColourSwizzleArgs Args = InProgram.GetOpArgs<mu::OP::ColourSwizzleArgs>(InParentAddress);
+		for (int32 Channel = 0; Channel < 4; ++Channel)
+		{
+			FString Caption = FString::Printf(TEXT("%d is %d from "), Channel, Args.sourceChannels[Channel]);
+			AddOpFunc(Args.sources[Channel], Caption);
+		}
+		break;
+	}
+
+	case mu::OP_TYPE::IM_LAYER:
+	{
+		mu::OP::ImageLayerArgs Args = InProgram.GetOpArgs<mu::OP::ImageLayerArgs>(InParentAddress);
+		AddOpFunc(Args.base, TEXT("base "));
+		if (Args.mask)
+		{
+			AddOpFunc(Args.mask, TEXT("mask "));
+		}
+		AddOpFunc(Args.blended, TEXT("blended "));
+		break;
+	}
+
+	case mu::OP_TYPE::IM_LAYERCOLOUR:
+	{
+		mu::OP::ImageLayerColourArgs Args = InProgram.GetOpArgs<mu::OP::ImageLayerColourArgs>(InParentAddress);
+		AddOpFunc(Args.base, TEXT("base "));
+		if (Args.mask)
+		{
+			AddOpFunc(Args.mask, TEXT("mask "));
+		}
+		AddOpFunc(Args.colour, TEXT("colour "));
+		break;
+	}
+
+	case mu::OP_TYPE::IM_MULTILAYER:
+	{
+		mu::OP::ImageMultiLayerArgs Args = InProgram.GetOpArgs<mu::OP::ImageMultiLayerArgs>(InParentAddress);
+		AddOpFunc(Args.rangeSize, TEXT("range "));
+		AddOpFunc(Args.base, TEXT("base "));
+		if (Args.mask)
+		{
+			AddOpFunc(Args.mask, TEXT("mask "));
+		}
+		AddOpFunc(Args.blended, TEXT("blended "));
+		break;
+	}
+
+	case mu::OP_TYPE::ME_ADDTAGS:
+	{
+		const uint8* OpData = InProgram.GetOpArgsPointer(InParentAddress);
+
+		mu::OP::ADDRESS SourceAddress;
+		FMemory::Memcpy(&SourceAddress, OpData, sizeof(mu::OP::ADDRESS));
+		OpData += sizeof(mu::OP::ADDRESS);
+
+		uint16 TagCount;
+		FMemory::Memcpy(&TagCount, OpData, sizeof(uint16));
+		OpData += sizeof(uint16);
+
+		FString Caption = FString::Printf(TEXT("add %d tags to "), TagCount);
+		AddOpFunc(SourceAddress, Caption);
+
+		break;
+	}
+
+	default:
+	{
+		// Generic list of child operations
+		bUseGeneric = true;
+		break;
+	}
+
+	}
+
+	if (bUseGeneric)
+	{
+		// Find children of the provided element
+		mu::ForEachReference(InProgram, InParentAddress, [this, &InProgram, AddOpFunc](mu::OP::ADDRESS ChildAddress)
+			{
+				AddOpFunc(ChildAddress,TEXT(""));
+			});
+	}
+	else
+	{
+		// Validate in case there is a mismatch in the custom processing of children and the generic one, which would cause problems.
+		ChildIndex = 0;
+
+		auto ValidateOpFunc = [this, InParentAddress, &ChildIndex](mu::OP::ADDRESS ChildAddress)
+		{
+			if (ChildAddress)
+			{
+				const FItemCacheKey Key = { InParentAddress, ChildAddress, ChildIndex };
+				const TSharedPtr<FMutableCodeTreeElement>* CachedItem = ItemCache.Find(Key);
+				check(CachedItem);
+			}
+			++ChildIndex;
+		};
+
+		mu::ForEachReference(InProgram, InParentAddress, [this, &InProgram, ValidateOpFunc](mu::OP::ADDRESS ChildAddress)
+			{
+				ValidateOpFunc(ChildAddress);
+			});
+	}
+}
+
+SMutableCodeViewer::EOperationComputationalCost SMutableCodeViewer::GetOperationTypeComputationalCost(mu::OP_TYPE OperationType) const
+{
+	if (VeryExpensiveOperationTypes.Contains(OperationType))
+	{
+		return EOperationComputationalCost::VeryExpensive;
+	}
+	else if (ExpensiveOperationTypes.Contains(OperationType))
+	{
+		return EOperationComputationalCost::Expensive;
+	}
+	else
+	{
+		return EOperationComputationalCost::Standard;
+	}
+}
+
+#pragma endregion 
+
+
+
+#pragma region CodeTree Callbacks
+
+
+TSharedRef<ITableRow> SMutableCodeViewer::GenerateRowForNodeTree(TSharedPtr<FMutableCodeTreeElement> InTreeNode, const TSharedRef<STableViewBase>& InOwnerTable)
+{
+	// Save the node for later access
+	TreeElements.Add(InTreeNode);
+	
+	// Generate a row element
+	TSharedRef<SMutableCodeTreeRow> Row = SNew(SMutableCodeTreeRow, InOwnerTable, InTreeNode);
+	
+	// Determine if a row should be painted as highlighted based on the selected item
+	if (TreeView->GetNumItemsSelected())
+	{
+		const TSharedPtr<FMutableCodeTreeElement> SelectedElement = TreeView->GetSelectedItems()[0];
+		if (InTreeNode->MutableOperation == SelectedElement->MutableOperation)
+		{
+			Row->Highlight();
+		}
+	}
+	
+	return Row;
+}
+
+
+void SMutableCodeViewer::GetChildrenForInfo(TSharedPtr<FMutableCodeTreeElement> InInfo, TArray<TSharedPtr<FMutableCodeTreeElement>>& OutChildren)
+{
+	if (!InInfo->MutableModel)
+	{
+		return;
+	}
+
+	check(MutableModel);
+	const mu::FProgram& Program = MutableModel->GetPrivate()->m_program;
+
+	mu::OP::ADDRESS ParentAddress = InInfo->MutableOperation;
+
+	// Generic case for unnamed children traversal.
+	uint32 ChildIndex = 0;
+	mu::ForEachReference(Program, InInfo->MutableOperation, [this, ParentAddress, &ChildIndex, &OutChildren](mu::OP::ADDRESS ChildAddress)
+	{
+		if (ChildAddress)
+		{
+			const FItemCacheKey Key = { ParentAddress, ChildAddress, ChildIndex };
+			const TSharedPtr<FMutableCodeTreeElement>* CachedItem = ItemCache.Find(Key);
+
+			if (CachedItem)
+			{
+				OutChildren.Add(*CachedItem);
+			}
+			else
+			{
+				// if all elements have been already cached this should never happen
+				checkNoEntry();
+			}
+		}
+		++ChildIndex;
+	});
+}
+
+void SMutableCodeViewer::OnExpansionChanged(TSharedPtr<FMutableCodeTreeElement> InItem, bool bInExpanded)
+{
+	// Update expanded state of the provided element
+	InItem->bIsExpanded = bInExpanded;
+	
+	// If an element gets expanded then contract (if found) the other element that uses the same address
+	if (bInExpanded)
+	{
+		const mu::OP::ADDRESS MutableOperation = InItem->MutableOperation;
+		const TSharedPtr<FMutableCodeTreeElement>* PreviouslyExpandedElement = ExpandedElements.Find(MutableOperation);
+		if (PreviouslyExpandedElement)
+		{
+			TreeView->SetItemExpansion(*PreviouslyExpandedElement, false);
+		}
+
+		// Only do this if in a situation where it may be required (do not do it if the tree has not been interacted with yet)
+		if (bShouldRecalculateStates)
+		{
+			// Find all the children (recursive) of this item.
+			TSet<TSharedPtr<FMutableCodeTreeElement>> FoundChildren;
+			GetVisibleChildren(InItem, FoundChildren);
+			for (const TSharedPtr<FMutableCodeTreeElement>& Child : FoundChildren)
+			{
+				// For each of the children found set it's state to be the one found on the expanded element
+				Child->SetElementCurrentState(InItem->GetStateIndex());
+			}
+		}
+		
+		// Cache this element as one currently expanded
+		ExpandedElements.Add(MutableOperation, InItem);
+	}
+	else
+	{
+		// Remove this element from the cache of expanded elements
+		ExpandedElements.Remove(InItem->MutableOperation);
+	}
+}
+
+void SMutableCodeViewer::GetVisibleChildren(TSharedPtr<FMutableCodeTreeElement> InInfo, TSet<TSharedPtr<FMutableCodeTreeElement>>& OutChildren)
+{
+	check(MutableModel);
+	const mu::FProgram& MutableProgram = MutableModel->GetPrivate()->m_program;
+	
+	TArray<TSharedPtr<FMutableCodeTreeElement>> ToSearchForChildren;
+	ToSearchForChildren.Add(InInfo);
+	while (!ToSearchForChildren.IsEmpty())
+	{
+		// Grab the first element in order to check for it's children
+		const TSharedPtr<FMutableCodeTreeElement> ToCheck = ToSearchForChildren[0];
+		ToSearchForChildren.RemoveAt(0);
+		
+		const mu::OP::ADDRESS ParentAddress = ToCheck->MutableOperation;
+	
+		// Generic case for unnamed children traversal.
+		uint32 ChildIndex = 0;
+		mu::ForEachReference(MutableProgram, ParentAddress, [this, ParentAddress, &ChildIndex, &OutChildren, &ToSearchForChildren ](mu::OP::ADDRESS ChildAddress)
+		{
+			if (ChildAddress)
+			{
+				const FItemCacheKey Key = { ParentAddress, ChildAddress, ChildIndex };
+				const TSharedPtr<FMutableCodeTreeElement> CachedItem = *ItemCache.Find(Key);
+
+				// Since we have already generated all elements CachedItem should be therefore always a valid pointer
+				check (CachedItem);
+				
+				// If the address has not been yet found then save it as one of the children affected
+				if (!OutChildren.Contains(CachedItem))
+				{
+					OutChildren.Add(CachedItem);
+
+					// And if the children is found to be expanded then also process it later to later return only the
+					// elements that are expanded in the tree view (using data manually set on each tree element)
+					if (CachedItem->bIsExpanded )
+					{
+						// Add for processing
+						ToSearchForChildren.Add(CachedItem);
+					}
+				}
+			}
+			++ChildIndex;
+		});
+	}
+
+	// Debug
+	// UE_LOG(LogTemp,Warning,TEXT("Found a total of %i children elements "),OutChildren.Num());
+}
+
+
+void SMutableCodeViewer::OnSelectionChanged(TSharedPtr<FMutableCodeTreeElement> InNode, ESelectInfo::Type InSelectInfo)
+{
+	if (bIsElementHighlighted)
+	{
+		ClearHighlightedItems();
+	}
+
+	TArray<TSharedPtr<FMutableCodeTreeElement>> SelectedNodes;
+	TreeView->GetSelectedItems(SelectedNodes);
+
+	PreviewBorder->ClearContent();
+
+	SelectedOperationAddress = 0;
+	bSelectedOperationIsImage = false;
+
+	if (SelectedNodes.IsEmpty())
+	{
+		return;
+	}
+
+	// Clear all selected items in the constant resources widget
+	ConstantsWidget->ClearSelectedConstantItems();
+	
+	// Find the duplicates for the selected tree element and highlight them
+	if (InNode)
+	{
+		HighlightDuplicatesOfEntry(InNode);
+	}
+
+	bIsPreviewPendingUpdate = true;
+
+	SelectedOperationAddress = SelectedNodes[0]->MutableOperation;
+	const mu::OP_TYPE OperationType = MutableModel->GetPrivate()->m_program.GetOpType(SelectedOperationAddress);
+	const mu::DATATYPE OperationDataType = mu::GetOpDataType(OperationType);
+	
+	switch (OperationDataType)
+	{
+	case mu::DT_LAYOUT:
+	{
+		// Create or reuse the UI
+		PrepareLayoutViewer();
+		break;
+	}
+
+	case mu::DT_IMAGE:
+	{
+		// Create or reuse the UI
+		bSelectedOperationIsImage = true;
+		PrepareImageViewer();
+		break;
+	}
+
+	case mu::DT_MESH:
+	{
+		// Create or reuse the UI
+		PrepareMeshViewer();
+		break;
+	}
+
+	case mu::DT_INSTANCE:
+	{
+		// Create or reuse the UI
+		PrepareInstanceViewer();
+		break;
+	}
+
+	case mu::DT_SCALAR:
+	{
+		// Create or reuse the UI
+		if (!PreviewScalarViewer)
+		{
+			PreviewScalarViewer = SNew(SMutableScalarViewer);
+		}
+
+		PreviewBorder->SetContent(PreviewScalarViewer.ToSharedRef());
+		break;
+	}
+
+	case mu::DT_STRING:
+	{
+		// Create or reuse the UI
+		PrepareStringViewer();
+		break;
+	}
+
+	case mu::DT_COLOUR:
+	{
+		// Create or reuse the UI
+		if (!PreviewColorViewer)
+		{
+			PreviewColorViewer = SNew(SMutableColorViewer);
+		}
+
+		PreviewBorder->SetContent(PreviewColorViewer.ToSharedRef());
+		break;
+	}
+
+	case mu::DT_INT:
+	{
+		// Create or reuse the UI
+		if (!PreviewIntViewer)
+		{
+			PreviewIntViewer = SNew(SMutableIntViewer);
+		}
+
+		PreviewBorder->SetContent(PreviewIntViewer.ToSharedRef());
+		break;
+	}
+
+	case mu::DT_BOOL:
+	{
+		// Create or reuse the UI
+		if (!PreviewBoolViewer)
+		{
+			PreviewBoolViewer = SNew(SMutableBoolViewer);
+		}
+
+		PreviewBorder->SetContent(PreviewBoolViewer.ToSharedRef());
+		break;
+	}
+
+	case mu::DT_PROJECTOR:
+	{
+		// Create or reuse the UI
+		PrepareProjectorViewer();
+		break;
+	}
+
+	
+	default:
+		// There is no viewer for this type yet.
+		break;
+
+	}
+}
+
+TSharedPtr<SWidget> SMutableCodeViewer::OnTreeContextMenuOpening()
+{
+	FMenuBuilder MenuBuilder(true, nullptr);
+	
+	if (TreeView->GetSelectedItems().Num())
+	{
+		if (TreeView->GetSelectedItems().Num() == 1)
+		{
+			MenuBuilder.AddMenuEntry(
+				LOCTEXT("Set_as_search_operation_type","Set as search Operation"),
+				LOCTEXT("Set_as_search_operation_type_Tooltip", "Sets the type of this operation as the type to be looking for when searching for operations on the tree view"),
+				FSlateIcon(),
+				FUIAction(FExecuteAction::CreateSP(this, &SMutableCodeViewer::OnSelectedOperationTypeFromTree)
+				)	
+			);
+		}
+
+		MenuBuilder.AddMenuSeparator();
+		
+		MenuBuilder.AddMenuEntry(
+			LOCTEXT("Code_Expand_Selected", "Expand Selected Operation"),
+			LOCTEXT("Code_Expand_Selected_Tooltip", "Expands only the selected Operation and leaves the other as they are."),
+			FSlateIcon(),
+			FUIAction(FExecuteAction::CreateSP(this, &SMutableCodeViewer::TreeExpandSelected)
+			)
+		);
+	}
+
+
+	MenuBuilder.AddMenuEntry(
+		LOCTEXT("Code_Expand_Instance", "Expand Instance-Level Operations"),
+		LOCTEXT("Code_Expand_Instance_Tooltip", "Expands all the operations in the tree that are instance operations (not images, meshes, booleans, etc.)."),
+		FSlateIcon(),
+		FUIAction(FExecuteAction::CreateSP(this, &SMutableCodeViewer::TreeExpandInstance)
+			//, FCanExecuteAction::CreateSP(this, &SMutableCodeViewer::HasAnyItemInPalette)
+		)
+	);
+
+	MenuBuilder.AddMenuEntry(
+		LOCTEXT("Code_Expand_Unique", "Expand All Unique Operations"),
+		LOCTEXT("Code_Expand_Unique_Tooltip", "Expands all the operations in the tree that have not been expanded yet."),
+		FSlateIcon(),
+		FUIAction(FExecuteAction::CreateSP(this, &SMutableCodeViewer::TreeExpandUnique)
+		)
+	);
+
+	return MenuBuilder.MakeWidget();
+}
+
+void SMutableCodeViewer::TreeExpandRecursive(TSharedPtr<FMutableCodeTreeElement> InInfo, bool bExpand)
+{
+	if (bExpand)
+	{
+		TreeExpandUnique();
+	}
+}
+
+void SMutableCodeViewer::OnRowReleased(const TSharedRef<ITableRow>& InTreeRow)
+{
+	SMutableCodeTreeRow* CastedTableRow = static_cast<SMutableCodeTreeRow*>(&InTreeRow.Get());
+	const TSharedPtr<FMutableCodeTreeElement>& RowElement = CastedTableRow->GetItem();
+	TreeElements.Remove(RowElement);
+}
+#pragma endregion
+
+#pragma region Highlight Methods
+void SMutableCodeViewer::HighlightDuplicatesOfEntry(const TSharedPtr<FMutableCodeTreeElement>& InTargetEntry)
+{
+	if (bIsElementHighlighted)
+	{
+		ClearHighlightedItems();
+	}
+
+	// Highlight the elements related to the currently selected item of the tree
+	HighlightedOperation = InTargetEntry->MutableOperation;
+
+	for (const TSharedPtr<FMutableCodeTreeElement>& TreeItem : TreeElements)
+	{
+		if (TreeItem.Get() != InTargetEntry.Get() && TreeItem->MutableOperation == HighlightedOperation)
+		{
+			TSharedPtr<ITableRow> TableRow = TreeView->WidgetFromItem(TreeItem);
+			SMutableCodeTreeRow* MutableRow = static_cast<SMutableCodeTreeRow*>(TableRow.Get());
+			MutableRow->Highlight();
+		}
+	}
+
+	bIsElementHighlighted = true;
+}
+
+void SMutableCodeViewer::ClearHighlightedItems()
+{
+	// Clear the previously highlighted elements
+	for (const TSharedPtr<FMutableCodeTreeElement>& HighlightedElement : TreeElements)
+	{
+		if (HighlightedElement->MutableOperation == HighlightedOperation)
+		{
+			TSharedPtr<ITableRow> TableRow = TreeView->WidgetFromItem(HighlightedElement);
+
+			if (TableRow.IsValid())
+			{
+				SMutableCodeTreeRow* MutableRow = static_cast<SMutableCodeTreeRow*>(TableRow.Get());
+				MutableRow->ResetHighlight();
+			}
+
+		}
+	}
+
+	bIsElementHighlighted = false;
+}
+#pragma endregion
+
+#pragma region Element Expansion Llogic
+
+
+void SMutableCodeViewer::TreeExpandElements(TArray<TSharedPtr<FMutableCodeTreeElement>>& InElementsToExpand, 
+	bool bForceExpandDuplicates /*= false*/,
+	mu::DATATYPE FilteringDataType /*= mu::DATATYPE::DT_NONE*/,
+	TSharedPtr<FProcessedOperationsBuffer> InExpandedOperationsBuffer /* = nullptr */)
+{
+	if (InElementsToExpand.IsEmpty())
+	{
+		return;
+	}
+	
+	// Initialization of recursive elements if this is the first invocation of method
+	{
+		if (!InExpandedOperationsBuffer)
+		{
+			InExpandedOperationsBuffer = MakeShared<FProcessedOperationsBuffer>();
+		}
+	}
+
+	// Load references to the arrays containing all the operations already worked on during another recursive call to this
+	// method
+	TArray<mu::OP::ADDRESS>& AlreadyExpandedOriginalOperations = InExpandedOperationsBuffer->ExpandedOriginalOperations;
+	TArray<mu::OP::ADDRESS>& AlreadyExpandedDuplicatedOperations = InExpandedOperationsBuffer->ExpandedDuplicatedOperations;
+	
+	// Array containing the children object found on Item.
+	TArray<TSharedPtr<FMutableCodeTreeElement>> Children;
+	
+	// Index of the current element being processed
+	int32 CurrentElementIndex = 0;
+	while (CurrentElementIndex < InElementsToExpand.Num() )
+	{
+		// Grab the current element to process and move the index forward once
+		TSharedPtr<FMutableCodeTreeElement> Item = InElementsToExpand[CurrentElementIndex++];
+		check(Item);
+
+		// Identifier of the element. May be repeated if there are elements duplicating another element
+		const mu::OP::ADDRESS OperationAddress = Item->MutableOperation;
+		
+		// Filter the elements being expanded if the user has defined a desired DATATYPE
+		if (FilteringDataType != mu::DATATYPE::DT_NONE)
+		{
+			const mu::OP_TYPE OperationType = Item->MutableModel->GetPrivate()->m_program.GetOpType(OperationAddress);
+			const mu::DATATYPE OperationDataType = mu::GetOpDataType(OperationType);
+
+			// If it is not of the desired type then ignore it and continue to the next pending element
+			if (OperationDataType != FilteringDataType)
+			{
+				continue;
+			}
+		}
+
+		// Reset the children array
+		Children.SetNum(0);
+		
+		// If not duplicated expand it and grab the children to be also expanded on the next loop
+		if (Item->DuplicatedOf == nullptr )
+		{
+			// Was this unique element expanded before (only valid if also expanding duplicates)
+			bool bHasBeenExpandedPreviously = false;
+			
+			// Mind duplicated original elements if dealing with duplicated operation expansions.
+			if (bForceExpandDuplicates)
+			{
+				// Make sure we have not already expanded this item to avoid recursive expansions of the same item and
+				// children
+				bHasBeenExpandedPreviously = AlreadyExpandedOriginalOperations.Contains(OperationAddress);
+			}
+			
+			// Only check for duplicated original elements when working with duplicates
+			if (!bHasBeenExpandedPreviously )
+			{
+				// Get the children of this unique element and prepare them for processing
+				GetChildrenForInfo(Item, Children);
+
+				// Call for the expansion of the children first
+				TreeExpandElements(Children, bForceExpandDuplicates, FilteringDataType,InExpandedOperationsBuffer);
+
+				// At this point all the children objects that needed expansion are already expanded so we can proceed with
+				// the expansion of this element
+				
+				{
+					// If we do expect to expand duplicates make sure we record this object as being expanded to be later able
+					// to block the expansion of duplicates of this object
+					if (bForceExpandDuplicates)
+					{
+						// Register this node as expanded so other nodes are able to check if it has already been worked with
+						AlreadyExpandedOriginalOperations.Add(OperationAddress);
+					}
+
+					// Only ask for the expansion of the element if we know it can be expanded due to it having
+					// children
+					if (!Children.IsEmpty())
+					{
+						// Expand this unique element
+						TreeView->SetItemExpansion(Item, true);
+					}
+				}
+			}
+			
+		}
+		// If it is a duplicated node
+		else
+		{
+			// Special behavior where we expand duplicates if parent is not found to be expanded
+			if (bForceExpandDuplicates)
+			{
+				// Was this element expanded as an original operation? We only want to expand the duplicate if the original
+				// was not duplicated before
+				bool bOriginalElementHasBeenExpanded = false;
+
+				// Was this element expanded on a duplicated element? we only want to expand the first duplicate!
+				const bool bOtherDuplicateOfSameOpWasExpandedBefore = AlreadyExpandedDuplicatedOperations.Contains(OperationAddress);
+
+				// Only check if there is another original element with the same operation if we know that there is not another
+				// duplicated element using this operation
+				if (!bOtherDuplicateOfSameOpWasExpandedBefore)
+				{
+					bOriginalElementHasBeenExpanded = AlreadyExpandedOriginalOperations.Contains(OperationAddress);
+				}
+				
+				// Was this operation expanded before?
+				const bool bWasOperationExpandedPreviously = bOriginalElementHasBeenExpanded || bOtherDuplicateOfSameOpWasExpandedBefore;
+
+				// If this operation have not yet been expanded then expand it!
+				// Duplicates do not have priority over original elements.
+				if ( !bWasOperationExpandedPreviously )
+				{
+					// Mark the children to be expanded later if conditions are met
+					GetChildrenForInfo(Item, Children);
+
+					// Expand the children objects
+					TreeExpandElements(Children, bForceExpandDuplicates, FilteringDataType,InExpandedOperationsBuffer);
+
+					// At this point all the children objects that needed expansion are already expanded so we can proceed with
+					// the expansion of this element
+					
+					{
+						// Record this node being expanded
+						AlreadyExpandedDuplicatedOperations.Add(OperationAddress);
+
+						// Only ask for the expansion of the element if we know it can be expanded due to it having
+						// children
+						if (!Children.IsEmpty())
+						{
+							// Expand the current element since we know it is from a operation not yet expanded
+							TreeView->SetItemExpansion(Item, true);
+						}
+					}
+				}
+			}
+		}
+	}
+}
+
+
+void SMutableCodeViewer::TreeExpandSelected()
+{
+	// Get the selected items and expand them excluding the duplicates
+	TArray<TSharedPtr<FMutableCodeTreeElement>> SelectedItems = TreeView->GetSelectedItems();
+	TreeExpandElements(SelectedItems,true);
+}
+
+void SMutableCodeViewer::TreeExpandUnique()
+{
+	// Expand the tree from the root and do not expand the duplicated elements
+	TreeExpandElements(RootNodes);
+}
+
+void SMutableCodeViewer::TreeExpandInstance()
+{
+	// Expand only the items that match the datatype provided
+	TreeExpandElements(RootNodes,false, mu::DT_INSTANCE);
+}
+
+#pragma endregion
+
+#pragma region Caching of operations related to constant resource
+
+void SMutableCodeViewer::CacheRootNodeAddresses()
+{
+	check (MutableModel);
+	check (RootNodeAddresses.IsEmpty())
+
+	TArray<mu::OP::ADDRESS> FoundRootNodeAddresses;
+	
+	const mu::Model::Private* ModelPrivate = MutableModel->GetPrivate();
+	const int32 StateCount = ModelPrivate->m_program.m_states.Num();
+	for ( int32 StateIndex=0; StateIndex < StateCount; ++StateIndex )
+	{
+		const mu::FProgram::FState& State = ModelPrivate->m_program.m_states[StateIndex];
+		FoundRootNodeAddresses.Add(State.m_root);
+	}
+
+	RootNodeAddresses = MoveTemp(FoundRootNodeAddresses);
+}
+
+ void SMutableCodeViewer::CacheAddressesRelatedWithConstantResource(const mu::DATATYPE ConstantDataType,
+	const int32 IndexOnConstantsArray)
+{
+	check(MutableModel);
+	check(RootNodeAddresses.Num());
+
+	if (IndexOnConstantsArray < 0)
+	{
+		// Not valid index.
+		UE_LOG(LogTemp,Error,TEXT("The provided index [%d] is not valid."),IndexOnConstantsArray );
+		return;
+	}
+
+	// Object containing all data required by the search operation to be able to be called recursively
+	FElementsSearchCache SearchPayload;
+	// Initialize the Search Payload with the root node addresses. This way the search will use them as the root nodes where
+	// to start searching
+	SearchPayload.SetupRootBatch(RootNodeAddresses);
+	
+	// Main update procedure run for the targeted state and the targeted parameter values
+	const mu::FProgram& Program = MutableModel->GetPrivate()->m_program; 
+	GetOperationsReferencingConstantResource(ConstantDataType,IndexOnConstantsArray,SearchPayload, Program);
+	
+	// At this point we did get all the addresses of operations that do involve the usage of our resource
+	if (SearchPayload.FoundElements.Num() > 0)
+	{
+		// Set the type operation type to CONST_BASED_NAVIGATION (used to tell the user what is happening)
+		TargetedTypeSelector->SetSelectedItem(ConstantBasedNavigationEntry);
+
+		// Dump the located resources array onto the navigation array since we have content to navigate over
+		NavigationElements = MoveTemp(SearchPayload.FoundElements);
+		SortElementsByTreeIndex(NavigationElements);
+		
+		// Reset the navigation index
+		NavigationIndex = -1;
+	}
+	else
+	{
+		UE_LOG(LogTemp,Error,TEXT("The provided constant index does not seem to be used anywere : Make sure the index is valid and that IsConstantResourceUsedByOperation() switch is up to date"));
+	}
+}
+
+
+void SMutableCodeViewer::GetOperationsReferencingConstantResource(
+	const mu::DATATYPE ConstantDataType,
+	const int32 IndexOnConstantsArray,
+	FElementsSearchCache& InSearchPayload,
+	const mu::FProgram& InProgram)
+{
+	// next batch of addresses to be explored 
+	TArray<FItemCacheKey> NextBatchAddressesData;
+	
+	for	(int32 ParentIndex = 0 ; ParentIndex < InSearchPayload.BatchData.Num(); ParentIndex++)
+	{
+		// Get one of the previous run "children" and treat as a parent to get it's children and process them
+		const mu::OP::ADDRESS& ParentAddress = InSearchPayload.BatchData[ParentIndex].Child;
+		
+		// Cache if same data type and we share the same address (means this op is pointing at the provided resource)
+		// It will cache duplicated entries
+		if (IsConstantResourceUsedByOperation(IndexOnConstantsArray, ConstantDataType, ParentAddress,InProgram))
+		{
+			// Since this element is related with the provided constant resource cache it on InSearchPayload.FoundElements
+			InSearchPayload.AddToFoundElements(ParentAddress,ParentIndex,ItemCache);
+		}
+		
+		// Get all NON PROCESSED the children of this operation to later be able to process them (on next recursive call)
+		InSearchPayload.CacheChildrenOfAddressIfNotProcessed(ParentAddress, InProgram, NextBatchAddressesData);
+	}
+
+	// At this point all the addresses to be computed on the next batch have already been set and will be computed on
+	// the next recursive call
+	
+	// Explore children if found 
+	if (NextBatchAddressesData.Num())
+	{
+		// Cache next batch data so the next invocations is able to locate the provided addresses on the itemsCache
+		InSearchPayload.BatchData = MoveTemp(NextBatchAddressesData);
+		
+		GetOperationsReferencingConstantResource(ConstantDataType, IndexOnConstantsArray, InSearchPayload, InProgram);
+	}
+}
+
+bool SMutableCodeViewer::IsConstantResourceUsedByOperation(const int32 IndexOnConstantsArray,
+	const mu::DATATYPE ConstantDataType, const mu::OP::ADDRESS OperationAddress, const mu::FProgram& InProgram) const
+{
+	// Cache the current operation type to know where to look and what to check
+	const mu::OP_TYPE OperationType = InProgram.GetOpType(OperationAddress);
+
+	// Making usage of the operation data type is not valid since some operations while return one type do, in fact,
+	// contain data from other types (like the mesh constant for example that contains mesh, skeleton and physics asset)
+	// const mu::DATATYPE DataType = mu::GetOpDataType(OperationType);
+	// if (DataType != ConstantDataType)
+	// {
+	// 	return false;
+	// }
+
+	// Is this operation referencing (by an index) the index we are providing from a constants array
+	bool bResourceLocated = false;
+	
+	// Check if the operation data type is compatible with the type of resources we are providing
+	switch (ConstantDataType)
+	{
+		case mu::DATATYPE::DT_STRING:
+			{
+				// TIP: To know if they represent a constant value check the code on code runner to see if they read from the constants array
+				if (OperationType == mu::OP_TYPE::ST_CONSTANT)
+				{
+					bResourceLocated = IndexOnConstantsArray == InProgram.GetOpArgs<mu::OP::ResourceConstantArgs>(OperationAddress).value;
+				}
+				else if (OperationType == mu::OP_TYPE::IN_ADDSTRING)
+				{
+					bResourceLocated = IndexOnConstantsArray == InProgram.GetOpArgs<mu::OP::InstanceAddArgs>(OperationAddress).name;
+				}
+				else if (OperationType == mu::OP_TYPE::IN_ADDMESH)
+				{
+					bResourceLocated = IndexOnConstantsArray == InProgram.GetOpArgs<mu::OP::InstanceAddArgs>(OperationAddress).name;
+				}
+				else if(OperationType == mu::OP_TYPE::IN_ADDIMAGE)
+				{
+					bResourceLocated = IndexOnConstantsArray == InProgram.GetOpArgs<mu::OP::InstanceAddArgs>(OperationAddress).name;
+				}
+				else if (OperationType == mu::OP_TYPE::IN_ADDVECTOR)
+				{
+					bResourceLocated = IndexOnConstantsArray == InProgram.GetOpArgs<mu::OP::InstanceAddArgs>(OperationAddress).name;
+				}
+				else if (OperationType == mu::OP_TYPE::IN_ADDSCALAR)
+				{
+					bResourceLocated = IndexOnConstantsArray == InProgram.GetOpArgs<mu::OP::InstanceAddArgs>(OperationAddress).name;
+				}
+				else if (OperationType == mu::OP_TYPE::IN_ADDCOMPONENT)
+				{
+					bResourceLocated = IndexOnConstantsArray == InProgram.GetOpArgs<mu::OP::InstanceAddArgs>(OperationAddress).name;
+				}
+				else if (OperationType == mu::OP_TYPE::IN_ADDSURFACE)
+				{
+					bResourceLocated = IndexOnConstantsArray == InProgram.GetOpArgs<mu::OP::InstanceAddArgs>(OperationAddress).name;
+				}
+				else if (OperationType == mu::OP_TYPE::ME_BINDSHAPE)
+				{
+					mu::OP::MeshBindShapeArgs Arguments = InProgram.GetOpArgs<mu::OP::MeshBindShapeArgs>(OperationAddress);
+					const uint8_t* Data = InProgram.GetOpArgsPointer(OperationAddress);
+
+					// Bones are stored after the args
+					Data += sizeof(Arguments);
+
+					// Iterate over the bones and check if they point to the same index on the string constants array
+					int32 NumBones;
+					FMemory::Memcpy(&NumBones, Data, sizeof(int32)); 
+					Data += sizeof(int32);
+				
+					for (int32 Bone = 0; Bone < NumBones; ++Bone)
+					{
+						// Exit once we know that the data is pointing to the index provided
+						if (*Data == IndexOnConstantsArray)
+						{
+							bResourceLocated = true;
+							break;
+						}
+						
+						Data += sizeof(int32);
+					}
+				}
+				
+				break;
+			}
+
+		case mu::DATATYPE::DT_IMAGE:
+			{
+				if (OperationType == mu::OP_TYPE::IM_CONSTANT)
+				{
+					bResourceLocated = IndexOnConstantsArray == InProgram.GetOpArgs<mu::OP::ResourceConstantArgs>(OperationAddress).value;
+				} 
+				
+				break;
+			}
+
+		case mu::DATATYPE::DT_MESH:
+			{
+				if (OperationType == mu::OP_TYPE::ME_CONSTANT)
+				{
+					bResourceLocated = IndexOnConstantsArray == InProgram.GetOpArgs<mu::OP::MeshConstantArgs>(OperationAddress).value;
+				}
+				break;
+			}
+
+		case mu::DATATYPE::DT_LAYOUT:
+			{
+				if (OperationType == mu::OP_TYPE::LA_CONSTANT)
+				{
+					bResourceLocated = IndexOnConstantsArray == InProgram.GetOpArgs<mu::OP::ResourceConstantArgs>(OperationAddress).value;
+				}
+				break;
+			}
+
+		case mu::DATATYPE::DT_PROJECTOR:
+			{
+				if (OperationType == mu::OP_TYPE::PR_CONSTANT)
+				{
+					bResourceLocated = IndexOnConstantsArray == InProgram.GetOpArgs<mu::OP::ResourceConstantArgs>(OperationAddress).value;
+				}
+				break;
+			}
+
+		case mu::DATATYPE::DT_MATRIX:
+			{
+				if (OperationType == mu::OP_TYPE::ME_TRANSFORM)
+				{
+					bResourceLocated = IndexOnConstantsArray == InProgram.GetOpArgs<mu::OP::MeshTransformArgs>(OperationAddress).matrix;
+				}
+				break;
+			}
+
+		case mu::DATATYPE::DT_SHAPE:
+			{
+				if (OperationType == mu::OP_TYPE::ME_CLIPMORPHPLANE)
+				{
+					const mu::OP::MeshClipMorphPlaneArgs Arguments =  InProgram.GetOpArgs<mu::OP::MeshClipMorphPlaneArgs>(OperationAddress);
+					
+					// Morph shape
+					bResourceLocated = IndexOnConstantsArray == Arguments.morphShape;
+					if (bResourceLocated)
+					{
+						break;
+					}
+
+					if (Arguments.vertexSelectionType == mu::OP::MeshClipMorphPlaneArgs::VS_SHAPE)
+					{
+						// Selection Shape
+						bResourceLocated = IndexOnConstantsArray == InProgram.GetOpArgs<mu::OP::MeshClipMorphPlaneArgs>(OperationAddress).vertexSelectionShapeOrBone;
+					}
+				}
+				break;
+			}
+		
+		case mu::DATATYPE::DT_CURVE:
+			{
+				if (OperationType == mu::OP_TYPE::SC_CURVE)
+				{
+					bResourceLocated = IndexOnConstantsArray == InProgram.GetOpArgs<mu::OP::ScalarCurveArgs>(OperationAddress).curve;
+				}
+				break;
+			}
+
+		case mu::DATATYPE::DT_SKELETON:
+			{
+				if (OperationType == mu::OP_TYPE::ME_CONSTANT)
+				{
+					bResourceLocated = IndexOnConstantsArray == InProgram.GetOpArgs<mu::OP::MeshConstantArgs>(OperationAddress).skeleton;
+				}
+				break;
+			}
+
+		// Invalid types
+		case mu::DATATYPE::DT_NONE:
+		default:
+			{
+				checkNoEntry();
+			}
+	}
+
+
+	return bResourceLocated;
+}
+
+
+
+#pragma endregion 
+
+
+namespace
+{
+	/** Test implementation to provide image parameters. It will generate some images of a fixed size and format. */
+	class TestResourceProvider : public mu::ExternalResourceProvider
+	{
+		static inline const mu::FImageDesc IMAGE_DESC = 
+			mu::FImageDesc(mu::FImageSize(1024, 1024), mu::EImageFormat::IF_RGBA_UBYTE, 1);
+
+	public:
+
+		/** */
+		TArray<TSoftObjectPtr<const UTexture>> ReferencedTextures;
+
+	public:
+
+
+		virtual TTuple<UE::Tasks::FTask, TFunction<void()>> GetImageAsync(FName Id, uint8 MipmapsToSkip, TFunction<void(mu::Ptr<mu::Image>)>& ResultCallback) override
+		{
+			MUTABLE_CPUPROFILER_SCOPE(TestImageProvider_GetImage);
+
+			int32 Size = IMAGE_DESC.m_size[0];
+			Size = FMath::Max(4, Size / (1 << MipmapsToSkip));
+
+			mu::Ptr<mu::Image> Image = new mu::Image(
+				Size, Size, IMAGE_DESC.m_lods,
+				IMAGE_DESC.m_format,
+				mu::EInitializationType::NotInitialized);
+
+			// Generate an alpha-tested circle with an horizontal gradient color.
+			uint8* Data = Image->GetLODData(0);
+			int32 CircleRadius = (Size * 2) / 5;
+			int32 CircleRadius2 = CircleRadius * CircleRadius;
+			int32 Color[3] = {255, 128, 0};
+
+			int32 LogSize = FMath::CeilLogTwo(Size);
+
+			int32 HalfSize = Size >> 1;
+			for (int32 RadY = -HalfSize; RadY < HalfSize; ++RadY)
+			{
+				int32 RadY2 = RadY * RadY;
+				for (int32 x = 0; x < Size; ++x)
+				{
+					int32 RadX = (x - HalfSize);
+					int32 R2 = RadX * RadX + RadY2;
+					int32 Opacity = FMath::Clamp(((CircleRadius2 - R2) * 512) / CircleRadius2 - 64, 0, 255);
+					Data[0] = uint8((Color[0] * x) >> LogSize);
+					Data[1] = uint8((Color[1] * x) >> LogSize);
+					Data[2] = uint8((Color[2] * x) >> LogSize);
+					Data[3] = uint8(Opacity);
+					Data += 4;
+				}
+			}
+
+			ResultCallback(Image);
+
+			return MakeTuple(UE::Tasks::MakeCompletedTask<void>(), []() -> void {});
+		}
+
+		virtual mu::FImageDesc GetImageDesc(FName Id, uint8 MipmapsToSkip) override
+		{
+			return IMAGE_DESC;
+		}
+
+
+		virtual TTuple<UE::Tasks::FTask, TFunction<void()>> GetReferencedImageAsync(const void* ModelPtr, int32 Id, uint8 MipmapsToSkip, TFunction<void(mu::Ptr<mu::Image>)>& ResultCallback)
+		{
+			check(ReferencedTextures.IsValidIndex(Id));
+
+			TSoftObjectPtr<const UTexture> TexturePtr = ReferencedTextures[Id].Get();
+			const UTexture2D* Texture = Cast<UTexture2D>( TexturePtr.Get() );
+			check(Texture);
+			
+			// In the editor the src data can be directly accessed
+			int32 MipIndex = (MipmapsToSkip < Texture->GetPlatformData()->Mips.Num()) ? MipmapsToSkip : Texture->GetPlatformData()->Mips.Num() - 1;
+			check(MipIndex >= 0);
+
+			mu::Ptr<mu::Image> ResultImage = new mu::Image();
+
+			FMutableSourceTextureData Tex(*Texture);
+			EUnrealToMutableConversionError Error = ConvertTextureUnrealSourceToMutable(ResultImage.get(), Tex, MipmapsToSkip);
+
+			if (Error != EUnrealToMutableConversionError::Success)
+			{
+				// This could happen in the editor, because some source textures may have changed while there was a background compilation.
+				// We just show a warning and move on. This cannot happen during cooks, so it is fine.
+				UE_LOG(LogMutable, Warning, TEXT("Failed to load some source texture data for [%s]. Some textures may be corrupted."), *Texture->GetName());
+			}
+
+			ResultCallback(ResultImage);
+
+			auto TrivialReturn = []() -> TTuple<UE::Tasks::FTask, TFunction<void()>>
+			{
+				return MakeTuple(UE::Tasks::MakeCompletedTask<void>(), []() -> void {});
+			};
+
+			return Invoke(TrivialReturn);
+		}
+
+
+		virtual TTuple<UE::Tasks::FTask, TFunction<void()>> GetMeshAsync(FName Id, TFunction<void(mu::Ptr<mu::Mesh>)>& ResultCallback) override
+		{
+			// TODO: This is a stub function where Mesh Parameters would be supported.
+
+			// Thread: worker
+			MUTABLE_CPUPROFILER_SCOPE(FUnrealMutableImageProvider::GetMeshAsync);
+
+			auto TrivialReturn = []() -> TTuple<UE::Tasks::FTask, TFunction<void()>>
+				{
+					return MakeTuple(UE::Tasks::MakeCompletedTask<void>(), []() -> void {});
+				};
+
+			// No USkeletalMesh was provided, cannot do anything, just provide a dummy mesh
+			UE_LOG(LogMutable, Warning, TEXT("No USkeletalMesh was provided for an application-specific mesh parameter."));
+
+			mu::Ptr<mu::Mesh> Dummy = new mu::Mesh;
+			ResultCallback(Dummy);
+			return Invoke(TrivialReturn);
+		}
+
+
+		virtual TTuple<UE::Tasks::FTask, TFunction<void()>> GetReferencedMeshAsync(const void* ModelPtr, int32 Id, TFunction<void(mu::Ptr<mu::Mesh>)>& ResultCallback) override
+		{
+			// TODO: This is a stub function where referenced Meshes would be supported.
+
+			// Thread: worker
+			MUTABLE_CPUPROFILER_SCOPE(FUnrealMutableImageProvider::GetReferencedMeshAsync);
+
+			auto TrivialReturn = []() -> TTuple<UE::Tasks::FTask, TFunction<void()>>
+				{
+					return MakeTuple(UE::Tasks::MakeCompletedTask<void>(), []() -> void {});
+				};
+
+			mu::Ptr<mu::Mesh> Mesh = new mu::Mesh();
+
+			// Not supported outside editor yet.
+			UE_LOG(LogMutable, Warning, TEXT("Failed to get referenced mesh. Not implemented yet."));
+
+			ResultCallback(Mesh);
+			return Invoke(TrivialReturn);
+
+		}
+
+	};
+}
+
+
+void SMutableCodeViewer::Tick(const FGeometry& AllottedGeometry, const double InCurrentTime, const float InDeltaTime)
+{
+	SCompoundWidget::Tick(AllottedGeometry, InCurrentTime, InDeltaTime);
+	
+	// After the tick we do know the tree has been refreshed, so all expansion and contraction operations have been
+	// completed and the new data has been loaded onto our listening arrays. Then its safe to expect the widgets to be
+	// there to be selected or inspected.
+	if (!TreeView->IsPendingRefresh())
+	{
+		/** If we have expanded the tree elements in order to reach one of them then continue the operation */
+		if (bWasUniqueExpansionInvokedForNavigation || bWasScrollToTargetRequested)
+		{
+			FocusViewOnNavigationTarget(ToFocusElement);
+		}
+	}
+	
+	
+	if (!bIsPreviewPendingUpdate)
+	{
+		return;
+	}
+
+	bIsPreviewPendingUpdate = false;
+
+	const mu::OP_TYPE OperationType = MutableModel->GetPrivate()->m_program.GetOpType(SelectedOperationAddress);
+	const mu::DATATYPE OperationDataType = mu::GetOpDataType(OperationType);
+
+	const mu::Ptr<mu::Settings> Settings = new mu::Settings();
+	const mu::Ptr<mu::System> System = new mu::System(Settings);
+
+	TSharedPtr<TestResourceProvider> ExternalResourceProvider = MakeShared<TestResourceProvider>();
+	ExternalResourceProvider->ReferencedTextures = ReferencedTextures;
+	System->SetExternalResourceProvider(ExternalResourceProvider);
+
+	System->GetPrivate()->BeginBuild(MutableModel);
+
+	switch (OperationDataType)
+	{
+	case mu::DT_LAYOUT:
+	{
+		check(PreviewLayoutViewer);
+		mu::Ptr<const mu::Layout> MutableLayout = System->GetPrivate()->BuildLayout(MutableModel, PreviewParameters.get(), SelectedOperationAddress);
+		PreviewLayoutViewer->SetLayout(MutableLayout);
+		break;
+	}
+
+	case mu::DT_IMAGE:
+	{
+		check(PreviewImageViewer);
+		mu::Ptr<const mu::Image> MutableImage = System->GetPrivate()->BuildImage(MutableModel, PreviewParameters.get(), SelectedOperationAddress, MipsToSkip, 0);
+		PreviewImageViewer->SetImage(MutableImage, 0);
+		break;
+	}
+
+	case mu::DT_MESH:
+	{
+		check(PreviewMeshViewer);
+		mu::Ptr<const mu::Mesh> MutableMesh = System->GetPrivate()->BuildMesh(MutableModel, PreviewParameters.get(), SelectedOperationAddress);
+		PreviewMeshViewer->SetMesh(MutableMesh);
+		break;
+	}
+
+	case mu::DT_INSTANCE:
+	{
+		check(PreviewInstanceViewer);
+		mu::Ptr<const mu::Instance> MutableInstance = System->GetPrivate()->BuildInstance(MutableModel, PreviewParameters.get(), SelectedOperationAddress);
+		PreviewInstanceViewer->SetInstance(MutableInstance, MutableModel, PreviewParameters, System);
+		break;
+	}
+
+	case mu::DT_BOOL:
+	{
+		check(PreviewBoolViewer);
+		const bool MutableBool = System->GetPrivate()->BuildBool(MutableModel, PreviewParameters.get(), SelectedOperationAddress);
+		PreviewBoolViewer->SetBool(MutableBool);
+		break;
+	}
+
+	case mu::DT_INT:
+	{
+		check(PreviewIntViewer);
+		const int32 MutableInt = System->GetPrivate()->BuildInt(MutableModel, PreviewParameters.get(), SelectedOperationAddress);
+		PreviewIntViewer->SetInt(MutableInt);
+		break;
+	}
+
+	case mu::DT_SCALAR:
+	{
+		check(PreviewScalarViewer);
+		const float MutableScalar = System->GetPrivate()->BuildScalar(MutableModel, PreviewParameters.get(), SelectedOperationAddress);
+		PreviewScalarViewer->SetScalar(MutableScalar);
+		break;
+	}
+
+	case mu::DT_STRING:
+	{
+		check(PreviewStringViewer);
+		mu::Ptr<const mu::String> MutableString = System->GetPrivate()->BuildString(MutableModel, PreviewParameters.get(), SelectedOperationAddress);
+		const FText MutableText = FText::FromString(FString(MutableString->GetValue()));
+		PreviewStringViewer->SetString(MutableText);
+		break;
+	}
+
+	case mu::DT_COLOUR:
+	{
+		check(PreviewColorViewer);
+		FVector4f Color = System->GetPrivate()->BuildColour(MutableModel, PreviewParameters.get(), SelectedOperationAddress);
+		PreviewColorViewer->SetColor(Color);
+		break;
+	}
+
+	case mu::DT_PROJECTOR:
+	{
+		check (PreviewProjectorViewer);
+		mu::FProjector Projector = System->GetPrivate()->BuildProjector(MutableModel, PreviewParameters.get(), SelectedOperationAddress);
+		PreviewProjectorViewer->SetProjector(Projector);
+	}
+	
+	default:
+#if UE_BUILD_DEBUG || UE_BUILD_DEVELOPMENT
+		UE_LOG(LogMutable, Log, TEXT("There is no previewer for the selected type of Mutable object"))
+#endif
+		// There is no viewer for this type.
+		break;
+	}
+
+	System->GetPrivate()->EndBuild();
+
+}
+
+
+void SMutableCodeViewer::OnPreviewParameterValueChanged(int32 ParamIndex)
+{
+	// This is deferred to the tick to avoid multiple updates per frame.
+	bIsPreviewPendingUpdate = true;
+}
+
+
+void SMutableCodeViewer::PrepareStringViewer()
+{
+	if (!PreviewStringViewer)
+	{
+		PreviewStringViewer = SNew(SMutableStringViewer);
+	}
+
+	PreviewBorder->SetContent(PreviewStringViewer.ToSharedRef());
+}
+
+void SMutableCodeViewer::PrepareImageViewer()
+{
+	if (!PreviewImageViewer)
+	{
+		PreviewImageViewer = SNew(SMutableImageViewer)
+			.GridSize(FIntPoint(8, 8));
+	}
+
+	PreviewBorder->SetContent(PreviewImageViewer.ToSharedRef());
+}
+
+void SMutableCodeViewer::PrepareMeshViewer()
+{
+	if (!PreviewMeshViewer)
+	{
+		PreviewMeshViewer = SNew(SMutableMeshViewer);
+	}
+
+	PreviewBorder->SetContent(PreviewMeshViewer.ToSharedRef());
+}
+
+void SMutableCodeViewer::PrepareInstanceViewer()
+{
+	if (!PreviewInstanceViewer)
+	{
+		PreviewInstanceViewer = SNew(SMutableInstanceViewer);
+	}
+
+	PreviewBorder->SetContent(PreviewInstanceViewer.ToSharedRef());
+}
+
+void SMutableCodeViewer::PrepareLayoutViewer()
+{
+	if (!PreviewLayoutViewer)
+	{
+		PreviewLayoutViewer = SNew(SMutableLayoutViewer);
+	}
+
+	PreviewBorder->SetContent(PreviewLayoutViewer.ToSharedRef());
+}
+
+
+void SMutableCodeViewer::PrepareProjectorViewer()
+{
+	if (!PreviewProjectorViewer)
+	{
+		PreviewProjectorViewer = SNew(SMutableProjectorViewer);
+	}
+
+	PreviewBorder->SetContent(PreviewProjectorViewer.ToSharedRef());
+}
+
+
+void SMutableCodeViewer::PreviewMutableString(const FString& InString)
+{
+	// Prepare the previewer object to receive data 
+	PrepareStringViewer();
+	
+	//  Provide the desired data to the previewer object
+	const FText TextToShow = FText::FromString(InString);
+	PreviewStringViewer->SetString(TextToShow);
+}
+
+
+void SMutableCodeViewer::PreviewMutableImage(mu::ImagePtrConst InImagePtr)
+{
+	PrepareImageViewer();
+	PreviewImageViewer->SetImage(InImagePtr,0);
+}
+
+
+void SMutableCodeViewer::PreviewMutableMesh(mu::MeshPtrConst InMeshPtr)
+{
+	PrepareMeshViewer();
+	PreviewMeshViewer->SetMesh(InMeshPtr);
+}
+
+
+void SMutableCodeViewer::PreviewMutableLayout(mu::Ptr<const mu::Layout> Layout)
+{
+	PrepareLayoutViewer();
+	PreviewLayoutViewer->SetLayout(Layout);
+}
+
+
+void SMutableCodeViewer::PreviewMutableProjector(const mu::FProjector* Projector)
+{
+	if (!Projector)
+	{
+		UE_LOG(LogTemp,Error,TEXT("Unable to preview data on null Projector pointer."))
+		return;
+	}
+	
+	PrepareProjectorViewer();
+	PreviewProjectorViewer->SetProjector(*Projector);
+}
+
+
+void SMutableCodeViewer::PreviewMutableSkeleton(mu::SkeletonPtrConst Skeleton)
+{
+	if (!PreviewSkeletonViewer)
+	{
+		PreviewSkeletonViewer = SNew(SMutableSkeletonViewer);
+	}
+
+	PreviewBorder->SetContent(PreviewSkeletonViewer.ToSharedRef());
+	
+	PreviewSkeletonViewer->SetSkeleton(Skeleton);
+}
+
+
+void SMutableCodeViewer::PreviewMutableCurve(const FRichCurve& Curve)
+{
+	if (!PreviewCurveViewer)
+	{
+		PreviewCurveViewer = SNew(SMutableCurveViewer);
+	}
+
+	PreviewBorder->SetContent(PreviewCurveViewer.ToSharedRef());
+	
+	PreviewCurveViewer->SetCurve(Curve);
+}
+
+// TODO: Implement matrix viewer
+void SMutableCodeViewer::PreviewMutableMatrix(const FMatrix44f& Mat)
+{
+	UE_LOG(LogMutable, Warning, TEXT("Previewer for Mutable Matrices not yet implemented"))
+}
+
+// TODO: Implement shape viewer
+void SMutableCodeViewer::PreviewMutableShape(const mu::FShape* Shape)
+{
+	UE_LOG(LogMutable, Warning, TEXT("Previewer for Mutable Shapes not yet implemented"))
+}
+
+
+FMutableCodeTreeElement::FMutableCodeTreeElement(int32 InIndexOnTree, const int32& InMutableStateIndex, const TSharedPtr<mu::Model, ESPMode::ThreadSafe>& InModel, mu::OP::ADDRESS InOperation, const FString& InCaption, const FSlateColor InLabelColor, const TSharedPtr<FMutableCodeTreeElement>* InDuplicatedOf)
+{
+	MutableModel = InModel;
+	MutableOperation = InOperation;
+	Caption = InCaption;
+	LabelColor = InLabelColor;
+	IndexOnTree = InIndexOnTree;
+	if (InDuplicatedOf)
+	{
+		DuplicatedOf = *InDuplicatedOf;
+	}
+
+	// Generate the label to be used to display this operation in the operation tree
+	GenerateLabelText();
+
+	// Process the data that can be extracted from the current state
+	SetElementCurrentState(InMutableStateIndex);
+}
+
+
+void FMutableCodeTreeElement::SetElementCurrentState(const int32& InStateIndex)
+{
+	// Skip operation if state is the same
+	if (InStateIndex == CurrentMutableStateIndex)
+	{
+		return;
+	}
+
+	// Check for an out of bounds value
+	check(MutableModel);
+	mu::FProgram& MutableProgram = MutableModel->GetPrivate()->m_program;
+	check(InStateIndex >= 0 && InStateIndex < MutableProgram.m_states.Num());
+
+	CurrentMutableStateIndex = InStateIndex;
+	const mu::FProgram::FState& CurrentState = MutableProgram.m_states[CurrentMutableStateIndex];
+
+	// Check if it is a dynamic resource
+	for (auto& DynamicResource : CurrentState.m_dynamicResources)
+	{
+		// If the operation gets located then mark it as dynamic resource
+		if (DynamicResource.Key == MutableOperation)
+		{
+			bIsDynamicResource = true;
+			break;
+		}
+	}
+	// Early exit: A dynamic resource can not be at the same time a state constant
+	if (bIsDynamicResource)
+	{
+		return;
+	}
+
+	// Check if it is a state constant
+	bIsStateConstant = CurrentState.m_updateCache.Contains(MutableOperation);
+}
+
+
+void FMutableCodeTreeElement::GenerateLabelText()
+{
+	const mu::FProgram& Program = MutableModel->GetPrivate()->m_program;
+	const mu::OP_TYPE Type = Program.GetOpType(MutableOperation);
+	FString OpName = mu::s_opNames[static_cast<int32>(Type)];
+	OpName.TrimEndInline();
+
+	// See if the operation type accepts additional information in the label
+	switch (Type)
+	{
+	case mu::OP_TYPE::BO_PARAMETER:
+	case mu::OP_TYPE::NU_PARAMETER:
+	case mu::OP_TYPE::SC_PARAMETER:
+	case mu::OP_TYPE::CO_PARAMETER:
+	case mu::OP_TYPE::PR_PARAMETER:
+	case mu::OP_TYPE::IM_PARAMETER:
+	case mu::OP_TYPE::ST_PARAMETER:
+	{
+		mu::OP::ParameterArgs Args = Program.GetOpArgs<mu::OP::ParameterArgs>(MutableOperation);
+		OpName += TEXT(" ");
+		OpName += Program.m_parameters[int32(Args.variable)].m_name;
+		break;
+	}
+
+	case mu::OP_TYPE::IM_SWIZZLE:
+	{
+		mu::OP::ImageSwizzleArgs Args = Program.GetOpArgs<mu::OP::ImageSwizzleArgs>(MutableOperation);
+		OpName += TEXT(" ");
+		OpName += StringCast<TCHAR>(mu::TypeInfo::s_imageFormatName[int32(Args.format)]).Get();
+		break;
+	}
+
+	case mu::OP_TYPE::IM_PIXELFORMAT:
+	{
+		mu::OP::ImagePixelFormatArgs Args = Program.GetOpArgs<mu::OP::ImagePixelFormatArgs>(MutableOperation);
+		OpName += TEXT(" ");
+		OpName += StringCast<TCHAR>(mu::TypeInfo::s_imageFormatName[int32(Args.format)]).Get();
+		OpName += TEXT(" or ");
+		OpName += StringCast<TCHAR>(mu::TypeInfo::s_imageFormatName[int32(Args.formatIfAlpha)]).Get();
+		break;
+	}
+
+	case mu::OP_TYPE::IM_MIPMAP:
+	{
+		mu::OP::ImageMipmapArgs Args = Program.GetOpArgs<mu::OP::ImageMipmapArgs>(MutableOperation);
+		OpName += FString::Printf(TEXT(" levels: %d-%d tail: %d"), Args.levels, Args.blockLevels, int32(Args.onlyTail));
+		break;
+	}
+
+	case mu::OP_TYPE::IM_RESIZE:
+	{
+		mu::OP::ImageResizeArgs Args = Program.GetOpArgs<mu::OP::ImageResizeArgs>(MutableOperation);
+		OpName += FString::Printf(TEXT(" %d x %d"), int32(Args.size[0]), int32(Args.size[1]));
+		break;
+	}
+
+	case mu::OP_TYPE::IM_RESIZEREL:
+	{
+		mu::OP::ImageResizeRelArgs Args = Program.GetOpArgs<mu::OP::ImageResizeRelArgs>(MutableOperation);
+		OpName += FString::Printf(TEXT(" %.3f x %.3f"), Args.factor[0], Args.factor[1]);
+		break;
+	}
+
+	case mu::OP_TYPE::IM_MULTILAYER:
+	{
+		mu::OP::ImageMultiLayerArgs Args = Program.GetOpArgs<mu::OP::ImageMultiLayerArgs>(MutableOperation);
+		OpName += TEXT(" rgb: ");
+		OpName += mu::TypeInfo::s_blendModeName[int32(Args.blendType)];
+		OpName += TEXT(", a: ");
+		OpName += mu::TypeInfo::s_blendModeName[int32(Args.blendTypeAlpha)];
+		OpName += FString::Printf(TEXT(" a from %d "), Args.BlendAlphaSourceChannel);
+		OpName += FString::Printf(TEXT(" range-id: %d"), Args.rangeId);
+		OpName += FString::Printf(TEXT(" mask-from-alpha: %d"), int32(Args.bUseMaskFromBlended));
+		break;
+	}
+
+	case mu::OP_TYPE::IM_LAYER:
+	{
+		mu::OP::ImageLayerArgs Args = Program.GetOpArgs<mu::OP::ImageLayerArgs>(MutableOperation);
+		OpName += TEXT(" rgb: ");
+		OpName += mu::TypeInfo::s_blendModeName[int32(Args.blendType)];
+		OpName += TEXT(", a: ");
+		OpName += mu::TypeInfo::s_blendModeName[int32(Args.blendTypeAlpha)];
+		OpName += FString::Printf(TEXT(" a from %d "), Args.BlendAlphaSourceChannel);
+		OpName += FString::Printf(TEXT(" flags %d"), Args.flags);
+		break;
+	}
+
+	case mu::OP_TYPE::IM_LAYERCOLOUR:
+	{
+		mu::OP::ImageLayerColourArgs Args = Program.GetOpArgs<mu::OP::ImageLayerColourArgs>(MutableOperation);
+		OpName += TEXT(" rgb: ");
+		OpName += mu::TypeInfo::s_blendModeName[int32(Args.blendType)];
+		OpName += TEXT(" a: ");
+		OpName += mu::TypeInfo::s_blendModeName[int32(Args.blendTypeAlpha)];
+		OpName += TEXT(" a from ");
+		OpName += FString::Printf(TEXT(" a from %d "), Args.BlendAlphaSourceChannel);
+		OpName += FString::Printf(TEXT(" flags %d"), Args.flags);
+		break;
+	}
+
+	case mu::OP_TYPE::IM_PLAINCOLOUR:
+	{
+		mu::OP::ImagePlainColourArgs Args = Program.GetOpArgs<mu::OP::ImagePlainColourArgs>(MutableOperation);
+		OpName += TEXT(" format: ");
+		OpName += StringCast<TCHAR>(mu::TypeInfo::s_imageFormatName[int32(Args.format)]).Get();
+		OpName += FString::Printf(TEXT(" size %d x %d"), Args.size[0], Args.size[1]);
+		OpName += FString::Printf(TEXT(" mips %d"), Args.LODs);
+		break;
+	}
+
+	case mu::OP_TYPE::IN_ADDIMAGE:
+	{
+		mu::OP::InstanceAddArgs Args = Program.GetOpArgs<mu::OP::InstanceAddArgs>(MutableOperation);
+		if (Program.m_constantStrings.IsValidIndex(Args.name))
+		{
+			OpName += TEXT(" name: ");
+			OpName += Program.m_constantStrings[Args.name];
+		}
+		break;
+	}
+
+	default:
+		break;
+	}
+
+	// Prepare the text shown on the UI side of the operation tree
+	MainLabel = FString::Printf(TEXT("%s %d : %s"), *Caption, int32(MutableOperation), *OpName);
+
+	// DEBUG : 
+	// FString IndexOnTree = FString::FromInt(IndexOnTree);
+	// IndexOnTree.Append(TEXT("- "));
+	// MainLabel.InsertAt(0,IndexOnTree);
+
+	// DEBUG : 
+	// FString RowStateIndex = FString::FromInt(GetStateIndex());
+	// RowStateIndex.Append(TEXT(" st "));
+	// MainLabel.InsertAt(0,RowStateIndex);
+
+	if (DuplicatedOf)
+	{
+		MainLabel.Append(TEXT(" (duplicated)"));
+	}
+}
+
+
+int32 FMutableCodeTreeElement::GetStateIndex() const
+{
+	return CurrentMutableStateIndex;
+}
+
+
+#undef LOCTEXT_NAMESPACE
