[mutable] Moved the Mutable plugin out of Experimental status into Beta.

diff --git a/Engine/Plugins/Mutable/Source/CustomizableObjectEditor/Private/MuCOE/SMutableGraphViewer.cpp b/Engine/Plugins/Mutable/Source/CustomizableObjectEditor/Private/MuCOE/SMutableGraphViewer.cpp
new file mode 100644
--- /dev/null
+++ b/Engine/Plugins/Mutable/Source/CustomizableObjectEditor/Private/MuCOE/SMutableGraphViewer.cpp
@@ -0,0 +1,808 @@
+// Copyright Epic Games, Inc. All Rights Reserved.
+
+#include "MuCOE/SMutableGraphViewer.h"
+
+#include "DesktopPlatformModule.h"
+#include "EditorDirectories.h"
+#include "Framework/MultiBox/MultiBoxBuilder.h"
+#include "Framework/Views/TableViewMetadata.h"
+#include "IDesktopPlatform.h"
+#include "Misc/Paths.h"
+#include "MuCOE/CustomizableObjectCompileRunnable.h"
+#include "MuCOE/CustomizableObjectEditorStyle.h"
+#include "MuCOE/SMutableCodeViewer.h"
+#include "MuCOE/UnrealEditorPortabilityHelpers.h"
+#include "Widgets/Docking/SDockTab.h"
+#include "Widgets/Input/STextComboBox.h"
+#include "Widgets/Input/SNumericDropDown.h"
+#include "Widgets/Views/STreeView.h"
+#include "ScopedTransaction.h"
+#include "MuT/NodeColourConstant.h"
+#include "MuT/NodeColourFromScalars.h"
+#include "MuT/NodeColourParameter.h"
+#include "MuT/NodeColourSampleImage.h"
+#include "MuT/NodeColourSwitch.h"
+#include "MuT/NodeComponentEdit.h"
+#include "MuT/NodeImageFormat.h"
+#include "MuT/NodeImageInterpolate.h"
+#include "MuT/NodeImageInvert.h"
+#include "MuT/NodeImageLayer.h"
+#include "MuT/NodeImageLayerColour.h"
+#include "MuT/NodeImageMipmap.h"
+#include "MuT/NodeImageMultiLayer.h"
+#include "MuT/NodeImagePlainColour.h"
+#include "MuT/NodeImageProject.h"
+#include "MuT/NodeImageResize.h"
+#include "MuT/NodeImageSwitch.h"
+#include "MuT/NodeImageSwizzle.h"
+#include "MuT/NodeImageTable.h"
+#include "MuT/NodeObjectGroup.h"
+#include "MuT/NodeObjectNew.h"
+#include "MuT/NodeSurfaceNew.h"
+#include "MuT/NodeSurfaceSwitch.h"
+#include "MuT/NodeSurfaceVariation.h"
+#include "MuT/NodeLOD.h"
+#include "MuT/NodeMeshConstant.h"
+#include "MuT/NodeMeshFormat.h"
+#include "MuT/NodeMeshFragment.h"
+#include "MuT/NodeMeshMakeMorph.h"
+#include "MuT/NodeMeshMorph.h"
+#include "MuT/NodeMeshTable.h"
+#include "MuT/NodeModifierMeshClipDeform.h"
+#include "MuT/NodeModifierMeshClipMorphPlane.h"
+#include "MuT/NodeModifierMeshClipWithUVMask.h"
+#include "MuT/NodeModifierMeshClipMorphPlane.h"
+#include "MuT/NodeModifierSurfaceEdit.h"
+#include "MuT/NodeScalarConstant.h"
+#include "MuT/NodeScalarCurve.h"
+#include "MuT/NodeScalarSwitch.h"
+#include "MuT/NodeScalarTable.h"
+#include "MuT/NodeObjectNew.h"
+
+#include "MuT/NodeObjectGroupPrivate.h"
+#include "MuT/NodeImageFormatPrivate.h"
+#include "MuT/NodeMeshFormatPrivate.h"
+#include "MuT/NodeMeshConstantPrivate.h"
+#include "MuT/NodeImageSwitchPrivate.h"
+#include "MuT/NodeImageLayerColourPrivate.h"
+#include "MuT/NodeImageLayerPrivate.h"
+#include "MuT/NodeImageMipmapPrivate.h"
+#include "MuT/NodeImageResizePrivate.h"
+#include "MuT/NodeImageInterpolatePrivate.h"
+#include "MuT/NodeImagePlainColourPrivate.h"
+#include "MuT/NodeImageProjectPrivate.h"
+#include "MuT/NodeMeshMorphPrivate.h"
+#include "MuT/NodeScalarEnumParameterPrivate.h"
+#include "MuT/NodeScalarParameterPrivate.h"
+#include "MuT/NodeMeshMakeMorphPrivate.h"
+#include "MuT/NodeScalarCurvePrivate.h"
+#include "MuT/NodeProjectorPrivate.h"
+#include "MuT/NodeImageInvertPrivate.h"
+#include "MuT/NodeImageSwizzlePrivate.h"
+#include "MuT/NodeImageMultiLayerPrivate.h"
+#include "Widgets/MutableExpanderArrow.h"
+
+class FExtender;
+class FReferenceCollector;
+class FUICommandList;
+class ITableRow;
+class SWidget;
+struct FGeometry;
+struct FSlateBrush;
+
+#define LOCTEXT_NAMESPACE "SMutableDebugger"
+
+// \todo: multi-column tree
+namespace MutableGraphTreeViewColumns
+{
+	static const FName Name("Name");
+};
+
+
+class SMutableGraphTreeRow : public STableRow<TSharedPtr<FMutableGraphTreeElement>>
+{
+public:
+
+	void Construct(const FArguments& Args, const TSharedRef<STableViewBase>& InOwnerTableView, const TSharedPtr<FMutableGraphTreeElement>& InRowItem)
+	{
+		RowItem = InRowItem;
+
+		FText MainLabel = FText::GetEmpty();
+		if (RowItem->MutableNode)
+		{
+			mu::Node::EType MutableType = RowItem->MutableNode->GetType()->Type;
+			
+			const FString LabelString = RowItem->Prefix.IsEmpty() 
+				? FString::Printf(TEXT("%d"), int32(MutableType) )
+				: FString::Printf( TEXT("%s : %d"), *RowItem->Prefix, int32(MutableType) );
+
+			MainLabel = FText::FromString(LabelString);
+			if (RowItem->DuplicatedOf)
+			{
+				MainLabel = FText::FromString( FString::Printf(TEXT("%d (Duplicated)"), int32(MutableType)));
+			}
+		}
+		else
+		{
+			MainLabel = FText::FromString( *RowItem->Prefix);
+		}
+
+
+		this->ChildSlot
+		[
+			SNew(SHorizontalBox)
+			+ SHorizontalBox::Slot()
+			.VAlign(VAlign_Center)
+			.AutoWidth()
+			[
+				SNew(SMutableExpanderArrow, SharedThis(this))
+			]
+
+			+ SHorizontalBox::Slot()
+			[
+				SNew(STextBlock)
+				.Text(MainLabel)
+			]
+		];
+
+		STableRow< TSharedPtr<FMutableGraphTreeElement> >::ConstructInternal(
+			STableRow::FArguments()
+			//.Style(FAppStyle::Get(), "DetailsView.TreeView.TableRow")
+			.ShowSelection(true)
+			, InOwnerTableView
+		);
+
+	}
+
+
+private:
+
+	TSharedPtr<FMutableGraphTreeElement> RowItem;
+};
+
+
+void SMutableGraphViewer::AddReferencedObjects(FReferenceCollector& Collector)
+{
+	// Add UObjects here if we own any at some point
+	//Collector.AddReferencedObject(CustomizableObject);
+}
+
+
+FString SMutableGraphViewer::GetReferencerName() const
+{
+	return TEXT("SMutableGraphViewer");
+}
+
+
+void SMutableGraphViewer::Construct(const FArguments& InArgs, const mu::NodePtr& InRootNode)
+{
+	DataTag = InArgs._DataTag;
+	ReferencedRuntimeTextures = InArgs._ReferencedRuntimeTextures;
+	ReferencedCompileTextures = InArgs._ReferencedCompileTextures;
+	RootNode = InRootNode;
+
+	FToolBarBuilder ToolbarBuilder(TSharedPtr<const FUICommandList>(), FMultiBoxCustomization::None, TSharedPtr<FExtender>(), true);
+	ToolbarBuilder.SetLabelVisibility(EVisibility::Visible);
+	ToolbarBuilder.SetStyle(&FAppStyle::Get(), "SlimToolBar");
+
+	ToolbarBuilder.AddWidget(SNew(STextBlock).Text(MakeAttributeLambda([this]() { return FText::FromString(DataTag); })));
+
+	ChildSlot
+	[
+		SNew(SVerticalBox)
+		+ SVerticalBox::Slot()
+		.AutoHeight()
+		.VAlign(VAlign_Center)
+		[
+			ToolbarBuilder.MakeWidget()
+		]
+		+ SVerticalBox::Slot()
+		.VAlign(VAlign_Fill)
+		[
+			SNew(SSplitter)
+			.Orientation(EOrientation::Orient_Horizontal)
+			+ SSplitter::Slot()
+			.Value(0.25f)
+			[
+				SNew(SBorder)
+				.BorderImage(UE_MUTABLE_GET_BRUSH("ToolPanel.GroupBorder"))
+				.Padding(FMargin(4.0f, 4.0f))
+				[
+					SAssignNew(TreeView, STreeView<TSharedPtr<FMutableGraphTreeElement>>)
+					.TreeItemsSource(&RootNodes)
+					.OnGenerateRow(this,&SMutableGraphViewer::GenerateRowForNodeTree)
+					.OnGetChildren(this, &SMutableGraphViewer::GetChildrenForInfo)
+					.OnSetExpansionRecursive(this, &SMutableGraphViewer::TreeExpandRecursive)
+					.OnContextMenuOpening(this, &SMutableGraphViewer::OnTreeContextMenuOpening)
+					.SelectionMode(ESelectionMode::Single)
+					.HeaderRow
+					(
+						SNew(SHeaderRow)
+						+ SHeaderRow::Column(MutableGraphTreeViewColumns::Name)
+						.FillWidth(25.f)
+						.DefaultLabel(LOCTEXT("Node Name", "Node Name"))
+					)
+				]
+			]
+			+ SSplitter::Slot()
+			.Value(0.75f)
+			[
+				SNew(SBorder)
+				.BorderImage(UE_MUTABLE_GET_BRUSH("ToolPanel.GroupBorder"))
+				.Padding(FMargin(4.0f, 4.0f))
+				//[
+				//	SubjectsTreeView->AsShared()
+				//]
+			]
+		]
+	];
+	
+	RebuildTree();
+}
+
+
+void SMutableGraphViewer::RebuildTree()
+{
+	RootNodes.Reset();
+	ItemCache.Reset();
+	MainItemPerNode.Reset();
+
+	RootNodes.Add(MakeShareable(new FMutableGraphTreeElement(RootNode)));
+	TreeView->RequestTreeRefresh();
+	TreeExpandUnique();
+}
+
+
+TSharedRef<ITableRow> SMutableGraphViewer::GenerateRowForNodeTree(TSharedPtr<FMutableGraphTreeElement> InTreeNode, const TSharedRef<STableViewBase>& InOwnerTable)
+{
+	TSharedRef<SMutableGraphTreeRow> Row = SNew(SMutableGraphTreeRow, InOwnerTable, InTreeNode);
+	return Row;
+}
+
+void SMutableGraphViewer::GetChildrenForInfo(TSharedPtr<FMutableGraphTreeElement> InInfo, TArray<TSharedPtr<FMutableGraphTreeElement>>& OutChildren)
+{
+// This is necessary because of problems with rtti information in other platforms. In any case, this part of the debugger is only useful in the standard editor.
+#if PLATFORM_WINDOWS
+	if (!InInfo->MutableNode)
+	{
+		return;
+	}
+
+	// If this is a duplicated of another row, don't provide its children.
+	if (InInfo->DuplicatedOf)
+	{
+		return;
+	}
+
+	mu::Node* ParentNode = InInfo->MutableNode.get();
+	uint32 InputIndex = 0;
+
+	auto AddChildFunc = [this, ParentNode, &InputIndex, &OutChildren](mu::Node* ChildNode, const FString& Prefix)
+	{
+		if (ChildNode)
+		{
+			FItemCacheKey Key = { ParentNode, ChildNode, InputIndex };
+			TSharedPtr<FMutableGraphTreeElement>* CachedItem = ItemCache.Find(Key);
+
+			if (CachedItem)
+			{
+				OutChildren.Add(*CachedItem);
+			}
+			else
+			{
+				TSharedPtr<FMutableGraphTreeElement>* MainItemPtr = MainItemPerNode.Find(ChildNode);
+				TSharedPtr<FMutableGraphTreeElement> Item = MakeShareable(new FMutableGraphTreeElement(ChildNode, MainItemPtr, Prefix));
+				OutChildren.Add(Item);
+				ItemCache.Add(Key, Item);
+
+				if (!MainItemPtr)
+				{
+					MainItemPerNode.Add(ChildNode, Item);
+				}
+			}
+		}
+		else
+		{
+			// No mutable node has been provided so create a dummy tree element
+			TSharedPtr<FMutableGraphTreeElement> Item = MakeShareable(new FMutableGraphTreeElement(nullptr, nullptr , Prefix));
+			OutChildren.Add(Item);
+		}
+		++InputIndex;
+	};
+
+	if (ParentNode->GetType() == mu::NodeObjectNew::GetStaticType())
+	{
+		mu::NodeObjectNew* ObjectNew = StaticCast<mu::NodeObjectNew*>(ParentNode);
+		for (int32 l = 0; l < ObjectNew->Components.Num(); ++l)
+		{
+			AddChildFunc(ObjectNew->Components[l].get(), TEXT("COMP") );
+		}
+
+		for (int32 Modifier = 0; Modifier < ObjectNew->Modifiers.Num(); Modifier++)
+		{
+			AddChildFunc(ObjectNew->Modifiers[Modifier].get(), FString::Printf(TEXT("MOD [%d]"), Modifier));
+		}
+
+		for (int32 l = 0; l < ObjectNew->Children.Num(); ++l)
+		{
+			AddChildFunc(ObjectNew->Children[l].get(), TEXT("CHILD"));
+		}
+	}
+
+	else if (ParentNode->GetType() == mu::NodeObjectGroup::GetStaticType())
+	{
+		mu::NodeObjectGroup* ObjectGroup = StaticCast<mu::NodeObjectGroup*>(ParentNode);
+		mu::NodeObjectGroup::Private* Private = ObjectGroup->GetPrivate();
+		for (int32 l = 0; l < Private->m_children.Num(); ++l)
+		{
+			AddChildFunc(Private->m_children[l].get(), TEXT("CHILD"));
+		}
+	}
+
+	else if (ParentNode->GetType() == mu::NodeSurfaceNew::GetStaticType())
+	{
+		mu::NodeSurfaceNew* SurfaceNew = StaticCast<mu::NodeSurfaceNew*>(ParentNode);
+		AddChildFunc(SurfaceNew->Mesh.get(), TEXT("MESH"));
+
+		for (int32 l = 0; l < SurfaceNew->Images.Num(); ++l)
+		{
+			AddChildFunc(SurfaceNew->Images[l].Image.get(), FString::Printf(TEXT("IMAGE [%s]"), *SurfaceNew->Images[l].Name));
+		}
+
+		for (int32 l = 0; l < SurfaceNew->Vectors.Num(); ++l)
+		{
+			AddChildFunc(SurfaceNew->Vectors[l].Vector.get(), FString::Printf(TEXT("VECTOR [%s]"), *SurfaceNew->Vectors[l].Name));
+		}
+
+		for (int32 l = 0; l < SurfaceNew->Scalars.Num(); ++l)
+		{
+			AddChildFunc(SurfaceNew->Scalars[l].Scalar.get(), FString::Printf(TEXT("SCALAR [%s]"), *SurfaceNew->Scalars[l].Name));
+		}
+
+		for (int32 l = 0; l < SurfaceNew->Strings.Num(); ++l)
+		{
+			AddChildFunc(SurfaceNew->Strings[l].String.get(), FString::Printf(TEXT("STRING [%s]"), *SurfaceNew->Strings[l].Name));
+		}
+	}
+
+	else if (ParentNode->GetType() == mu::NodeModifierSurfaceEdit::GetStaticType())
+	{
+		mu::NodeModifierSurfaceEdit* SurfaceEdit = StaticCast<mu::NodeModifierSurfaceEdit*>(ParentNode);
+		AddChildFunc(SurfaceEdit->MeshAdd.get(), TEXT("MESH_ADD"));
+		AddChildFunc(SurfaceEdit->MeshRemove.get(), TEXT("MESH_REMOVE"));
+		AddChildFunc(SurfaceEdit->MorphFactor.get(), FString::Printf(TEXT("MORPH_FACTOR [%s]"), *SurfaceEdit->MeshMorph) );
+
+		for (int32 l = 0; l < SurfaceEdit->Textures.Num(); ++l)
+		{
+			AddChildFunc(SurfaceEdit->Textures[l].Extend.get(), FString::Printf(TEXT("EXTEND [%d]"), l));
+			AddChildFunc(SurfaceEdit->Textures[l].PatchImage.get(), FString::Printf(TEXT("PATCH IMAGE [%d]"), l));
+			AddChildFunc(SurfaceEdit->Textures[l].PatchMask.get(), FString::Printf(TEXT("PATCH MASK [%d]"), l));
+
+		}
+	}
+
+	else if (ParentNode->GetType() == mu::NodeSurfaceSwitch::GetStaticType())
+	{
+		mu::NodeSurfaceSwitch* SurfaceSwitch = StaticCast<mu::NodeSurfaceSwitch*>(ParentNode);
+		AddChildFunc(SurfaceSwitch->Parameter.get(), TEXT("PARAM"));
+		for (int32 l = 0; l < SurfaceSwitch->Options.Num(); ++l)
+		{
+			AddChildFunc(SurfaceSwitch->Options[l].get(), FString::Printf(TEXT("OPTION [%d]"), l));
+		}
+	}
+
+	else if (ParentNode->GetType() == mu::NodeSurfaceVariation::GetStaticType())
+	{
+		mu::NodeSurfaceVariation* SurfaceVar = StaticCast<mu::NodeSurfaceVariation*>(ParentNode);
+		for (int32 l = 0; l < SurfaceVar->DefaultSurfaces.Num(); ++l)
+		{
+			AddChildFunc(SurfaceVar->DefaultSurfaces[l].get(), FString::Printf(TEXT("DEF SURF [%d]"), l));
+		}
+		for (int32 l = 0; l < SurfaceVar->DefaultModifiers.Num(); ++l)
+		{
+			AddChildFunc(SurfaceVar->DefaultModifiers[l].get(), FString::Printf(TEXT("DEF MOD [%d]"), l));
+		}
+
+		for (int32 v = 0; v < SurfaceVar->Variations.Num(); ++v)
+		{
+			const mu::NodeSurfaceVariation::FVariation Var = SurfaceVar->Variations[v];
+			for (int32 l = 0; l < Var.Surfaces.Num(); ++l)
+			{
+				AddChildFunc(Var.Surfaces[l].get(), FString::Printf(TEXT("VAR [%s] SURF [%d]"), *Var.Tag, l));
+			}
+			for (int32 l = 0; l < Var.Modifiers.Num(); ++l)
+			{
+				AddChildFunc(Var.Modifiers[l].get(), FString::Printf(TEXT("VAR [%s] MOD [%d]"), *Var.Tag, l));
+			}
+		}
+	}
+	
+	else if (ParentNode->GetType() == mu::NodeLOD::GetStaticType())
+	{
+		mu::NodeLOD* LodVar = StaticCast<mu::NodeLOD*>(ParentNode);
+
+		for (int32 SurfaceIndex = 0; SurfaceIndex < LodVar->Surfaces.Num(); SurfaceIndex++)
+		{
+			AddChildFunc(LodVar->Surfaces[SurfaceIndex].get(), FString::Printf(TEXT("SURFACE [%d]"), SurfaceIndex));
+		}
+	}
+	
+	else if (ParentNode->GetType() == mu::NodeComponentNew::GetStaticType())
+	{
+		mu::NodeComponentNew* ComponentVar = StaticCast<mu::NodeComponentNew*>(ParentNode);
+		for (int32 LODIndex = 0; LODIndex < ComponentVar->LODs.Num(); LODIndex++)
+		{
+			AddChildFunc(ComponentVar->LODs[LODIndex].get(), FString::Printf(TEXT("LOD [%d]"), LODIndex));
+		}
+	}
+
+	else if (ParentNode->GetType() == mu::NodeComponentEdit::GetStaticType())
+	{
+		mu::NodeComponentEdit* ComponentVar = StaticCast<mu::NodeComponentEdit*>(ParentNode);
+		for (int32 LODIndex = 0; LODIndex < ComponentVar->LODs.Num(); LODIndex++)
+		{
+			AddChildFunc(ComponentVar->LODs[LODIndex].get(), FString::Printf(TEXT("LOD [%d]"), LODIndex));
+		}
+	}
+
+	else if (ParentNode->GetType() == mu::NodeMeshConstant::GetStaticType())
+	{
+		mu::NodeMeshConstant* MeshConstantVar = StaticCast<mu::NodeMeshConstant*>(ParentNode);
+		mu::NodeMeshConstant::Private* Private = MeshConstantVar->GetPrivate();
+		for (int32 LayoutIndex = 0; LayoutIndex < Private->Layouts.Num(); LayoutIndex++)
+		{
+			AddChildFunc(Private->Layouts[LayoutIndex].get(), FString::Printf(TEXT("LAYOUT [%d]"), LayoutIndex));
+		}
+	}
+
+	else if (ParentNode->GetType() == mu::NodeImageFormat::GetStaticType())
+	{
+		mu::NodeImageFormat* ImageFormatVar = StaticCast<mu::NodeImageFormat*>(ParentNode);
+		mu::NodeImageFormat::Private* Private = ImageFormatVar->GetPrivate();
+		AddChildFunc(Private->m_source.get(), FString::Printf(TEXT("SOURCE IMAGE")));
+	}
+
+	else if (ParentNode->GetType() == mu::NodeMeshFormat::GetStaticType())
+	{
+		mu::NodeMeshFormat* MeshFormatVar = StaticCast<mu::NodeMeshFormat*>(ParentNode);
+		mu::NodeMeshFormat::Private* Private = MeshFormatVar->GetPrivate();
+		AddChildFunc(Private->Source.get(), FString::Printf(TEXT("SOURCE MESH")));
+	}
+
+	else if (ParentNode->GetType() == mu::NodeModifierMeshClipMorphPlane::GetStaticType())
+	{
+		// Nothing to show
+	}
+
+	else if (ParentNode->GetType() == mu::NodeModifierMeshClipWithMesh::GetStaticType())
+	{
+		mu::NodeModifierMeshClipWithMesh* ModifierMeshClipWithMeshVar = StaticCast<mu::NodeModifierMeshClipWithMesh*>(ParentNode);
+		AddChildFunc(ModifierMeshClipWithMeshVar->ClipMesh.get(), FString::Printf(TEXT("CLIP MESH")));
+	}
+
+	else if (ParentNode->GetType() == mu::NodeModifierMeshClipDeform::GetStaticType())
+	{
+		mu::NodeModifierMeshClipDeform* ModifierMeshClipDeformVar = StaticCast<mu::NodeModifierMeshClipDeform*>(ParentNode);
+		AddChildFunc(ModifierMeshClipDeformVar->ClipMesh.get(), FString::Printf(TEXT("CLIP MESH")));
+	}
+
+	else if (ParentNode->GetType() == mu::NodeModifierMeshClipWithUVMask::GetStaticType())
+	{
+		mu::NodeModifierMeshClipWithUVMask* ModifierMeshClipWithUVMaskVar = StaticCast<mu::NodeModifierMeshClipWithUVMask*>(ParentNode);
+		AddChildFunc(ModifierMeshClipWithUVMaskVar->ClipMask.get(), FString::Printf(TEXT("CLIP MASK")));
+		AddChildFunc(ModifierMeshClipWithUVMaskVar->ClipLayout.get(), FString::Printf(TEXT("CLIP LAYOUT")));
+	}
+
+	else if (ParentNode->GetType() == mu::NodeImageSwitch::GetStaticType())
+	{
+		mu::NodeImageSwitch* ImageSwitchVar = StaticCast<mu::NodeImageSwitch*>(ParentNode);
+		mu::NodeImageSwitch::Private* Private = ImageSwitchVar->GetPrivate();
+		AddChildFunc(Private->m_pParameter.get(), FString::Printf(TEXT("PARAM")));
+		for (int32 OptionIndex = 0; OptionIndex < Private->m_options.Num(); OptionIndex++)
+		{
+			AddChildFunc(Private->m_options[OptionIndex].get(), FString::Printf(TEXT("OPTION [%d]"), OptionIndex));
+		}
+	}
+
+	else if (ParentNode->GetType() == mu::NodeImageMipmap::GetStaticType())
+	{
+		mu::NodeImageMipmap* ImageMipMapVar = StaticCast<mu::NodeImageMipmap*>(ParentNode);
+		mu::NodeImageMipmap::Private* Private = ImageMipMapVar->GetPrivate();
+		AddChildFunc(Private->m_pSource.get(), FString::Printf(TEXT("SOURCE")));
+		AddChildFunc(Private->m_pFactor.get(), FString::Printf(TEXT("FACTOR")));
+	}
+
+	else if (ParentNode->GetType() == mu::NodeImageLayer::GetStaticType())
+	{
+		mu::NodeImageLayer* ImageLayerVar = StaticCast<mu::NodeImageLayer*>(ParentNode);
+		mu::NodeImageLayer::Private* Private = ImageLayerVar->GetPrivate();
+		AddChildFunc(Private->m_pBase.get(), FString::Printf(TEXT("BASE")));
+		AddChildFunc(Private->m_pMask.get(), FString::Printf(TEXT("MASK")));
+		AddChildFunc(Private->m_pBlended.get(), FString::Printf(TEXT("BLEND")));
+	}
+	
+	else if (ParentNode->GetType() == mu::NodeImageLayerColour::GetStaticType())
+	{
+		mu::NodeImageLayerColour* ImageLayerColourVar = StaticCast<mu::NodeImageLayerColour*>(ParentNode);
+		mu::NodeImageLayerColour::Private* Private = ImageLayerColourVar->GetPrivate();
+		AddChildFunc(Private->m_pBase.get(), FString::Printf(TEXT("BASE")));
+		AddChildFunc(Private->m_pMask.get(), FString::Printf(TEXT("MASK")));
+		AddChildFunc(Private->m_pColour.get(), FString::Printf(TEXT("COLOR")));
+	}
+
+	else if (ParentNode->GetType() == mu::NodeImageResize::GetStaticType())
+	{
+		mu::NodeImageResize* ImageResizeVar = StaticCast<mu::NodeImageResize*>(ParentNode);
+		mu::NodeImageResize::Private* Private = ImageResizeVar->GetPrivate();
+		AddChildFunc(Private->m_pBase.get(), FString::Printf(TEXT("BASE")));
+	}
+
+	else if (ParentNode->GetType() == mu::NodeMeshMorph::GetStaticType())
+	{
+		mu::NodeMeshMorph* MeshMorphVar = StaticCast<mu::NodeMeshMorph*>(ParentNode);
+		mu::NodeMeshMorph::Private* Private = MeshMorphVar->GetPrivate();
+		AddChildFunc(Private->Base.get(), FString::Printf(TEXT("BASE")));
+		AddChildFunc(Private->Morph.get(), FString::Printf(TEXT("MORPH")));
+		AddChildFunc(Private->Factor.get(), FString::Printf(TEXT("FACTOR")));
+	}
+
+	else if (ParentNode->GetType() == mu::NodeImageProject::GetStaticType())
+	{
+		mu::NodeImageProject* ImageProjectVar = StaticCast<mu::NodeImageProject*>(ParentNode);
+		mu::NodeImageProject::Private* Private = ImageProjectVar->GetPrivate();
+		AddChildFunc(Private->m_pProjector.get(), FString::Printf(TEXT("PROJECTOR")));
+		AddChildFunc(Private->m_pMesh.get(), FString::Printf(TEXT("MESH")));
+		AddChildFunc(Private->m_pImage.get(), FString::Printf(TEXT("IMAGE")));
+		AddChildFunc(Private->m_pMask.get(), FString::Printf(TEXT("MASK")));
+		AddChildFunc(Private->m_pAngleFadeStart.get(), FString::Printf(TEXT("FADE START ANGLE")));
+		AddChildFunc(Private->m_pAngleFadeEnd.get(), FString::Printf(TEXT("FADE END ANGLE")));
+	}
+
+	else if (ParentNode->GetType() == mu::NodeImagePlainColour::GetStaticType())
+	{
+		mu::NodeImagePlainColour* ImagePlainColourVar = StaticCast<mu::NodeImagePlainColour*>(ParentNode);
+        mu::NodeImagePlainColour::Private* Private = ImagePlainColourVar->GetPrivate();
+		AddChildFunc(Private->m_pColour.get(), FString::Printf(TEXT("COLOR")));
+	}
+
+	else if (ParentNode->GetType() == mu::NodeLayout::GetStaticType())
+	{
+		// Nothing to show
+	}
+
+	else if (ParentNode->GetType() == mu::NodeScalarEnumParameter::GetStaticType())
+	{
+		mu::NodeScalarEnumParameter* ScalarEnumParameterVar = StaticCast<mu::NodeScalarEnumParameter*>(ParentNode);
+		mu::NodeScalarEnumParameter::Private* Private = ScalarEnumParameterVar->GetPrivate();
+		for (int32 RangeIndex = 0; RangeIndex < Private->m_ranges.Num(); RangeIndex++)
+		{
+			AddChildFunc(Private->m_ranges[RangeIndex].get(), FString::Printf(TEXT("RANGE [%d]"), RangeIndex));
+		}
+	}
+
+	else if (ParentNode->GetType() == mu::NodeMeshFragment::GetStaticType())
+	{
+		mu::NodeMeshFragment* MeshFragmentVar = StaticCast<mu::NodeMeshFragment*>(ParentNode);
+		AddChildFunc(MeshFragmentVar->SourceMesh.get(), FString::Printf(TEXT("MESH")));
+	}
+
+	else if (ParentNode->GetType() == mu::NodeColourSampleImage::GetStaticType())
+	{
+		mu::NodeColourSampleImage* ColorSampleImageVar = StaticCast<mu::NodeColourSampleImage*>(ParentNode);
+		AddChildFunc(ColorSampleImageVar->Image.get(), FString::Printf(TEXT("IMAGE")));
+		AddChildFunc(ColorSampleImageVar->X.get(), FString::Printf(TEXT("X")));
+		AddChildFunc(ColorSampleImageVar->Y.get(), FString::Printf(TEXT("Y")));
+	}
+
+	else if (ParentNode->GetType() == mu::NodeImageInterpolate::GetStaticType())
+	{
+		mu::NodeImageInterpolate* ImageInterpolateVar = StaticCast<mu::NodeImageInterpolate*>(ParentNode);
+		mu::NodeImageInterpolate::Private* Private = ImageInterpolateVar->GetPrivate();
+		AddChildFunc(Private->m_pFactor.get(), FString::Printf(TEXT("FACTOR")));
+		for (int32 TargetIndex = 0; TargetIndex < Private->m_targets.Num(); TargetIndex++)
+		{
+			AddChildFunc(Private->m_targets[TargetIndex].get(), FString::Printf(TEXT("TARGET [%d]"), TargetIndex));
+		}
+	}
+	
+	else if (ParentNode->GetType() == mu::NodeScalarConstant::GetStaticType())
+	{
+		// Nothing to show
+	}
+
+	else if (ParentNode->GetType() == mu::NodeScalarParameter::GetStaticType())
+	{
+		mu::NodeScalarParameter* ScalarParameterVar = StaticCast<mu::NodeScalarParameter*>(ParentNode);
+		mu::NodeScalarParameter::Private* Private = ScalarParameterVar->GetPrivate();
+		for (int32 RangeIndex = 0; RangeIndex < Private->m_ranges.Num(); RangeIndex++)
+		{
+			AddChildFunc(Private->m_ranges[RangeIndex].get(), FString::Printf(TEXT("RANGE [%d]"), RangeIndex));
+		}
+	}
+
+	else if (ParentNode->GetType() == mu::NodeColourParameter::GetStaticType())
+	{
+		mu::NodeColourParameter* ColorParameterVar = StaticCast<mu::NodeColourParameter*>(ParentNode);
+		for (int32 RangeIndex = 0; RangeIndex < ColorParameterVar->Ranges.Num(); RangeIndex++)
+		{
+			AddChildFunc(ColorParameterVar->Ranges[RangeIndex].get(), FString::Printf(TEXT("RANGE [%d]"), RangeIndex));
+		}
+	}
+
+	else if (ParentNode->GetType() == mu::NodeColourConstant::GetStaticType())
+	{
+		// Nothing to show
+	}
+
+
+	else if (ParentNode->GetType() == mu::NodeImageConstant::GetStaticType())
+	{
+		// Nothing to show
+	}
+
+	
+	else if (ParentNode->GetType() == mu::NodeScalarCurve::GetStaticType())
+	{
+		mu::NodeScalarCurve* ScalarCurveVar = StaticCast<mu::NodeScalarCurve*>(ParentNode);
+		AddChildFunc(ScalarCurveVar->CurveSampleValue.get(), FString::Printf(TEXT("INPUT")));
+	}
+
+	else if (ParentNode->GetType() == mu::NodeMeshMakeMorph::GetStaticType())
+	{
+		mu::NodeMeshMakeMorph* MeshMakeMorphVar = StaticCast<mu::NodeMeshMakeMorph*>(ParentNode);
+		mu::NodeMeshMakeMorph::Private* Private = MeshMakeMorphVar->GetPrivate();
+		AddChildFunc(Private->m_pBase.get(), FString::Printf(TEXT("BASE")));
+		AddChildFunc(Private->m_pTarget.get(), FString::Printf(TEXT("TARGET")));
+	}
+	
+	else if (ParentNode->GetType() == mu::NodeProjectorParameter::GetStaticType())
+	{
+		mu::NodeProjectorParameter* ProjectorParameterVar = StaticCast<mu::NodeProjectorParameter*>(ParentNode);
+		mu::NodeProjectorParameter::Private* Private = ProjectorParameterVar->GetPrivate();
+		for (int32 RangeIndex = 0; RangeIndex < Private->m_ranges.Num(); RangeIndex++)
+		{
+			AddChildFunc(Private->m_ranges[RangeIndex].get(), FString::Printf(TEXT("RANGE [%d]"), RangeIndex));
+		}
+	}
+
+	else if (ParentNode->GetType() == mu::NodeProjectorConstant::GetStaticType())
+	{
+		// Nothing to show
+	}
+
+	else if (ParentNode->GetType() == mu::NodeColourSwitch::GetStaticType())
+	{
+		mu::NodeColourSwitch* ColorSwitchVar = StaticCast<mu::NodeColourSwitch*>(ParentNode);
+		AddChildFunc(ColorSwitchVar->Parameter.get(), TEXT("PARAM"));
+		for (int32 OptionIndex = 0; OptionIndex < ColorSwitchVar->Options.Num(); ++OptionIndex)
+		{
+			AddChildFunc(ColorSwitchVar->Options[OptionIndex].get(), FString::Printf(TEXT("OPTION [%d]"), OptionIndex));
+		}
+	}
+
+	else if (ParentNode->GetType() == mu::NodeImageSwizzle::GetStaticType())
+	{
+		mu::NodeImageSwizzle* ImageSwizzleVar = StaticCast<mu::NodeImageSwizzle*>(ParentNode);
+		mu::NodeImageSwizzle::Private* Private = ImageSwizzleVar->GetPrivate();
+		for (int32 SourceIndex = 0; SourceIndex < Private->m_sources.Num(); ++SourceIndex)
+		{
+			AddChildFunc(Private->m_sources[SourceIndex].get(), FString::Printf(TEXT("SOURCE [%d]"), SourceIndex));
+		}
+	}
+	
+	else if (ParentNode->GetType() == mu::NodeImageInvert::GetStaticType())
+	{
+		mu::NodeImageInvert* ImageInvertVar = StaticCast<mu::NodeImageInvert*>(ParentNode);
+		mu::NodeImageInvert::Private* Private = ImageInvertVar->GetPrivate();
+		AddChildFunc(Private->m_pBase.get(), TEXT("BASE"));
+	}
+
+	else if (ParentNode->GetType() == mu::NodeImageMultiLayer::GetStaticType())
+	{
+		mu::NodeImageMultiLayer* ImageMultilayerVar = StaticCast<mu::NodeImageMultiLayer*>(ParentNode);
+		mu::NodeImageMultiLayer::Private* Private = ImageMultilayerVar->GetPrivate();
+		AddChildFunc(Private->m_pBase.get(), FString::Printf(TEXT("BASE")));
+		AddChildFunc(Private->m_pMask.get(), FString::Printf(TEXT("MASK")));
+		AddChildFunc(Private->m_pBlended.get(), FString::Printf(TEXT("BLEND")));
+		AddChildFunc(Private->m_pRange.get(), FString::Printf(TEXT("RANGE")));
+	}
+	
+	else if (ParentNode->GetType() == mu::NodeImageTable::GetStaticType())
+	{
+		// No nodes to show
+	}
+
+	else if (ParentNode->GetType() == mu::NodeMeshTable::GetStaticType())
+	{
+		mu::NodeMeshTable* MeshTableVar = StaticCast<mu::NodeMeshTable*>(ParentNode);
+		for (int32 LayoutIndex = 0; LayoutIndex < MeshTableVar->Layouts.Num(); ++LayoutIndex)
+		{
+			AddChildFunc(MeshTableVar->Layouts[LayoutIndex].get(), FString::Printf(TEXT("LAYOUT [%d]"), LayoutIndex));
+		}
+	}
+
+	else if (ParentNode->GetType() == mu::NodeScalarTable::GetStaticType())
+	{
+		// Nothing to show
+	}
+
+	else if (ParentNode->GetType() == mu::NodeScalarSwitch::GetStaticType())
+	{
+		mu::NodeScalarSwitch* ScalarSwitchVar = StaticCast<mu::NodeScalarSwitch*>(ParentNode);
+		AddChildFunc(ScalarSwitchVar->Parameter.get(), TEXT("PARAM"));
+		for (int32 OptionIndex = 0; OptionIndex < ScalarSwitchVar->Options.Num(); ++OptionIndex)
+		{
+			AddChildFunc(ScalarSwitchVar->Options[OptionIndex].get(), FString::Printf(TEXT("OPTION [%d]"), OptionIndex));
+		}
+	}
+	
+	else if (ParentNode->GetType() == mu::NodeColourFromScalars::GetStaticType())
+	{
+		mu::NodeColourFromScalars* ScalarTableVar = StaticCast<mu::NodeColourFromScalars*>(ParentNode);
+		AddChildFunc(ScalarTableVar->X.get(), TEXT("X"));
+		AddChildFunc(ScalarTableVar->Y.get(), TEXT("Y"));
+		AddChildFunc(ScalarTableVar->Z.get(), TEXT("Z"));
+		AddChildFunc(ScalarTableVar->W.get(), TEXT("W"));
+	}
+	
+	else
+	{
+		UE_LOG(LogMutable,Error,TEXT("The node of type %d has not been implemented, so its children won't be added to the tree."), int32(ParentNode->GetType()->Type));
+
+		// Add a placeholder to the tree
+		const FString Prefix =  FString::Printf(TEXT("[%d] NODE TYPE NOT IMPLEMENTED"), int32(ParentNode->GetType()->Type));
+		AddChildFunc(nullptr, Prefix);
+	}
+#endif
+}
+
+
+TSharedPtr<SWidget> SMutableGraphViewer::OnTreeContextMenuOpening()
+{
+	FMenuBuilder MenuBuilder(true, nullptr);
+
+	MenuBuilder.AddMenuEntry(
+		LOCTEXT("Graph_Expand_Instance", "Expand Instance-Level Operations"),
+		LOCTEXT("Graph_Expand_Instance_Tooltip", "Expands all the operations in the tree that are instance operations (not images, meshes, booleans, etc.)."),
+		FSlateIcon(),
+		FUIAction(FExecuteAction::CreateSP(this, &SMutableGraphViewer::TreeExpandUnique)
+			//, FCanExecuteAction::CreateSP(this, &SMutableCodeViewer::HasAnyItemInPalette)
+		)
+	);
+
+	return MenuBuilder.MakeWidget();
+}
+
+
+void SMutableGraphViewer::TreeExpandRecursive(TSharedPtr<FMutableGraphTreeElement> InInfo, bool bExpand)
+{
+	if (bExpand)
+	{
+		TreeExpandUnique();
+	}
+}
+
+
+void SMutableGraphViewer::TreeExpandUnique()
+{
+	TArray<TSharedPtr<FMutableGraphTreeElement>> Pending = RootNodes;
+
+	TSet<TSharedPtr<FMutableGraphTreeElement>> Processed;
+
+	TArray<TSharedPtr<FMutableGraphTreeElement>> Children;
+
+	while (!Pending.IsEmpty())
+	{
+		TSharedPtr<FMutableGraphTreeElement> Item = Pending.Pop();
+		TreeView->SetItemExpansion(Item, true);
+
+		Children.SetNum(0);
+		GetChildrenForInfo(Item, Children);
+		Pending.Append(Children);
+	}
+}
+
+
+#undef LOCTEXT_NAMESPACE 
+
+
diff --git a/Engine/Plugins/Mutable/Source/CustomizableObjectEditor/Private/MuCOE/SMutableObjectViewer.cpp b/Engine/Plugins/Mutable/Source/CustomizableObjectEditor/Private/MuCOE/SMutableObjectViewer.cpp
new file mode 100644
--- /dev/null
+++ b/Engine/Plugins/Mutable/Source/CustomizableObjectEditor/Private/MuCOE/SMutableObjectViewer.cpp
@@ -0,0 +1,358 @@
+// Copyright Epic Games, Inc. All Rights Reserved.
+
+#include "MuCOE/SMutableObjectViewer.h"
+
+#include "Framework/MultiBox/MultiBoxBuilder.h"
+#include "Framework/Views/TableViewMetadata.h"
+#include "Interfaces/ITargetPlatform.h"
+#include "Interfaces/ITargetPlatformManagerModule.h"
+#include "MuCO/CustomizableObject.h"
+#include "MuCOE/CustomizableObjectCompileRunnable.h"
+#include "MuCOE/CustomizableObjectEditorStyle.h"
+#include "MuCOE/SMutableCodeViewer.h"
+#include "MuCOE/SMutableGraphViewer.h"
+#include "Widgets/Docking/SDockTab.h"
+#include "Widgets/Input/STextComboBox.h"
+#include "Widgets/Input/SNumericDropDown.h"
+#include "Widgets/Views/STreeView.h"
+
+class FExtender;
+class FUICommandList;
+class ITableRow;
+class SWidget;
+struct FSlateBrush;
+
+#define LOCTEXT_NAMESPACE "SMutableDebugger"
+
+
+void SMutableObjectViewer::AddReferencedObjects(FReferenceCollector& Collector)
+{
+	Collector.AddReferencedObject(CustomizableObject);
+}
+
+
+FString SMutableObjectViewer::GetReferencerName() const
+{
+	return TEXT("SMutableObjectViewer");
+}
+
+
+void SMutableObjectViewer::Construct(const FArguments& InArgs, UCustomizableObject* InObject)
+{
+	check(InObject);
+	
+	CustomizableObject = InObject;
+
+	// Initialize the debugger compile options
+	CompileOptions = InObject->GetPrivate()->GetCompileOptions();
+
+	{
+		ITargetPlatformManagerModule* TPM = GetTargetPlatformManager();
+		const TArray<ITargetPlatform*>& Platforms = TPM->GetActiveTargetPlatforms();
+		CompileOptions.TargetPlatform = Platforms.IsEmpty() ? nullptr : Platforms[0];
+	}
+
+	FToolBarBuilder ToolbarBuilder(TSharedPtr<const FUICommandList>(), FMultiBoxCustomization::None, TSharedPtr<FExtender>(), true);
+	ToolbarBuilder.SetLabelVisibility(EVisibility::Visible);
+	ToolbarBuilder.SetStyle(&FAppStyle::Get(), "SlimToolBar");
+
+	ToolbarBuilder.BeginSection("Compilation");
+
+	ToolbarBuilder.AddToolBarButton(
+		FUIAction(FExecuteAction::CreateSP(this, &SMutableObjectViewer::GenerateMutableGraphPressed)),
+		NAME_None,
+		LOCTEXT("GenerateMutableGraph", "Unreal to Mutable Graph"),
+		LOCTEXT("GenerateMutableGraphTooltip", "Generate a mutable graph from the customizable object source graph."),
+		FSlateIcon(FCustomizableObjectEditorStyle::Get().GetStyleSetName(), "CustomizableObjectDebugger.GenerateMutableGraph", "CustomizableObjectDebugger.GenerateMutableGraph.Small"),
+		EUserInterfaceActionType::Button
+	);
+
+	ToolbarBuilder.AddToolBarButton(
+		FUIAction(FExecuteAction::CreateSP(this, &SMutableObjectViewer::CompileMutableCodePressed)),
+		NAME_None,
+		LOCTEXT("GenerateMutableCode", "Unreal to Mutable Code"),
+		LOCTEXT("GenerateMutableCodeTooltipFromGraph", "Generate a mutable code from the customizable object source graph."),
+		FSlateIcon(FCustomizableObjectEditorStyle::Get().GetStyleSetName(), "CustomizableObjectDebugger.CompileMutableCode", "CustomizableObjectDebugger.CompileMutableCode.Small"),
+		EUserInterfaceActionType::Button
+	);
+
+	ToolbarBuilder.AddComboButton(
+		FUIAction(),
+		FOnGetContent::CreateSP(this, &SMutableObjectViewer::GenerateCompileOptionsMenuContent),
+		LOCTEXT("Compile_Options_Label", "Compile Options"),
+		LOCTEXT("Compile_Options_Tooltip", "Change Compile Options"),
+		TAttribute<FSlateIcon>(),
+		true);
+
+	ToolbarBuilder.EndSection();
+	
+	ChildSlot
+	[
+		SNew(SVerticalBox)
+		+ SVerticalBox::Slot()
+		.AutoHeight()
+		.VAlign(VAlign_Center)
+		[
+			ToolbarBuilder.MakeWidget()
+		]
+		+ SVerticalBox::Slot()
+		.FillHeight(1.0f)
+		[
+			SAssignNew(DebuggerContentsBox, SBorder)
+		]
+	];
+}
+
+
+void SMutableObjectViewer::GenerateMutableGraphPressed()
+{
+	// Convert from Unreal graph to Mutable graph.
+	FCustomizableObjectCompiler Compiler;
+	TArray<TSoftObjectPtr<const UTexture>> RuntimeTextures;
+	TArray<FMutableSourceTextureData> CompilerTextures;
+	mu::Ptr<mu::Node> RootNode = Compiler.Export(CustomizableObject, CompileOptions, RuntimeTextures, CompilerTextures);
+	if (!RootNode)
+	{
+		// TODO: Show errors
+		return;
+	}
+
+	FString DataTag = FString::Printf(TEXT("Mutable Graph : %s"), *CompileOptions.TargetPlatform->PlatformName());
+
+	TSharedRef<SMutableGraphViewer> GraphViewer = SNew(SMutableGraphViewer, RootNode)
+		.DataTag(DataTag)
+		.ReferencedRuntimeTextures(RuntimeTextures)
+		.ReferencedCompileTextures(CompilerTextures);
+	
+	DebuggerContentsBox->ClearContent();
+	DebuggerContentsBox->SetContent(GraphViewer);
+}
+
+
+void SMutableObjectViewer::CompileMutableCodePressed()
+{
+	TArray<TSoftObjectPtr<const UTexture>> RuntimeTextures;
+	TArray<FMutableSourceTextureData> CompilerTextures;
+	if (CompileOptions.bForceLargeLODBias)
+	{
+		// Debug compile with many different biases
+		constexpr int32 MaxBias = 15;
+		for (int32 Bias = 0; Bias < MaxBias; ++Bias)
+		{
+			CompileOptions.DebugBias = Bias;
+
+			FCustomizableObjectCompiler Compiler;
+			RuntimeTextures.Empty();
+			CompilerTextures.Empty();
+			mu::NodePtr RootNode = Compiler.Export(CustomizableObject, CompileOptions, RuntimeTextures, CompilerTextures);
+			if (!RootNode)
+			{
+				// TODO: Show errors
+				ensure(false);
+				return;
+			}
+
+			// Do the compilation to Mutable Code synchronously.
+			TSharedPtr<FCustomizableObjectCompileRunnable> CompileTask = MakeShareable(new FCustomizableObjectCompileRunnable(RootNode));
+			CompileTask->Options = CompileOptions;
+			CompileTask->ReferencedTextures = CompilerTextures;
+			CompileTask->Init();
+			CompileTask->Run();
+		}
+	}
+
+	// Convert from Unreal graph to Mutable graph.
+	FCustomizableObjectCompiler Compiler;
+	RuntimeTextures.Empty();
+	CompilerTextures.Empty();
+	mu::NodePtr RootNode = Compiler.Export(CustomizableObject, CompileOptions, RuntimeTextures, CompilerTextures);
+	if (!RootNode)
+	{
+		// TODO: Show errors
+		ensure(false);
+		return;
+	}
+
+	// Do the compilation to Mutable Code synchronously.
+	TSharedPtr<FCustomizableObjectCompileRunnable> CompileTask = MakeShareable(new FCustomizableObjectCompileRunnable(RootNode));
+	CompileTask->Options = CompileOptions;
+	CompileTask->ReferencedTextures = CompilerTextures;
+	CompileTask->Init();
+	CompileTask->Run();
+
+	FString DataTag = FString::Printf( TEXT("Mutable Code : %s opt %d"), *CompileOptions.TargetPlatform->PlatformName(), CompileOptions.OptimizationLevel );
+
+	TSharedRef<SMutableCodeViewer> CodeViewer = SNew(SMutableCodeViewer, CompileTask->Model, RuntimeTextures)
+		.DataTag(DataTag);
+	
+	DebuggerContentsBox->ClearContent();
+	DebuggerContentsBox->SetContent(CodeViewer);
+}
+
+
+TSharedRef<SWidget> SMutableObjectViewer::GenerateCompileOptionsMenuContent()
+{
+	const bool bShouldCloseWindowAfterMenuSelection = false;
+	FMenuBuilder MenuBuilder(bShouldCloseWindowAfterMenuSelection, nullptr);
+
+	// settings
+	MenuBuilder.BeginSection("Optimization", LOCTEXT("MutableCompileOptimizationHeading", "Optimization"));
+	{
+		// Unreal Graph to Mutable Graph options
+		//-----------------------------------
+
+		// Platform
+		DebugPlatformStrings.Empty();
+
+		TSharedPtr<FString> SelectedPlatform;
+		{
+			ITargetPlatformManagerModule* TPM = GetTargetPlatformManager();
+			check(TPM);
+
+			ITargetPlatform* CurrentPlatform = NULL;
+			const TArray<ITargetPlatform*>& Platforms = TPM->GetTargetPlatforms();
+			for (const ITargetPlatform* Platform : Platforms)
+			{
+				TSharedPtr<FString> ThisPlatform = MakeShareable(new FString(Platform->PlatformName()));
+				DebugPlatformStrings.Add(ThisPlatform);
+				if (Platform == CompileOptions.TargetPlatform)
+				{
+					SelectedPlatform = ThisPlatform;
+				}
+			}
+		}
+
+		if (!SelectedPlatform.IsValid() && DebugPlatformStrings.Num())
+		{
+			SelectedPlatform = DebugPlatformStrings[0];
+		}
+
+		DebugPlatformCombo =
+			SNew(STextComboBox)
+			.OptionsSource(&DebugPlatformStrings)
+			.InitiallySelectedItem(SelectedPlatform)
+			.OnSelectionChanged(this, &SMutableObjectViewer::OnChangeDebugPlatform)
+			;
+
+		MenuBuilder.AddWidget(DebugPlatformCombo.ToSharedRef(), LOCTEXT("MutableDebugPlatform", "Target Platform"));
+
+		// Compilation options
+		//-----------------------------------
+
+		// Optimisation level
+		CompileOptimizationStrings.Empty();
+		CompileOptimizationStrings.Add(MakeShareable(new FString(LOCTEXT("OptimizationNone", "None").ToString())));
+		CompileOptimizationStrings.Add(MakeShareable(new FString(LOCTEXT("OptimizationMin", "Minimal").ToString())));
+		CompileOptimizationStrings.Add(MakeShareable(new FString(LOCTEXT("OptimizationMax", "Maximum").ToString())));
+
+		CompileOptions.OptimizationLevel = FMath::Min(CompileOptions.OptimizationLevel, CompileOptimizationStrings.Num() - 1);
+
+		CompileOptimizationCombo =
+			SNew(STextComboBox)
+			.OptionsSource(&CompileOptimizationStrings)
+			.InitiallySelectedItem(CompileOptimizationStrings[CompileOptions.OptimizationLevel])
+			.OnSelectionChanged(this, &SMutableObjectViewer::OnChangeCompileOptimizationLevel)
+			;
+		MenuBuilder.AddWidget(CompileOptimizationCombo.ToSharedRef(), LOCTEXT("MutableCompileOptimizationLevel", "Optimization Level"));
+
+		{
+			CompileTextureCompressionStrings.Empty();
+			CompileTextureCompressionStrings.Add(MakeShareable(new FString(LOCTEXT("MutableTextureCompressionNone", "None").ToString())));
+			CompileTextureCompressionStrings.Add(MakeShareable(new FString(LOCTEXT("MutableTextureCompressionFast", "Fast").ToString())));
+			CompileTextureCompressionStrings.Add(MakeShareable(new FString(LOCTEXT("MutableTextureCompressionHighQuality", "High Quality").ToString())));
+
+			int32 SelectedCompression = FMath::Clamp(int32(CompileOptions.TextureCompression), 0, CompileTextureCompressionStrings.Num() - 1);
+			CompileTextureCompressionCombo =
+				SNew(STextComboBox)
+				.OptionsSource(&CompileTextureCompressionStrings)
+				.InitiallySelectedItem(CompileTextureCompressionStrings[SelectedCompression])
+				.OnSelectionChanged(this, &SMutableObjectViewer::OnChangeCompileTextureCompressionType)
+				;
+
+			MenuBuilder.AddWidget(CompileTextureCompressionCombo.ToSharedRef(), LOCTEXT("MutableCompileTextureCompressionType", "Texture Compression"));
+		}
+
+		// Image tiling
+		// Unfortunately SNumericDropDown doesn't work with integers at the time of writing.
+		TArray<SNumericDropDown<float>::FNamedValue> TilingOptions;
+		TilingOptions.Add(SNumericDropDown<float>::FNamedValue(0, FText::FromString(TEXT("0")), FText::FromString(TEXT("Disabled"))));
+		TilingOptions.Add(SNumericDropDown<float>::FNamedValue(64, FText::FromString(TEXT("64")), FText::FromString(TEXT("64"))));
+		TilingOptions.Add(SNumericDropDown<float>::FNamedValue(128, FText::FromString(TEXT("128")), FText::FromString(TEXT("128"))));
+		TilingOptions.Add(SNumericDropDown<float>::FNamedValue(256, FText::FromString(TEXT("256")), FText::FromString(TEXT("256"))));
+		TilingOptions.Add(SNumericDropDown<float>::FNamedValue(512, FText::FromString(TEXT("512")), FText::FromString(TEXT("512"))));
+
+		CompileTilingCombo = SNew(SNumericDropDown<float>)
+			.DropDownValues(TilingOptions)
+			.Value_Lambda( [&]() { return float(CompileOptions.ImageTiling); })
+			.OnValueChanged_Lambda( [&](float Value) { CompileOptions.ImageTiling = int32(Value); } )
+			;
+		MenuBuilder.AddWidget(CompileTilingCombo.ToSharedRef(), LOCTEXT("MutableCompileImageTiling", "Image Tiling"));
+
+		// Disk as cache
+		MenuBuilder.AddMenuEntry(
+			LOCTEXT("Generate_MutableUseDisk", "Enable compiling using the disk as memory."),
+			LOCTEXT("Generate_MutableUseDiskTooltip", "This is very slow but supports compiling huge objects. It requires a lot of free space in the OS disk."),
+			FSlateIcon(),
+			FUIAction(
+				FExecuteAction::CreateLambda([this]() { CompileOptions.bUseDiskCompilation = !CompileOptions.bUseDiskCompilation; }),
+				FCanExecuteAction(),
+				FIsActionChecked::CreateLambda([this]() { return CompileOptions.bUseDiskCompilation; })),
+			NAME_None,
+			EUserInterfaceActionType::ToggleButton
+		);
+
+		// Debug LODBias
+		MenuBuilder.AddMenuEntry(
+			LOCTEXT("ForceLargeLODBias", "Force a large texture LODBias."),
+			LOCTEXT("ForceLargeLODBiasTooltip", "This is useful to test compilation of special cook modes."),
+			FSlateIcon(),
+			FUIAction(
+				FExecuteAction::CreateLambda([this]() { CompileOptions.bForceLargeLODBias = !CompileOptions.bForceLargeLODBias; }),
+				FCanExecuteAction(),
+				FIsActionChecked::CreateLambda([this]() { return CompileOptions.bForceLargeLODBias; })),
+			NAME_None,
+			EUserInterfaceActionType::ToggleButton
+		);
+	}
+	MenuBuilder.EndSection();
+
+	return MenuBuilder.MakeWidget();
+}
+
+
+void SMutableObjectViewer::OnChangeCompileOptimizationLevel(TSharedPtr<FString> NewSelection, ESelectInfo::Type SelectInfo)
+{
+	CompileOptions.OptimizationLevel = CompileOptimizationStrings.Find(NewSelection);
+}
+
+
+void SMutableObjectViewer::OnChangeCompileTextureCompressionType(TSharedPtr<FString> NewSelection, ESelectInfo::Type)
+{
+	CompileOptions.TextureCompression = ECustomizableObjectTextureCompression(CompileTextureCompressionStrings.Find(NewSelection));
+}
+
+
+void SMutableObjectViewer::OnChangeDebugPlatform(TSharedPtr<FString> NewSelection, ESelectInfo::Type SelectInfo)
+{
+	if (!CustomizableObject) return;
+
+	CompileOptions.TargetPlatform = nullptr;
+
+	ITargetPlatformManagerModule* TPM = GetTargetPlatformManager();
+	check(TPM);
+	const TArray<ITargetPlatform*>& Platforms = TPM->GetTargetPlatforms();
+	check(Platforms.Num());
+
+	CompileOptions.TargetPlatform = Platforms[0];
+
+	for (int32 Index = 1; Index < Platforms.Num(); Index++)
+	{
+		if (Platforms[Index]->PlatformName() == *NewSelection)
+		{
+			CompileOptions.TargetPlatform = Platforms[Index];
+			break;
+		}
+	}
+}
+
+#undef LOCTEXT_NAMESPACE 
diff --git a/Engine/Plugins/Mutable/Source/CustomizableObjectEditor/Private/MuCOE/SMutableObjectViewer.h b/Engine/Plugins/Mutable/Source/CustomizableObjectEditor/Private/MuCOE/SMutableObjectViewer.h
new file mode 100644
--- /dev/null
+++ b/Engine/Plugins/Mutable/Source/CustomizableObjectEditor/Private/MuCOE/SMutableObjectViewer.h
@@ -0,0 +1,67 @@
+// Copyright Epic Games, Inc. All Rights Reserved.
+
+#pragma once
+
+#include "MuCO/CustomizableObjectCompilerTypes.h"
+#include "Widgets/SCompoundWidget.h"
+#include "Widgets/Input/SNumericDropDown.h"
+
+class STableViewBase;
+namespace ESelectInfo { enum Type : int; }
+template <typename ItemType> class STreeView;
+
+// Forward declarations
+class UCustomizableObject;
+class FMutableObjectTreeElement;
+class FReferenceCollector;
+class FTabManager;
+class ITableRow;
+class STextComboBox;
+class SWidget;
+
+
+/** This widget shows the internals of a CustomizableObject for debugging purposes. 
+ * This is not the Unreal source graph in the UCustomizableObject, but an intermediate step of the compilation process.
+ */
+class SMutableObjectViewer :
+	public SCompoundWidget,
+	public FGCObject
+{
+public:
+
+	// SWidget interface
+	SLATE_BEGIN_ARGS(SMutableObjectViewer) {}
+	SLATE_END_ARGS()
+
+	void Construct(const FArguments& InArgs, UCustomizableObject*);
+
+	// FGCObject interface
+	virtual void AddReferencedObjects(FReferenceCollector& Collector) override;
+	virtual FString GetReferencerName() const override;
+
+private:
+	
+	/** Slate whose content space contains the currently open debugging view */
+	TSharedPtr<SBorder> DebuggerContentsBox;
+	
+	/** The Mutable Graph to show, represented by its root. */
+	TObjectPtr<UCustomizableObject> CustomizableObject;
+
+	/** Compilation options to use in the debugger operations. */
+	FCompilationOptions CompileOptions;
+	TArray< TSharedPtr<FString> > DebugPlatformStrings;
+	TSharedPtr<STextComboBox> DebugPlatformCombo;
+
+	/** UI callbacks */
+	void GenerateMutableGraphPressed();
+	void CompileMutableCodePressed();
+	TSharedRef<SWidget> GenerateCompileOptionsMenuContent();
+	TSharedPtr<STextComboBox> CompileOptimizationCombo;
+	TArray< TSharedPtr<FString> > CompileOptimizationStrings;
+	TSharedPtr<STextComboBox> CompileTextureCompressionCombo;
+	TArray< TSharedPtr<FString> > CompileTextureCompressionStrings;
+	TSharedPtr<SNumericDropDown<float>> CompileTilingCombo;
+	void OnChangeCompileOptimizationLevel(TSharedPtr<FString> NewSelection, ESelectInfo::Type SelectInfo);
+	void OnChangeCompileTextureCompressionType(TSharedPtr<FString> NewSelection, ESelectInfo::Type);
+	void OnChangeDebugPlatform(TSharedPtr<FString> NewSelection, ESelectInfo::Type SelectInfo);
+};
diff --git a/Engine/Plugins/Mutable/Source/CustomizableObjectEditor/Private/MuCOE/SMutableParametersWidget.cpp b/Engine/Plugins/Mutable/Source/CustomizableObjectEditor/Private/MuCOE/SMutableParametersWidget.cpp
new file mode 100644
--- /dev/null
+++ b/Engine/Plugins/Mutable/Source/CustomizableObjectEditor/Private/MuCOE/SMutableParametersWidget.cpp
@@ -0,0 +1,739 @@
+// Copyright Epic Games, Inc. All Rights Reserved.
+
+#include "MuCOE/SMutableParametersWidget.h"
+
+#include "SSearchableComboBox.h"
+#include "MuCOE/UnrealEditorPortabilityHelpers.h"
+#include "MuR/ModelPrivate.h"
+#include "MuR/ParametersPrivate.h"
+#include "Templates/SharedPointer.h"
+#include "Widgets/SBoxPanel.h"
+#include "Widgets/Colors/SColorBlock.h"
+#include "Widgets/Colors/SColorPicker.h"
+#include "Widgets/Input/SVectorInputBox.h"
+#include "Widgets/Layout/SGridPanel.h"
+
+struct FGeometry;
+
+
+#define LOCTEXT_NAMESPACE "CustomizableObjectEditor"
+
+
+void SMutableParametersWidget::Construct(const FArguments& InArgs)
+{
+	OnParametersValueChanged = InArgs._OnParametersValueChanged;
+	MutableParameters = InArgs._Parameters;
+	if (MutableParameters)
+	{
+		bIsPendingUpdate = true;
+	}
+	
+	ChildSlot
+	[
+		SNew(SVerticalBox)
+		+ SVerticalBox::Slot()
+		.AutoHeight()
+		.Padding(0, 4.0f)
+		.HAlign(HAlign_Fill)
+		[
+			SAssignNew(ParamBox, SVerticalBox)
+		]
+	];
+}
+
+
+void SMutableParametersWidget::Tick(const FGeometry& AllottedGeometry, const double InCurrentTime, const float InDeltaTime)
+{
+	SCompoundWidget::Tick(AllottedGeometry, InCurrentTime, InDeltaTime);
+
+	if (!bIsPendingUpdate)
+	{
+		return;
+	}
+
+	bIsPendingUpdate = false;
+
+	ParamBox->ClearChildren();
+
+	if (!MutableParameters)
+	{
+		return;
+	}
+	
+	for ( int32 ParamIndex=0; ParamIndex< MutableParameters->GetCount(); ++ParamIndex )
+	{
+		FString ParamName = MutableParameters->GetName(ParamIndex);
+
+		TSharedPtr<SHorizontalBox> ParameterBox;
+
+		ParamBox->AddSlot()
+			.Padding(1.0f, 6.0f)
+			.HAlign(HAlign_Fill)
+			.VAlign(VAlign_Fill)
+			.AutoHeight()
+			[
+				SAssignNew(ParameterBox, SHorizontalBox)
+
+#if UE_BUILD_DEBUG
+				+ SHorizontalBox::Slot()
+				.AutoWidth()
+				.Padding(4,0)
+				.VAlign(VAlign_Center)
+				[
+					SNew(STextBlock)
+					.ColorAndOpacity(FSlateColor(FLinearColor(0.5f,0.5f,0.5f,1.0f)))
+					.Text(FText::FromString(FString::FromInt(ParamIndex)))
+					.Justification(ETextJustify::Left)
+					.AutoWrapText(true)
+				]
+#endif
+
+				+ SHorizontalBox::Slot()
+				.FillWidth(0.66f)
+				.VAlign(VAlign_Center)
+				[
+					SNew(STextBlock)
+					.Text(FText::FromString(ParamName))
+					.Justification(ETextJustify::Left)
+					.AutoWrapText(true)
+				]
+			];
+		
+		
+		// Use the parameter descriptor to know what parameter value determines the amount of values this one has
+		mu::FParameterDesc& ParameterDescription = MutableParameters->GetPrivate()->m_pModel->GetPrivate()->m_program.
+															m_parameters[ParamIndex];
+		TArray<uint32> Ranges = ParameterDescription.m_ranges;
+		if (Ranges.Num() >= 1)
+		{
+			if (Ranges.Num() > 1)
+			{
+				// todo: Only supports first dimension
+				UE_LOG(LogTemp,Warning,TEXT("Currently the debugger only suports showing the first dimension of each range."))
+			}
+
+			// This parameter does have more than one value (multivalue)
+			
+			// A range index can have more than one range
+			constexpr int32 Dimension = 0;		// Only operate over the first dimension
+			// Default amount of values to display
+			int32 PositionCount = 1;			
+			
+			// grab only first value (first dimension)
+			const uint32 RangeIdx = Ranges[Dimension];
+			// Get the parameter that controls the amount of values of current parameter
+			const int32 ParameterControllingValueCount = MutableParameters->GetPrivate()->m_pModel->GetPrivate()->m_program.m_ranges[RangeIdx].m_dimensionParameter;
+
+			// If a parameter index has been reported then we know that another parameter's value determines the amount of values we have
+			if (ParameterControllingValueCount >= 0)
+			{
+				// Get the numerical value of that parameters
+				mu::PARAMETER_TYPE ParameterType = MutableParameters->GetType(ParameterControllingValueCount);
+				if (ParameterType == mu::PARAMETER_TYPE::T_INT)
+				{
+					PositionCount = MutableParameters->GetIntValue(ParameterControllingValueCount);
+				}
+				else if (ParameterType == mu::PARAMETER_TYPE::T_FLOAT)
+				{
+					PositionCount = MutableParameters->GetFloatValue(ParameterControllingValueCount);
+				}
+				else
+				{
+					checkNoEntry();
+					// Unable to get the amount of options of a non numerical value
+				}
+
+				// Show one value at least.
+				if (PositionCount < 1)
+				{
+					PositionCount = 1;
+				}
+			}
+
+			// If no parameter is controlling the amount of values of this one or if the amount of values is 0 the
+			// amount of values that will therefore be used will be 1.
+
+			// Slate where to place all the values for each of the positions of the current parameter at the target dimension
+			TSharedPtr<SVerticalBox> ParameterValuesCollection = SNew(SVerticalBox);
+			
+			// Iterate over all the values for the parameter and show them. Updating their value is managed by the UI
+			for (int32 Position = 0; Position < PositionCount; Position++)
+			{
+				mu::RangeIndexPtr RangeIndex = MutableParameters->NewRangeIndex(ParamIndex);
+				RangeIndex->SetPosition(Dimension,Position);
+			
+				// Show position data if more than one is set to avoid UI clutter
+				if (PositionCount > 1)
+				{
+					// Add identifying data on top
+					ParameterValuesCollection->AddSlot()
+					.AutoHeight()
+					.HAlign(EHorizontalAlignment::HAlign_Right)
+					.Padding(4,4,4,2)
+					[
+						SNew(STextBlock)
+						.Text(FText::FromString(FString::Printf(TEXT("Value : %i"), Position )))
+						.Justification(ETextJustify::Right)
+						.AutoWrapText(true)
+					];
+				}
+			
+				// Generate a new temporary slate container to be later added to the vertical container for all dimensions and positions
+				TSharedPtr<SHorizontalBox> ParameterProxyContainer = SNew(SHorizontalBox);
+
+				// Process the parameter 
+				GenerateAndAttachParameterSlate(ParamIndex,ParameterProxyContainer,RangeIndex);
+
+				// Add the horizontal slot we have used to place the proxy now onto our vertical box (so it is placed inside a vertical slot)
+				ParameterValuesCollection->AddSlot()
+				.AutoHeight()
+				[
+					ParameterProxyContainer.ToSharedRef()
+				];
+			}
+		
+			// Add the collection of parameter value objects to the ParameterBox
+			ParameterBox->AddSlot()
+			.FillWidth(1)
+			.Padding(4)
+			.HAlign(HAlign_Fill)
+			[
+				ParameterValuesCollection.ToSharedRef()
+			];
+
+		}
+		else
+		{
+			// The parameter only can have one value (no multivalue) so draw it normally
+			GenerateAndAttachParameterSlate(ParamIndex,ParameterBox,nullptr);
+		}
+		
+	}
+}
+
+
+void SMutableParametersWidget::ScheduleUpdateIfRequired(const int32& InParameterIndex)
+{
+	// Update the ui to display the new amount of values we want to show for the multidimensional values if this parameter's value
+	// controls the amount of values another parameter has.
+	const bool bIsRangeSize = MutableParameters->GetPrivate()->m_pModel->GetPrivate()->m_program.m_ranges.ContainsByPredicate([InParameterIndex](const mu::FRangeDesc& r)
+		{ return r.m_dimensionParameter == InParameterIndex; });
+
+	// Schedule update since we know the value does change the amount of values shown by another parameter (multivalue parameter)
+	if (bIsRangeSize && !bIsPendingUpdate)
+	{
+		bIsPendingUpdate = true;
+	}
+}
+
+
+void SMutableParametersWidget::GenerateAndAttachParameterSlate(const int32 ParamIndex, TSharedPtr<SHorizontalBox> ParameterHorizontalBox,  mu::RangeIndexPtrConst RangeIndex)
+{
+	switch (MutableParameters->GetType(ParamIndex) )
+	{
+	case mu::PARAMETER_TYPE::T_BOOL:
+		ParameterHorizontalBox->AddSlot()
+			.Padding(4.0f)
+			.VAlign(VAlign_Center)
+			[
+				SNew(SCheckBox)
+				.Visibility(this, &SMutableParametersWidget::GetParameterVisibility, ParamIndex)
+				.IsChecked(this, &SMutableParametersWidget::GetBoolParameterValue, ParamIndex, RangeIndex)
+				.OnCheckStateChanged(this, &SMutableParametersWidget::OnBoolParameterChanged, ParamIndex, RangeIndex)
+			];
+		break;
+
+	case mu::PARAMETER_TYPE::T_FLOAT:
+		ParameterHorizontalBox->AddSlot()
+			.Padding(4.0f)
+			.VAlign(VAlign_Center)
+			[
+				SNew(SNumericEntryBox<float>)
+				.Visibility(this, &SMutableParametersWidget::GetParameterVisibility, ParamIndex)
+				.AllowSpin(true)
+				.MinSliderValue(0.0f)
+				.MaxSliderValue(1.0f)			// TODO: Use actual max possible value instead of hard-coding it
+				.Value(this, &SMutableParametersWidget::GetFloatParameterValue, ParamIndex, RangeIndex)
+				.OnValueChanged(this, &SMutableParametersWidget::OnFloatParameterChanged, ParamIndex, RangeIndex)
+				.OnValueCommitted(this,&SMutableParametersWidget::OnFloatParameterCommitted,ParamIndex)
+			];
+		break;
+
+	case mu::PARAMETER_TYPE::T_COLOUR:
+		ParameterHorizontalBox->AddSlot()
+			.Padding(4.0f)
+			.VAlign(VAlign_Center)
+			[
+				SNew(SColorBlock)
+				.Color(this, &SMutableParametersWidget::GetColorParameterValue, ParamIndex, RangeIndex)
+				.ShowBackgroundForAlpha(false)
+				.AlphaDisplayMode(EColorBlockAlphaDisplayMode::Ignore)
+				.UseSRGB(false)
+				.OnMouseButtonDown(this, &SMutableParametersWidget::OnColorBlockMouseButtonDown, ParamIndex, RangeIndex)
+				.Size(FVector2D(10.0f, 10.0f))
+			];
+		break;
+
+	case mu::PARAMETER_TYPE::T_INT:
+	{
+		// If we have a list of options, add a combo box too
+		TSharedPtr<SSearchableComboBox> ParamComboBox;
+		const int32 ValueCount = MutableParameters->GetIntPossibleValueCount(ParamIndex);
+		if (ValueCount > 0)
+		{
+			const FString ToolTipText = FString("None");
+
+			TSharedPtr<TArray<TSharedPtr<FString>>>* FoundOptions = IntParameterOptions.Find(ParamIndex);
+			TArray<TSharedPtr<FString>>& OptionNamesAttribute = FoundOptions && FoundOptions->IsValid() ?
+																*FoundOptions->Get() :
+																*(IntParameterOptions.Add(ParamIndex, MakeShared<TArray<TSharedPtr<FString>>>()).Get());
+
+			OptionNamesAttribute.Empty();
+
+			const int32 Value = MutableParameters->GetIntValue( ParamIndex , RangeIndex);
+			const int32 ValueIndex = MutableParameters->GetIntValueIndex( ParamIndex, Value );
+			for (int32 i = 0; i < ValueCount; ++i)
+			{
+				const FString& ValueText = MutableParameters->GetIntPossibleValueName( ParamIndex, i );
+				OptionNamesAttribute.Add(MakeShared<FString>(ValueText));				
+			}
+
+			ParameterHorizontalBox->AddSlot()
+				.Padding(4.0f)
+				.VAlign(VAlign_Center)
+				[
+					SAssignNew(ParamComboBox, SSearchableComboBox)
+					.SearchVisibility(this, &SMutableParametersWidget::GetParameterVisibility, ParamIndex)
+					.ToolTipText(FText::FromString(ToolTipText))
+					.OptionsSource(&OptionNamesAttribute)
+					.InitiallySelectedItem(OptionNamesAttribute[ValueIndex])
+					.Method(EPopupMethod::UseCurrentWindow)
+					.OnSelectionChanged(this, &SMutableParametersWidget::OnIntParameterTextChanged, ParamIndex, RangeIndex)
+					.OnGenerateWidget(this, &SMutableParametersWidget::OnGenerateWidgetIntParameter)
+					.Content()
+					[
+						SNew(STextBlock)
+						.Text_Lambda([this, ParamIndex, RangeIndex, OptionNamesAttribute]() -> FText
+						{
+							const int32 Value = MutableParameters ? MutableParameters->GetIntValue(ParamIndex, RangeIndex) : 0;
+							const int32 ValueIndex = MutableParameters ? MutableParameters->GetIntValueIndex(ParamIndex, Value) : 0;
+
+							return FText::FromString(*OptionNamesAttribute[ValueIndex]);
+						})
+					]
+				];
+		}
+
+		ParameterHorizontalBox->AddSlot()
+			.Padding(4.0f)
+			.VAlign(VAlign_Center)
+			[
+				SNew(SNumericEntryBox<int32>)
+				.Visibility(this, &SMutableParametersWidget::GetParameterVisibility, ParamIndex)
+				.AllowSpin(true)
+				.MinSliderValue(0)
+				.MaxSliderValue(this, &SMutableParametersWidget::GetIntParameterValueMax, ParamIndex)
+				.Value(this, &SMutableParametersWidget::GetIntParameterValue, ParamIndex, RangeIndex)
+				.OnValueChanged(this, &SMutableParametersWidget::OnIntParameterChanged, ParamIndex, ParamComboBox, RangeIndex)
+			];
+		break;
+	}
+
+	case mu::PARAMETER_TYPE::T_PROJECTOR:
+		ParameterHorizontalBox->AddSlot()
+			.Padding(4.0f)
+			.VAlign(VAlign_Center)
+			[
+				SNew(SGridPanel)
+				.FillColumn(0, 1.0f)
+
+				+SGridPanel::Slot(0,0)
+				.VAlign(VAlign_Center)
+				[
+					SNew(STextBlock)
+					.Text(FText::FromString(TEXT("Location")))
+					.Font(UE_MUTABLE_GET_FONTSTYLE(TEXT("PropertyWindow.NormalFont")))
+				]
+
+				+SGridPanel::Slot(1, 0)
+				[
+					SNew(SNumericVectorInputBox<FVector::FReal>)
+					.X(this, &SMutableParametersWidget::GetProjectorLocation, EAxis::X, ParamIndex, RangeIndex)
+					.Y(this, &SMutableParametersWidget::GetProjectorLocation, EAxis::Y, ParamIndex, RangeIndex)
+					.Z(this, &SMutableParametersWidget::GetProjectorLocation, EAxis::Z, ParamIndex, RangeIndex)
+					.bColorAxisLabels(true)
+					.OnXChanged(this, &SMutableParametersWidget::SetProjectorLocation, ETextCommit::Default, EAxis::X, false, ParamIndex, RangeIndex)
+					.OnYChanged(this, &SMutableParametersWidget::SetProjectorLocation, ETextCommit::Default, EAxis::Y, false, ParamIndex, RangeIndex)
+					.OnZChanged(this, &SMutableParametersWidget::SetProjectorLocation, ETextCommit::Default, EAxis::Z, false, ParamIndex, RangeIndex)
+					.OnXCommitted(this, &SMutableParametersWidget::SetProjectorLocation, EAxis::X, true, ParamIndex, RangeIndex)
+					.OnYCommitted(this, &SMutableParametersWidget::SetProjectorLocation, EAxis::Y, true, ParamIndex, RangeIndex)
+					.OnZCommitted(this, &SMutableParametersWidget::SetProjectorLocation, EAxis::Z, true, ParamIndex, RangeIndex)
+					.AllowSpin(true)
+					.SpinDelta(1)
+				]
+
+				+ SGridPanel::Slot(0, 1)
+				.VAlign(VAlign_Center)
+				[
+					SNew(STextBlock)
+					.Text(FText::FromString(TEXT("Scale")))
+					.Font(UE_MUTABLE_GET_FONTSTYLE("PropertyWindow.NormalFont"))
+				]
+
+				+ SGridPanel::Slot(1, 1)
+				[
+					SNew(SNumericVectorInputBox<FVector::FReal>)
+					.X(this, &SMutableParametersWidget::GetProjectorScale, EAxis::X, ParamIndex, RangeIndex)
+					.Y(this, &SMutableParametersWidget::GetProjectorScale, EAxis::Y, ParamIndex, RangeIndex)
+					.Z(this, &SMutableParametersWidget::GetProjectorScale, EAxis::Z, ParamIndex, RangeIndex)
+					.bColorAxisLabels(true)
+					.OnXChanged(this, &SMutableParametersWidget::SetProjectorScale, ETextCommit::Default, EAxis::X, false, ParamIndex, RangeIndex)
+					.OnYChanged(this, &SMutableParametersWidget::SetProjectorScale, ETextCommit::Default, EAxis::Y, false, ParamIndex, RangeIndex)
+					.OnZChanged(this, &SMutableParametersWidget::SetProjectorScale, ETextCommit::Default, EAxis::Z, false, ParamIndex, RangeIndex)
+					.OnXCommitted(this, &SMutableParametersWidget::SetProjectorScale, EAxis::X, true, ParamIndex, RangeIndex)
+					.OnYCommitted(this, &SMutableParametersWidget::SetProjectorScale, EAxis::Y, true, ParamIndex, RangeIndex)
+					.OnZCommitted(this, &SMutableParametersWidget::SetProjectorScale, EAxis::Z, true, ParamIndex, RangeIndex)
+					.AllowSpin(true)
+					.SpinDelta(1)
+				]
+
+			];
+		break;
+
+	default:
+		// Unsupported parameter type
+		ParameterHorizontalBox->AddSlot()
+			.Padding(4.0f)
+			.VAlign(VAlign_Center)
+			[
+				SNew(STextBlock)
+				.Text(FText::FromString(TEXT("Unsupported parameter type.")))
+				.Justification(ETextJustify::Left)
+				.AutoWrapText(true)
+			];
+	}
+}
+
+
+EVisibility SMutableParametersWidget::GetParameterVisibility(int32 ParamIndex) const
+{
+	// \TODO: Parameter relevancy?
+	if (!MutableParameters
+		|| ParamIndex >= MutableParameters->GetCount() )
+	{
+		return EVisibility::HitTestInvisible;
+	}
+
+	return EVisibility::Visible;
+}
+
+
+ECheckBoxState SMutableParametersWidget::GetBoolParameterValue(int32 ParamIndex,  mu::RangeIndexPtrConst RangeIndex) const
+{
+	if ( !MutableParameters 
+		|| ParamIndex>=MutableParameters->GetCount()
+		|| MutableParameters->GetType(ParamIndex)!=mu::PARAMETER_TYPE::T_BOOL )
+	{
+		return ECheckBoxState::Undetermined;
+	}	
+	
+	return MutableParameters->GetBoolValue(ParamIndex,RangeIndex) ? ECheckBoxState::Checked : ECheckBoxState::Unchecked;
+}
+
+
+void SMutableParametersWidget::OnBoolParameterChanged(ECheckBoxState InCheckboxState, int32 ParamIndex,  mu::RangeIndexPtrConst RangeIndex)
+{
+	if (!MutableParameters
+		|| ParamIndex >= MutableParameters->GetCount()
+		|| MutableParameters->GetType(ParamIndex) != mu::PARAMETER_TYPE::T_BOOL)
+	{
+		return;
+	}
+
+	MutableParameters->SetBoolValue( ParamIndex, InCheckboxState==ECheckBoxState::Checked, RangeIndex );
+
+	OnParametersValueChanged.ExecuteIfBound(ParamIndex);
+}
+
+
+TOptional<float> SMutableParametersWidget::GetFloatParameterValue(int32 ParamIndex, mu::RangeIndexPtrConst RangeIndex) const
+{
+	if (!MutableParameters
+		|| ParamIndex >= MutableParameters->GetCount()
+		|| MutableParameters->GetType(ParamIndex) != mu::PARAMETER_TYPE::T_FLOAT)
+	{
+		return TOptional<float>();
+	}
+
+	return MutableParameters->GetFloatValue(ParamIndex,RangeIndex);
+}
+
+
+void SMutableParametersWidget::OnFloatParameterChanged(float InValue, int32 ParamIndex, mu::RangeIndexPtrConst RangeIndex)
+{
+	if (!MutableParameters
+		|| ParamIndex >= MutableParameters->GetCount()
+		|| MutableParameters->GetType(ParamIndex) != mu::PARAMETER_TYPE::T_FLOAT)
+	{
+		return;
+	}
+
+	MutableParameters->SetFloatValue(ParamIndex, InValue, RangeIndex);
+	
+	OnParametersValueChanged.ExecuteIfBound(ParamIndex);
+}
+
+void SMutableParametersWidget::OnFloatParameterCommitted(float Value, ETextCommit::Type CommitType, int32 ParamIndex)
+{
+	// Think on the possibility of asking for a reconstruction if the new value is a valid integer
+	if (Value >= 1 && Value == (int)Value)
+	{
+		ScheduleUpdateIfRequired(ParamIndex);
+	}
+}
+
+
+TOptional<int32> SMutableParametersWidget::GetIntParameterValue (int32 ParamIndex, mu::RangeIndexPtrConst RangeIndex) const
+{
+	if (!MutableParameters
+		|| ParamIndex >= MutableParameters->GetCount()
+		|| MutableParameters->GetType(ParamIndex) != mu::PARAMETER_TYPE::T_INT)
+	{
+		return TOptional<int32>();
+	}
+
+	int32 Value = MutableParameters->GetIntValue(ParamIndex, RangeIndex);
+	int32 ValueIndex = MutableParameters->GetIntValueIndex(ParamIndex, Value);
+
+	return ValueIndex;
+}
+
+
+TOptional<int32> SMutableParametersWidget::GetIntParameterValueMax(int32 ParamIndex) const
+{
+	if (!MutableParameters
+		|| ParamIndex >= MutableParameters->GetCount()
+		|| MutableParameters->GetType(ParamIndex) != mu::PARAMETER_TYPE::T_INT)
+	{
+		return TOptional<int32>();
+	}
+
+	int32 ValueCount = MutableParameters->GetIntPossibleValueCount(ParamIndex);
+
+	if (!ValueCount)
+	{
+		ValueCount = 16;
+	}
+
+	return ValueCount -1;
+}
+
+
+void SMutableParametersWidget::OnIntParameterChanged(int32 InValue, int32 ParamIndex, TSharedPtr<SSearchableComboBox> Combo, mu::RangeIndexPtrConst RangeIndex )
+{
+	if (!MutableParameters
+		|| ParamIndex >= MutableParameters->GetCount()
+		|| MutableParameters->GetType(ParamIndex) != mu::PARAMETER_TYPE::T_INT)
+	{
+		return;
+	}
+
+	// The slider value is actually the value index for the parameter
+	if (InValue >= 0 && InValue < MutableParameters->GetIntPossibleValueCount(ParamIndex))
+	{
+		int32 RealValue = MutableParameters->GetIntPossibleValue(ParamIndex, InValue);
+		MutableParameters->SetIntValue(ParamIndex, RealValue,RangeIndex);
+		
+		// Update the text combo if any
+		if (Combo)
+		{
+			const FString Text = MutableParameters->GetIntPossibleValueName(ParamIndex, InValue);
+			Combo->RefreshOptions();
+			Combo->SetSelectedItem(MakeShared<FString>(Text), ESelectInfo::Direct);
+		}
+		
+		OnParametersValueChanged.ExecuteIfBound(ParamIndex);
+		
+		ScheduleUpdateIfRequired(ParamIndex);
+	}
+}
+
+
+void SMutableParametersWidget::OnIntParameterTextChanged(TSharedPtr<FString> Selection, ESelectInfo::Type, int32 ParamIndex, mu::RangeIndexPtrConst RangeIndex)
+{
+	if (!Selection
+		|| !MutableParameters
+		|| ParamIndex >= MutableParameters->GetCount()
+		|| MutableParameters->GetType(ParamIndex) != mu::PARAMETER_TYPE::T_INT)
+	{
+		return;
+	}
+
+	int32 ValueIndex = MutableParameters->GetIntValueIndex(ParamIndex, StringCast<ANSICHAR>(**Selection).Get());
+	if (ValueIndex >= 0 && ValueIndex < MutableParameters->GetIntPossibleValueCount(ParamIndex))
+	{
+		int32 RealValue = MutableParameters->GetIntPossibleValue(ParamIndex, ValueIndex);
+		MutableParameters->SetIntValue(ParamIndex, RealValue,RangeIndex);
+		
+		OnParametersValueChanged.ExecuteIfBound(ParamIndex);
+		
+		ScheduleUpdateIfRequired(ParamIndex);
+	}
+}
+
+
+TSharedRef<SWidget> SMutableParametersWidget::OnGenerateWidgetIntParameter(TSharedPtr<FString> InItem) const
+{
+	return SNew(STextBlock).Text(FText::FromString(*InItem.Get()));
+}
+
+
+FLinearColor SMutableParametersWidget::GetColorParameterValue(int32 ParamIndex, mu::RangeIndexPtrConst RangeIndex) const
+{
+	if (!MutableParameters
+		|| ParamIndex >= MutableParameters->GetCount()
+		|| MutableParameters->GetType(ParamIndex) != mu::PARAMETER_TYPE::T_COLOUR)
+	{
+		return FLinearColor();
+	}
+
+	FLinearColor Result;
+	MutableParameters->GetColourValue(ParamIndex, &Result.R, &Result.G, &Result.B, &Result.A, RangeIndex);
+
+	return Result;
+}
+
+
+FReply SMutableParametersWidget::OnColorBlockMouseButtonDown(const FGeometry& MyGeometry, const FPointerEvent& MouseEvent, int32 ParamIndex, mu::RangeIndexPtrConst RangeIndex)
+{
+	if (MouseEvent.GetEffectingButton() != EKeys::LeftMouseButton)
+	{
+		return FReply::Unhandled();
+	}
+
+	FLinearColor Col = GetColorParameterValue(ParamIndex,RangeIndex);
+
+	FColorPickerArgs args;
+	args.bIsModal = true;
+	args.bUseAlpha = true;
+	args.bOnlyRefreshOnMouseUp = false;
+	args.InitialColor = Col;
+	args.OnColorCommitted = FOnLinearColorValueChanged::CreateSP(this, &SMutableParametersWidget::OnSetColorFromColorPicker, ParamIndex, RangeIndex);
+	OpenColorPicker(args);
+
+	return FReply::Handled();
+}
+
+
+void SMutableParametersWidget::OnSetColorFromColorPicker(FLinearColor NewColor, int32 ParamIndex,  mu::RangeIndexPtrConst RangeIndex)
+{
+	if (!MutableParameters
+		|| ParamIndex >= MutableParameters->GetCount()
+		|| MutableParameters->GetType(ParamIndex) != mu::PARAMETER_TYPE::T_COLOUR)
+	{
+		return;
+	}
+	
+	MutableParameters->SetColourValue(ParamIndex, NewColor.R, NewColor.G, NewColor.B, NewColor.A, RangeIndex);
+
+	OnParametersValueChanged.ExecuteIfBound(ParamIndex);
+}
+
+
+TOptional<FVector::FReal> SMutableParametersWidget::GetProjectorLocation(EAxis::Type Axis, int32 ParamIndex, mu::RangeIndexPtrConst RangeIndex) const
+{ 
+	if (!MutableParameters
+		|| ParamIndex >= MutableParameters->GetCount()
+		|| MutableParameters->GetType(ParamIndex) != mu::PARAMETER_TYPE::T_PROJECTOR)
+	{
+		return 0.0f;
+	}
+	
+	mu::PROJECTOR_TYPE Type;
+	FVector3f Pos;
+	FVector3f Dir;
+	FVector3f Up;
+	FVector3f Scale;
+	float Angle;
+	MutableParameters->GetProjectorValue(ParamIndex, &Type, &Pos, &Dir, &Up, &Scale, &Angle , RangeIndex);
+
+	return Pos[ Axis-EAxis::X ];
+}
+
+
+void SMutableParametersWidget::SetProjectorLocation(FVector::FReal NewValue, ETextCommit::Type, EAxis::Type Axis, bool bCommitted, int32 ParamIndex, mu::RangeIndexPtrConst RangeIndex)
+{
+	if (!MutableParameters
+		|| ParamIndex >= MutableParameters->GetCount()
+		|| MutableParameters->GetType(ParamIndex) != mu::PARAMETER_TYPE::T_PROJECTOR)
+	{
+		return;
+	}
+	
+	//mu::RangeIndexPtr UpdatedRange = RangeIndexPayload.GetRangeIndex();
+	
+	mu::PROJECTOR_TYPE Type;
+	FVector3f Pos;
+	FVector3f Dir;
+	FVector3f Up;
+	FVector3f Scale;
+	float Angle;
+	MutableParameters->GetProjectorValue(ParamIndex, &Type, &Pos, &Dir, &Up, &Scale, &Angle, RangeIndex);
+
+	Pos[Axis - EAxis::X] = NewValue;
+
+	MutableParameters->SetProjectorValue(ParamIndex, Pos, Dir, Up, Scale, Angle, RangeIndex);
+
+	OnParametersValueChanged.ExecuteIfBound(ParamIndex);
+}
+
+
+TOptional<FVector::FReal> SMutableParametersWidget::GetProjectorScale(EAxis::Type Axis, int32 ParamIndex, mu::RangeIndexPtrConst RangeIndex) const
+{
+	if (!MutableParameters
+		|| ParamIndex >= MutableParameters->GetCount()
+		|| MutableParameters->GetType(ParamIndex) != mu::PARAMETER_TYPE::T_PROJECTOR)
+	{
+		return 0.0f;
+	}
+
+	mu::PROJECTOR_TYPE Type;
+	FVector3f Pos;
+	FVector3f Dir;
+	FVector3f Up;
+	FVector3f Scale;
+	float Angle;
+	MutableParameters->GetProjectorValue(ParamIndex, &Type, &Pos, &Dir, &Up, &Scale, &Angle, RangeIndex);
+
+	return Scale[Axis - EAxis::X];
+}
+
+
+void SMutableParametersWidget::SetProjectorScale(FVector::FReal NewValue, ETextCommit::Type, EAxis::Type Axis, bool bCommitted, int32 ParamIndex,  mu::RangeIndexPtrConst RangeIndex)
+{
+	if (!MutableParameters
+		|| ParamIndex >= MutableParameters->GetCount()
+		|| MutableParameters->GetType(ParamIndex) != mu::PARAMETER_TYPE::T_PROJECTOR)
+	{
+		return;
+	}
+
+	
+	mu::PROJECTOR_TYPE Type;
+	FVector3f Pos;
+	FVector3f Dir;
+	FVector3f Up;
+	FVector3f Scale;
+	float Angle;
+	MutableParameters->GetProjectorValue(ParamIndex, &Type, &Pos, &Dir, &Up, &Scale, &Angle, RangeIndex);
+
+	Scale[Axis - EAxis::X] = NewValue;
+
+	MutableParameters->SetProjectorValue(ParamIndex, Pos, Dir, Up, Scale, Angle, RangeIndex);
+
+	OnParametersValueChanged.ExecuteIfBound(ParamIndex);
+}
+
+
+#undef LOCTEXT_NAMESPACE
+
diff --git a/Engine/Plugins/Mutable/Source/CustomizableObjectEditor/Private/MuCOE/SMutableParametersWidget.h b/Engine/Plugins/Mutable/Source/CustomizableObjectEditor/Private/MuCOE/SMutableParametersWidget.h
new file mode 100644
--- /dev/null
+++ b/Engine/Plugins/Mutable/Source/CustomizableObjectEditor/Private/MuCOE/SMutableParametersWidget.h
@@ -0,0 +1,107 @@
+// Copyright Epic Games, Inc. All Rights Reserved.
+
+#pragma once
+
+#include "MuR/Parameters.h"
+#include "MuR/Ptr.h"
+#include "Templates/SharedPointer.h"
+#include "Widgets/SBoxPanel.h"
+#include "Widgets/SCompoundWidget.h"
+
+enum class ECheckBoxState : uint8;
+namespace ESelectInfo { enum Type : int; }
+namespace ETextCommit { enum Type : int; }
+
+class FString;
+class SSearchableComboBox;
+class SVerticalBox;
+struct FGeometry;
+struct FPointerEvent;
+
+
+/** Delegate that is executed when a parameter value changes */
+DECLARE_DELEGATE_OneParam(FOnMutableParameterValueChanged, int32);
+
+
+/**  This widget displays and allows edition of a set of parameter values for the lower level Mutable Parameters 
+* object. It is meant to be used for tools and debugging and it doesn't have any knowledge of Customizable Objects
+* and the additional parameter details stored there.
+*/
+class SMutableParametersWidget : public SCompoundWidget
+{
+public:
+
+	SLATE_BEGIN_ARGS(SMutableParametersWidget) {}
+
+	/** Parameters to show and edit. */
+	SLATE_ARGUMENT_DEFAULT(mu::ParametersPtr, Parameters) { nullptr };
+
+	/** Called when any parameter value has changed, with the parameter index as argument.  */
+	SLATE_EVENT(FOnMutableParameterValueChanged, OnParametersValueChanged)
+
+	SLATE_END_ARGS()
+
+	void Construct(const FArguments& InArgs);
+
+	// SWidget interface
+	void Tick(const FGeometry& AllottedGeometry, const double InCurrentTime, const float InDeltaTime) override;
+
+private:
+
+	/** Reference to the mutable parameters. */
+	mu::ParametersPtr MutableParameters;
+
+	/** Is true, the parameter object has changed and we need to update. */
+	bool bIsPendingUpdate = false;
+
+	/** Delegate called when the check box changes state */
+	FOnMutableParameterValueChanged OnParametersValueChanged;
+
+	/** Dynamically filled box with per-parameter widgets */
+	TSharedPtr<SVerticalBox> ParamBox;
+
+	/** Map from ParamIndexInObject to the int param's selector options */
+	TMap<int32, TSharedPtr<TArray<TSharedPtr<FString>>>> IntParameterOptions;
+
+	/** Check if the Parameter index provided is from a parameter that controls the amount of values another parameter/s
+	 * should have exposed in the UI. If the parameter does so then the redraw of the parameters widget will be scheduled
+	 * for the next update.
+	 * @param InParameterIndex The index of the parameter we changed it's value and want to check if that should or not
+	 * trigger the Parameter's Widget redraw.
+	 */
+	void ScheduleUpdateIfRequired(const int32& InParameterIndex);
+	
+	/** Parameter slate generation methods. */
+
+	/** Generate a new slate for the provided parameter index and attaches it to the also provided Horizontal box slate.
+	 * @param ParamIndex The index of the mutable parameter to represent using a slate object.
+	 * @param ParameterHorizontalBox The Horizontal box slate object to contain the newly generated slate object representing the parameter.
+	 * @param RangeIndex Range object pointing at the parameters' target dimension and position. Providing a nullptr is expected for parameters with
+	 * a single value (not a multivalue parameter)
+	 * This object can also represent single value parameters but must be provided in either case.
+	 */
+	void GenerateAndAttachParameterSlate(const int32 ParamIndex, TSharedPtr<SHorizontalBox> ParameterHorizontalBox, mu::RangeIndexPtrConst RangeIndex);
+	
+	/** Internal UI callbacks. */
+	EVisibility GetParameterVisibility(int32 ParamIndex) const;
+	ECheckBoxState GetBoolParameterValue(int32 ParamIndex, mu::RangeIndexPtrConst RangeIndex) const;
+	void OnBoolParameterChanged(ECheckBoxState CheckboxState, int32 ParamIndex, mu::RangeIndexPtrConst RangeIndex);
+	TOptional<float> GetFloatParameterValue(int32 ParamIndex, mu::RangeIndexPtrConst RangeIndex) const;
+	void OnFloatParameterChanged(float Value, int32 ParamIndex, mu::RangeIndexPtrConst RangeIndex);
+	void OnFloatParameterCommitted(float Value, ETextCommit::Type CommitType, int32 ParamIndex);
+	FLinearColor GetColorParameterValue(int32 ParamIndex, mu::RangeIndexPtrConst RangeIndex) const;
+	FReply OnColorBlockMouseButtonDown(const FGeometry& MyGeometry, const FPointerEvent& MouseEvent, int32 ParamIndex, mu::RangeIndexPtrConst RangeIndex);
+	void OnSetColorFromColorPicker(FLinearColor NewColor, int32 ParamIndex, mu::RangeIndexPtrConst RangeIndex);
+	TOptional<int32> GetIntParameterValue(int32 ParamIndex,  mu::RangeIndexPtrConst RangeIndex) const;
+	TOptional<int32> GetIntParameterValueMax(int32 ParamIndex) const;
+	void OnIntParameterChanged(int32 Value, int32 ParamIndex, TSharedPtr<SSearchableComboBox> Combo,  mu::RangeIndexPtrConst RangeIndex);
+	void OnIntParameterTextChanged(TSharedPtr<FString> Selection, ESelectInfo::Type SelectInfo, int32 ParamIndex, mu::RangeIndexPtrConst RangeIndex);
+	TSharedRef<SWidget> OnGenerateWidgetIntParameter(TSharedPtr<FString> InItem) const;
+
+	/** Projector UI callbacks */
+	TOptional<FVector::FReal> GetProjectorLocation(EAxis::Type Axis, int32 ParamIndex,  mu::RangeIndexPtrConst RangeIndex) const;
+	void SetProjectorLocation(FVector::FReal NewValue, ETextCommit::Type, EAxis::Type, bool bCommitted, int32 ParamIndex,  mu::RangeIndexPtrConst RangeIndex);
+	TOptional<FVector::FReal> GetProjectorScale(EAxis::Type Axis, int32 ParamIndex, mu::RangeIndexPtrConst RangeIndex) const;
+	void SetProjectorScale(FVector::FReal NewValue, ETextCommit::Type, EAxis::Type, bool bCommitted, int32 ParamIndex, mu::RangeIndexPtrConst RangeIndex);
+
+};
diff --git a/Engine/Plugins/Mutable/Source/CustomizableObjectEditor/Private/MuCOE/StatePerformanceTesting.cpp b/Engine/Plugins/Mutable/Source/CustomizableObjectEditor/Private/MuCOE/StatePerformanceTesting.cpp
new file mode 100644
--- /dev/null
+++ b/Engine/Plugins/Mutable/Source/CustomizableObjectEditor/Private/MuCOE/StatePerformanceTesting.cpp
@@ -0,0 +1,654 @@
+// Copyright Epic Games, Inc. All Rights Reserved.
+
+
+#include "MuCOE/StatePerformanceTesting.h"
+
+#include "Components/SkeletalMeshComponent.h"
+#include "Editor.h"
+#include "Logging/MessageLog.h"
+#include "MuCO/CustomizableObject.h"
+#include "MuCO/CustomizableObjectInstance.h"
+#include "MuCO/CustomizableObjectInstancePrivate.h"
+#include "MuCO/CustomizableSkeletalComponent.h"
+
+
+#define LOCTEXT_NAMESPACE "CustomizableObjectEditor"
+
+
+FTestingCustomizableObject::FTestingCustomizableObject(TestType TestParameter)
+	: Instance(nullptr)
+	, BatchTestInCourse(false)
+	, BatchTestFinished(false)
+	, ActiveBatchIndex(-1)
+	, ActiveBatch(nullptr)
+	, TestResourcesPresent(false)
+	, Test(TestParameter)
+{
+	HelperCallback = NewObject<UUpdateClassHelper>(GetTransientPackage());
+	HelperCallback->StateData = this;
+}
+
+
+FTestingCustomizableObject::~FTestingCustomizableObject()
+{
+	if (BatchTestInCourse)
+	{
+		FinishTest();
+	}
+
+	HelperCallback = nullptr;
+	Instance = nullptr;	
+}
+
+
+void FTestingCustomizableObject::InitTest()
+{
+
+}
+
+
+void FTestingCustomizableObject::AddReferencedObjects(FReferenceCollector& Collector)
+{
+	Collector.AddReferencedObject(Instance);
+	Collector.AddReferencedObject(HelperCallback);
+
+	for (auto& Elem : CustomizableSkeletalComponents)
+	{
+		Collector.AddReferencedObject(Elem);
+	}
+	for (auto& Elem : SkeletalMeshComponents)
+	{
+		Collector.AddReferencedObject(Elem);
+	}
+}
+
+
+FString FTestingCustomizableObject::GetReferencerName() const
+{
+	return TEXT("FTestingCustomizableObject");
+}
+
+
+void FTestingCustomizableObject::StartTest(UCustomizableObject* CustomizableObject)
+{
+	if (BatchTestInCourse)
+	{
+		FinishTest();
+	}
+	
+	if (!CustomizableObject->IsCompiled())
+	{
+		FText Msg(LOCTEXT("StateTestCompileRequired","Please compile the Customizable Object before starting the test"));
+		FMessageLog MessageLog("Mutable");
+		MessageLog.Notify(Msg, EMessageSeverity::Info, true);
+		return;
+	}
+
+	TestResourcesPresent = true;
+
+	Instance = CustomizableObject->CreateInstance();
+
+	const int32 NumComponents = CustomizableObject->GetComponentCount();
+	SkeletalMeshComponents.Reset(NumComponents);
+	CustomizableSkeletalComponents.Reset(NumComponents);
+
+	for (int32 ComponentIndex = 0; ComponentIndex < NumComponents; ++ComponentIndex)
+	{
+		UCustomizableSkeletalComponent* CustomizableSkeletalComponent = NewObject<UCustomizableSkeletalComponent>(UCustomizableSkeletalComponent::StaticClass());
+		if (CustomizableSkeletalComponent)
+		{
+			CustomizableSkeletalComponents.Add(CustomizableSkeletalComponent);
+			
+			CustomizableSkeletalComponent->CustomizableObjectInstance = Instance;
+			CustomizableSkeletalComponent->ComponentIndex = ComponentIndex;
+			//CustomizableSkeletalComponent->ComponentName = CustomizableObject->GetPrivate()->Components[ComponentIndex].Name;
+
+
+			USkeletalMeshComponent* SkeletalMeshComponent = NewObject<USkeletalMeshComponent>(GetTransientPackage(), NAME_None, RF_Transient);
+			if (SkeletalMeshComponent)
+			{
+				SkeletalMeshComponents.Add(SkeletalMeshComponent);
+
+				CustomizableSkeletalComponent->AttachToComponent(SkeletalMeshComponent, FAttachmentTransformRules::KeepRelativeTransform);
+			}
+		}
+	}
+
+	if (SkeletalMeshComponents.Num() && SkeletalMeshComponents.Num() == CustomizableSkeletalComponents.Num())
+	{
+		Instance->UpdatedDelegate.AddDynamic(HelperCallback, &UUpdateClassHelper::DelegatedCallback);
+		Instance->UpdateSkeletalMeshAsync(true);
+		InitStateAndParameterSequence();
+	}
+}
+
+
+void FTestingCustomizableObject::InitStateAndParameterSequence()
+{
+	BatchTestFinished = false;
+	BatchTestInCourse = true;
+
+	ArrayBatch.Empty();
+
+	InitTest();
+	
+	TickDelegate = FTickerDelegate::CreateRaw(this, &FTestingCustomizableObject::RunStateTestTick);
+	TickDelegateHandle = FTSTicker::GetCoreTicker().AddTicker(TickDelegate, 0.f);
+
+	if (ArrayBatch.Num() > 0)
+	{
+		ActiveBatchIndex = 0;
+		ActiveBatch = &ArrayBatch[0];
+	}
+}
+
+
+void FTestingCustomizableObject::FinishTest()
+{
+	WorstTimeDelegate.Unbind();
+	StateTestEnded.ExecuteIfBound();
+	if (Instance)
+	{
+		//Instance->UpdateBeginDelegate.Unbind();
+		Instance->UpdatedDelegate.RemoveDynamic(HelperCallback, &UUpdateClassHelper::DelegatedCallback);
+	}
+
+	for(UCustomizableSkeletalComponent* CustomizableSkeletalComponent : CustomizableSkeletalComponents)
+	{
+		if (CustomizableSkeletalComponent)
+		{
+			CustomizableSkeletalComponent->DestroyComponent();
+		}
+	}
+
+	for (USkeletalMeshComponent* SkeletalMeshComponent : SkeletalMeshComponents)
+	{
+		if (SkeletalMeshComponent)
+		{
+			SkeletalMeshComponent->DestroyComponent();
+		}
+	}
+
+	CustomizableSkeletalComponents.Empty();
+	SkeletalMeshComponents.Empty();
+
+	Instance = nullptr; // NOTE: is this the right thing? doing like in the CO editor destructor
+
+	TestResourcesPresent = false;
+
+	TickDelegate = nullptr;
+	FTSTicker::GetCoreTicker().RemoveTicker(TickDelegateHandle);
+
+	GEditor->ForceGarbageCollection();
+
+	BatchTestInCourse = false;
+	BatchTestFinished = true;
+}
+
+
+void FTestingCustomizableObject::InstanceUpdateEnd()
+{
+	const int32 MutableRuntimeCycles = Instance->GetPrivate()->LastUpdateMutableRuntimeCycles;
+
+	if ((ActiveBatch != nullptr) && !ActiveBatch->TestFinished)
+	{
+		ActiveBatch->TestFinished = true;
+		ActiveBatch->ElapsedTime = FPlatformTime::ToMilliseconds(MutableRuntimeCycles);
+	}
+}
+
+
+bool FTestingCustomizableObject::RunStateTestTick(float DeltaTime)
+{
+	if (!ActiveBatch->TestStarted)
+	{
+		ActiveBatch->TestStarted = true;
+		LaunchNextBatch();
+		Instance->UpdateSkeletalMeshAsync(true);
+	}
+	
+	if ((ActiveBatch->TestStarted && ActiveBatch->TestFinished && ((ActiveBatchIndex + 1) < ArrayBatch.Num())))
+	{
+		InstanceTestUpdated(ActiveBatch);
+		ActiveBatchIndex++;
+		ActiveBatch = &ArrayBatch[ActiveBatchIndex];
+	}
+	
+	if (((ActiveBatchIndex + 1) == ArrayBatch.Num()) && ActiveBatch->TestStarted && ActiveBatch->TestFinished)
+	{
+		FinishTest();
+		return false;
+	}
+
+	return true;
+}
+
+
+bool FTestingCustomizableObject::GetTestResourcesPresent()
+{
+	return TestResourcesPresent;
+}
+
+
+int32 FTestingCustomizableObject::GetCurrentActiveBatchIndex()
+{
+	return ActiveBatchIndex;
+}
+
+
+int32 FTestingCustomizableObject::GetNumBatches()
+{
+	return ArrayBatch.Num();
+}
+
+bool FTestingCustomizableObject::GetBatchTestInCourse()
+{
+	return BatchTestInCourse;
+}
+
+void FTestingCustomizableObject::LaunchNextBatch()
+{
+	if (ActiveBatch == nullptr)
+	{
+		return;
+	}
+
+	if (ActiveBatch->IsStateChange)
+	{
+		Instance->GetPrivate()->SetState(ActiveBatch->State);
+	}
+	else
+	{
+		switch (ActiveBatch->ParameterType)
+		{
+			case EMutableParameterType::Bool:
+			{
+				Instance->SetBoolParameterSelectedOption(ActiveBatch->Name, ActiveBatch->ValueBool);
+				break;
+			}
+			case EMutableParameterType::Int:
+			{
+				Instance->SetIntParameterSelectedOption(ActiveBatch->Name, ActiveBatch->OptionName);
+				break;
+			}
+			case EMutableParameterType::Float:
+			{
+				Instance->SetFloatParameterSelectedOption(ActiveBatch->Name, ActiveBatch->ValueFloat);
+				break;
+			}
+			case EMutableParameterType::Color:
+			{
+				Instance->SetColorParameterSelectedOption(ActiveBatch->Name, ActiveBatch->ColorValue);
+				break;
+			}
+			case EMutableParameterType::Projector:
+			{
+				// For projectors we have small random offsets to apply to the position
+				FVector Position = Instance->GetProjectorPosition(ActiveBatch->Name);
+				Position *= ActiveBatch->VectorValue;
+				Instance->SetProjectorPosition(ActiveBatch->Name, Position);
+				break;
+			}
+			case EMutableParameterType::Texture:
+			{
+				// What to do in this case?
+				break;
+			}
+			default:
+			{
+				break;
+			}
+		}
+	}
+}
+
+
+void FTestingCustomizableObject::InstanceTestUpdated(BatchParameterChange* BatchChange)
+{
+
+}
+
+
+void UUpdateClassHelper::DelegatedCallback(UCustomizableObjectInstance* Instance)
+{
+	if (StateData != nullptr)
+	{
+		StateData->InstanceUpdateEnd();
+	}
+}
+
+
+FRuntimeTest::FRuntimeTest() : FTestingCustomizableObject(TestType::RuntimeParameter)
+{
+}
+
+
+void FRuntimeTest::LoadStateData(int32 StateIndex)
+{
+	UCustomizableObject* CustomizableObject = Instance->GetCustomizableObject(); 
+
+	if (!CustomizableObject)
+	{
+		return;
+	}
+
+	StateWholeInformation StateData;
+	StateData.State = StateIndex;
+	StateData.StateName = CustomizableObject->GetStateName(StateIndex);
+
+	ArrayBatch.Add(BatchParameterChange(StateIndex, StateData.StateName));
+
+	uint32 MaxParamChanges = 16;
+	
+	// TODO: test no-state available cases
+	uint32 StateParameterCount = CustomizableObject->GetStateParameterCount(StateIndex);
+	for (uint32 i = 0; i < StateParameterCount; ++i)
+	{
+		uint32 ParamIndexInObject = CustomizableObject->GetStateParameterIndex(StateIndex, i);
+		FString RunTimeParameter = CustomizableObject->GetParameterName(ParamIndexInObject);
+		EMutableParameterType Temp = CustomizableObject->GetParameterType(ParamIndexInObject);
+		switch (CustomizableObject->GetParameterType(ParamIndexInObject))
+		{
+		case EMutableParameterType::None:
+		{
+			break;
+		}
+		case EMutableParameterType::Bool:
+		{
+			const TArray<FCustomizableObjectBoolParameterValue>& BoolParameters = Instance->GetPrivate()->GetDescriptor().GetBoolParameters();
+			FCustomizableObjectBoolParameterValue Parameter;
+			Parameter.ParameterName = RunTimeParameter;
+			const uint32 MaxIndex = BoolParameters.Num();
+			for (uint32 j = 0; j < MaxIndex; ++j)
+			{
+				if (BoolParameters[j].ParameterName == Parameter.ParameterName)
+				{
+					Parameter.ParameterValue = BoolParameters[j].ParameterValue;
+					StateData.BoolParametersIndex.Add(j);
+					for (uint32 k = 0; k < (MaxParamChanges / 2); ++k)
+					{
+						ArrayBatch.Add(BatchParameterChange(StateIndex, StateData.StateName, RunTimeParameter, Parameter.ParameterValue));
+						ArrayBatch.Add(BatchParameterChange(StateIndex, StateData.StateName, RunTimeParameter, !Parameter.ParameterValue));
+					}
+				}
+			}
+			StateData.BoolParameters.Add(Parameter);
+			break;
+		}
+		case EMutableParameterType::Int:
+		{
+			const TArray<FCustomizableObjectIntParameterValue>& IntParameters = Instance->GetPrivate()->GetDescriptor().GetIntParameters();
+			FCustomizableObjectIntParameterValue Parameter;
+			Parameter.ParameterName = RunTimeParameter;
+			const uint32 MaxIndex = IntParameters.Num();
+			for (uint32 j = 0; j < MaxIndex; ++j)
+			{
+				if (IntParameters[j].ParameterName == Parameter.ParameterName)
+				{
+					//Parameter.ParameterValue = Descriptor.IntParameters[j].ParameterValue;
+					Parameter.ParameterValueName = IntParameters[j].ParameterValueName;
+					StateData.IntParametersIndex.Add(j);
+
+					int32 ParameterIndex = CustomizableObject->FindParameter(IntParameters[j].ParameterName);
+
+					uint32 NumOptions = CustomizableObject->GetIntParameterNumOptions(ParameterIndex);
+					for (uint32 l = 0; l < NumOptions && l < MaxParamChanges; ++l)
+					{
+						FString PossibleValue = CustomizableObject->GetIntParameterAvailableOption(ParameterIndex, l);
+						ArrayBatch.Add(BatchParameterChange(StateIndex, StateData.StateName, RunTimeParameter, PossibleValue));
+					}
+				}
+			}
+			StateData.IntParameters.Add(Parameter);
+			break; 
+		}
+		case EMutableParameterType::Float:
+		{
+			const TArray<FCustomizableObjectFloatParameterValue>& FloatParameters = Instance->GetPrivate()->GetDescriptor().GetFloatParameters();
+			FCustomizableObjectFloatParameterValue Parameter;
+			Parameter.ParameterName = RunTimeParameter;
+			const uint32 MaxIndex = FloatParameters.Num();
+			for (uint32 j = 0; j < MaxIndex; ++j)
+			{
+				if (FloatParameters[j].ParameterName == Parameter.ParameterName)
+				{
+					Parameter.ParameterValue = FloatParameters[j].ParameterValue;
+					StateData.FloatParametersIndex.Add(j);
+
+					for (uint32 k = 0; k < MaxParamChanges; ++k)
+					{
+						ArrayBatch.Add(BatchParameterChange(StateIndex, StateData.StateName, RunTimeParameter, FMath::RandRange(0.0f, 1.0f)));
+					}
+					break;
+				}
+			}
+			StateData.FloatParameters.Add(Parameter);
+			break;
+		}
+		case EMutableParameterType::Color:
+		{
+			const TArray<FCustomizableObjectVectorParameterValue>& VectorParameters = Instance->GetPrivate()->GetDescriptor().GetVectorParameters();
+			FCustomizableObjectVectorParameterValue Parameter;
+			Parameter.ParameterName = RunTimeParameter;
+			const uint32 MaxIndex = VectorParameters.Num();
+			for (uint32 j = 0; j < MaxIndex; ++j)
+			{
+				if (VectorParameters[j].ParameterName == Parameter.ParameterName)
+				{
+					Parameter.ParameterValue = VectorParameters[j].ParameterValue;
+					StateData.VectorParametersIndex.Add(j);
+					for (uint32 k = 0; k < MaxParamChanges; ++k)
+					{
+						ArrayBatch.Add(BatchParameterChange(StateIndex, StateData.StateName, RunTimeParameter, FLinearColor(FMath::RandRange(0.0f, 1.0f), FMath::RandRange(0.0f, 1.0f), FMath::RandRange(0.0f, 1.0f), 1.0f)));
+					}
+					break;
+				}
+			}
+			StateData.VectorParameters.Add(Parameter);
+			break;
+		}
+		case EMutableParameterType::Projector:
+		{
+			const TArray<FCustomizableObjectVectorParameterValue>& VectorParameters = Instance->GetPrivate()->GetDescriptor().GetVectorParameters();
+			FCustomizableObjectVectorParameterValue Parameter;
+			Parameter.ParameterName = RunTimeParameter;
+			const uint32 MaxIndex = VectorParameters.Num();
+			for (uint32 j = 0; j < MaxIndex; ++j)
+			{
+				if (VectorParameters[j].ParameterName == Parameter.ParameterName)
+				{
+					Parameter.ParameterValue = VectorParameters[j].ParameterValue;
+					StateData.VectorParametersIndex.Add(j);
+					for (uint32 k = 0; k < MaxParamChanges; ++k)
+					{
+						// For projector we multiply the position with a small factor around 1
+						ArrayBatch.Add(BatchParameterChange(StateIndex, StateData.StateName, RunTimeParameter, FVector(FMath::RandRange(0.9f, 1.1f), FMath::RandRange(0.9f, 1.1f), FMath::RandRange(0.9f, 1.1f))));
+					}
+					break;
+				}
+			}
+			StateData.VectorParameters.Add(Parameter);
+			break;
+		}
+		case EMutableParameterType::Texture:
+		{
+			const TArray<FCustomizableObjectTextureParameterValue>& TextureParameters = Instance->GetPrivate()->GetDescriptor().GetTextureParameters();
+			FCustomizableObjectTextureParameterValue Parameter;
+			Parameter.ParameterName = RunTimeParameter;
+			const uint32 MaxIndex = TextureParameters.Num();
+			for (uint32 j = 0; j < MaxIndex; ++j)
+			{
+				if (TextureParameters[j].ParameterName == Parameter.ParameterName)
+				{
+					Parameter.ParameterValue = TextureParameters[j].ParameterValue;
+					StateData.TextureParametersIndex.Add(j);
+					// How would be the testing here done?
+					break;
+				}
+			}
+			StateData.TextureParameters.Add(Parameter);
+			break;
+		}
+		default:
+		{
+			break;
+		}
+		}
+	}
+
+	ArrayState.Add(StateData);
+}
+
+
+void FRuntimeTest::InitTest()
+{
+	MapStateEnterData.Empty();
+	MapStateUpdateData.Empty();
+	MapStateRuntimeData.Empty();
+
+	UCustomizableObject* CustomizableObject = Instance->GetCustomizableObject();
+
+	if (!CustomizableObject)
+	{
+		return;
+	}
+
+	// Build a batch of changes for each state
+	const int32 NumStates = Instance->GetCustomizableObject()->GetStateCount();
+
+	// Enter each state and, for each runtime parameter, set random values (in the case
+	// of integer parameter, set all possible values)
+	for (int32 i = 0; i < NumStates; ++i)
+	{
+		LoadStateData(i);
+	}
+
+	const int32 NumStateChanges = NumStates * 50;
+	uint32 IndexState;
+	for (int32 i = 0; i < NumStateChanges; ++i)
+	{
+		IndexState = FMath::Rand() % NumStates;
+		ArrayBatch.Add(BatchParameterChange(IndexState, ArrayState[IndexState].StateName));
+	}
+}
+
+
+TMap<uint32, MeasuredData> FRuntimeTest::GetMapStateEnterData()
+{
+	return MapStateEnterData;
+}
+
+
+TMap<uint32, MeasuredData> FRuntimeTest::GetMapStateUpdateData()
+{
+	return MapStateUpdateData;
+}
+
+
+TMap<uint32, TArray<MeasuredData>> FRuntimeTest::GetMapStateRuntimeData()
+{
+	return MapStateRuntimeData;
+}
+
+
+void FRuntimeTest::InstanceTestUpdated(BatchParameterChange* BatchChange)
+{
+	if (BatchChange->IsStateChange)
+	{
+		MeasuredData* DataPtr = MapStateEnterData.Find(BatchChange->State);
+
+		if (DataPtr == nullptr)
+		{
+			MeasuredData& Data = MapStateEnterData.Add(BatchChange->State, MeasuredData(BatchChange->StateName));
+			AddMeasuredDataAndReportTime(&Data, BatchChange->ElapsedTime, BatchChange->State);
+			TArray<MeasuredData> ArrayMeasuredData;
+			MapStateRuntimeData.Add(BatchChange->State, ArrayMeasuredData);
+		}
+		else
+		{
+			AddMeasuredDataAndReportTime(DataPtr, BatchChange->ElapsedTime, BatchChange->State);
+		}
+	}
+	else
+	{
+		// Add current change in course to statistics about the instance being previewed
+		MeasuredData* DataPtr = FindStateRuntimeParameter(BatchChange->State, BatchChange->Name);
+
+		if (DataPtr == nullptr)
+		{
+			MeasuredData Data = MeasuredData(BatchChange->Name);
+			AddMeasuredDataAndReportTime(&Data, BatchChange->ElapsedTime, BatchChange->State, BatchChange->Name);
+			TArray<MeasuredData>* ArrayMeasuredData = MapStateRuntimeData.Find(BatchChange->State);
+			ArrayMeasuredData->Add(Data);
+		}
+		else
+		{
+			AddMeasuredDataAndReportTime(DataPtr, BatchChange->ElapsedTime, BatchChange->State, BatchChange->Name);
+		}
+
+		DataPtr = MapStateUpdateData.Find(BatchChange->State);
+		if (DataPtr == nullptr)
+		{
+			MeasuredData Data = MeasuredData(BatchChange->StateName);
+			AddMeasuredDataAndReportTime(&Data, BatchChange->ElapsedTime, BatchChange->State, BatchChange->Name);
+			MapStateUpdateData.Add(BatchChange->State, Data);
+		}
+		else
+		{
+			AddMeasuredDataAndReportTime(DataPtr, BatchChange->ElapsedTime, BatchChange->State, BatchChange->Name);
+		}
+	}
+}
+
+
+MeasuredData* FRuntimeTest::FindStateRuntimeParameter(uint32 State, FString& Name)
+{
+	TArray<MeasuredData>* ArrayData = MapStateRuntimeData.Find(State);
+
+	const uint32 MaxIndex = ArrayData->Num();
+	for (uint32 i = 0; i < MaxIndex; ++i)
+	{
+		if ((*ArrayData)[i].Name == Name)
+		{
+			return &(*ArrayData)[i];
+		}
+	}
+
+	return nullptr;
+}
+
+
+void FRuntimeTest::AddMeasuredDataAndReportTime(MeasuredData* DataPtr, float ElapsedTime, int32 StateIndex, FString ParameterName)
+{
+	if (DataPtr)
+	{
+		if (ElapsedTime > DataPtr->MaxTime)
+		{
+			if (ParameterName.IsEmpty())
+			{
+				WorstTimeDelegate.ExecuteIfBound(ElapsedTime, StateIndex, INDEX_NONE, 0/* TODO: Get lod!*/, Instance, DataPtr);
+			}
+			else
+			{
+				int32 ParameterIndexInCO = INDEX_NONE;
+
+				if (Instance)
+				{
+					if (const UCustomizableObject* Object = Instance->GetCustomizableObject())
+					{
+						ParameterIndexInCO = Object->FindParameter(ParameterName);
+					}
+				}
+
+				if (ParameterIndexInCO != INDEX_NONE)
+				{
+					WorstTimeDelegate.ExecuteIfBound(ElapsedTime, StateIndex, ParameterIndexInCO, 0/* TODO: Get lod!*/, Instance, DataPtr);
+				}
+			}
+		}
+		DataPtr->AddMeasureData(ElapsedTime);
+	}
+}
+
+
+#undef LOCTEXT_NAMESPACE
