[mutable] Moved the Mutable plugin out of Experimental status into Beta.

diff --git a/Engine/Plugins/Mutable/Source/CustomizableObjectEditor/Private/MuCOE/StatePerformanceTesting.h b/Engine/Plugins/Mutable/Source/CustomizableObjectEditor/Private/MuCOE/StatePerformanceTesting.h
new file mode 100644
--- /dev/null
+++ b/Engine/Plugins/Mutable/Source/CustomizableObjectEditor/Private/MuCOE/StatePerformanceTesting.h
@@ -0,0 +1,475 @@
+// Copyright Epic Games, Inc. All Rights Reserved.
+
+#pragma once
+
+#include "Containers/Ticker.h"
+#include "MuCO/CustomizableObjectParameterTypeDefinitions.h"
+#include "UObject/GCObject.h"
+
+#include "StatePerformanceTesting.generated.h"
+
+class UCustomizableObject;
+class UCustomizableObjectInstance;
+class UCustomizableSkeletalComponent;
+class USkeletalMeshComponent;
+class UUpdateClassHelper;
+struct FFrame;
+
+
+enum class TestType
+{
+	RuntimeParameter,
+	CustomizableObjectMemory
+};
+
+
+struct MeasuredData
+{
+	MeasuredData() :
+		Name("")
+		, MeanTime(0.0f)
+		, MinTime(FLT_MAX)
+		, MaxTime(FLT_MIN)
+		, Accumulated(0.0f)
+		, NumMeasures(0.0f)
+	{
+
+	}
+
+	MeasuredData(FString& NameParam) :
+		Name(NameParam)
+		, MeanTime(0.0f)
+		, MinTime(FLT_MAX)
+		, MaxTime(FLT_MIN)
+		, Accumulated(0.0f)
+		, NumMeasures(0.0f)
+	{
+
+	}
+
+	void AddMeasureData(float ElapsedTime)
+	{
+		Accumulated += ElapsedTime;
+		MinTime = FMath::Min(MinTime, ElapsedTime);
+		MaxTime = FMath::Max(MaxTime, ElapsedTime);
+		NumMeasures += 1.0f;
+		MeanTime = Accumulated / NumMeasures;
+	}
+
+
+	MeasuredData operator+=(const MeasuredData& Other)
+	{
+		if (this->Name != Other.Name)
+		{
+			return *this;
+		}
+
+		if ((HasInitialValues(*this) && HasInitialValues(Other)) ||
+			(!HasInitialValues(*this) && HasInitialValues(Other)))
+		{
+			return *this;
+		}
+		else if (HasInitialValues(*this) && !HasInitialValues(Other))
+		{
+			*this = Other;
+			return *this;
+		}
+
+		this->Accumulated += Other.Accumulated;
+		this->MaxTime = FMath::Max(this->MaxTime, Other.MaxTime);
+		this->MinTime = FMath::Min(this->MinTime, Other.MinTime);
+		this->NumMeasures += Other.NumMeasures;
+		this->MeanTime = this->Accumulated / this->NumMeasures;
+		return *this;
+	}
+
+	bool HasInitialValues(const MeasuredData& Data)
+	{
+		if ((Data.MeanTime == 0.0f) &&
+			(Data.MinTime == FLT_MAX) &&
+			(Data.MaxTime == FLT_MIN) &&
+			(Data.Accumulated == 0.0f) &&
+			(Data.NumMeasures == 0.0f))
+			{
+				return true;
+			}
+
+		return false;
+	}
+
+
+	FString Name;
+	float MeanTime;
+	float MinTime;
+	float MaxTime;
+	float Accumulated;
+	float NumMeasures;
+};
+
+
+struct BatchParameterChange
+{
+	/** Float change constructor */
+	BatchParameterChange(int32 StateParam, FString& StateName, FString& NameParam, float ValueFloatParam) :
+		  State(StateParam)
+		, StateName(StateName)
+		, ParameterType(EMutableParameterType::Float)
+		, Name(NameParam)
+		, ValueFloat(ValueFloatParam)
+		, ValueBool(false)
+		, IsStateChange(false)
+		, ElapsedTime(0.0f)
+		, TestStarted(false)
+		, TestFinished(false)
+	{
+
+	}
+
+	/** Integer parameter change constructor */
+	BatchParameterChange(int32 StateParam, FString& StateName, FString& NameParam, FString ParameterOptionParam) :
+		  State(StateParam)
+		, StateName(StateName)
+		, ParameterType(EMutableParameterType::Int)
+		, Name(NameParam)
+		, ValueFloat(-1.0f)
+		, OptionName(ParameterOptionParam)
+		, ValueBool(false)
+		, IsStateChange(false)
+		, ElapsedTime(0.0f)
+		, TestStarted(false)
+		, TestFinished(false)
+	{
+
+	}
+
+	/** Bool change constructor */
+	BatchParameterChange(int32 StateParam, FString& StateName, FString& NameParam, bool ValueBoolParam) :
+		  State(StateParam)
+		, StateName(StateName)
+		, ParameterType(EMutableParameterType::Bool)
+		, Name(NameParam)
+		, ValueFloat(-1.0f)
+		, ValueBool(ValueBoolParam)
+		, IsStateChange(false)
+		, ElapsedTime(0.0f)
+		, TestStarted(false)
+		, TestFinished(false)
+	{
+
+	}
+
+	/** Color change constructor */
+	BatchParameterChange(int32 StateParam, FString& StateName, FString& NameParam, FLinearColor ColorValueParam) :
+		  State(StateParam)
+		, StateName(StateName)
+		, ParameterType(EMutableParameterType::Color)
+		, Name(NameParam)
+		, ValueFloat(-1.0f)
+		, ValueBool(false)
+		, ColorValue(ColorValueParam)
+		, IsStateChange(false)
+		, ElapsedTime(0.0f)
+		, TestStarted(false)
+		, TestFinished(false)
+	{
+
+	}
+
+	/** Projector position change constructor */
+	BatchParameterChange(int32 StateParam, FString& StateName, FString& NameParam, FVector VectorValueParam) :
+		  State(StateParam)
+		, StateName(StateName)
+		, ParameterType(EMutableParameterType::Projector)
+		, Name(NameParam)
+		, ValueFloat(-1.0f)
+		, ValueBool(false)
+		, VectorValue(VectorValueParam)
+		, IsStateChange(false)
+		, ElapsedTime(0.0f)
+		, TestStarted(false)
+		, TestFinished(false)
+	{
+
+	}
+
+	/** State change constructor */
+	BatchParameterChange(int32 StateParam, FString& NameParam) :
+		  State(StateParam)
+		, StateName(NameParam)
+		, ParameterType(EMutableParameterType::None)
+		, Name(NameParam)
+		, ValueFloat(-1.0f)
+		, ValueBool(false)
+		, IsStateChange(true)
+		, ElapsedTime(0.0f)
+		, TestStarted(false)
+		, TestFinished(false)
+	{
+
+	}
+
+	/** Default constructor */
+	BatchParameterChange() :
+		  State(-1)
+		, ParameterType(EMutableParameterType::None)
+		, ValueFloat(-1.0f)
+		, ValueBool(false)
+		, IsStateChange(false)
+		, ElapsedTime(0.0f)
+		, TestStarted(false)
+		, TestFinished(false)
+	{
+
+	}
+
+	/** State of the parameter, in case this change in the batch is for a parameter. Otherwise, is the state to set (as a batch change) */
+	int32 State = -1;
+
+	/** State name (parameter name in the case of a parameter, or state name in the case of a state) */
+	FString StateName;
+
+	/** Parameter type, in case this change in the batch is for a parameter */
+	EMutableParameterType ParameterType = EMutableParameterType::None;
+
+	/** Name (parameter name in the case of a parameter, or state name in the case of a state) */
+	FString Name;
+
+	/** Value to change the parameter to, in case it's a float parameter */
+	float ValueFloat = -1.0f;
+
+	/** Value to change the parameter to, in case it's an integer parameter */
+	FString OptionName;
+
+	/** Value to change the parameter to, in case it's a bool parameter */
+	bool ValueBool = false;
+
+	/** Value to change the parameter to, in case it's a color parameter */
+	FLinearColor ColorValue;
+
+	/** Value to change the parameter to, in case it's a projector parameter: offset the position */
+	FVector VectorValue;
+
+	/** Flag to distinguish between state and parameter changes */
+	bool IsStateChange = false;
+
+	/** Th elpased time to perform this batch change */
+	float ElapsedTime = 0.0f;
+
+	/** Flag to know if the test for this batch test started */
+	bool TestStarted = false;
+
+	/** Flag to know if the test for this batch test finished */
+	bool TestFinished = false;
+};
+
+
+struct StateWholeInformation
+{
+	/** State index */
+	int32 State;
+
+	/** State name */
+	FString StateName;
+
+	/** Copy of run time boolean parameters */
+	TArray<struct FCustomizableObjectBoolParameterValue> BoolParameters;
+
+	/** Index in the Customizable Object of the parameters in BoolParameters */
+	TArray<int32> BoolParametersIndex;
+
+	/** Copy of run time integer parameters */
+	TArray<struct FCustomizableObjectIntParameterValue> IntParameters;
+
+	/** Index in the Customizable Object of the parameters in IntParameters */
+	TArray<int32> IntParametersIndex;
+
+	/** Copy of run time float parameters */
+	TArray<struct FCustomizableObjectFloatParameterValue> FloatParameters;
+
+	/** Index in the Customizable Object of the parameters in FloatParameters */
+	TArray<int32> FloatParametersIndex;
+
+	/** Copy of run time texture parameters */
+	TArray<struct FCustomizableObjectTextureParameterValue> TextureParameters;
+
+	/** Index in the Customizable Object of the parameters in TextureParameters */
+	TArray<int32> TextureParametersIndex;
+
+	/** Copy of run time vector parameters */
+	TArray<struct FCustomizableObjectVectorParameterValue> VectorParameters;
+
+	/** Index in the Customizable Object of the parameters in VectorParameters */
+	TArray<int32> VectorParametersIndex;
+
+	/** Copy of run time projector parameters */
+	TArray<struct FCustomizableObjectProjectorParameterValue> ProjectorParameters;
+
+	/** Index in the Customizable Object of the parameters in ProjectorParameters */
+	TArray<int32> ProjectorParametersIndex;
+};
+
+
+// Callbacks to notify of worst cases found when running the test. It can be called multiple times per test, as worse cases pop up.
+DECLARE_DELEGATE_SixParams(FStateTestWorstTimeFound, float, int32, int32, uint32, UCustomizableObjectInstance*, MeasuredData*);
+// Callback to notify that the test has finished, before data is cleared
+DECLARE_DELEGATE(FStateTestEnded);
+
+
+class FTestingCustomizableObject : public FGCObject
+{
+public:
+
+	// FGCObject interface
+	virtual void AddReferencedObjects(FReferenceCollector& Collector) override;
+	virtual FString GetReferencerName() const override;
+
+public:
+	/** Default constructor */
+	FTestingCustomizableObject(TestType TestParameter);
+
+	/** Default destructor */
+	virtual ~FTestingCustomizableObject();
+
+	/** Init test */
+	virtual void InitTest();
+
+	/** Start the test */
+	void StartTest(UCustomizableObject* CustomizableObject);
+
+	/** Inits the loop through all states and runtime parameters */
+	void InitStateAndParameterSequence();
+
+	/** Finish the test */
+	void FinishTest();
+
+	/** Callback for instance update end delegate */
+	void InstanceUpdateEnd();
+
+	/** Getter of ActiveBatchIndex */
+	int32 GetCurrentActiveBatchIndex();
+
+	/** Getter of ArrayBatch.Num() */
+	int32 GetNumBatches();
+
+	/** Getter of BatchTestInCourse */
+	bool GetBatchTestInCourse();
+
+	/** Coordinates the advaces of each test of the batch test */
+	bool RunStateTestTick(float DeltaTime);
+
+	/** Instance to build to perform the test */
+	TObjectPtr<UCustomizableObjectInstance> Instance;
+
+	/** Instance skeletal component */
+	TArray<TObjectPtr<UCustomizableSkeletalComponent>> CustomizableSkeletalComponents;
+
+	/** Instance skeletal mesh component */
+	TArray<TObjectPtr<USkeletalMeshComponent>> SkeletalMeshComponents;
+
+	/** UObject class to be able to use the update callback */
+	UPROPERTY()
+	TObjectPtr<UUpdateClassHelper> HelperCallback;
+
+	/** Getter of TestResourcesPresent */
+	bool GetTestResourcesPresent();
+
+	/** Delegate called when the worst generation time is found UP UNTIL NOW, with its info and a pointer to the instance. It can be called multiple times on a single test run. */
+	FStateTestWorstTimeFound WorstTimeDelegate;
+
+	/** Delegate called when the stress test finishes, before any data is erased. */
+	FStateTestEnded StateTestEnded;
+
+protected:
+	/** Performs the change in state or runtime parameter according to the information in ActiveBatch */
+	void LaunchNextBatch();
+
+	/** To be called when a running test updates its instance parameters after a BatchParameterChange execution */
+	virtual void InstanceTestUpdated(BatchParameterChange* BatchChange);
+
+	/** Flag to know if there's a batch test running */
+	bool BatchTestInCourse;
+
+	/** Flag to know if the last ran batch test finished */
+	bool BatchTestFinished;
+
+	/** All Customizable Object state information to generate a proper batch of changes */
+	TArray<StateWholeInformation> ArrayState;
+
+	/** Array with the batch of changes to apply */
+	TArray<BatchParameterChange> ArrayBatch;
+
+	/** Active batch change index */
+	int32 ActiveBatchIndex;
+
+	/** Active batch change */
+	BatchParameterChange* ActiveBatch;
+
+	/** Handle for the delegate to run the Tick method in charge of coordinating the test */
+	FTSTicker::FDelegateHandle TickDelegateHandle;
+
+	/** Delegate to run the Tick method in charge of coordinating the test */
+	FTickerDelegate TickDelegate;
+
+	/** Flag to control when test resources are built (for AddReferencedObjects method) */
+	bool TestResourcesPresent;
+
+	/** Enum to know the type of test is being carried out by this instance */
+	TestType Test;
+};
+
+
+UCLASS()
+class UUpdateClassHelper : public UObject
+{
+public:
+	GENERATED_BODY()
+
+	/** Method to assign for the callback */
+	UFUNCTION()
+	void DelegatedCallback(UCustomizableObjectInstance* Instance);
+
+	/** Pointer to existing class to notify */
+	FTestingCustomizableObject* StateData = nullptr;
+};
+
+
+/** Class used to make the runtime parameter and state enter tests */
+class FRuntimeTest : public FTestingCustomizableObject
+{
+public:
+	/** Default constructor */
+	FRuntimeTest();
+
+	/** Loads the state and runtime parameter data of the instance */
+	void LoadStateData(int32 StateIndex);
+
+	/** Init test */
+	void InitTest() final override;
+
+	/** Getter of MapStateEnterData (copy to avoid any possible thread issues) */
+	TMap<uint32, MeasuredData> GetMapStateEnterData();
+
+	/** Getter of MapStateUpdateData (copy to avoid any possible thread issues) */
+	TMap<uint32, MeasuredData> GetMapStateUpdateData();
+
+	/** Getter of MapStateRuntimeData (copy to avoid any possible thread issues) */
+	TMap<uint32, TArray<MeasuredData>> GetMapStateRuntimeData();
+
+protected:
+	/** To be called when a running test updates its instance parameters after a BatchParameterChange execution */
+	void InstanceTestUpdated(BatchParameterChange* BatchChange) final override;
+
+	/** Retrieve pointer to runtime parameter of name given as parameter, for state given as parameter */
+	MeasuredData* FindStateRuntimeParameter(uint32 State, FString& Name);
+
+	void AddMeasuredDataAndReportTime(MeasuredData* DataPtr, float ElapsedTime, int32 StateIndex, FString ParameterName = FString());
+
+	/** Map storing state id -> state entering time data */
+	TMap<uint32, MeasuredData> MapStateEnterData;
+
+	/** Map storing state id -> state update time data */
+	TMap<uint32, MeasuredData> MapStateUpdateData;
+
+	/** Map storing state id -> state runtime parameters measure data */
+	TMap<uint32, TArray<MeasuredData>> MapStateRuntimeData;
+};
diff --git a/Engine/Plugins/Mutable/Source/CustomizableObjectEditor/Private/MuCOE/StressTest.cpp b/Engine/Plugins/Mutable/Source/CustomizableObjectEditor/Private/MuCOE/StressTest.cpp
new file mode 100644
--- /dev/null
+++ b/Engine/Plugins/Mutable/Source/CustomizableObjectEditor/Private/MuCOE/StressTest.cpp
@@ -0,0 +1,1088 @@
+// Copyright Epic Games, Inc. All Rights Reserved.
+
+#include "MuCOE/StressTest.h"
+
+#include "Materials/MaterialInstance.h"
+#include "MuCO/CustomizableObject.h"
+#include "MuCO/CustomizableObjectInstance.h"
+#include "MuCO/CustomizableSkeletalComponent.h"
+#include "MuCO/UnrealPortabilityHelpers.h"
+#include "MuCOE/StatePerformanceTesting.h"
+#include "Engine/SkeletalMesh.h"
+#include "MuCO/CustomizableObjectInstancePrivate.h"
+#include "UObject/Package.h"
+
+class FSkeletalMeshLODRenderData;
+class FSkeletalMeshRenderData;
+
+
+void ULiveInstance::DelegatedCallback(UCustomizableObjectInstance* UpdatedInstance)
+{
+	AddInstanceInformation();
+	Updated = true;
+}
+
+
+void ULiveInstance::AddInstanceInformation()
+{
+	if(!Instance->HasAnySkeletalMesh())
+	{
+		UE_LOG(LogMutable, Warning, TEXT("ERROR: No skeletal meshes found in ULiveInstance::AddInstanceInformation"));
+		return;
+	}
+
+	const int32 MaxIndexLOD = Instance->GetNumLODsAvailable();
+	for (int32 i = 0; i < MaxIndexLOD; ++i)
+	{
+		TMap<FString, MaterialBriefInfo> MapMaterialInfo;
+		uint32 CurrentLODTextureSize = 0;
+		uint32 CurrentLODNumFaces = 0;
+
+		TMap<FString, MeasuredData> MapMeasuredData = ULiveInstance::InitInstanceInformation();
+
+		const int32 NumComponents = Instance->GetNumComponents();
+		for (int32 c = 0; c < NumComponents; ++c)
+		{
+			USkeletalMesh* SkeletalMesh = Instance->GetSkeletalMesh(c);
+			if (!SkeletalMesh)
+			{
+				continue;
+			}
+
+			FSkeletalMeshRenderData* RenderData = SkeletalMesh->GetResourceForRendering();
+			TIndirectArray<FSkeletalMeshLODRenderData>* ArrayLODRenderData = &RenderData->LODRenderData;
+
+			ArrayLODRenderData = &RenderData->LODRenderData;
+			if (i < SkeletalMesh->GetLODNum())
+			{
+				const FSkeletalMeshLODInfo& LODInfo = *SkeletalMesh->GetLODInfo(i);
+
+				CurrentLODNumFaces += RenderData->LODRenderData[i].GetTotalFaces();
+
+				MapMeasuredData["Number of faces"].AddMeasureData((*ArrayLODRenderData)[i].GetTotalFaces());
+				MapMeasuredData["Number of texture coordinates"].AddMeasureData((*ArrayLODRenderData)[i].GetNumTexCoords());
+				const TArray<FSkeletalMaterial>& Materials = SkeletalMesh->GetMaterials();
+				MapMeasuredData["Number of materials"].AddMeasureData(LODInfo.LODMaterialMap.Num());
+
+				for (int32 j = 0; j < LODInfo.LODMaterialMap.Num(); ++j)
+				{
+					if (Materials.IsValidIndex(LODInfo.LODMaterialMap[j]))
+					{
+						UMaterialInterface* MaterialInterface = Materials[LODInfo.LODMaterialMap[j]].MaterialInterface;
+						UMaterialInstance* BaseMaterial = Cast<UMaterialInstance>(MaterialInterface);
+						FString MaterialPath = MaterialInterface->GetPathName();
+						MaterialBriefInfo Info;
+
+						Info.Name = BaseMaterial ? BaseMaterial->GetName() : MaterialInterface->GetName();
+
+						if (BaseMaterial && MaterialPath.Contains("Transient"))
+						{
+							for (int32 k = 0; k < BaseMaterial->TextureParameterValues.Num(); ++k)
+							{
+								if (UTexture* Texture = BaseMaterial->TextureParameterValues[k].ParameterValue)
+								{
+									FString TextPath = Texture->GetPathName();
+									if (Texture->GetPathName().Contains("Transient"))
+									{
+										FIntVector Size = FIntVector(0, 0, 0);
+										Info.ArrayTextureName.Add(Texture->GetName());
+										Size.X = Texture->GetSurfaceWidth();
+										Size.Y = Texture->GetSurfaceHeight();
+										uint32 CurrentTextureSize = Texture->CalcTextureMemorySizeEnum(ETextureMipCount::TMC_AllMips);
+										CurrentLODTextureSize += CurrentTextureSize;
+										MapMeasuredData["Texture size"].AddMeasureData(float(CurrentTextureSize));
+										Info.ArrayTextureSize.Add(Size);
+										Info.ArrayTextureFormat.Add(((UTexture2D*)Texture)->GetPixelFormat());
+									}
+								}
+							}
+							MapMaterialInfo.Add(Info.Name, Info);
+						}
+					}
+				}
+				MapMeasuredData["Instance texture size"].AddMeasureData(CurrentLODTextureSize);
+			}
+		}
+
+		TestDelegate.ExecuteIfBound(MapMeasuredData, MapMaterialInfo, CurrentLODTextureSize, CurrentLODNumFaces, uint32(i), Instance);
+
+		NumInstanceUpdated++;
+	}
+}
+
+
+TMap<FString, MeasuredData> ULiveInstance::InitInstanceInformation()
+{
+	TMap<FString, MeasuredData> MapResult;
+
+	TArray<FString> ArrayParam;
+
+	ArrayParam.Add("Number of faces");
+	ArrayParam.Add("Number of materials");
+	ArrayParam.Add("Number of texture coordinates");
+	ArrayParam.Add("Texture size");
+	ArrayParam.Add("Instance texture size");
+
+	int32 MaxIndex = ArrayParam.Num();
+
+	for (int32 i = 0; i < MaxIndex; ++i)
+	{
+		MapResult.Add(ArrayParam[i], MeasuredData(ArrayParam[i]));
+	}
+
+	return MapResult;
+}
+
+
+FRunningStressTest::FRunningStressTest()
+{
+	NumUpdateRequested = 0;
+}
+
+
+bool FRunningStressTest::RunStressTestTick(float DeltaTime)
+{
+	if (PendingInstanceCount <= 0)
+	{
+		VerifyInstancesThreshold -= DeltaTime;
+
+		if (VerifyInstancesThreshold > 0.0f)
+		{
+			return true;
+		}
+
+		// Verify all instances in the test made their last update before finising the test properly
+		for (TPair<UCustomizableObjectInstance*, ULiveInstance*>& it : LiveInstances)
+		{
+			if (!it.Value->Updated)
+			{
+				VerifyInstancesThreshold = 5.0f;
+				return true;
+			}
+		}
+
+		if (NumUpdatedDone == NumUpdateRequested)
+		{
+			UE_LOG(LogMutable, Log, TEXT("\t\t\t Final Stress Test results for model %s"), *CustomizableObject->GetName());
+		}
+
+		WriteTestResults();
+
+		PrintMostExpensivePerLOD();
+
+		UE_LOG(LogMutable, Log, TEXT("///////////////////////////////////////////////////////////////////////////////////"));
+
+		UE_LOG(LogMutable, Log, TEXT("Finished stress test for model [%s]."), *CustomizableObject->GetName());
+		UE_LOG(LogMutable, Log, TEXT("\n"));
+
+		FinishTest();
+
+		return false;
+	}
+	else
+	{
+		int thisFrameTime = int(DeltaTime*1000.0f);
+		TotalTimeMs += thisFrameTime;
+
+		// Create a new instance if necessary
+		int remainingTime = thisFrameTime;
+		NextInstanceTimeMs -= remainingTime;
+		while (NextInstanceTimeMs < 0)
+		{
+			--PendingInstanceCount;
+
+			// Remaining time
+			remainingTime = -NextInstanceTimeMs;
+
+			UCustomizableObjectInstance* Instance = CustomizableObject->CreateInstance();
+
+			if (!ArrayLODResized)
+			{
+				NumLOD = CustomizableObject->GetNumLODs();
+				ArrayLODMostExpensiveGeometry.SetNum(NumLOD);
+				ArrayLODMostExpensiveTexture.SetNum(NumLOD);
+				ArrayMeasuredDataPerLOD.SetNum(NumLOD);
+				ArrayMaterialInfoPerLOD.SetNum(NumLOD);
+
+				for (uint32 i = 0; i < NumLOD; ++i)
+				{
+					ArrayMeasuredDataPerLOD[i] = ULiveInstance::InitInstanceInformation();
+				}
+
+				ArrayLODResized = true;
+			}
+
+			Instance->SetRandomValues();
+
+			UCustomizableSkeletalComponent* Component = NewObject<UCustomizableSkeletalComponent>(GetTransientPackage(), NAME_None, RF_Transient);
+			Component->CustomizableObjectInstance = Instance;
+
+			ULiveInstance* Data = NewObject<ULiveInstance>(GetTransientPackage());
+			Instance->UpdatedDelegate.AddDynamic(Data, &ULiveInstance::DelegatedCallback);
+			Instance->UpdateSkeletalMeshAsync(true);
+			Data->Updated = false;
+
+			Data->TimeToDie = FMath::Max(InstanceLifeTimeMs, InstanceLifeTimeMs + (FMath::Rand() % (InstanceLifeTimeMsVar * 2)) - InstanceLifeTimeMsVar);
+			Data->TimeToUpdate = InstanceUpdateTimeMs
+				+ (FMath::Rand() % (InstanceUpdateTimeMsVar * 2))
+				- InstanceUpdateTimeMsVar;
+
+			Data->Instance = Instance;
+			Data->NumInstanceUpdated = 0;
+			Data->TestDelegate.BindRaw(this, &FRunningStressTest::AddInstanceData);
+
+			LiveInstances.Add(Instance, Data);
+			NumUpdateRequested++;
+
+			NextInstanceTimeMs = CreateInstanceTimeMs
+				+ (FMath::Rand() % (CreateInstanceTimeMsVar * 2))
+				- CreateInstanceTimeMsVar;
+
+			if (remainingTime > 0)
+			{
+				NextInstanceTimeMs -= remainingTime;
+			}
+
+			UE_LOG(LogMutable, Log, TEXT("Stress test created instance. [%d active]"), LiveInstances.Num());
+		}
+
+		// Update instances that need it
+		for (TPair<UCustomizableObjectInstance*, ULiveInstance*>& it : LiveInstances)
+		{
+			if (!it.Value->Updated)
+			{
+				continue;
+			}
+
+			it.Value->TimeToUpdate -= thisFrameTime;
+			if (it.Value->TimeToUpdate < 0)
+			{
+				it.Key->SetRandomValues();
+				it.Key->UpdateSkeletalMeshAsync(true);
+
+				it.Value->Updated = false;
+				it.Value->TimeToUpdate = InstanceUpdateTimeMs
+					+ (FMath::Rand() % (InstanceUpdateTimeMsVar * 2))
+					- InstanceUpdateTimeMsVar;
+			}
+		}
+
+		// Destroy instances that expired
+		TArray<UCustomizableObjectInstance*> ToRemove;
+		for (TPair<UCustomizableObjectInstance*, ULiveInstance*>& it : LiveInstances)
+		{
+			if (!it.Value) continue;
+
+			it.Value->TimeToDie -= thisFrameTime;
+			if (it.Value->TimeToDie < 0 && it.Value->Updated)
+			{
+				it.Key->UpdatedDelegate.RemoveDynamic(it.Value, &ULiveInstance::DelegatedCallback);
+				it.Value = nullptr;
+				ToRemove.Add(it.Key);
+			}
+		}
+
+		for (UCustomizableObjectInstance* i : ToRemove)
+		{
+			LiveInstances.Remove(i);
+		}
+
+		if (ToRemove.Num())
+		{
+			UE_LOG(LogMutable, Log, TEXT("Stress test removed instance. [%d active]"), LiveInstances.Num());
+		}
+	}
+
+	return true;
+}
+
+
+void FRunningStressTest::InitMeasureTest()
+{
+	StressTestReadyToReset = false;
+	NumUpdateRequested = 0;
+	NumUpdatedDone = 0;
+	ArrayMeasuredDataPerLOD.Empty();
+	ArrayMaterialInfoPerLOD.Empty();
+
+	int32 MaxValue = TNumericLimits<int32>::Max();
+	int32 MinValue = TNumericLimits<int32>::Min();
+
+	TextureMinSize = FIntVector(MaxValue, MaxValue, 0);
+	TextureMaxSize = FIntVector(MinValue, MinValue, 0);
+
+	TestInCourse = true;
+
+	MinTextureSize = TNumericLimits<uint32>::Max();
+	MeanTextureSize = 0;
+	MaxTextureSize = 0;
+
+	StressTestTickDelegate = FTickerDelegate::CreateRaw(this, &FRunningStressTest::RunStressTestTick);
+	StressTestTickDelegateHandle = FTSTicker::GetCoreTicker().AddTicker(StressTestTickDelegate, 0.f);
+
+	ArrayLODMostExpensiveGeometry.Empty();
+	ArrayLODMostExpensiveTexture.Empty();
+
+	ArrayLODResized = false;
+	NumLOD = 0;
+}
+
+
+void FRunningStressTest::WriteTestResults()
+{
+	FString Line;
+	Line = "///////////////////////////////////////////////////////////////////////////////////";
+	UE_LOG(LogMutable, Log, TEXT("%s"), *Line);
+
+	for (uint32 i = 0; i < NumLOD; ++i)
+	{
+		int32 MaxOffsetName = 0;
+		for (TPair<FString, MeasuredData>& it : ArrayMeasuredDataPerLOD[i])
+		{
+			MaxOffsetName = FMath::Max(MaxOffsetName, it.Key.Len());
+		}
+		MaxOffsetName += 5;
+
+		MinTextureSize = uint32(ArrayMeasuredDataPerLOD[i]["Texture size"].MinTime);
+		MeanTextureSize = uint32(ArrayMeasuredDataPerLOD[i]["Texture size"].MeanTime);
+		MaxTextureSize = uint32(ArrayMeasuredDataPerLOD[i]["Texture size"].MaxTime);
+
+		Line = "Results for LOD level ";
+		Line += FString::Printf(TEXT("%d"), i);
+		UE_LOG(LogMutable, Log, TEXT("%s:"), *Line);
+
+		uint32 MaxLength1 = MaxOffsetName + 25;
+		uint32 MaxLength2 = MaxLength1 + 25;
+
+		Line = "\t Parameter";
+		AppendWhitespace(Line, MaxOffsetName - Line.Len());
+		Line += "Min";
+		AppendWhitespace(Line, MaxLength1 - Line.Len());
+		Line += "Mean";
+		AppendWhitespace(Line, MaxLength2 - Line.Len());
+		Line += "Max";
+		UE_LOG(LogMutable, Log, TEXT("%s"), *Line);
+
+		Line = "\t ---------------------------------------------------------------------------------------------";
+		UE_LOG(LogMutable, Log, TEXT("%s"), *Line);
+
+		for (TPair<FString, MeasuredData>& it : ArrayMeasuredDataPerLOD[i])
+		{
+			if (it.Key == "Texture size")
+			{
+				continue;
+			}
+
+			Line = "\t ";
+			Line += FString::Printf(TEXT("%s: "), *(it.Key));
+			AppendWhitespace(Line, MaxOffsetName - Line.Len());
+			Line += FString::Printf(TEXT("%d"), int32(it.Value.MinTime));
+			AppendWhitespace(Line, MaxLength1 - Line.Len());
+			Line += FString::Printf(TEXT("%d"), int32(it.Value.MeanTime));
+			AppendWhitespace(Line, MaxLength2 - Line.Len());
+			Line += FString::Printf(TEXT("%d"), int32(it.Value.MaxTime));
+			UE_LOG(LogMutable, Log, TEXT("%s"), *Line);
+		}
+
+		// Texture size information
+		Line = "\t Texture size ";
+		AppendWhitespace(Line, MaxOffsetName - Line.Len());
+		float Value = (MinTextureSize > 1048576) ? float(MinTextureSize) / 1048576.0f : float(MinTextureSize) / 1024.0f;
+		Line += FString::Printf(TEXT("%.2f"), Value);
+		Line += (MinTextureSize > 1048576) ? "(MB)" : "(KB)";
+		AppendWhitespace(Line, MaxLength1 - Line.Len());
+		Value = (MeanTextureSize > 1048576) ? float(MeanTextureSize) / 1048576.0f : float(MeanTextureSize) / 1024.0f;
+		Line += FString::Printf(TEXT("%.2f"), Value);
+		Line += (MeanTextureSize > 1048576) ? "(MB)" : "(KB)";
+		AppendWhitespace(Line, MaxLength2 - Line.Len());
+		Value = (MaxTextureSize > 1048576) ? float(MaxTextureSize) / 1048576.0f : float(MaxTextureSize) / 1024.0f;
+		Line += FString::Printf(TEXT("%.2f"), Value);
+		Line += (MaxTextureSize > 1048576) ? "(MB)" : "(KB)";
+		UE_LOG(LogMutable, Log, TEXT("%s\n"), *Line);
+	}
+}
+
+
+void FRunningStressTest::AddInstanceData(MeasureDataMap MapMeasuredParam, MaterialInfoMap MapMaterialParam, uint32 CurrentLODTextureSize, uint32 CurrentLODNumFaces, uint32 LOD, UCustomizableObjectInstance* Instance)
+{
+	for (TPair<FString, MeasuredData>& it : MapMeasuredParam)
+	{
+		MeasuredData* Data = ArrayMeasuredDataPerLOD[LOD].Find(it.Key);
+
+		if (Data != nullptr)
+		{
+			*Data += it.Value;
+		}
+	}
+
+	if (ArrayLODMostExpensiveGeometry[LOD].CurrentLODNumFaces <= CurrentLODNumFaces)
+	{
+		CaptureInstanceParameters(Instance, &ArrayLODMostExpensiveGeometry[LOD]);
+		ArrayLODMostExpensiveGeometry[LOD].CurrentLODNumFaces = CurrentLODNumFaces;
+		ArrayLODMostExpensiveGeometry[LOD].CurrentLODTextureSize = CurrentLODTextureSize;
+		ArrayLODMostExpensiveGeometry[LOD].MapMaterial = MapMaterialParam;
+	}
+
+	WorstFacesDelegate.ExecuteIfBound(CurrentLODNumFaces, LOD, Instance, &ArrayMeasuredDataPerLOD[LOD]["Number of faces"]);
+
+	if (ArrayLODMostExpensiveTexture[LOD].CurrentLODTextureSize <= CurrentLODTextureSize)
+	{
+		CaptureInstanceParameters(Instance, &ArrayLODMostExpensiveTexture[LOD]);
+		ArrayLODMostExpensiveTexture[LOD].CurrentLODNumFaces = CurrentLODNumFaces;
+		ArrayLODMostExpensiveTexture[LOD].CurrentLODTextureSize = CurrentLODTextureSize;
+		ArrayLODMostExpensiveTexture[LOD].MapMaterial = MapMaterialParam;
+	}
+
+	WorstTextureDelegate.ExecuteIfBound(MapMaterialParam, CurrentLODTextureSize, LOD, Instance, &ArrayMeasuredDataPerLOD[LOD]["Instance texture size"]);
+
+	for (TPair<FString, MaterialBriefInfo>& it : MapMaterialParam)
+	{
+		MaterialBriefInfo* Data = ArrayMaterialInfoPerLOD[LOD].Find(it.Key);
+
+		if (Data == nullptr)
+		{
+			ArrayMaterialInfoPerLOD[LOD].Add(it.Key, it.Value);
+
+			const int32 MaxIndex = it.Value.ArrayTextureSize.Num();
+			for (int32 i = 0; i < MaxIndex; ++i)
+			{
+				FIntVector& Value = it.Value.ArrayTextureSize[i];
+
+				if ((Value.X <= TextureMinSize.X) || (Value.Y <= TextureMinSize.Y))
+				{
+					TextureMinSize = Value;
+				}
+
+				if ((Value.X >= TextureMaxSize.X) || (Value.Y >= TextureMaxSize.Y))
+				{
+					TextureMaxSize = Value;
+				}
+			}
+		}
+	}
+
+	NumUpdatedDone++;
+
+	FString Line = "Test update instance number ";
+	Line += FString::Printf(TEXT("%d"), NumUpdatedDone);
+	UE_LOG(LogMutable, Log, TEXT("%s"), *Line);
+
+	WriteTestResults();
+}
+
+
+void FRunningStressTest::FinishTest()
+{
+	WorstFacesDelegate.Unbind();
+	WorstTextureDelegate.Unbind();
+	StressTestEnded.ExecuteIfBound();
+	FTSTicker::GetCoreTicker().RemoveTicker(StressTestTickDelegateHandle);
+	StressTestTickDelegate = nullptr;
+
+	StressTestReadyToReset = true;
+
+	for (TPair<UCustomizableObjectInstance*, ULiveInstance*>& it : LiveInstances)
+	{
+		it.Value->TestDelegate.Unbind();
+	}
+
+	LiveInstances.Empty();
+	TestInCourse = false;
+}
+
+
+void FRunningStressTest::CaptureInstanceParameters(UCustomizableObjectInstance* Instance, MostExpensiveInstanceData* Destination)
+{
+	Destination->BoolParameters = Instance->GetPrivate()->GetDescriptor().GetBoolParameters();
+	Destination->IntParameters = Instance->GetPrivate()->GetDescriptor().GetIntParameters();
+	Destination->FloatParameters = Instance->GetPrivate()->GetDescriptor().GetFloatParameters();
+	Destination->TextureParameters = Instance->GetPrivate()->GetDescriptor().GetTextureParameters();
+	Destination->VectorParameters = Instance->GetPrivate()->GetDescriptor().GetVectorParameters();
+	Destination->ProjectorParameters = Instance->GetPrivate()->GetDescriptor().GetProjectorParameters();
+}
+
+
+TMap<class UCustomizableObjectInstance*, ULiveInstance*>& FRunningStressTest::GetLiveInstances()
+{
+	return LiveInstances;
+}
+
+
+bool FRunningStressTest::GetStressTestReadyToReset()
+{
+	return StressTestReadyToReset;
+}
+
+
+bool FRunningStressTest::GetTestInCourse()
+{
+	return TestInCourse;
+}
+
+
+void FRunningStressTest::SetCreateInstanceTimeMs(int32 CreateInstanceTimeMsParam)
+{
+	CreateInstanceTimeMs = CreateInstanceTimeMsParam;
+}
+
+
+void FRunningStressTest::SetCreateInstanceTimeMsVar(int32 CreateInstanceTimeMsVarParam)
+{
+	CreateInstanceTimeMsVar = CreateInstanceTimeMsVarParam;
+}
+
+
+void FRunningStressTest::SetInstanceUpdateTimeMs(int32 InstanceUpdateTimeMsParam)
+{
+	InstanceUpdateTimeMs = InstanceUpdateTimeMsParam;
+}
+
+
+void FRunningStressTest::SetInstanceUpdateTimeMsVar(int32 InstanceUpdateTimeMsVarParam)
+{
+	InstanceUpdateTimeMsVar = InstanceUpdateTimeMsVarParam;
+}
+
+
+void FRunningStressTest::SetInstanceLifeTimeMs(int32 InstanceLifeTimeMsParam)
+{
+	InstanceLifeTimeMs = InstanceLifeTimeMsParam;
+}
+
+
+void FRunningStressTest::SetInstanceLifeTimeMsVar(int32 InstanceLifeTimeMsVarParam)
+{
+	InstanceLifeTimeMsVar = InstanceLifeTimeMsVarParam;
+}
+
+
+void FRunningStressTest::SetPendingInstanceCount(int32 PendingInstanceCountParam)
+{
+	PendingInstanceCount = PendingInstanceCountParam;
+}
+
+
+void FRunningStressTest::SetCustomizableObject(UCustomizableObject* CustomizableObjectParam)
+{
+	CustomizableObject = CustomizableObjectParam;
+}
+
+
+void FRunningStressTest::SetStressTestReadyToReset(bool StressTestReadyToResetParam)
+{
+	StressTestReadyToReset = StressTestReadyToResetParam;
+}
+
+
+void FRunningStressTest::SetVerifyInstancesThreshold(float VerifyInstancesThresholdParam)
+{
+	VerifyInstancesThreshold = VerifyInstancesThresholdParam;
+}
+
+
+void FRunningStressTest::SetNextInstanceTimeMs(float NextInstanceTimeMsParam)
+{
+	NextInstanceTimeMs = NextInstanceTimeMsParam;
+}
+
+
+void FRunningStressTest::AppendWhitespace(FString& Target, int32 Amount)
+{
+	for (int32 i = 0; i < Amount; ++i)
+	{
+		Target += " ";
+	}
+}
+
+
+void FRunningStressTest::PrintMostExpensivePerLOD()
+{
+	const uint32 MaxIndexLOD = ArrayLODMostExpensiveGeometry.Num();
+
+	// For each LOD level
+	for (uint32 i = 0; i < MaxIndexLOD; ++i)
+	{
+		FString Line = "///////////////////////////////////////////////////////////////////////////////////";
+		UE_LOG(LogMutable, Log, TEXT("%s:"), *Line);
+
+		Line = "Instance detailed information for LOD level ";
+		Line += FString::Printf(TEXT("%d:"), i);
+		UE_LOG(LogMutable, Log, TEXT("%s\n"), *Line);
+
+		Line = "\t Information for most expensive instance in terms of geometry size";
+		UE_LOG(LogMutable, Log, TEXT("%s:"), *Line);
+		MostExpensiveInstanceData& DataGeometry = ArrayLODMostExpensiveGeometry[i];
+		PrintMostExpensiveInstanceData(DataGeometry, i);
+
+		Line = "\t Information for most expensive instance in terms of material texture size";
+		UE_LOG(LogMutable, Log, TEXT("%s:"), *Line);
+		MostExpensiveInstanceData& DataTexture = ArrayLODMostExpensiveTexture[i];
+		PrintMostExpensiveInstanceData(DataTexture, i);
+	}
+}
+
+
+void FRunningStressTest::PrintParameterData(const MostExpensiveInstanceData& Data)
+{
+	uint32 MaxLength = FindLongestParameterName(Data);
+	MaxLength += 5;
+
+	FString Line;
+
+	if (Data.BoolParameters.Num() > 0)
+	{
+		PrintParameterHeader(FString("\t\t Bool Parameters"), MaxLength, false);
+		for (const FCustomizableObjectBoolParameterValue& it : Data.BoolParameters)
+		{
+			Line = FString::Printf(TEXT("\t\t\t %s"), *it.ParameterName);
+			AppendWhitespace(Line, MaxLength - Line.Len());
+			Line += FString::Printf(TEXT("%d"), it.ParameterValue);
+			UE_LOG(LogMutable, Log, TEXT("%s"), *Line);
+		}
+	}
+
+	if (Data.IntParameters.Num() > 0)
+	{
+		PrintParameterHeader(FString("\t\t Integer Parameters"), MaxLength, true);
+		for (const FCustomizableObjectIntParameterValue& it : Data.IntParameters)
+		{
+			Line = FString::Printf(TEXT("\t\t\t %s"), *it.ParameterName);
+			AppendWhitespace(Line, MaxLength - Line.Len());
+			Line += FString::Printf(TEXT("%s"), *it.ParameterValueName);
+			UE_LOG(LogMutable, Log, TEXT("%s"), *Line);
+		}
+	}
+
+	if (Data.FloatParameters.Num() > 0)
+	{
+		PrintParameterHeader(FString("\t\t Float Parameters"), MaxLength, true);
+		for (const FCustomizableObjectFloatParameterValue& it : Data.FloatParameters)
+		{
+			Line = FString::Printf(TEXT("\t\t\t %s"), *it.ParameterName);
+			AppendWhitespace(Line, MaxLength - Line.Len());
+			Line += FString::Printf(TEXT("%.2f"), it.ParameterValue);
+			UE_LOG(LogMutable, Log, TEXT("%s"), *Line);
+		}
+	}
+
+	if (Data.TextureParameters.Num() > 0)
+	{
+		PrintParameterHeader(FString("\t\t Texture Parameters"), MaxLength, true);
+		for (const FCustomizableObjectTextureParameterValue& it : Data.TextureParameters)
+		{
+			Line = FString::Printf(TEXT("\t\t\t %s"), *it.ParameterName);
+			AppendWhitespace(Line, MaxLength - Line.Len());
+			Line += it.ParameterValue.ToString();
+			UE_LOG(LogMutable, Log, TEXT("%s"), *Line);
+		}
+	}
+
+	if (Data.VectorParameters.Num() > 0)
+	{
+		PrintParameterHeader(FString("\t\t Vector Parameters"), MaxLength, true);
+		for (const FCustomizableObjectVectorParameterValue& it : Data.VectorParameters)
+		{
+			Line = FString::Printf(TEXT("\t\t\t %s"), *it.ParameterName);
+			AppendWhitespace(Line, MaxLength - Line.Len());
+			Line += FString::Printf(TEXT("(%.2f,"), it.ParameterValue.R);
+			Line += FString::Printf(TEXT("%.2f,"), it.ParameterValue.G);
+			Line += FString::Printf(TEXT("%.2f,"), it.ParameterValue.B);
+			Line += FString::Printf(TEXT("%.2f)"), it.ParameterValue.A);
+			UE_LOG(LogMutable, Log, TEXT("%s"), *Line);
+		}
+	}
+
+	if (Data.ProjectorParameters.Num() > 0)
+	{
+		PrintParameterHeader(FString("\t\t Projector Parameters"), MaxLength, true);
+		for (const FCustomizableObjectProjectorParameterValue& it : Data.ProjectorParameters)
+		{
+			Line = FString::Printf(TEXT("\t\t\t %s"), *it.ParameterName);
+			AppendWhitespace(Line, MaxLength - Line.Len());
+			Line += FString::Printf(TEXT("(%.2f,"), it.Value.Position.X);
+			Line += FString::Printf(TEXT("%.2f,"), it.Value.Position.Y);
+			Line += FString::Printf(TEXT("%.2f)"), it.Value.Position.Z);
+			UE_LOG(LogMutable, Log, TEXT("%s"), *Line);
+		}
+	}
+}
+
+uint32 FRunningStressTest::FindLongestParameterName(const MostExpensiveInstanceData& Data)
+{
+	int32 Result = 0;
+	for (const FCustomizableObjectBoolParameterValue& it : Data.BoolParameters)
+	{
+		Result = FMath::Max(Result, it.ParameterName.Len());
+	}
+
+	for (const FCustomizableObjectIntParameterValue& it : Data.IntParameters)
+	{
+		Result = FMath::Max(Result, it.ParameterName.Len());
+	}
+
+	for (const FCustomizableObjectFloatParameterValue& it : Data.FloatParameters)
+	{
+		Result = FMath::Max(Result, it.ParameterName.Len());
+	}
+
+	for (const FCustomizableObjectTextureParameterValue& it : Data.TextureParameters)
+	{
+		Result = FMath::Max(Result, it.ParameterName.Len());
+	}
+
+	for (const FCustomizableObjectVectorParameterValue& it : Data.VectorParameters)
+	{
+		Result = FMath::Max(Result, it.ParameterName.Len());
+	}
+
+	for (const FCustomizableObjectProjectorParameterValue& it : Data.ProjectorParameters)
+	{
+		Result = FMath::Max(Result, it.ParameterName.Len());
+	}
+
+	return Result;
+}
+
+
+void FRunningStressTest::PrintParameterHeader(FString Param, int32 Value, bool AddNewLine)
+{
+	if (AddNewLine)
+	{
+		UE_LOG(LogMutable, Log, TEXT("\n"));
+	}
+	FString Line = Param;
+	UE_LOG(LogMutable, Log, TEXT("\t\t\t %s"), *Line);
+	Line = "\t\t\t Name";
+	AppendWhitespace(Line, Value - Line.Len());
+	Line += "Value";
+	UE_LOG(LogMutable, Log, TEXT("%s"), *Line);
+	UE_LOG(LogMutable, Log, TEXT("\t\t\t ------------------------------------------"));
+}
+
+
+void FRunningStressTest::PrintMostExpensiveInstanceData(const MostExpensiveInstanceData& Data, uint32 LOD)
+{
+	FString Line;
+	Line = FString::Printf(TEXT("\t\t Number of faces: "));
+	Line += FString::Printf(TEXT("%d"), Data.CurrentLODNumFaces);
+	UE_LOG(LogMutable, Log, TEXT("%s"), *Line);
+
+	// Texture size information
+	Line = "\t\t Texture size: ";
+	float Value = (MinTextureSize > 1048576) ? float(MinTextureSize) / 1048576.0f : float(MinTextureSize) / 1024.0f;
+	Line += FString::Printf(TEXT("min=%.2f"), Value);
+	Line += (MinTextureSize > 1048576) ? "(MB)" : "(KB)";
+	Value = (MaxTextureSize > 1048576) ? float(MaxTextureSize) / 1048576.0f : float(MaxTextureSize) / 1024.0f;
+	Line += FString::Printf(TEXT(", max=%.2f"), Value);
+	Line += (MaxTextureSize > 1048576) ? "(MB)" : "(KB)";
+	UE_LOG(LogMutable, Log, TEXT("%s\n"), *Line);
+
+	// Material information
+	const TMap<FString, MaterialBriefInfo>& MapMaterialLOD = Data.MapMaterial;
+
+	Line = FString::Printf(TEXT("\t\t Material information"));
+
+	if (ArrayMaterialInfoPerLOD[LOD].Num() == 1)
+	{
+		Line += FString::Printf(TEXT(" (1 element):"));
+	}
+	else
+	{
+		Line += FString::Printf(TEXT(" (%d elements):"), MapMaterialLOD.Num());
+	}
+
+	UE_LOG(LogMutable, Log, TEXT("%s"), *Line);
+
+	int32 MaxOffsetName = 0;
+	for (const TPair<FString, MaterialBriefInfo>& itMaterial : MapMaterialLOD)
+	{
+		MaxOffsetName = FMath::Max(MaxOffsetName, itMaterial.Key.Len());
+	}
+	MaxOffsetName += 20;
+	int32 SecondOffsetTexture = MaxOffsetName + 20;
+
+	FString TextureFormatString;
+
+	for (const TPair<FString, MaterialBriefInfo>& itMaterial : MapMaterialLOD)
+	{
+		Line = FString::Printf(TEXT("\t\t\t Material %s: "), *(itMaterial.Key));
+		UE_LOG(LogMutable, Log, TEXT("%s"), *Line);
+
+		int32 MaxIndex = itMaterial.Value.ArrayTextureName.Num();
+		for (int32 i = 0; i < MaxIndex; ++i)
+		{
+			Line = FString::Printf(TEXT("\t\t\t\t Texture %s"), *(itMaterial.Value.ArrayTextureName[i]));
+			AppendWhitespace(Line, MaxOffsetName - Line.Len());
+			Line += FString::Printf(TEXT("size=(%d,"), itMaterial.Value.ArrayTextureSize[i].X);
+			Line += FString::Printf(TEXT("%d)"), itMaterial.Value.ArrayTextureSize[i].Y);
+			AppendWhitespace(Line, SecondOffsetTexture - Line.Len());
+			TextureFormatString = StringyfyEPixelFormat(itMaterial.Value.ArrayTextureFormat[i]);
+			Line += FString::Printf(TEXT("format=%s"), *TextureFormatString);
+			UE_LOG(LogMutable, Log, TEXT("%s"), *Line);
+		}
+	}
+
+	UE_LOG(LogMutable, Log, TEXT("\n"));
+
+	if ((Data.BoolParameters.Num() > 0) || (Data.IntParameters.Num() > 0) || (Data.FloatParameters.Num() > 0) ||
+		(Data.TextureParameters.Num() > 0) || (Data.VectorParameters.Num() > 0) || (Data.ProjectorParameters.Num() > 0))
+	{
+		Line = FString::Printf(TEXT("\t\t Parameter information"));
+		UE_LOG(LogMutable, Log, TEXT("%s"), *Line);
+
+		// Parameter information
+		PrintParameterData(Data);
+	}
+}
+
+
+FString FRunningStressTest::StringyfyEPixelFormat(EPixelFormat Value)
+{
+	switch (Value)
+	{
+		case PF_Unknown:
+		{
+			return FString("PF_Unknown");
+		}
+		case PF_A32B32G32R32F:
+		{
+			return FString("PF_A32B32G32R32F");
+		}
+		case PF_B8G8R8A8:
+		{
+			return FString("PF_B8G8R8A8");
+		}
+		case PF_G8:
+		{
+			return FString("PF_G8");
+		}
+		case PF_G16:
+		{
+			return FString("PF_G16");
+		}
+		case PF_DXT1:
+		{
+			return FString("PF_DXT1");
+		}
+		case PF_DXT3:
+		{
+			return FString("PF_DXT3");
+		}
+		case PF_DXT5:
+		{
+			return FString("PF_DXT5");
+		}
+		case PF_UYVY:
+		{
+			return FString("PF_UYVY");
+		}
+		case PF_FloatRGB:
+		{
+			return FString("PF_FloatRGB");
+		}
+		case PF_FloatRGBA:
+		{
+			return FString("PF_FloatRGBA");
+		}
+		case PF_DepthStencil:
+		{
+			return FString("PF_DepthStencil");
+		}
+		case PF_ShadowDepth:
+		{
+			return FString("PF_ShadowDepth");
+		}
+		case PF_R32_FLOAT:
+		{
+			return FString("PF_R32_FLOAT");
+		}
+		case PF_G16R16:
+		{
+			return FString("PF_G16R16");
+		}
+		case PF_G16R16F:
+		{
+			return FString("PF_G16R16F");
+		}
+		case PF_G16R16F_FILTER:
+		{
+			return FString("PF_G16R16F_FILTER");
+		}
+		case PF_G32R32F:
+		{
+			return FString("PF_G32R32F");
+		}
+		case PF_A2B10G10R10:
+		{
+			return FString("PF_A2B10G10R10");
+		}
+		case PF_A16B16G16R16:
+		{
+			return FString("PF_A16B16G16R16");
+		}
+		case PF_D24:
+		{
+			return FString("PF_D24");
+		}
+		case PF_R16F:
+		{
+			return FString("PF_R16F");
+		}
+		case PF_R16F_FILTER:
+		{
+			return FString("PF_R16F_FILTER");
+		}
+		case PF_BC5:
+		{
+			return FString("PF_BC5");
+		}
+		case PF_V8U8:
+		{
+			return FString("PF_V8U8");
+		}
+		case PF_A1:
+		{
+			return FString("PF_A1");
+		}
+		case PF_FloatR11G11B10:
+		{
+			return FString("PF_FloatR11G11B10");
+		}
+		case PF_A8:
+		{
+			return FString("PF_A8");
+		}
+		case PF_R32_UINT:
+		{
+			return FString("PF_R32_UINT");
+		}
+		case PF_R32_SINT:
+		{
+			return FString("PF_R32_SINT");
+		}
+		case PF_PVRTC2:
+		{
+			return FString("PF_PVRTC2");
+		}
+		case PF_PVRTC4:
+		{
+			return FString("PF_PVRTC4");
+		}
+		case PF_R16_UINT:
+		{
+			return FString("PF_R16_UINT");
+		}
+		case PF_R16_SINT:
+		{
+			return FString("PF_R16_SINT");
+		}
+		case PF_R16G16B16A16_UINT:
+		{
+			return FString("PF_R16G16B16A16_UINT");
+		}
+		case PF_R16G16B16A16_SINT:
+		{
+			return FString("PF_R16G16B16A16_SINT");
+		}
+		case PF_R5G6B5_UNORM:
+		{
+			return FString("PF_R5G6B5_UNORM");
+		}
+		case PF_R8G8B8A8:
+		{
+			return FString("PF_R8G8B8A8");
+		}
+		case PF_A8R8G8B8:
+		{
+			return FString("PF_A8R8G8B8");
+		}
+		case PF_BC4:
+		{
+			return FString("PF_BC4");
+		}
+		case PF_R8G8:
+		{
+			return FString("PF_R8G8");
+		}
+		case PF_ATC_RGB:
+		{
+			return FString("PF_ATC_RGB");
+		}
+		case PF_ATC_RGBA_E:
+		{
+			return FString("PF_ATC_RGBA_E");
+		}
+		case PF_ATC_RGBA_I:
+		{
+			return FString("PF_ATC_RGBA_I");
+		}
+		case PF_X24_G8:
+		{
+			return FString("PF_X24_G8");
+		}
+		case PF_ETC1:
+		{
+			return FString("PF_ETC1");
+		}
+		case PF_ETC2_RGB:
+		{
+			return FString("PF_ETC2_RGB");
+		}
+		case PF_ETC2_RGBA:
+		{
+			return FString("PF_ETC2_RGBA");
+		}
+		case PF_R32G32B32A32_UINT:
+		{
+			return FString("PF_R32G32B32A32_UINT");
+		}
+		case PF_R16G16_UINT:
+		{
+			return FString("PF_R16G16_UINT");
+		}
+		case PF_ASTC_4x4:
+		{
+			return FString("PF_ASTC_4x4");
+		}
+		case PF_ASTC_6x6:
+		{
+			return FString("PF_ASTC_6x6");
+		}
+		case PF_ASTC_8x8:
+		{
+			return FString("PF_ASTC_8x8");
+		}
+		case PF_ASTC_10x10:
+		{
+			return FString("PF_ASTC_10x10");
+		}
+		case PF_ASTC_12x12:
+		{
+			return FString("PF_ASTC_12x12");
+		}
+		case PF_BC6H:
+		{
+			return FString("PF_BC6H");
+		}
+		case PF_BC7:
+		{
+			return FString("PF_BC7");
+		}
+		case PF_R8_UINT:
+		{
+			return FString("PF_R8_UINT");
+		}
+		case PF_L8:
+		{
+			return FString("PF_L8");
+		}
+		case PF_MAX:
+		{
+			return FString("PF_MAX");
+		}
+		default:
+		{
+			return FString("Error");
+		}
+	}
+}
diff --git a/Engine/Plugins/Mutable/Source/CustomizableObjectEditor/Private/MuCOE/StressTest.h b/Engine/Plugins/Mutable/Source/CustomizableObjectEditor/Private/MuCOE/StressTest.h
new file mode 100644
--- /dev/null
+++ b/Engine/Plugins/Mutable/Source/CustomizableObjectEditor/Private/MuCOE/StressTest.h
@@ -0,0 +1,295 @@
+// Copyright Epic Games, Inc. All Rights Reserved.
+
+#pragma once
+
+#include "Containers/Ticker.h"
+#include "MuCO/CustomizableObjectParameterTypeDefinitions.h"
+
+#include "StressTest.generated.h"
+
+enum EPixelFormat : uint8;
+struct MeasuredData;
+
+class UCustomizableObject;
+class UCustomizableObjectInstance;
+struct FFrame;
+
+
+struct MaterialBriefInfo
+{
+	MaterialBriefInfo()
+	{
+
+	}
+
+	MaterialBriefInfo(FString& NameParam, TArray<FString>& ArrayTextureNameParam, TArray<FIntVector> ArrayTextureSizeParam, TArray<EPixelFormat> ArrayTextureFormatParam) :
+		Name(NameParam)
+		, ArrayTextureName(ArrayTextureNameParam)
+		, ArrayTextureSize(ArrayTextureSizeParam)
+		, ArrayTextureFormat(ArrayTextureFormatParam)
+	{
+
+	}
+
+	FString Name;
+	TArray<FString> ArrayTextureName;
+	TArray<FIntVector> ArrayTextureSize;
+	TArray<EPixelFormat> ArrayTextureFormat;
+};
+
+
+struct MostExpensiveInstanceData
+{
+	MostExpensiveInstanceData():
+		  CurrentLODNumFaces(0)
+		, CurrentLODTextureSize(0)
+	{
+
+	}
+
+	FString Name;
+	uint32 CurrentLODNumFaces;
+	uint32 CurrentLODTextureSize;
+	TMap<FString, MaterialBriefInfo> MapMaterial;
+	TArray<struct FCustomizableObjectBoolParameterValue> BoolParameters;
+	TArray<struct FCustomizableObjectIntParameterValue> IntParameters;
+	TArray<struct FCustomizableObjectFloatParameterValue> FloatParameters;
+	TArray<struct FCustomizableObjectTextureParameterValue> TextureParameters;
+	TArray<struct FCustomizableObjectVectorParameterValue> VectorParameters;
+	TArray<struct FCustomizableObjectProjectorParameterValue> ProjectorParameters;
+};
+
+
+typedef TMap<FString, MeasuredData> MeasureDataMap;
+typedef TMap<FString, MaterialBriefInfo> MaterialInfoMap;
+
+
+// Callback to add update information to test
+DECLARE_DELEGATE_SixParams(FStressTestInstanceUpdate, MeasureDataMap, MaterialInfoMap, uint32, uint32, uint32, UCustomizableObjectInstance*);
+// Callbacks to notify of worst cases found when running the test. It can be called multiple times per test, as worse cases pop up.
+DECLARE_DELEGATE_FourParams(FStressTestWorstNumFacesFound, uint32, uint32, UCustomizableObjectInstance*, MeasuredData*);
+DECLARE_DELEGATE_FiveParams(FStressTestWorstTextureSizeFound, MaterialInfoMap, uint32, uint32, UCustomizableObjectInstance*, MeasuredData*);
+// Callback to notify that the test has finished, before data is cleared
+DECLARE_DELEGATE(FStressTestEnded);
+
+UCLASS()
+class ULiveInstance : public UObject
+{
+public:
+	GENERATED_BODY()
+
+	/** Method to assign for the callback */
+	UFUNCTION()
+	void DelegatedCallback(UCustomizableObjectInstance* Instance);
+
+	/** Add to ArrayInformation the relevant information from Instance */
+	void AddInstanceInformation();
+
+	/** Init elements in MapMeasuredData the elevant information from Instance */
+	static TMap<FString, MeasuredData> InitInstanceInformation();
+
+	/** Amount of time an instance is kept before destroying it (the instance needs to complete at least the first skeletal mesh update) */
+	int TimeToDie;
+
+	/** After an instance has been updated, if this value is negative, a new skeleta lmesh update with random values is requested */
+	int TimeToUpdate;
+
+	/** Flag to control whether the instance has been updated */
+	bool Updated;
+
+	/** Instace being updated with random values */
+	UCustomizableObjectInstance* Instance;
+
+	/** Number of instances updated */
+	uint32 NumInstanceUpdated;
+
+	/** Delgate to call the owning test and add results */
+	FStressTestInstanceUpdate TestDelegate;
+};
+
+
+class FRunningStressTest
+{
+public:
+	/** Default constructor */
+	FRunningStressTest();
+
+	/** Ticker used to contorl logic during stress test */
+	bool RunStressTestTick(float DeltaTime);
+
+	/** Setup map to gather data and ticker */
+	void InitMeasureTest();
+
+	/** Show in the console log the test results */
+	void WriteTestResults();
+
+	/** Callback for new instance update information */
+	void AddInstanceData(MeasureDataMap MapMeasuredParam, MaterialInfoMap MapMaterialParam, uint32 CurrentLODTextureSize, uint32 CurrentLODNumFaces, uint32 LOD, UCustomizableObjectInstance* Instance);
+
+	/** Finish current test */
+	void FinishTest();
+
+	/** Will retrieve the parameters present in Instance and set them in Destination */
+	void CaptureInstanceParameters(UCustomizableObjectInstance* Instance, MostExpensiveInstanceData* Destination);
+
+	/** Getter of LiveInstances */
+	TMap<class UCustomizableObjectInstance*, ULiveInstance*>& GetLiveInstances();
+
+	/** Getter of StressTestReadyToReset */
+	bool GetStressTestReadyToReset();
+
+	/** Getter of TestInCourse */
+	bool GetTestInCourse();
+
+	/** Setter of CreateInstanceTimeMs */
+	void SetCreateInstanceTimeMs(int32 CreateInstanceTimeMsParam);
+
+	/** Setter of CreateInstanceTimeMsVar */
+	void SetCreateInstanceTimeMsVar(int32 CreateInstanceTimeMsVarParam);
+
+	/** Setter of InstanceUpdateTimeMs */
+	void SetInstanceUpdateTimeMs(int32 InstanceUpdateTimeMsParam);
+
+	/** Setter of InstanceUpdateTimeMsVar */
+	void SetInstanceUpdateTimeMsVar(int32 InstanceUpdateTimeMsVarParam);
+
+	/** Setter of InstanceLifeTimeMs */
+	void SetInstanceLifeTimeMs(int32 InstanceLifeTimeMsParam);
+
+	/** Setter of InstanceLifeTimeMsVar */
+	void SetInstanceLifeTimeMsVar(int32 InstanceLifeTimeMsVarParam);
+
+	/** Setter of PendingInstanceCount */
+	void SetPendingInstanceCount(int32 PendingInstanceCountParam);
+
+	/** Setter of CustomizableObject */
+	void SetCustomizableObject(UCustomizableObject* CustomizableObjectParam);
+
+	/** Setter of StressTestReadyToReset */
+	void SetStressTestReadyToReset(bool StressTestReadyToResetParam);
+
+	/** Setter of SetVerifyInstancesThreshold */
+	void SetVerifyInstancesThreshold(float VerifyInstancesThresholdParam);
+
+	/** Setter of NextInstanceTimeMs */
+	void SetNextInstanceTimeMs(float NextInstanceTimeMsParam);
+
+	/** Delegate called when the worst case for number of faces is found UP UNTIL NOW, with its info and a pointer to the instance. It can be called multiple times on a single test run. */
+	FStressTestWorstNumFacesFound WorstFacesDelegate;
+
+	/** Delegate called when the worst case for texture size is found UP UNTIL NOW, with its info and a pointer to the instance. It can be called multiple times on a single test run. */
+	FStressTestWorstTextureSizeFound WorstTextureDelegate;
+
+	/** Delegate called when the stress test finishes, before any data is erased. */
+	FStressTestEnded StressTestEnded;
+
+	/** Return the string verison of the pixel format enum given as parameter */
+	static FString StringyfyEPixelFormat(EPixelFormat Value);
+
+protected:
+	/** Append Amount whitespace characters to Target */
+	void AppendWhitespace(FString& Target, int32 Amount);
+
+	/** Output face, texture size and material information in ArrayLODMostExpensive */
+	void PrintMostExpensivePerLOD();
+
+	/** Print parameter data in Data */
+	void PrintParameterData(const MostExpensiveInstanceData& Data);
+
+	/** Utility method to find the length of the longest parameter name in Data */
+	uint32 FindLongestParameterName(const MostExpensiveInstanceData& Data);
+
+	/** Utility method to add the headers used in PrintMostExpensivePerLOD */
+	void PrintParameterHeader(FString Param, int32 Value, bool AddNewLine);
+
+	/** Utility method to print the parameter information */
+	void PrintMostExpensiveInstanceData(const MostExpensiveInstanceData& Data, uint32 LOD);
+
+	/** Value used to determine how frequently a new instance should be made */
+	int32 CreateInstanceTimeMs;
+
+	/** Value used to determine how frequently a new instance should be made */
+	int32 CreateInstanceTimeMsVar;
+
+	/** Value used to determine how frequently an instance should be updated (the ULiveInstance::Updated flag is used as restriction) */
+	int32 InstanceUpdateTimeMs;
+
+	/** Value used to determine how frequently an instance should be updated (the ULiveInstance::Updated flag is used as restriction) */
+	int32 InstanceUpdateTimeMsVar;
+
+	/** Value used to set the maximum amount of time an instance is updated before being destroyed (the ULiveInstance::Updated flag is used as restriction) */
+	int32 InstanceLifeTimeMs;
+
+	/** Value used to set the maximum amount of time an instance is updated before being destroyed (the ULiveInstance::Updated flag is used as restriction) */
+	int32 InstanceLifeTimeMsVar;
+
+	/** Number of instances for the test */
+	int32 PendingInstanceCount;
+
+	/** Elapsed time in the current test */
+	int64 TotalTimeMs = 0;
+
+	/** Threshold time to build a new instance in the ticker method */
+	int NextInstanceTimeMs = 0;
+
+	/** Ticker delegate handle */
+	FTSTicker::FDelegateHandle StressTestTickDelegateHandle;
+
+	/** Ticker handle */
+	FTickerDelegate StressTestTickDelegate;
+
+	/** Map with the customizable object instances and the corresponding data */
+	UPROPERTY()
+	TMap<class UCustomizableObjectInstance*, ULiveInstance*> LiveInstances;
+
+	/** Threshold time to verify all instances in the test have finished updating */
+	float VerifyInstancesThreshold;
+
+	/** Customizable object to build the instances */
+	UCustomizableObject* CustomizableObject;
+
+	/** True when all the instances in the stress test have finished being updated */
+	bool StressTestReadyToReset;
+
+	/** Counter to know the total number of instance update requested */
+	uint32 NumUpdateRequested;
+
+	/** Counter to know the total number of instance update already done */
+	uint32 NumUpdatedDone;
+
+	/** Information about each random values skeleton update carried out by this instance, each array element covers per-LOD data */
+	TArray<TMap<FString, MeasuredData>> ArrayMeasuredDataPerLOD;
+
+	/** To have a per-material breakdown information */
+	TArray<TMap<FString, MaterialBriefInfo>> ArrayMaterialInfoPerLOD;
+
+	/** Smallest texture found in the test */
+	FIntVector TextureMinSize;
+
+	/** Biggest texture found in the test */
+	FIntVector TextureMaxSize;
+
+	/** Flag to know if a test is in course */
+	bool TestInCourse;
+
+	/** Min texture size in bytes found in the test */
+	uint32 MinTextureSize;
+
+	/** Mean texture size in bytes found in the test */
+	uint32 MeanTextureSize;
+
+	/** Max texture size in bytes found in the test */
+	uint32 MaxTextureSize;
+
+	/** Array with per-LOD detailed information of the most expensive instance built regarding geometry data */
+	TArray<MostExpensiveInstanceData> ArrayLODMostExpensiveGeometry;
+
+	/** Array with per-LOD detailed information of the most expensive instance built regarding texture size data */
+	TArray<MostExpensiveInstanceData> ArrayLODMostExpensiveTexture;
+
+	/** Flag to control when to initialize the number of elements of ArrayLODMostExpensive in each test */
+	bool ArrayLODResized;
+
+	/** Number of LOD of the Customizable Object */
+	uint32 NumLOD;
+};
diff --git a/Engine/Plugins/Mutable/Source/CustomizableObjectEditor/Private/MuCOE/Thumbnails/CustomizableObjectInstanceThumbnailRenderer.cpp b/Engine/Plugins/Mutable/Source/CustomizableObjectEditor/Private/MuCOE/Thumbnails/CustomizableObjectInstanceThumbnailRenderer.cpp
new file mode 100644
--- /dev/null
+++ b/Engine/Plugins/Mutable/Source/CustomizableObjectEditor/Private/MuCOE/Thumbnails/CustomizableObjectInstanceThumbnailRenderer.cpp
@@ -0,0 +1,133 @@
+// Copyright Epic Games, Inc. All Rights Reserved.
+
+#include "MuCOE/Thumbnails/CustomizableObjectInstanceThumbnailRenderer.h"
+
+#include "CanvasTypes.h"
+#include "MuCO/CustomizableObjectInstance.h"
+#include "MuCO/CustomizableSkeletalComponent.h"
+#include "MuCO/CustomizableObjectCompilerTypes.h"
+#include "MuCOE/CustomizableObjectCompiler.h"
+#include "SceneView.h"
+#include "TextureResource.h"
+#include "ThumbnailHelpers.h"
+#include "UObject/ConstructorHelpers.h"
+
+class FRenderTarget;
+class USkeletalMesh;
+
+#define LOCTEXT_NAMESPACE "CustomizableObjectInstanceThumbnailRenderer"
+
+UCustomizableObjectInstanceThumbnailRenderer::UCustomizableObjectInstanceThumbnailRenderer()
+	: Super()
+{
+	ThumbnailScene = nullptr;
+	//CustomizableObjectInstance = ConstructObject<UCustomizableObjectInstance>(UCustomizableObjectInstance::StaticClass(), GetTransientPackage(), NAME_None, RF_Transient);
+
+	// Structure to hold one-time initialization
+	struct FConstructorStatics
+	{
+		ConstructorHelpers::FObjectFinder<UTexture2D> PSysThumbnail_NoImage;
+		ConstructorHelpers::FObjectFinder<UTexture2D> PSysThumbnail_OOD;
+		FConstructorStatics()
+			: PSysThumbnail_NoImage(TEXT("/Engine/EditorMaterials/ParticleSystems/PSysThumbnail_NoImage"))
+			, PSysThumbnail_OOD(TEXT("/Engine/EditorMaterials/ParticleSystems/PSysThumbnail_OOD"))
+		{
+		}
+	};
+	static FConstructorStatics ConstructorStatics;
+
+	NoImage = ConstructorStatics.PSysThumbnail_NoImage.Object;
+}
+
+void UCustomizableObjectInstanceThumbnailRenderer::Draw(UObject* Object, int32 X, int32 Y, uint32 Width, uint32 Height, FRenderTarget* RenderTarget, FCanvas* Canvas, bool bAdditionalViewFamily)
+{
+	UCustomizableObjectInstance* CustomizableObjectInstance = Cast<UCustomizableObjectInstance>(Object);
+	
+	if (CustomizableObjectInstance)
+	{
+		UCustomizableObject* CustomizableObject = CustomizableObjectInstance->GetCustomizableObject();
+
+		if (!CustomizableObject)
+		{
+			return;
+		}
+
+		if ( ThumbnailScene == nullptr )
+		{
+			ThumbnailScene = new FSkeletalMeshThumbnailScene();
+		}
+
+		//CustomizableObjectInstance = ConstructObject<UCustomizableObjectInstance>(UCustomizableObjectInstance::StaticClass(), GetTransientPackage(), NAME_None, RF_Transient);
+		//CustomizableObjectInstance->SetObject(CustomizableObject);
+		//CustomizableObjectInstance->bShowOnlyRuntimeParameters = false;
+
+		if (!CustomizableObject->IsCompiled())
+		{
+			// Use default compilation options?
+			FCompilationOptions Options;
+
+			FCustomizableObjectCompiler Compiler;
+			// Compiler.Compile(CustomizableObjectInstance->CustomizableObject, Options, false); // TODO Can not be sync!
+		}
+
+		// The compilation might fail because the graph is not valid or incomplete, so it has to be checked again
+		if (CustomizableObject->GetRefSkeletalMesh() && CustomizableObject->IsCompiled()) // TODO MultiComp
+		{
+			CustomizableObjectInstance->SetObject(CustomizableObject);
+			
+			// \todo removed
+			//CustomizableObjectInstance->GetPrivate()->ReloadParametersFromObject();
+			
+			UCustomizableSkeletalComponent* PreviewCustomizableSkeletalComponent = NewObject<UCustomizableSkeletalComponent>(UCustomizableSkeletalComponent::StaticClass());
+			PreviewCustomizableSkeletalComponent->CustomizableObjectInstance = CustomizableObjectInstance;
+
+			PreviewCustomizableSkeletalComponent->UpdateSkeletalMeshAsync();
+			USkeletalMesh* SkeletalMesh = PreviewCustomizableSkeletalComponent->GetSkeletalMesh();
+
+			if (SkeletalMesh)
+			{
+				ThumbnailScene->SetSkeletalMesh(SkeletalMesh);
+
+				FSceneViewFamilyContext ViewFamily(FSceneViewFamily::ConstructionValues(RenderTarget, ThumbnailScene->GetScene(), FEngineShowFlags(ESFIM_Game))
+					// \todo This differs from engine version to engine version. While thumbnails are not enable, we can disable this code.
+					//.SetWorldTimes(FApp::GetCurrentTime() - GStartTime, FApp::GetDeltaTime(), FApp::GetCurrentTime() - GStartTime)
+					.SetTime(UThumbnailRenderer::GetTime())
+					.SetAdditionalViewFamily(bAdditionalViewFamily));
+
+				ViewFamily.EngineShowFlags.DisableAdvancedFeatures();
+				ViewFamily.EngineShowFlags.MotionBlur = 0;
+				ViewFamily.EngineShowFlags.LOD = 0;
+
+				RenderViewFamily(Canvas, &ViewFamily, ThumbnailScene->CreateView(&ViewFamily, X, Y, Width, Height));
+				ThumbnailScene->SetSkeletalMesh(nullptr);
+			}
+			else
+			{
+				// Use the texture interface to draw
+				Canvas->DrawTile(X, Y, Width, Height, 0.f, 0.f, 1.f, 1.f, FLinearColor::White, NoImage->GetResource(), false);
+			}
+
+			PreviewCustomizableSkeletalComponent->DestroyComponent();
+		}
+		else
+		{
+			// Use the texture interface to draw
+			Canvas->DrawTile(X, Y, Width, Height, 0.f, 0.f, 1.f, 1.f, FLinearColor::White, NoImage->GetResource(), false);
+		}
+	}
+}
+
+void UCustomizableObjectInstanceThumbnailRenderer::BeginDestroy()
+{
+	if ( ThumbnailScene != nullptr )
+	{
+		delete ThumbnailScene;
+		ThumbnailScene = nullptr;
+	}
+
+	Super::BeginDestroy();
+}
+
+
+#undef LOCTEXT_NAMESPACE
+
diff --git a/Engine/Plugins/Mutable/Source/CustomizableObjectEditor/Private/MuCOE/Thumbnails/CustomizableObjectThumbnailRenderer.cpp b/Engine/Plugins/Mutable/Source/CustomizableObjectEditor/Private/MuCOE/Thumbnails/CustomizableObjectThumbnailRenderer.cpp
new file mode 100644
--- /dev/null
+++ b/Engine/Plugins/Mutable/Source/CustomizableObjectEditor/Private/MuCOE/Thumbnails/CustomizableObjectThumbnailRenderer.cpp
@@ -0,0 +1,121 @@
+// Copyright Epic Games, Inc. All Rights Reserved.
+
+#include "MuCOE/Thumbnails/CustomizableObjectThumbnailRenderer.h"
+
+#include "CanvasTypes.h"
+#include "MuCO/CustomizableSkeletalComponent.h"
+#include "MuCO/UnrealPortabilityHelpers.h"
+#include "MuCO/CustomizableObjectCompilerTypes.h"
+#include "MuCOE/CustomizableObjectCompiler.h"
+#include "SceneView.h"
+#include "TextureResource.h"
+#include "ThumbnailHelpers.h"
+#include "UObject/ConstructorHelpers.h"
+
+class FRenderTarget;
+class USkeletalMesh;
+
+
+
+UCustomizableObjectThumbnailRenderer::UCustomizableObjectThumbnailRenderer()
+	: Super()
+{
+	ThumbnailScene = nullptr;
+	//CustomizableObjectInstance = ConstructObject<UCustomizableObjectInstance>(UCustomizableObjectInstance::StaticClass(), GetTransientPackage(), NAME_None, RF_Transient);
+
+	// Structure to hold one-time initialization
+	struct FConstructorStatics
+	{
+		ConstructorHelpers::FObjectFinder<UTexture2D> PSysThumbnail_NoImage;
+		ConstructorHelpers::FObjectFinder<UTexture2D> PSysThumbnail_OOD;
+		FConstructorStatics()
+			: PSysThumbnail_NoImage(TEXT("/Engine/EditorMaterials/ParticleSystems/PSysThumbnail_NoImage"))
+			, PSysThumbnail_OOD(TEXT("/Engine/EditorMaterials/ParticleSystems/PSysThumbnail_OOD"))
+		{
+		}
+	};
+	static FConstructorStatics ConstructorStatics;
+
+	NoImage = ConstructorStatics.PSysThumbnail_NoImage.Object;
+}
+
+void UCustomizableObjectThumbnailRenderer::Draw(UObject* Object, int32 X, int32 Y, uint32 Width, uint32 Height, FRenderTarget* RenderTarget, FCanvas* Canvas, bool bAdditionalViewFamily)
+{
+	UCustomizableObject* CustomizableObject = Cast<UCustomizableObject>(Object);
+	
+	if (CustomizableObject != nullptr)
+	{
+		if ( ThumbnailScene == nullptr )
+		{
+			ThumbnailScene = new FSkeletalMeshThumbnailScene();
+		}
+
+		if (!CustomizableObject->IsCompiled())
+		{
+			// Use default compilation options?
+			FCompilationOptions Options = CustomizableObject->GetPrivate()->GetCompileOptions();
+
+			FCustomizableObjectCompiler Compiler;
+			// Compiler.Compile(CustomizableObject, Options, false); // TODO Can not be sync!
+		}
+
+		// The compilation might fail because the graph is not valid or incomplete, so it has to be checked again
+		if (CustomizableObject->GetRefSkeletalMesh() && CustomizableObject->IsCompiled()) // TODO MultiComp
+		{
+			CustomizableObjectInstance = CustomizableObject->CreateInstance();
+
+			if (CustomizableObjectInstance)
+			{
+				UCustomizableSkeletalComponent* PreviewCustomizableSkeletalComponent = NewObject<UCustomizableSkeletalComponent>(UCustomizableSkeletalComponent::StaticClass());
+				PreviewCustomizableSkeletalComponent->CustomizableObjectInstance = CustomizableObjectInstance;
+
+				PreviewCustomizableSkeletalComponent->UpdateSkeletalMeshAsync();
+				USkeletalMesh* SkeletalMesh = PreviewCustomizableSkeletalComponent->GetSkeletalMesh();
+
+				if (SkeletalMesh && SkeletalMesh->GetLODNum())
+				{
+					ThumbnailScene->SetSkeletalMesh(SkeletalMesh);
+
+					FSceneViewFamilyContext ViewFamily(FSceneViewFamily::ConstructionValues(RenderTarget, ThumbnailScene->GetScene(), FEngineShowFlags(ESFIM_Game))
+						// \todo This differs from engine version to engine version. While thumbnails are not enable, we can disable this code.
+						//.SetWorldTimes(FApp::GetCurrentTime() - GStartTime, FApp::GetDeltaTime(), FApp::GetCurrentTime() - GStartTime) 
+						.SetTime(UThumbnailRenderer::GetTime())
+						.SetAdditionalViewFamily(bAdditionalViewFamily));
+
+					ViewFamily.EngineShowFlags.DisableAdvancedFeatures();
+					ViewFamily.EngineShowFlags.MotionBlur = 0;
+					ViewFamily.EngineShowFlags.LOD = 0;
+
+					RenderViewFamily(Canvas, &ViewFamily, ThumbnailScene->CreateView(&ViewFamily, X, Y, Width, Height));
+					ThumbnailScene->SetSkeletalMesh(nullptr);
+				}
+				else
+				{
+					// Use the texture interface to draw
+					Canvas->DrawTile(X, Y, Width, Height, 0.f, 0.f, 1.f, 1.f, FLinearColor::White, NoImage->GetResource(), false);
+				}
+
+				SkeletalMesh = nullptr;
+				PreviewCustomizableSkeletalComponent->DestroyComponent();
+				CustomizableObjectInstance = nullptr;
+			}
+		}
+		else
+		{
+			// Use the texture interface to draw
+			Canvas->DrawTile(X, Y, Width, Height, 0.f, 0.f, 1.f, 1.f, FLinearColor::White, NoImage->GetResource(), false);
+		}
+	}
+}
+
+void UCustomizableObjectThumbnailRenderer::BeginDestroy()
+{
+	if ( ThumbnailScene != nullptr )
+	{
+		delete ThumbnailScene;
+		ThumbnailScene = nullptr;
+	}
+
+	Super::BeginDestroy();
+}
+
diff --git a/Engine/Plugins/Mutable/Source/CustomizableObjectEditor/Public/MuCOE/CustomizableObjectEditor_Deprecated.h b/Engine/Plugins/Mutable/Source/CustomizableObjectEditor/Public/MuCOE/CustomizableObjectEditor_Deprecated.h
new file mode 100644
--- /dev/null
+++ b/Engine/Plugins/Mutable/Source/CustomizableObjectEditor/Public/MuCOE/CustomizableObjectEditor_Deprecated.h
@@ -0,0 +1,198 @@
+﻿// Copyright Epic Games, Inc. All Rights Reserved.
+
+#pragma once
+
+#include "Containers/Array.h"
+#include "UObject/ObjectPtr.h"
+#include "EdGraph/EdGraphPin.h"
+
+#include "CustomizableObjectEditor_Deprecated.generated.h"
+
+class UTexture2D;
+struct FLODReductionSettings;
+
+
+// Place to hide all deprecated data structures. They are still needed for deserialization backwards compatibility.
+
+// UCustomizableObjectNodeMaterial
+// Deprecated, do not use!
+USTRUCT()
+struct CUSTOMIZABLEOBJECTEDITOR_API FCustomizableObjectNodeMaterialImage
+{
+	GENERATED_USTRUCT_BODY()
+
+	UPROPERTY()
+	FString Name;
+
+	UPROPERTY()
+	int32 UVLayout = 0;
+
+	UPROPERTY()
+	TObjectPtr<UTexture2D> ReferenceTexture = nullptr;
+	
+	UPROPERTY()
+	int32 LayerIndex = 0;
+
+	UPROPERTY()
+	FString PinName;
+};
+
+// Deprecated, do not use!
+USTRUCT()
+struct CUSTOMIZABLEOBJECTEDITOR_API FCustomizableObjectNodeMaterialVector
+{
+	GENERATED_USTRUCT_BODY()
+	
+	UPROPERTY()
+	FString Name;
+	
+	UPROPERTY()
+	int32 LayerIndex = 0;
+
+	UPROPERTY()
+	FString PinName;
+};
+
+// Deprecated, do not use!
+USTRUCT()
+struct CUSTOMIZABLEOBJECTEDITOR_API FCustomizableObjectNodeMaterialScalar
+{
+	GENERATED_USTRUCT_BODY()
+
+	UPROPERTY()
+	FString Name;
+	
+	UPROPERTY()
+	int32 LayerIndex = 0;
+
+	UPROPERTY()
+	FString PinName;
+};
+
+// Deprecated, do not use!
+constexpr int32 UV_LAYOUT_DEFAULT = -2;
+
+// UCustomizableObjectNodeModifierEditMeshSection
+// Deprecated, do not use!
+USTRUCT()
+struct CUSTOMIZABLEOBJECTEDITOR_API FCustomizableObjectNodeEditMaterialImage
+{
+	GENERATED_USTRUCT_BODY()
+
+	UPROPERTY()
+	FString Name;
+};
+
+// UCustomizableObjectNodeModifierExtendMeshSection
+// Deprecated, do not use!
+USTRUCT()
+struct CUSTOMIZABLEOBJECTEDITOR_API FCustomizableObjectNodeExtendMaterialImage
+{
+	GENERATED_USTRUCT_BODY()
+
+	UPROPERTY()
+	FString Name;
+};
+
+// UCustomizableObjectNodeSkeletalMesh
+// Deprecated, do not use!
+USTRUCT()
+struct CUSTOMIZABLEOBJECTEDITOR_API FCustomizableObjectNodeSkeletalMeshMaterial
+{
+	GENERATED_USTRUCT_BODY()
+
+	UPROPERTY()
+	FString Name;
+
+	UPROPERTY()
+	TObjectPtr<UEdGraphPin_Deprecated> MeshPin_DEPRECATED;
+
+	UPROPERTY()
+	TArray<TObjectPtr<UEdGraphPin_Deprecated>> LayoutPins_DEPRECATED;
+
+	UPROPERTY()
+	TArray<TObjectPtr<UEdGraphPin_Deprecated>> ImagePins_DEPRECATED;
+
+	UPROPERTY()
+	FEdGraphPinReference MeshPinRef;
+
+	UPROPERTY()
+	TArray<FEdGraphPinReference> LayoutPinsRef;
+
+	UPROPERTY()
+	TArray<FEdGraphPinReference> ImagePinsRef;
+};
+
+// Deprecated, do not use!
+USTRUCT()
+struct CUSTOMIZABLEOBJECTEDITOR_API FCustomizableObjectNodeSkeletalMeshLOD
+{
+	GENERATED_USTRUCT_BODY()
+
+	UPROPERTY()
+	TArray<FCustomizableObjectNodeSkeletalMeshMaterial> Materials;
+};
+
+// UCustomizableObjectNodeColorVariation
+// Deprecated, do not use!
+USTRUCT()
+struct CUSTOMIZABLEOBJECTEDITOR_API FCustomizableObjectColorVariation
+{
+	GENERATED_USTRUCT_BODY()
+
+	UPROPERTY(EditAnywhere, Category = CustomizableObject)
+	FString Tag;
+};
+
+// UCustomizableObjectFloatVariation
+// Deprecated, do not use!
+USTRUCT()
+struct CUSTOMIZABLEOBJECTEDITOR_API FCustomizableObjectFloatVariation
+{
+	GENERATED_USTRUCT_BODY()
+
+	UPROPERTY(EditAnywhere, Category = CustomizableObject)
+	FString Tag;
+};
+
+// UCustomizableObjectMaterialVariation
+// Deprecated, do not use!
+USTRUCT()
+struct CUSTOMIZABLEOBJECTEDITOR_API FCustomizableObjectMaterialVariation
+{
+	GENERATED_USTRUCT_BODY()
+
+	UPROPERTY(EditAnywhere, Category = CustomizableObject)
+	FString Tag;
+};
+
+// UCustomizableObjectMeshVariation
+// Deprecated, do not use!
+USTRUCT()
+struct CUSTOMIZABLEOBJECTEDITOR_API FCustomizableObjectMeshVariation
+{
+	GENERATED_USTRUCT_BODY()
+
+	UPROPERTY(EditAnywhere, Category = CustomizableObject)
+	FString Tag;
+};
+
+
+// UCustomizableObjectNodeObject
+// Deprecated, do not use!
+USTRUCT()
+struct FComponentSettings
+{
+	GENERATED_USTRUCT_BODY()
+	
+	UPROPERTY()
+	FString ComponentName;
+
+	UPROPERTY()
+	TArray<FLODReductionSettings> LODReductionSettings;
+};
+
+
+#if UE_ENABLE_INCLUDE_ORDER_DEPRECATED_IN_5_2
+#include "Engine/Texture2D.h"
+#endif
diff --git a/Engine/Plugins/Mutable/Source/CustomizableObjectEditor/Public/MuCOE/CustomizableObjectGraph.h b/Engine/Plugins/Mutable/Source/CustomizableObjectEditor/Public/MuCOE/CustomizableObjectGraph.h
new file mode 100644
--- /dev/null
+++ b/Engine/Plugins/Mutable/Source/CustomizableObjectEditor/Public/MuCOE/CustomizableObjectGraph.h
@@ -0,0 +1,42 @@
+// Copyright Epic Games, Inc. All Rights Reserved.
+
+#pragma once
+
+#include "EdGraph/EdGraph.h"
+
+#include "CustomizableObjectGraph.generated.h"
+
+class UObject;
+
+
+UCLASS()
+class CUSTOMIZABLEOBJECTEDITOR_API UCustomizableObjectGraph : public UEdGraph
+{
+public:
+	GENERATED_BODY()
+
+	UCustomizableObjectGraph();
+
+	// UObject Interface
+	virtual void PostLoad() override;
+	void PostRename(UObject * OldOuter, const FName OldName) override;
+
+	// Own Interface
+	void NotifyNodeIdChanged(const FGuid& OldGuid, const FGuid& NewGuid);
+
+	FGuid RequestNotificationForNodeIdChange(const FGuid& OldGuid, const FGuid& NodeToNotifyGuid);
+
+	void PostDuplicate(bool bDuplicateForPIE) override;
+
+	/** Adds the necessary nodes for a CO to work */
+	void AddEssentialGraphNodes();
+
+private:
+
+	// Request Node Id Update Map
+	TMap<FGuid, TSet<FGuid>> NodesToNotifyMap;
+
+	// Guid map with the key beeing the old Guid and the Value the new one, filled after duplicating COs
+	TMap<FGuid, FGuid> NotifiedNodeIdsMap;
+};
+
diff --git a/Engine/Plugins/Mutable/Source/CustomizableObjectEditor/Public/MuCOE/EdGraphSchema_CustomizableObject.h b/Engine/Plugins/Mutable/Source/CustomizableObjectEditor/Public/MuCOE/EdGraphSchema_CustomizableObject.h
new file mode 100644
--- /dev/null
+++ b/Engine/Plugins/Mutable/Source/CustomizableObjectEditor/Public/MuCOE/EdGraphSchema_CustomizableObject.h
@@ -0,0 +1,175 @@
+// Copyright Epic Games, Inc. All Rights Reserved.
+
+#pragma once
+
+#include "EdGraphSchema_K2_Actions.h"
+
+#include "EdGraphSchema_CustomizableObject.generated.h"
+
+enum class EMaterialParameterType : uint8;
+
+class UEdGraph;
+class UEdGraphNode;
+class UEdGraphPin;
+class UObject;
+struct FAssetData;
+struct FEdGraphPinType;
+class ICustomizableObjectEditor;
+
+/** Action to add a node to the graph */
+USTRUCT()
+struct CUSTOMIZABLEOBJECTEDITOR_API FCustomizableObjectSchemaAction_NewNode : public FEdGraphSchemaAction
+{
+	GENERATED_USTRUCT_BODY();
+
+	/** Template of node we want to create */
+	UPROPERTY()
+	TObjectPtr<class UEdGraphNode> NodeTemplate;
+
+
+	FCustomizableObjectSchemaAction_NewNode() 
+		: FEdGraphSchemaAction()
+		, NodeTemplate(NULL)
+	{}
+
+	FCustomizableObjectSchemaAction_NewNode(const FString& InNodeCategory, const FText& InMenuDesc, const FText& InToolTip, const int32 InGrouping, const FText& InKeywords = FText(), int32 InSectionID = 0)
+		: FEdGraphSchemaAction(FText::FromString(InNodeCategory), InMenuDesc, InToolTip, InGrouping, InKeywords, InSectionID)
+		, NodeTemplate(NULL)
+	{}
+
+	// FEdGraphSchemaAction interface
+	virtual UEdGraphNode* PerformAction(class UEdGraph* ParentGraph, UEdGraphPin* FromPin, const FVector2D Location, bool bSelectNewNode = true) override;
+	virtual UEdGraphNode* PerformAction(class UEdGraph* ParentGraph, TArray<UEdGraphPin*>& FromPins, const FVector2D Location, bool bSelectNewNode = true) override;
+	virtual void AddReferencedObjects( FReferenceCollector& Collector ) override;
+
+	/** Reimplementaiton of EdGraphSchema::CreateNode(...). Performs the overlap calculation before calling AutowireNewNode(...). AutowireNewNode can induce a call to a ReconstrucNode() which removes pins required for the calculation. */
+	static UEdGraphNode* CreateNode(class UEdGraph* ParentGraph, UEdGraphPin* FromPin, const FVector2D Location, class UEdGraphNode* InNodeTemplate);
+	// End of FEdGraphSchemaAction interface
+
+	template <typename NodeType>
+	static NodeType* InstantSpawn(class UEdGraph* ParentGraph, NodeType* InTemplateNode, const FVector2D Location)
+	{
+		FEdGraphSchemaAction_K2NewNode Action;
+		Action.NodeTemplate = InTemplateNode;
+
+		return Cast<NodeType>(Action.PerformAction(ParentGraph, NULL, Location));
+	}
+};
+
+
+/** Action to paste clipboard contents into the graph */
+USTRUCT()
+struct CUSTOMIZABLEOBJECTEDITOR_API FCustomizableObjectSchemaAction_Paste : public FEdGraphSchemaAction
+{
+	GENERATED_USTRUCT_BODY();
+
+	// Simple type info
+	static FName StaticGetTypeId() { static FName Type("FCustomizableObjectSchemaAction_Paste"); return Type; }
+	virtual FName GetTypeId() const override { return StaticGetTypeId(); }
+
+	FCustomizableObjectSchemaAction_Paste()
+		: FEdGraphSchemaAction()
+	{}
+
+	FCustomizableObjectSchemaAction_Paste(const FText& InNodeCategory, const FText& InMenuDesc, const FString& InToolTip, const int32 InGrouping)
+		: FEdGraphSchemaAction(InNodeCategory, InMenuDesc, FText::FromString(InToolTip), InGrouping)
+	{}
+
+	// FEdGraphSchemaAction interface
+	virtual UEdGraphNode* PerformAction(class UEdGraph* ParentGraph, UEdGraphPin* FromPin, const FVector2D Location, bool bSelectNewNode = true) override;
+	// End of FEdGraphSchemaAction interface
+};
+
+
+UCLASS()
+class CUSTOMIZABLEOBJECTEDITOR_API UEdGraphSchema_CustomizableObject : public UEdGraphSchema
+{
+public:
+	GENERATED_BODY()
+
+	// Allowable PinType.PinCategory values
+	static const FName PC_Object;
+	static const FName PC_Component;
+	static const FName PC_Material;
+	static const FName PC_Modifier;
+	static const FName PC_Mesh;
+	static const FName PC_PassThroughMesh;
+	static const FName PC_Layout;
+	static const FName PC_Image;
+	static const FName PC_PassThroughImage;
+	static const FName PC_Projector;
+	static const FName PC_GroupProjector;
+	static const FName PC_Color;
+	static const FName PC_Float;
+	static const FName PC_Bool;
+	static const FName PC_Enum;
+	static const FName PC_Stack;
+	static const FName PC_MaterialAsset;
+	static const FName PC_Wildcard;
+	static const FName PC_PoseAsset;
+
+	// EdGraphSchema interface
+	virtual void GetAssetsGraphHoverMessage(const TArray<FAssetData>& Assets, const UEdGraph* HoverGraph, FString& OutTooltipText, bool& OutOkIcon) const override;
+	virtual void GetGraphContextActions(FGraphContextMenuBuilder& ContextMenuBuilder) const override;
+	virtual const FPinConnectionResponse CanCreateConnection(const UEdGraphPin* A, const UEdGraphPin* B) const override;
+	virtual FLinearColor GetPinTypeColor(const FEdGraphPinType& PinType) const override;
+	virtual bool ShouldHidePinDefaultValue(UEdGraphPin* Pin) const override;
+	virtual void DroppedAssetsOnGraph(const TArray<FAssetData>& Assets, const FVector2D& GraphPosition, UEdGraph* Graph) const override;
+	virtual void OnPinConnectionDoubleCicked(UEdGraphPin* PinA, UEdGraphPin* PinB, const FVector2D& GraphPosition) const override;
+	virtual void GetContextMenuActions(UToolMenu* Menu, UGraphNodeContextMenuContext* Context) const override;
+	virtual void BreakNodeLinks(UEdGraphNode& TargetNode) const override;
+	virtual void BreakPinLinks(UEdGraphPin& TargetPin, bool bSendsNodeNotification) const override;
+	virtual void BreakSinglePinLink(UEdGraphPin* SourcePin, UEdGraphPin* TargetPin) const override;
+	virtual bool TryCreateConnection(UEdGraphPin* PinA, UEdGraphPin* PinB) const override;
+	virtual FPinConnectionResponse MovePinLinks(UEdGraphPin& MoveFromPin, UEdGraphPin& MoveToPin, bool bIsIntermediateMove, bool bNotifyLinkedNodes) const override;
+	virtual FConnectionDrawingPolicy* CreateConnectionDrawingPolicy(int32 InBackLayerID, int32 InFrontLayerID, float InZoomFactor, const FSlateRect& InClippingRect, FSlateWindowElementList& InDrawElements, UEdGraph* InGraphObj) const override;
+
+	//  Own interface
+	static FLinearColor GetPinTypeColor(const FName& PinType);
+
+	void GetContextMenuActionsReconstructAllChildNodes(UToolMenu* Menu, TWeakObjectPtr<UGraphNodeContextMenuContext> Context) const;
+
+	/** Utility to create a node creating action */
+	static TSharedPtr<FCustomizableObjectSchemaAction_NewNode> AddNewNodeAction(FGraphActionListBuilderBase& ContextMenuBuilder, const FString& Category, const FText& MenuDesc, const FText& Tooltip, const int32 Grouping = 0, const FString& Keywords = FString());
+
+	/** */
+	void GetBreakLinkToSubMenuActions( class FMenuBuilder& MenuBuilder, UEdGraphPin* InGraphPin ) const;
+	
+private: 
+/** Enum containing all the object types that we are able to convert onto a node when dragging
+ * and dropping and asset of that type onto the CO graph.
+ * Each value of this enumeration will, in practice, have a CO node to be represented by.
+ */
+// TODO: Replace the usage of this enum class with something similar to typeID (not casting)
+enum class ESpawnableObjectType : int32
+{
+	None = -1,		// Invalid value
+	
+	UTexture2D,
+	USkeletalMesh,
+	UStaticMesh,
+	UMaterialInterface,
+};
+
+public:
+
+	/**
+	 * Provided an asset object check if mutable does have a relative node to be able to host this asset's data.
+	 *  @param InAsset - The asset object we want to check if can be used to generate a new mutable node containing it.
+	 *  @param OutObjectType - The type of object of the asset.
+	 *  @return True if the asset can be used to spawn a mutable node, false if not.
+	 *  @note All possible object types that we can convert onto mutable nodes must be defined here.
+	 */
+	bool IsSpawnableAsset(const FAssetData& InAsset, ESpawnableObjectType& OutObjectType) const;
+
+	static TSharedPtr<class ICustomizableObjectEditor> GetCustomizableObjectEditor(const class UEdGraph* ParentGraph);
+	UEdGraphNode* AddComment(class UEdGraph* ParentGraph, UEdGraphPin* FromPin, const FVector2D Location, bool bSelectNewNode = true) const;
+	
+	/** Given a pin category name, get its friendly name (user readable). */
+	static FText GetPinCategoryName(const FName& PinCategory);
+};
+
+
+#if UE_ENABLE_INCLUDE_ORDER_DEPRECATED_IN_5_2
+#include "MaterialTypes.h"
+#endif
diff --git a/Engine/Plugins/Mutable/Source/CustomizableObjectEditor/Public/MuCOE/ICustomizableObjectDebugger.h b/Engine/Plugins/Mutable/Source/CustomizableObjectEditor/Public/MuCOE/ICustomizableObjectDebugger.h
new file mode 100644
--- /dev/null
+++ b/Engine/Plugins/Mutable/Source/CustomizableObjectEditor/Public/MuCOE/ICustomizableObjectDebugger.h
@@ -0,0 +1,22 @@
+// Copyright Epic Games, Inc. All Rights Reserved.
+
+#pragma once
+
+#include "Toolkits/AssetEditorToolkit.h"
+
+class UCustomizableObject;
+
+// Public interface to Customizable Object Editor
+class ICustomizableObjectDebugger : public FAssetEditorToolkit
+{
+public:
+	// Retrieves the current Customizable Object displayed in the Editor.
+	virtual UCustomizableObject* GetCustomizableObject() = 0;
+
+};
+
+
+
+#if UE_ENABLE_INCLUDE_ORDER_DEPRECATED_IN_5_2
+#include "Toolkits/IToolkitHost.h"
+#endif
diff --git a/Engine/Plugins/Mutable/Source/CustomizableObjectEditor/Public/MuCOE/ICustomizableObjectEditor.h b/Engine/Plugins/Mutable/Source/CustomizableObjectEditor/Public/MuCOE/ICustomizableObjectEditor.h
new file mode 100644
--- /dev/null
+++ b/Engine/Plugins/Mutable/Source/CustomizableObjectEditor/Public/MuCOE/ICustomizableObjectEditor.h
@@ -0,0 +1,49 @@
+// Copyright Epic Games, Inc. All Rights Reserved.
+
+#pragma once
+
+#include "MuCOE/ICustomizableObjectInstanceEditor.h"
+
+class UEdGraphNode;
+
+class UCustomizableObject;
+class UCustomizableObjectInstance;
+class UClass;
+class UCustomizableObjectNode;
+
+
+// Public interface to Customizable Object Editor
+class ICustomizableObjectEditor : public ICustomizableObjectInstanceEditor
+{
+public:
+	// Retrieves the current Customizable Object displayed in the Editor.
+	virtual UCustomizableObject* GetCustomizableObject() = 0;
+	
+	// Checks whether nodes can currently be pasted */
+	virtual bool CanPasteNodes() const {return false;}
+
+	// Paste nodes at a specific location
+	virtual void PasteNodesHere(const FVector2D& Location) {}
+
+	virtual UEdGraphNode* CreateCommentBox(const FVector2D& NodePos) { return nullptr; }
+
+	virtual void UpdateGraphNodeProperties() {}
+
+	virtual void UpdateObjectProperties() {}
+
+	// Make sure the graph editor selects a specific node
+	virtual void SelectNode(const UEdGraphNode* Node) {}
+
+	virtual void ReconstructNode(UEdGraphNode* Node) {}
+
+	/** Reconstructs all child the nodes that match the given type.
+	 * @param StartNode Root node to start the graph traversal. This one also will be reconstructed.
+	 * @param NodeType Node types to reconstruct. */
+	virtual void ReconstructAllChildNodes(UCustomizableObjectNode& StartNode, const UClass& NodeType) {}
+};
+
+
+
+#if UE_ENABLE_INCLUDE_ORDER_DEPRECATED_IN_5_2
+#include "Toolkits/IToolkitHost.h"
+#endif
diff --git a/Engine/Plugins/Mutable/Source/CustomizableObjectEditor/Public/MuCOE/MutableUtils.h b/Engine/Plugins/Mutable/Source/CustomizableObjectEditor/Public/MuCOE/MutableUtils.h
new file mode 100644
--- /dev/null
+++ b/Engine/Plugins/Mutable/Source/CustomizableObjectEditor/Public/MuCOE/MutableUtils.h
@@ -0,0 +1,33 @@
+// Copyright Epic Games, Inc. All Rights Reserved.
+
+#pragma once
+
+#include "Math/Vector2D.h"
+
+class USkeletalMesh;
+class UStaticMesh;
+
+
+/** Returns the mesh UV. */
+TArray<FVector2f> GetUV(const USkeletalMesh& SkeletalMesh, const int32 LODIndex, const int32 SectionIndex, const int32 UVIndex);
+
+
+/** Returns the mesh UV. */
+TArray<FVector2f> GetUV(const UStaticMesh& StaticMesh, const int32 LODIndex, const int32 SectionIndex, const int32 UVIndex);
+
+
+/** Returns if the given point is inside the 0-1 bounds. */
+bool HasNormalizedBounds(const FVector2f& Point);
+
+
+/** Check if a given LOD and material in a skeletal mesh is closed. */
+bool IsMeshClosed(const USkeletalMesh* Mesh, int LOD, int MaterialIndex);
+
+
+/** Check if a given LOD and material in a skeletal mesh is closed. */
+bool IsMeshClosed(const UStaticMesh* Mesh, int LOD, int MaterialIndex);
+
+#if UE_ENABLE_INCLUDE_ORDER_DEPRECATED_IN_5_2
+#include "Engine/SkeletalMesh.h"
+#include "Engine/StaticMesh.h"
+#endif
