[mutable] Moved the Mutable plugin out of Experimental status into Beta.

diff --git a/Engine/Plugins/Mutable/Source/CustomizableObjectEditor/Public/MuCOE/Nodes/CustomizableObjectNode.h b/Engine/Plugins/Mutable/Source/CustomizableObjectEditor/Public/MuCOE/Nodes/CustomizableObjectNode.h
new file mode 100644
--- /dev/null
+++ b/Engine/Plugins/Mutable/Source/CustomizableObjectEditor/Public/MuCOE/Nodes/CustomizableObjectNode.h
@@ -0,0 +1,367 @@
+// Copyright Epic Games, Inc. All Rights Reserved.
+
+#pragma once
+
+#include "EdGraph/EdGraph.h" // IWYU pragma: keep
+#include "EdGraph/EdGraphPin.h"
+#include "MuCO/CustomizableObject.h"
+#include "MuCO/CustomizableObjectPrivate.h"
+
+#include "CustomizableObjectNode.generated.h"
+
+class FArchive;
+class SWidget;
+class UCustomizableObjectNodeObject;
+class UCustomizableObjectNodeRemapPins;
+class UCustomizableObjectNodeRemapPinsByName;
+class UCustomizableObjectNodeRemapPinsByPosition;
+class IDetailsView;
+struct FPropertyChangedEvent;
+
+/** Abstract base class for all pin data. */
+UCLASS(Abstract)
+class UCustomizableObjectNodePinData : public UObject
+{
+	GENERATED_BODY()
+
+public:
+	UCustomizableObjectNodePinData();
+
+	/** False if inherited types are different. */
+	bool operator==(const UCustomizableObjectNodePinData& Other) const;
+	
+	virtual bool operator!=(const UCustomizableObjectNodePinData& Other) const;
+
+	/** Virtual function used to copy pin data when remapping pins. */
+	virtual void Copy(const UCustomizableObjectNodePinData& Other) {}
+	
+	/** Add backwards compatible code here.
+	 * When called, it is guaranteed that all nodes in this graph will have executed the PostLoad function. */
+	virtual void BackwardsCompatibleFixup(int32 CustomizableObjectCustomVersion) {}
+
+protected:
+	/** Virtual function used to perform comparision between different UCustomizableObjectNodePinData inherited types. */
+	virtual bool Equals(const UCustomizableObjectNodePinData& Other) const;
+};
+
+/** Encapsulation of parameters for the FPostEditChangePropertyDelegate delegate function.
+ *
+ * The delegate requires the function being a UFUNCTION which requires all parameters to be UCLASS, USTRUC or UENUM.
+ */
+USTRUCT()
+struct FPostEditChangePropertyDelegateParameters
+{
+	GENERATED_BODY()
+
+	UObject* Node;
+	FPropertyChangedEvent* FPropertyChangedEvent;
+};
+
+DECLARE_DYNAMIC_MULTICAST_DELEGATE_OneParam(FPostEditChangePropertyDelegate, FPostEditChangePropertyDelegateParameters&, Parameters); // Deprecated, use non-dynamic version.
+DECLARE_MULTICAST_DELEGATE_TwoParams(FPostEditChangePropertyRegularDelegate, UObject*, FPropertyChangedEvent&);
+
+DECLARE_MULTICAST_DELEGATE(FPostReconstructNodeDelegate);
+
+DECLARE_MULTICAST_DELEGATE(FNodeConnectionListChangedDelegate);
+
+DECLARE_MULTICAST_DELEGATE_OneParam(FPinConnectionListChangedDelegate, UEdGraphPin*);
+
+DECLARE_MULTICAST_DELEGATE(FDestroyNodeDelegate);
+
+using FRemapPinsDelegateParameter = TMap<UEdGraphPin*, UEdGraphPin*>; // Required for the delegate macro.
+DECLARE_MULTICAST_DELEGATE_OneParam(FRemapPinsDelegate, const FRemapPinsDelegateParameter&);
+
+/** Base class of all Customizable Object nodes.
+ * 
+ * The Customizable Object node system is build on top of the following premises. To avoid breaking this system, these premise must be hold!
+ * 
+ * PREMISES:
+ * 1. A node can only be modified when a node is reconstructed (i.e., inside the ReconstructNode, AllocateDefaultPins functions). Exceptionally, a pin can also be created inside the BeginConstruct.
+ *    We consider that a node has been modified if one of the following things happen:
+ *        - Creation, destruction and modification of pins.
+ *        - Modification of any data that would modify any pin.
+ */
+UCLASS(abstract)
+class CUSTOMIZABLEOBJECTEDITOR_API UCustomizableObjectNode : public UEdGraphNode
+{
+public:
+	GENERATED_BODY()
+
+	// UObject interface
+	virtual void Serialize(FArchive& Ar) override;
+	virtual void PostLoad() override; // Do work at PostBackwardsCompatibleFixup.
+
+	// UEdGraphNode interface
+	virtual void AutowireNewNode(UEdGraphPin* FromPin) override;
+	virtual void NodeConnectionListChanged() override;
+	virtual void PinConnectionListChanged(UEdGraphPin* Pin) override;
+	virtual void DestroyNode() override;
+	virtual TSharedPtr<SGraphNode> CreateVisualWidget() override;
+
+	/** Allocates the default pins using the empty remap pins action. Usually called from CreateNode. */
+	virtual void AllocateDefaultPins() override final; // Final. Override AllocateDefaultPins(UCustomizableObjectNodeRemapPins* RemapPins) instead
+	
+	/** Reconstructs the node using its default remap pins action. */
+	virtual void ReconstructNode() override final; // Final. Override ReconstructNode(UCustomizableObjectNodeRemapPins* RemapPins) instead
+	virtual void PostInitProperties() override;
+
+	// Own interface
+	/** Called at the beginning of the node lifecycle.*/
+	virtual void BeginConstruct() {};
+
+	/** Add backwards compatible code here.
+	 * When called, it is guaranteed that all nodes in this graph will have executed the PostLoad function. */
+	virtual void BackwardsCompatibleFixup(int32 CustomizableObjectCustomVersion);
+
+	/** Reconstructs the node pins using the provided AllocateDefaultPins function.
+	 * Safe to call in BackwardsCompatibleFixup.
+	 * Notice that this function will not call the node AllocateDefaultPins function nor the ReconstructNode overrides.
+	 *
+	 * Since this function is used BackwardsCompatibleFixup, it should not have major changes (or they will break existing fixups).
+	 * If changes are needed, duplicate this function so old fixups remain unchanged. */
+	void FixupReconstructPins(UCustomizableObjectNodeRemapPins* RemapPinsAction, TFunction<void(UCustomizableObjectNode*, UCustomizableObjectNodeRemapPins*)> AllocateDefaultPins);
+	
+	/** Add post load work here.
+	 * When called, it is guaranteed that all nodes in this graph will have executed the BackwardsCompatibleFixup function.
+	 * Notice that no compatibility code should not go here. Use BackwardsCompatibleFixup instead. */
+	virtual void PostBackwardsCompatibleFixup() {}
+	
+	/** Virtual implementation of RemovePin. Allows to do work before removing a pin.
+	 * Use this function instead of RemovePin. RemovePin does not removes possible attached PinData. */
+	virtual bool CustomRemovePin(UEdGraphPin& Pin);
+
+	/**
+	 * Subclasses should override this to return true and set OutCategory if this node should be
+	 * auto-added to the right-click context menu in the graph editor.
+	 * 
+	 * Some nodes are added manually in the graph editor code and don't need to do this.
+	 */
+	virtual bool ShouldAddToContextMenu(FText& OutCategory) const;
+
+	void GetInputPins(TArray<class UEdGraphPin*>& OutInputPins) const;
+	void GetOutputPins(TArray<class UEdGraphPin*>& OutOutputPins) const;
+	UEdGraphPin* GetOutputPin(int32 OutputIndex) const;
+	
+	virtual bool IsNodeOutDatedAndNeedsRefresh() { return false; }
+	virtual FString GetRefreshMessage() const { return "Refresh Node."; }
+	void SetRefreshNodeWarning();
+	void RemoveWarnings();
+
+	// Used to replace Ids of referenced nodes by their new ids after duplicating the CustomizableObject.
+	virtual void UpdateReferencedNodeId(const FGuid& NewGuid) {};
+
+	TSharedPtr<class ICustomizableObjectEditor> GetGraphEditor() const;
+
+	virtual bool ProvidesCustomPinRelevancyTest() const { return false; } // Override and return true if IsPinRelevant is overridden.
+	virtual bool IsPinRelevant(const UEdGraphPin*) const { return true; } // Override if default pins allocated with AllocateDefaultPins do not provide enough information.
+
+	struct FAttachedErrorDataView
+	{
+		TArrayView<const float> UnassignedUVs;
+	};
+	
+	/** 
+	 * Check if two pins can be connected. 
+	 * 
+	 * @param	InOwnedInputPin	Input pin which belongs to this node.
+	 * @param	InOutputPin		Output pin which belongs to another node. If the node connects to iself, it could belong to this node.
+	 * @param	bOutIsOtherNodeBlocklisted		Is the other node of a type we are not allowed to connect?
+	 * @param	bOutArePinsCompatible		Does InOutputPin pin share the same type as the InOwnedInputPin?
+	 * @return	True if the pin types match and if the other node is not one of the Blocklisted types.
+	 */
+	virtual bool CanConnect( const UEdGraphPin* InOwnedInputPin, const UEdGraphPin* InOutputPin, bool& bOutIsOtherNodeBlocklisted, bool& bOutArePinsCompatible) const;
+
+	// Used during compilation process to cache the node for all LODs, or generate it specifically for each of them.
+	virtual bool IsAffectedByLOD() const;
+
+	/** Return an array of tags that this node will enable and apply to its data. Null if none, or of it doesn't apply. */
+	virtual TArray<FString>* GetEnableTags();
+
+
+	// Get the CustomizableObject graph that owns this node
+	class UCustomizableObjectGraph* GetCustomizableObjectGraph() const;
+
+	/** If returns true, only a single link can go out of every output pin. False by default. */
+	virtual bool IsSingleOutputNode() const;
+
+	// Wrapper of the pin creation method, to cope with multiple Unreal Engine versions
+	UEdGraphPin* CustomCreatePin( EEdGraphPinDirection Direction, const FName& Type, const FName& Name, bool bIsArray=false );
+
+	/** Create pin with attached data.
+	 *
+     * @param PinData Pin data to be saved. */
+	UEdGraphPin* CustomCreatePin(EEdGraphPinDirection Direction, const FName& Type, const FName& Name, UCustomizableObjectNodePinData* PinData);
+
+	
+	virtual void AddAttachedErrorData(const FAttachedErrorDataView& ) {};
+	virtual void ResetAttachedErrorData() {};
+
+	/** 
+	 * When creating a new connection, break all previous connections. This method can be overridden.
+	 * 
+	 * @return false by default.
+	 */
+	virtual bool ShouldBreakExistingConnections(const UEdGraphPin* InputPin, const UEdGraphPin* OutputPin) const;
+
+	/** Called when a new connection is perform. Used to break an specific existing connectoin. Empty default implementation. This method can be overridden.
+	 *
+	 * @param	InputPin	New connection input pin.
+	 * @param	OutputPin	New connection output pin.
+	 */
+	virtual void BreakExistingConnectionsPostConnection(UEdGraphPin* InputPin, UEdGraphPin* OutputPin) {};
+
+	void PostEditChangeProperty(FPropertyChangedEvent& PropertyChangedEvent) override;
+
+	/** Custom post duplicate function. Called at the beginning of duplication, before the nodes have their Guid updated. */
+	virtual void BeginPostDuplicate(bool bDuplicateForPIE) {};
+	
+	/** Gets all non orphan pins this node owns. */
+	TArray<UEdGraphPin*> GetAllNonOrphanPins() const;
+
+	/** Gets all orphan pins this node owns. */
+	TArray<UEdGraphPin*> GetAllOrphanPins() const;
+
+	/**
+	 * Specialization of ReconstructNode UEdGraphNode function.
+	 *
+	 * @param RemapPins pointer to a remap pins action. Remap pins action which will be used once the node has been reconstructed. Can be null.
+	 */
+	virtual void ReconstructNode(UCustomizableObjectNodeRemapPins* RemapPinsMode);
+
+	/** Return the LOD which this node belongs to. Return -1 if it is not connected (directly or indirectly) to a LOD pin. */
+	int32 GetLOD() const;
+	
+	/** Set a pin to be hidden or not. */
+	void SetPinHidden(UEdGraphPin& Pin, bool bHidden);
+
+	/** Set an array of pins to be hidden or not. 
+	 *
+	 * @param Pins Input pins. All pins must be non null.
+	 */
+	void SetPinHidden(const TArray<UEdGraphPin*>& Pins, bool bHidden);
+
+	/** Return true if a pin can be hidden. Override it to have a custom behaivour. */
+	virtual bool CanPinBeHidden(const UEdGraphPin& Pin) const;
+
+	/** @return true if the details should spawn a Pin Viewer. */
+	virtual bool HasPinViewer() const;
+
+	/** Returns pin custom details. Override if required. */
+	virtual TSharedPtr<IDetailsView> CustomizePinDetails(const UEdGraphPin& Pin) const;
+
+	/** Given the pin data, returns its associated pin. */
+	UEdGraphPin* GetPin(const UCustomizableObjectNodePinData& PinData);
+	
+	/** See GetPinData(const FEdGraphPinReference&). */
+	UCustomizableObjectNodePinData* GetPinData(const UEdGraphPin &Pin) const;
+	
+	/** Given a pin which has attached data, get the typed pin data.
+	 * @param Pin Pin which has attached data.
+	 * @return Pin data already typed. */
+	template<typename T>
+	T& GetPinData(const UEdGraphPin& Pin) const
+	{
+		return *CastChecked<T>(PinsDataId[Pin.PinId]);
+	}
+	
+	UPROPERTY()
+	FPostEditChangePropertyDelegate PostEditChangePropertyDelegate;
+
+	FPostEditChangePropertyRegularDelegate PostEditChangePropertyRegularDelegate;
+
+	FPostReconstructNodeDelegate PostReconstructNodeDelegate;
+
+	FNodeConnectionListChangedDelegate NodeConnectionListChangedDelegate;
+
+	FPinConnectionListChangedDelegate PinConnectionListChangedDelegate;
+
+	FDestroyNodeDelegate DestroyNodeDelegate;
+
+	FRemapPinsDelegate RemapPinsDelegate;
+
+protected:
+	/**
+	 * Specialization of AllocateDefaultPins UEdGraphNode function. Override.
+	 * 
+	 * Given the node context (this), create all the required pins.
+	 *
+	 * @param RemapPins remap pins action. Use this object to save any information required to later perfrom the remap pins action. Can be null.
+	 */
+	virtual void AllocateDefaultPins(UCustomizableObjectNodeRemapPins* RemapPins) {};
+
+	/**
+	 * Creates and returns the node default node remap pins action.
+	 * 
+	 * By default all nodes remap by name. Override if a node requires a diferent default remap pins action.
+	 */
+	virtual UCustomizableObjectNodeRemapPins* CreateRemapPinsDefault() const;
+
+	/**
+	 * Creates and returns a remap pins by name action. 
+	 * 
+	 * Override a if a node requires an specialized remap pins by name.
+	 */
+	virtual UCustomizableObjectNodeRemapPinsByName* CreateRemapPinsByName() const;
+
+	/**
+	 * Creates and returns a remap pins by position action.
+	 *
+	 * Override a if a node requires an specialized remap pins by position.
+	 */
+	virtual UCustomizableObjectNodeRemapPinsByPosition* CreateRemapPinsByPosition() const;
+	
+	/** Allows to perform work when remapping a pin.
+	  * Copies pin data from old pin to new pin. Keeps the id of the new pin. */
+	virtual void RemapPin(UEdGraphPin& NewPin, const UEdGraphPin& OldPin);
+
+	/** Allows to perform work when remapping a pin. */
+	virtual void RemapPins(const TMap<UEdGraphPin*, UEdGraphPin*>& PinsToRemap);
+
+	/*** Allows to perform work when remapping the pin data. */
+	virtual void RemapPinsData(const TMap<UEdGraphPin*, UEdGraphPin*>& PinsToRemap);
+	
+	/** Attach data to a given pin. */
+	void AddPinData(const UEdGraphPin& Pin, UCustomizableObjectNodePinData& PinData);
+
+private:
+	/** Data attached to a given pin. Not all pins contain data.
+	 * Data must be eliminated once the pins is removed. This is done automatically on CustomRemovePin.
+	 * Do not use RemovePin since it will not remove the data. */
+	UPROPERTY()
+	TMap<FGuid, TObjectPtr<UCustomizableObjectNodePinData>> PinsDataId;
+
+	// Deprecated properties in favour of PinsDataId
+	UPROPERTY()
+	TMap<FEdGraphPinReference, TObjectPtr<UCustomizableObjectNodePinData>> PinsData_DEPRECATED;
+};
+
+
+/** Return the specified node located in another Customizable Object graph.
+ *
+ * @param Customizable Object object to look for the node. If nulltpr, it returns nullptr.
+ * @param NodeGuid node guid. If invalid, it returns nullptr.
+ * 
+ * Return nullptr if not found.
+ */
+template<class FNodeType>
+FNodeType* GetCustomizableObjectExternalNode(UCustomizableObject* Object, const FGuid& NodeGuid)
+{
+	FNodeType* Result = nullptr;
+
+	if (Object && Object->GetPrivate()->GetSource() && NodeGuid.IsValid())
+	{
+		TArray<FNodeType*> GroupNodes;
+		Object->GetPrivate()->GetSource()->GetNodesOfClass<FNodeType>(GroupNodes);
+
+		for (FNodeType* GroupNode : GroupNodes)
+		{
+			if (NodeGuid == GroupNode->NodeGuid)
+			{
+				Result = GroupNode;
+				break;
+			}
+		}
+	}
+
+	return Result;
+}
diff --git a/Engine/Plugins/Mutable/Source/CustomizableObjectEditor/Public/MuCOE/Nodes/CustomizableObjectNodeAnimationPose.h b/Engine/Plugins/Mutable/Source/CustomizableObjectEditor/Public/MuCOE/Nodes/CustomizableObjectNodeAnimationPose.h
new file mode 100644
--- /dev/null
+++ b/Engine/Plugins/Mutable/Source/CustomizableObjectEditor/Public/MuCOE/Nodes/CustomizableObjectNodeAnimationPose.h
@@ -0,0 +1,46 @@
+// Copyright Epic Games, Inc. All Rights Reserved.
+
+#pragma once
+
+#include "MuCOE/Nodes/CustomizableObjectNode.h"
+
+#include "CustomizableObjectNodeAnimationPose.generated.h"
+
+namespace ENodeTitleType { enum Type : int; }
+
+class FString;
+class UCustomizableObjectNodeRemapPins;
+class UEdGraphPin;
+class UObject;
+class UPoseAsset;
+struct FPropertyChangedEvent;
+
+
+UCLASS()
+class CUSTOMIZABLEOBJECTEDITOR_API UCustomizableObjectNodeAnimationPose : public UCustomizableObjectNode
+{
+public:
+	GENERATED_BODY()
+
+	UCustomizableObjectNodeAnimationPose();
+
+	UPROPERTY(EditAnywhere, Category=CustomizableObject)
+	TObjectPtr<UPoseAsset> PoseAsset;
+
+	// UObject interface.
+	virtual void PostEditChangeProperty(FPropertyChangedEvent& PropertyChangedEvent) override;
+
+	// Begin EdGraphNode interface
+	UEdGraphPin* GetInputMeshPin() const;
+	UEdGraphPin* GetTablePosePin() const;
+	virtual FText GetNodeTitle(ENodeTitleType::Type TitleType) const override;
+	virtual FLinearColor GetNodeTitleColor() const override;
+	// End EdGraphNode interface
+
+	// UCustomizableObjectNode interface
+	void AllocateDefaultPins(UCustomizableObjectNodeRemapPins* RemapPins) override;
+
+	// Retrieve the pose information from the PoseAsset
+	static void StaticRetrievePoseInformation(UPoseAsset* PoseAsset, class USkeletalMesh* RefSkeletalMesh, TArray<FName>& OutArrayBoneName, TArray<FTransform>& OutArrayTransform);
+};
+
diff --git a/Engine/Plugins/Mutable/Source/CustomizableObjectEditor/Public/MuCOE/Nodes/CustomizableObjectNodeColorArithmeticOp.h b/Engine/Plugins/Mutable/Source/CustomizableObjectEditor/Public/MuCOE/Nodes/CustomizableObjectNodeColorArithmeticOp.h
new file mode 100644
--- /dev/null
+++ b/Engine/Plugins/Mutable/Source/CustomizableObjectEditor/Public/MuCOE/Nodes/CustomizableObjectNodeColorArithmeticOp.h
@@ -0,0 +1,65 @@
+// Copyright Epic Games, Inc. All Rights Reserved.
+
+#pragma once
+
+#include "MuCOE/Nodes/CustomizableObjectNode.h"
+
+#include "CustomizableObjectNodeColorArithmeticOp.generated.h"
+
+namespace ENodeTitleType { enum Type : int; }
+
+class UCustomizableObjectNodeRemapPins;
+class UEdGraphPin;
+class UObject;
+struct FPropertyChangedEvent;
+
+
+UENUM(BlueprintType)
+enum class EColorArithmeticOperation : uint8
+{
+	E_Add 	UMETA(DisplayName = "+"),
+	E_Sub 	UMETA(DisplayName = "-"),
+	E_Mul	UMETA(DisplayName = "x"),
+	E_Div	UMETA(DisplayName = "/")
+};
+
+UCLASS()
+class CUSTOMIZABLEOBJECTEDITOR_API UCustomizableObjectNodeColorArithmeticOp : public UCustomizableObjectNode
+{
+public:
+	GENERATED_BODY()
+
+	UCustomizableObjectNodeColorArithmeticOp();
+
+	// UObject interface.
+	void PostEditChangeProperty(FPropertyChangedEvent& PropertyChangedEvent) override;
+
+	// Begin EdGraphNode interface
+	FText GetNodeTitle(ENodeTitleType::Type TitleType) const override;
+	FLinearColor GetNodeTitleColor() const override;
+	FText GetTooltipText() const override;
+	bool ShouldOverridePinNames() const override { return true; }
+
+	// UCustomizableObjectNode interface
+	void AllocateDefaultPins(UCustomizableObjectNodeRemapPins* RemapPins) override;
+
+	// First Operand
+	UEdGraphPin* XPin() const
+	{
+		return FindPin(TEXT("A"));
+	}
+
+	// Second Operand
+	UEdGraphPin* YPin() const
+	{
+		return FindPin(TEXT("B"));
+	}
+
+	UPROPERTY(EditAnywhere, Category = ColorArithmeticOperation)
+	EColorArithmeticOperation Operation;
+
+	// UCustomizableObjectNode interface
+	virtual bool IsAffectedByLOD() const { return false; }
+
+};
+
diff --git a/Engine/Plugins/Mutable/Source/CustomizableObjectEditor/Public/MuCOE/Nodes/CustomizableObjectNodeColorFromFloats.h b/Engine/Plugins/Mutable/Source/CustomizableObjectEditor/Public/MuCOE/Nodes/CustomizableObjectNodeColorFromFloats.h
new file mode 100644
--- /dev/null
+++ b/Engine/Plugins/Mutable/Source/CustomizableObjectEditor/Public/MuCOE/Nodes/CustomizableObjectNodeColorFromFloats.h
@@ -0,0 +1,59 @@
+// Copyright Epic Games, Inc. All Rights Reserved.
+
+#pragma once
+
+#include "MuCOE/Nodes/CustomizableObjectNode.h"
+
+#include "CustomizableObjectNodeColorFromFloats.generated.h"
+
+namespace ENodeTitleType { enum Type : int; }
+
+class UCustomizableObjectNodeRemapPins;
+class UEdGraphPin;
+class UObject;
+struct FPropertyChangedEvent;
+
+
+UCLASS()
+class CUSTOMIZABLEOBJECTEDITOR_API UCustomizableObjectNodeColorFromFloats : public UCustomizableObjectNode
+{
+public:
+	GENERATED_BODY()
+
+	UCustomizableObjectNodeColorFromFloats();
+
+	// UObject interface.
+	void PostEditChangeProperty(FPropertyChangedEvent& PropertyChangedEvent) override;
+
+	// Begin EdGraphNode interface
+	FText GetNodeTitle(ENodeTitleType::Type TitleType) const override;
+	FLinearColor GetNodeTitleColor() const override;
+	FText GetTooltipText() const override;
+
+	// UCustomizableObjectNode interface
+	void AllocateDefaultPins(UCustomizableObjectNodeRemapPins* RemapPins) override;
+
+	UEdGraphPin* RPin() const
+	{
+		return FindPin(TEXT("R"));
+	}
+
+	UEdGraphPin* GPin() const
+	{
+		return FindPin(TEXT("G"));
+	}
+
+	UEdGraphPin* BPin() const
+	{
+		return FindPin(TEXT("B"));
+	}
+
+	UEdGraphPin* APin() const
+	{
+		return FindPin(TEXT("A"));
+	}
+
+	// UCustomizableObjectNode interface
+	virtual bool IsAffectedByLOD() const override { return false; }
+};
+
diff --git a/Engine/Plugins/Mutable/Source/CustomizableObjectEditor/Public/MuCOE/Nodes/CustomizableObjectNodeColorParameter.h b/Engine/Plugins/Mutable/Source/CustomizableObjectEditor/Public/MuCOE/Nodes/CustomizableObjectNodeColorParameter.h
new file mode 100644
--- /dev/null
+++ b/Engine/Plugins/Mutable/Source/CustomizableObjectEditor/Public/MuCOE/Nodes/CustomizableObjectNodeColorParameter.h
@@ -0,0 +1,41 @@
+// Copyright Epic Games, Inc. All Rights Reserved.
+
+#pragma once
+
+#include "MuCOE/Nodes/CustomizableObjectNode.h"
+
+#include "CustomizableObjectNodeColorParameter.generated.h"
+
+namespace ENodeTitleType { enum Type : int; }
+
+class UCustomizableObjectNodeRemapPins;
+class UObject;
+
+
+UCLASS()
+class CUSTOMIZABLEOBJECTEDITOR_API UCustomizableObjectNodeColorParameter : public UCustomizableObjectNode
+{
+public:
+	GENERATED_BODY()
+
+	UPROPERTY(EditAnywhere, Category=CustomizableObject, meta = (DontUpdateWhileEditing))
+	FLinearColor DefaultValue = FLinearColor(1, 1, 1, 1);
+
+	UPROPERTY(EditAnywhere, Category=CustomizableObject)
+	FString ParameterName = "Default Name";
+
+	UPROPERTY(EditAnywhere, Category = UI, meta = (DisplayName = "Parameter UI Metadata"))
+	FMutableParamUIMetadata ParamUIMetadata;
+
+	// Begin EdGraphNode interface
+	FText GetNodeTitle(ENodeTitleType::Type TitleType) const override;
+	FLinearColor GetNodeTitleColor() const override;
+	FText GetTooltipText() const override;
+	virtual void OnRenameNode(const FString& NewName) override;
+	virtual bool GetCanRenameNode() const override { return true; }
+
+	// UCustomizableObjectNode interface
+	void AllocateDefaultPins(UCustomizableObjectNodeRemapPins* RemapPins) override;
+	bool IsAffectedByLOD() const override { return false; }
+};
+
diff --git a/Engine/Plugins/Mutable/Source/CustomizableObjectEditor/Public/MuCOE/Nodes/CustomizableObjectNodeColorSwitch.h b/Engine/Plugins/Mutable/Source/CustomizableObjectEditor/Public/MuCOE/Nodes/CustomizableObjectNodeColorSwitch.h
new file mode 100644
--- /dev/null
+++ b/Engine/Plugins/Mutable/Source/CustomizableObjectEditor/Public/MuCOE/Nodes/CustomizableObjectNodeColorSwitch.h
@@ -0,0 +1,22 @@
+// Copyright Epic Games, Inc. All Rights Reserved.
+
+#pragma once
+
+#include "MuCOE/Nodes/CustomizableObjectNodeSwitchBase.h"
+
+#include "CustomizableObjectNodeColorSwitch.generated.h"
+
+
+UCLASS()
+class CUSTOMIZABLEOBJECTEDITOR_API UCustomizableObjectNodeColorSwitch : public UCustomizableObjectNodeSwitchBase
+{
+public:
+	GENERATED_BODY()
+
+	// UCustomizableObjectNode interface
+	virtual void BackwardsCompatibleFixup(int32 CustomizableObjectCustomVersion) override;
+
+	// UCustomizableObjectNodeSwitchBase interface
+	virtual FName GetCategory() const override;
+};
+
diff --git a/Engine/Plugins/Mutable/Source/CustomizableObjectEditor/Public/MuCOE/Nodes/CustomizableObjectNodeColorVariation.h b/Engine/Plugins/Mutable/Source/CustomizableObjectEditor/Public/MuCOE/Nodes/CustomizableObjectNodeColorVariation.h
new file mode 100644
--- /dev/null
+++ b/Engine/Plugins/Mutable/Source/CustomizableObjectEditor/Public/MuCOE/Nodes/CustomizableObjectNodeColorVariation.h
@@ -0,0 +1,28 @@
+// Copyright Epic Games, Inc. All Rights Reserved.
+
+#pragma once
+
+#include "MuCOE/Nodes/CustomizableObjectNodeVariation.h"
+
+#include "CustomizableObjectNodeColorVariation.generated.h"
+
+struct FCustomizableObjectColorVariation;
+
+
+UCLASS()
+class CUSTOMIZABLEOBJECTEDITOR_API UCustomizableObjectNodeColorVariation : public UCustomizableObjectNodeVariation
+{
+public:
+	GENERATED_BODY()
+
+	// Deprecated properties
+	UPROPERTY()
+	TArray<FCustomizableObjectColorVariation> Variations_DEPRECATED;
+
+	// UCustomizableObjectNode interface
+	virtual void BackwardsCompatibleFixup(int32 CustomizableObjectCustomVersion) override;
+	
+	// UCustomizableObjectNodeVariation interface
+	virtual FName GetCategory() const override;
+};
+
diff --git a/Engine/Plugins/Mutable/Source/CustomizableObjectEditor/Public/MuCOE/Nodes/CustomizableObjectNodeComponent.h b/Engine/Plugins/Mutable/Source/CustomizableObjectEditor/Public/MuCOE/Nodes/CustomizableObjectNodeComponent.h
new file mode 100644
--- /dev/null
+++ b/Engine/Plugins/Mutable/Source/CustomizableObjectEditor/Public/MuCOE/Nodes/CustomizableObjectNodeComponent.h
@@ -0,0 +1,21 @@
+// Copyright Epic Games, Inc. All Rights Reserved.
+
+#pragma once
+
+#include "MuCOE/Nodes/CustomizableObjectNode.h"
+
+#include "CustomizableObjectNodeComponent.generated.h"
+
+class UEdGraphPin;
+
+
+UCLASS(abstract)
+class CUSTOMIZABLEOBJECTEDITOR_API UCustomizableObjectNodeComponent : public UCustomizableObjectNode
+{
+public:
+	GENERATED_BODY()
+
+	// UCustomizableObjectNode interface
+	virtual bool IsAffectedByLOD() const override;
+};
+
diff --git a/Engine/Plugins/Mutable/Source/CustomizableObjectEditor/Public/MuCOE/Nodes/CustomizableObjectNodeCurve.h b/Engine/Plugins/Mutable/Source/CustomizableObjectEditor/Public/MuCOE/Nodes/CustomizableObjectNodeCurve.h
new file mode 100644
--- /dev/null
+++ b/Engine/Plugins/Mutable/Source/CustomizableObjectEditor/Public/MuCOE/Nodes/CustomizableObjectNodeCurve.h
@@ -0,0 +1,55 @@
+// Copyright Epic Games, Inc. All Rights Reserved.
+
+#pragma once
+
+#include "MuCOE/Nodes/CustomizableObjectNode.h"
+
+#include "CustomizableObjectNodeCurve.generated.h"
+
+namespace ENodeTitleType { enum Type : int; }
+
+class FArchive;
+class UCurveBase;
+class UCustomizableObjectNodeRemapPins;
+class UEdGraphPin;
+class UObject;
+struct FPropertyChangedEvent;
+
+
+UCLASS()
+class CUSTOMIZABLEOBJECTEDITOR_API UCustomizableObjectNodeCurve : public UCustomizableObjectNode
+{
+public:
+	GENERATED_BODY()
+
+	UCustomizableObjectNodeCurve();
+
+	// UObject interface.
+	void PostEditChangeProperty(FPropertyChangedEvent& PropertyChangedEvent) override;
+	virtual void Serialize(FArchive& Ar) override;
+
+	// Begin EdGraphNode interface
+	FText GetNodeTitle(ENodeTitleType::Type TitleType) const override;
+	FLinearColor GetNodeTitleColor() const override;
+	FText GetTooltipText() const override;
+
+	// UCustomizableObjectNode interface
+	void AllocateDefaultPins(UCustomizableObjectNodeRemapPins* RemapPins) override;
+
+	UEdGraphPin* InputPin() const;
+	UEdGraphPin* CurvePins(int32 Index) const;
+	int32 GetNumCurvePins() const;
+
+	UPROPERTY(EditAnywhere, Category = Curve)
+	TObjectPtr<UCurveBase> CurveAsset;
+
+	// UCustomizableObjectNode interface
+	virtual bool IsAffectedByLOD() const override { return false; }
+
+	bool ProvidesCustomPinRelevancyTest() const override { return true; }
+	bool IsPinRelevant(const UEdGraphPin* Pin) const override;
+
+	/** Override the default behaivour of remap pins. Use remap pins by position by default. */
+	UCustomizableObjectNodeRemapPins* CreateRemapPinsDefault() const override;
+};
+
diff --git a/Engine/Plugins/Mutable/Source/CustomizableObjectEditor/Public/MuCOE/Nodes/CustomizableObjectNodeEnumParameter.h b/Engine/Plugins/Mutable/Source/CustomizableObjectEditor/Public/MuCOE/Nodes/CustomizableObjectNodeEnumParameter.h
new file mode 100644
--- /dev/null
+++ b/Engine/Plugins/Mutable/Source/CustomizableObjectEditor/Public/MuCOE/Nodes/CustomizableObjectNodeEnumParameter.h
@@ -0,0 +1,62 @@
+// Copyright Epic Games, Inc. All Rights Reserved.
+
+#pragma once
+
+#include "MuCOE/Nodes/CustomizableObjectNode.h"
+
+#include "CustomizableObjectNodeEnumParameter.generated.h"
+
+namespace ENodeTitleType { enum Type : int; }
+
+class UCustomizableObjectNodeRemapPins;
+class UObject;
+struct FPropertyChangedEvent;
+
+
+USTRUCT()
+struct CUSTOMIZABLEOBJECTEDITOR_API FCustomizableObjectNodeEnumValue
+{
+	GENERATED_USTRUCT_BODY()
+
+	UPROPERTY(EditAnywhere, Category=CustomizableObject)
+	FString Name;
+
+	UPROPERTY(EditAnywhere, Category = UI, meta = (DisplayName = "Parameter UI Metadata"))
+	FMutableParamUIMetadata ParamUIMetadata;
+};
+
+
+UCLASS()
+class CUSTOMIZABLEOBJECTEDITOR_API UCustomizableObjectNodeEnumParameter : public UCustomizableObjectNode
+{
+public:
+	GENERATED_BODY()
+
+	UPROPERTY(EditAnywhere, Category=CustomizableObject)
+	int32 DefaultIndex = 0;
+
+	UPROPERTY(EditAnywhere, Category=CustomizableObject)
+	FString ParameterName = "Default Name";
+
+	UPROPERTY(EditAnywhere, Category = UI, meta = (DisplayName = "Parameter UI Metadata"))
+	FMutableParamUIMetadata ParamUIMetadata;
+
+	UPROPERTY(EditAnywhere, Category=CustomizableObject)
+	TArray<FCustomizableObjectNodeEnumValue> Values;
+
+	// UObject interface.
+	void PostEditChangeProperty(FPropertyChangedEvent& PropertyChangedEvent) override;
+
+	// Begin EdGraphNode interface
+	FText GetNodeTitle(ENodeTitleType::Type TitleType) const override;
+	FLinearColor GetNodeTitleColor() const override;
+	FText GetTooltipText() const override;
+	virtual void OnRenameNode(const FString& NewName) override;
+	virtual bool GetCanRenameNode() const override { return true; }
+
+	// UCustomizableObjectNode interface
+	void AllocateDefaultPins(UCustomizableObjectNodeRemapPins* RemapPins) override;
+	virtual bool IsAffectedByLOD() const override { return false; }
+
+};
+
diff --git a/Engine/Plugins/Mutable/Source/CustomizableObjectEditor/Public/MuCOE/Nodes/CustomizableObjectNodeExposePin.h b/Engine/Plugins/Mutable/Source/CustomizableObjectEditor/Public/MuCOE/Nodes/CustomizableObjectNodeExposePin.h
new file mode 100644
--- /dev/null
+++ b/Engine/Plugins/Mutable/Source/CustomizableObjectEditor/Public/MuCOE/Nodes/CustomizableObjectNodeExposePin.h
@@ -0,0 +1,59 @@
+// Copyright Epic Games, Inc. All Rights Reserved.
+
+#pragma once
+
+#include "MuCOE/Nodes/CustomizableObjectNode.h"
+
+#include "CustomizableObjectNodeExposePin.generated.h"
+
+namespace ENodeTitleType { enum Type : int; }
+
+class UCustomizableObjectNodeRemapPins;
+class UEdGraphPin;
+class UObject;
+struct FPropertyChangedEvent;
+
+
+DECLARE_MULTICAST_DELEGATE(FOnNameChangedDelegate);
+
+
+/** Export Node. */
+UCLASS()
+class CUSTOMIZABLEOBJECTEDITOR_API UCustomizableObjectNodeExposePin : public UCustomizableObjectNode
+{
+public:
+	GENERATED_BODY()
+
+	// UObject interface
+	void PostEditChangeProperty(FPropertyChangedEvent& PropertyChangedEvent) override;
+
+	// EdGraphNode interface
+	FText GetNodeTitle(ENodeTitleType::Type TitleType) const override;
+	FLinearColor GetNodeTitleColor() const override;
+	FText GetTooltipText() const override;
+
+	virtual bool CanConnect(const UEdGraphPin* InOwnedInputPin, const UEdGraphPin* InOutputPin, bool& bOutIsOtherNodeBlocklisted, bool& bOutArePinsCompatible) const override;
+
+	// UCustomizableObjectNode interface
+	void AllocateDefaultPins(UCustomizableObjectNodeRemapPins* RemapPins) override;
+
+	/** Return the Expose Pin Node expose pin name. */
+	FString GetNodeName() const;
+	
+	// This is actually PinCategory
+	UPROPERTY()
+	FName PinType;
+	
+	UEdGraphPin* InputPin() const
+	{
+		return FindPin(TEXT("Object"));
+	}
+
+	/** Boradcasted when the UPROPERTY Name changes. */
+	FOnNameChangedDelegate OnNameChangedDelegate;
+	
+private:
+	UPROPERTY(Category = CustomizableObject, EditAnywhere)
+	FString Name = "Default Name";
+};
+
diff --git a/Engine/Plugins/Mutable/Source/CustomizableObjectEditor/Public/MuCOE/Nodes/CustomizableObjectNodeFloatParameter.h b/Engine/Plugins/Mutable/Source/CustomizableObjectEditor/Public/MuCOE/Nodes/CustomizableObjectNodeFloatParameter.h
new file mode 100644
--- /dev/null
+++ b/Engine/Plugins/Mutable/Source/CustomizableObjectEditor/Public/MuCOE/Nodes/CustomizableObjectNodeFloatParameter.h
@@ -0,0 +1,60 @@
+// Copyright Epic Games, Inc. All Rights Reserved.
+
+#pragma once
+
+#include "MuCOE/Nodes/CustomizableObjectNode.h"
+
+#include "CustomizableObjectNodeFloatParameter.generated.h"
+
+namespace ENodeTitleType { enum Type : int; }
+
+class FArchive;
+class UCustomizableObjectNodeRemapPins;
+class UEdGraphPin;
+class UObject;
+struct FEdGraphPinReference;
+struct FPropertyChangedEvent;
+
+
+USTRUCT()
+struct CUSTOMIZABLEOBJECTEDITOR_API FCustomizableObjectNodeFloatDescription
+{
+	GENERATED_USTRUCT_BODY()
+
+	UPROPERTY(EditAnywhere, Category=CustomizableObject)
+	FString Name;
+};
+
+
+UCLASS()
+class CUSTOMIZABLEOBJECTEDITOR_API UCustomizableObjectNodeFloatParameter : public UCustomizableObjectNode
+{
+public:
+	GENERATED_BODY()
+
+	UPROPERTY(EditAnywhere, Category=CustomizableObject)
+	float DefaultValue = 1.0f;
+
+	UPROPERTY(EditAnywhere, Category=CustomizableObject)
+	FString ParameterName = "Default Name";
+
+	UPROPERTY(EditAnywhere, Category = UI, meta = (DisplayName = "Parameter UI Metadata"))
+	FMutableParamUIMetadata ParamUIMetadata;
+	
+	// UObject interface.
+	void PostEditChangeProperty(FPropertyChangedEvent& PropertyChangedEvent) override;
+	
+	// Begin EdGraphNode interface
+	FText GetNodeTitle(ENodeTitleType::Type TitleType) const override;
+	FLinearColor GetNodeTitleColor() const override;
+	FText GetTooltipText() const override;
+	virtual bool GetCanRenameNode() const override { return true; }
+
+	// UCustomizableObjectNode interface
+	void AllocateDefaultPins(UCustomizableObjectNodeRemapPins* RemapPins) override;
+	bool IsAffectedByLOD() const override;
+	void BackwardsCompatibleFixup(int32 CustomizableObjectCustomVersion) override;
+	virtual void OnRenameNode(const FString& NewName) override;
+
+};
+
diff --git a/Engine/Plugins/Mutable/Source/CustomizableObjectEditor/Public/MuCOE/Nodes/CustomizableObjectNodeFloatSwitch.h b/Engine/Plugins/Mutable/Source/CustomizableObjectEditor/Public/MuCOE/Nodes/CustomizableObjectNodeFloatSwitch.h
new file mode 100644
--- /dev/null
+++ b/Engine/Plugins/Mutable/Source/CustomizableObjectEditor/Public/MuCOE/Nodes/CustomizableObjectNodeFloatSwitch.h
@@ -0,0 +1,21 @@
+// Copyright Epic Games, Inc. All Rights Reserved.
+
+#pragma once
+
+#include "MuCOE/Nodes/CustomizableObjectNodeSwitchBase.h"
+
+#include "CustomizableObjectNodeFloatSwitch.generated.h"
+
+
+UCLASS()
+class CUSTOMIZABLEOBJECTEDITOR_API UCustomizableObjectNodeFloatSwitch : public UCustomizableObjectNodeSwitchBase
+{
+public:
+	GENERATED_BODY()
+
+	// UCustomizableObjectNode interface
+	virtual void BackwardsCompatibleFixup(int32 CustomizableObjectCustomVersion) override;
+
+	// UCustomizableObjectNodeSwitchBase interface
+	virtual FName GetCategory() const override;
+};
diff --git a/Engine/Plugins/Mutable/Source/CustomizableObjectEditor/Public/MuCOE/Nodes/CustomizableObjectNodeFloatVariation.h b/Engine/Plugins/Mutable/Source/CustomizableObjectEditor/Public/MuCOE/Nodes/CustomizableObjectNodeFloatVariation.h
new file mode 100644
--- /dev/null
+++ b/Engine/Plugins/Mutable/Source/CustomizableObjectEditor/Public/MuCOE/Nodes/CustomizableObjectNodeFloatVariation.h
@@ -0,0 +1,28 @@
+// Copyright Epic Games, Inc. All Rights Reserved.
+
+#pragma once
+
+#include "MuCOE/Nodes/CustomizableObjectNodeVariation.h"
+
+#include "CustomizableObjectNodeFloatVariation.generated.h"
+
+struct FCustomizableObjectFloatVariation;
+
+
+UCLASS()
+class CUSTOMIZABLEOBJECTEDITOR_API UCustomizableObjectNodeFloatVariation : public UCustomizableObjectNodeVariation
+{
+public:
+	GENERATED_BODY()
+	
+	// Deprecated properties
+	UPROPERTY()
+	TArray<FCustomizableObjectFloatVariation> Variations_DEPRECATED;
+
+	// UCustomizableObjectNode interface
+	virtual void BackwardsCompatibleFixup(int32 CustomizableObjectCustomVersion) override;
+	
+	// UCustomizableObjectNodeVariation interface
+	virtual FName GetCategory() const override;
+};
+
diff --git a/Engine/Plugins/Mutable/Source/CustomizableObjectEditor/Public/MuCOE/Nodes/CustomizableObjectNodeLayoutBlocks.h b/Engine/Plugins/Mutable/Source/CustomizableObjectEditor/Public/MuCOE/Nodes/CustomizableObjectNodeLayoutBlocks.h
new file mode 100644
--- /dev/null
+++ b/Engine/Plugins/Mutable/Source/CustomizableObjectEditor/Public/MuCOE/Nodes/CustomizableObjectNodeLayoutBlocks.h
@@ -0,0 +1,61 @@
+// Copyright Epic Games, Inc. All Rights Reserved.
+
+#pragma once
+
+#include "MuCOE/Nodes/CustomizableObjectNode.h"
+
+#include "CustomizableObjectNodeLayoutBlocks.generated.h"
+
+class UCustomizableObjectLayout;
+enum class ECustomizableObjectTextureLayoutPackingStrategy : uint8;
+namespace ENodeTitleType { enum Type : int; }
+struct FCustomizableObjectLayoutBlock;
+
+class FArchive;
+class UCustomizableObjectNodeRemapPins;
+class UEdGraphPin;
+class UObject;
+struct FPropertyChangedEvent;
+
+
+UCLASS()
+class CUSTOMIZABLEOBJECTEDITOR_API UCustomizableObjectNodeLayoutBlocks : public UCustomizableObjectNode
+{
+public:
+	GENERATED_BODY()
+
+	UCustomizableObjectNodeLayoutBlocks();
+
+	UPROPERTY()
+	FIntPoint GridSize_DEPRECATED;
+
+	/** Used with the fixed layout strategy. */
+	UPROPERTY()
+	FIntPoint MaxGridSize_DEPRECATED;
+
+	UPROPERTY()
+	TArray<FCustomizableObjectLayoutBlock> Blocks_DEPRECATED;
+
+	UPROPERTY()
+	ECustomizableObjectTextureLayoutPackingStrategy PackingStrategy_DEPRECATED;
+
+	UPROPERTY()
+	TObjectPtr<UCustomizableObjectLayout> Layout = nullptr;
+
+	// EdGraphNode interface
+	void Serialize(FArchive& Ar) override;
+
+	// UCustomizableObjectNode interface
+	virtual void BackwardsCompatibleFixup(int32 CustomizableObjectCustomVersion) override;
+	void AllocateDefaultPins(UCustomizableObjectNodeRemapPins* RemapPins) override;
+
+	UEdGraphPin* OutputPin() const
+	{
+		return FindPin(TEXT("Layout"));
+	}
+};
+
+
+#if UE_ENABLE_INCLUDE_ORDER_DEPRECATED_IN_5_2
+#include "MuCOE/CustomizableObjectLayout.h"
+#endif
diff --git a/Engine/Plugins/Mutable/Source/CustomizableObjectEditor/Public/MuCOE/Nodes/CustomizableObjectNodeMaterial.h b/Engine/Plugins/Mutable/Source/CustomizableObjectEditor/Public/MuCOE/Nodes/CustomizableObjectNodeMaterial.h
new file mode 100644
--- /dev/null
+++ b/Engine/Plugins/Mutable/Source/CustomizableObjectEditor/Public/MuCOE/Nodes/CustomizableObjectNodeMaterial.h
@@ -0,0 +1,335 @@
+// Copyright Epic Games, Inc. All Rights Reserved.
+
+#pragma once
+
+#include "MuCOE/Nodes/CustomizableObjectNodeMaterialBase.h"
+#include "MuCOE/RemapPins/CustomizableObjectNodeRemapPinsByName.h"
+
+#include "CustomizableObjectNodeMaterial.generated.h"
+
+enum class EMaterialParameterType : uint8;
+namespace ENodeTitleType { enum Type : int; }
+
+class FArchive;
+class SGraphNode;
+class SWidget;
+class UCustomizableObjectNodeRemapPins;
+class UEdGraphPin;
+class UMaterial;
+class UMaterialInterface;
+class UObject;
+class UTexture2D;
+struct FCustomizableObjectNodeMaterialImage;
+struct FCustomizableObjectNodeMaterialScalar;
+struct FCustomizableObjectNodeMaterialVector;
+struct FEdGraphPinReference;
+struct FFrame;
+struct FPropertyChangedEvent;
+
+
+/** Custom remap pins by name action.
+ *
+ * Remap pins by Texture Parameter Id. */
+UCLASS()
+class UCustomizableObjectNodeMaterialRemapPinsByName : public UCustomizableObjectNodeRemapPinsByName
+{
+	GENERATED_BODY()
+public:
+	virtual bool Equal(const UCustomizableObjectNode& Node, const UEdGraphPin& OldPin, const UEdGraphPin& NewPin) const override;
+
+	virtual void RemapPins(const UCustomizableObjectNode& Node, const TArray<UEdGraphPin*>& OldPins, const TArray<UEdGraphPin*>& NewPins, TMap<UEdGraphPin*, UEdGraphPin*>& PinsToRemap, TArray<UEdGraphPin*>& PinsToOrphan) override;
+
+	bool HasSavedPinData(const UCustomizableObjectNode& Node, const UEdGraphPin &Pin) const;
+};
+
+
+/** Base class for all Material Parameters. */
+UCLASS()
+class CUSTOMIZABLEOBJECTEDITOR_API UCustomizableObjectNodeMaterialPinDataParameter : public UCustomizableObjectNodePinData
+{
+	GENERATED_BODY()
+
+public:
+
+	/** Parameter id + layer index */
+	UPROPERTY()
+	FNodeMaterialParameterId MaterialParameterId;
+
+	/** Texture Parameter Id. */
+	UPROPERTY()
+	FGuid ParameterId_DEPRECATED;
+
+	/** Returns true if all properties are in its default state. */
+	virtual bool IsDefault() const;
+};
+
+
+/** Node pin mode. All pins set to EPinMode::Default will use this this mode. */
+UENUM()
+enum class ENodePinMode
+{
+	Mutable UMETA(ToolTip = "All Material Texture Parameters go through Mutable."),
+	Passthrough UMETA(ToolTip = "All Material Texture Parameters are not modified by Mutable.")
+};
+
+
+/** Image pin, pin mode. */
+UENUM()
+enum class EPinMode
+{
+	Default UMETA(DisplayName = "Node Defined", ToolTip = "Use node's \"Default Texture Parameter Mode\"."),
+	Mutable UMETA(ToolTip = "The Material Texture Parameters goes through Mutable."),
+	Passthrough UMETA(ToolTip = "The Material Texture Parameters is not modified by Mutable.")
+};
+
+
+/** Image Pin, UV Layout Mode. */
+UENUM()
+enum class EUVLayoutMode
+{
+	/** Does not override the UV Index specified in the Material. */
+	FromMaterial,
+	/* Texture should not be transformed by any layout. These textures will not be reduced automatically for LODs. */
+	Ignore,
+	/** User specified UV Index. */
+	Index
+};
+
+
+/** Enum to FText. */
+FText EPinModeToText(EPinMode PinMode);
+
+UCLASS()
+class CUSTOMIZABLEOBJECTEDITOR_API UCustomizableObjectNodeMaterial : public UCustomizableObjectNodeMaterialBase
+{
+public:
+	GENERATED_BODY()
+	
+	// UObject interface.
+	virtual void Serialize(FArchive& Ar) override;
+	virtual void PostLoad() override;
+	virtual void PostEditChangeProperty(FPropertyChangedEvent& PropertyChangedEvent) override;
+
+	// EdGraphNode interface
+	virtual FText GetNodeTitle(ENodeTitleType::Type TitleType) const override;
+	virtual FLinearColor GetNodeTitleColor() const override;
+	virtual FText GetTooltipText() const override;
+	virtual TSharedPtr<SGraphNode> CreateVisualWidget() override;
+	virtual void PinConnectionListChanged(UEdGraphPin* Pin) override;
+	virtual void PostPasteNode() override;
+	virtual bool CanConnect(const UEdGraphPin* InOwnedInputPin, const UEdGraphPin* InOutputPin, bool& bOutIsOtherNodeBlocklisted, bool& bOutArePinsCompatible) const override;
+
+	// UCustomizableObjectNode interface
+	virtual void BackwardsCompatibleFixup(int32 CustomizableObjectCustomVersion) override;
+	virtual void AllocateDefaultPins(UCustomizableObjectNodeRemapPins* RemapPins) override;
+	virtual bool CanPinBeHidden(const UEdGraphPin& Pin) const override;
+	virtual bool HasPinViewer() const override;
+	virtual UCustomizableObjectNodeRemapPinsByName* CreateRemapPinsDefault() const override;
+	virtual bool ProvidesCustomPinRelevancyTest() const override { return true; }
+	virtual bool IsPinRelevant(const UEdGraphPin* Pin) const override;
+	virtual bool CustomRemovePin(UEdGraphPin& Pin) override;
+	virtual void BreakExistingConnectionsPostConnection(UEdGraphPin* InputPin, UEdGraphPin* OutputPin) override;
+	virtual bool IsNodeOutDatedAndNeedsRefresh() override;
+	virtual FString GetRefreshMessage() const override;
+	virtual TSharedPtr<IDetailsView> CustomizePinDetails(const UEdGraphPin& Pin) const override;
+	virtual TArray<FString>* GetEnableTags() override;
+
+	
+	// UCustomizableObjectNodeMaterialBase interface
+	virtual TArray<UCustomizableObjectLayout*> GetLayouts() const override;
+	virtual UEdGraphPin* OutputPin() const override;
+	virtual UMaterialInterface* GetMaterial() const override;
+	virtual bool IsReuseMaterialBetweenLODs() const override;
+	virtual int32 GetNumParameters(EMaterialParameterType Type) const override;
+	virtual FNodeMaterialParameterId GetParameterId(EMaterialParameterType Type, int32 ParameterIndex) const override;
+	virtual FName GetParameterName(EMaterialParameterType Type, int32 ParameterIndex) const override;
+	virtual int32 GetParameterLayerIndex(EMaterialParameterType Type, int32 ParameterIndex) const override;
+	virtual FText GetParameterLayerName(EMaterialParameterType Type, int32 ParameterIndex) const override;
+	virtual bool HasParameter(const FNodeMaterialParameterId& ParameterId) const override;
+	virtual UEdGraphPin* GetParameterPin(EMaterialParameterType Type, int32 ParameterIndex) const override;
+	virtual UEdGraphPin* GetParameterPin(const FNodeMaterialParameterId& ParameterId) const override;
+	virtual bool IsImageMutableMode(int32 ImageIndex) const override;
+	virtual bool IsImageMutableMode(const UEdGraphPin& Pin) const override;
+	virtual UTexture2D* GetImageReferenceTexture(int32 ImageIndex) const override;
+	virtual UTexture2D* GetImageValue(int32 ImageIndex) const override;
+	virtual int32 GetImageUVLayout(int32 ImageIndex) const override;
+	virtual UEdGraphPin* GetMeshPin() const override;
+	virtual UEdGraphPin* GetMaterialAssetPin() const override;
+	virtual const UCustomizableObjectNodeMaterial* GetMaterialNode() const override;
+	virtual bool RealMaterialDataHasChanged() const override;
+	virtual FPostImagePinModeChangedDelegate* GetPostImagePinModeChangedDelegate() override;
+	
+	// Own Interface
+
+	void SetMaterial(UMaterialInterface* InMaterial);
+
+	static bool HasParameter(const UMaterialInterface* InMaterial, const FNodeMaterialParameterId& ParameterId);
+	static int32 GetParameterLayerIndex(const UMaterialInterface* InMaterial, EMaterialParameterType Type, int32 ParameterIndex);
+
+	UPROPERTY(EditAnywhere, Category = Tags)
+	TArray<FString> Tags;
+
+private:
+	/** Set the Pin Mode of a Texture Parameter Pin. */
+	void SetImagePinMode(UEdGraphPin& Pin, EPinMode PinMode) const;
+	
+	/** Delegate called when a Texture Parameter Pin Mode changes. */
+	FPostImagePinModeChangedDelegate PostImagePinModeChangedDelegate;
+
+	UPROPERTY(EditAnywhere, Category=CustomizableObject)
+	TObjectPtr<UMaterialInterface> Material = nullptr;
+
+	UPROPERTY(EditAnywhere, Category=CustomizableObject, DisplayName = "Default Texture Parameter Mode", Meta = (ToolTip = "All Mateiral Texture Parameters set to \"Node Defined\" will use this mode."))
+	ENodePinMode TextureParametersMode = ENodePinMode::Passthrough;
+
+	UPROPERTY()
+	int32 MeshComponentIndex_DEPRECATED = 0;
+
+public:
+	/** Selects which Mesh component of the Instance this material belongs to */
+	UPROPERTY()
+	FName MeshComponentName_DEPRECATED;
+
+private:
+	/** Materials will be reused between LODs when possible unless explicitly disabled. */
+	UPROPERTY(EditAnywhere, Category = CustomizableObject)
+	bool bReuseMaterialBetweenLODs = true;
+
+	/** Last static or skeletal mesh connected. Used to remove the callback once disconnected. */
+	TWeakObjectPtr<UCustomizableObjectNode> LastMeshNodeConnected;
+
+	/** List of material parameter types that are actually relevant to mutable. */
+	static const TArray<EMaterialParameterType> ParameterTypes;
+
+	/** Relates a Parameter id (key) (and layer if is a layered material) to a Pin (value). Only used to improve performance.
+	  * If a deprecated pin and a non-deprecated pin have the same Parameter id, this the non-deprecated one prevails. */
+	UPROPERTY()
+	TMap<FNodeMaterialParameterId, FEdGraphPinReference> PinsParameterMap;
+	
+	/** Create the pin data of the given parameter type. */
+	UCustomizableObjectNodeMaterialPinDataParameter* CreatePinData(EMaterialParameterType Type, int32 ParameterIndex);
+
+	/** Allocate a pin for each parameter of the given type. */
+	void AllocateDefaultParameterPins(EMaterialParameterType Type);
+
+	/** Set the default Material from the connected static or skeletal mesh. */
+	void SetDefaultMaterial();
+
+	/** Connected NodeStaticMesh or NodeSkeletalMesh Mesh UPROPERTY changed callback function. Sets the default material. */
+	UFUNCTION()
+	void MeshPostEditChangeProperty(FPostEditChangePropertyDelegateParameters& Parameters);
+
+	/** Format pin name. */
+	FName GetPinName(EMaterialParameterType Type, int32 ParameterIndex) const;
+
+	/** Returns the texture coordinate of the given Material Expression. Returns -1 if not found. */
+	static int32 GetExpressionTextureCoordinate(UMaterial* Material, const FGuid &ImageId);
+
+	/** Return the Pin Category given the node NodePinMode. */
+	static FName NodePinModeToImagePinMode(ENodePinMode NodePinMode);
+
+	/** Return the Pin Category given a PinMode. */
+	FName GetImagePinMode(EPinMode PinMode) const;
+
+	/** Return the Pin Category given a Pin. */
+	FName GetImagePinMode(const UEdGraphPin& Pin) const;
+
+	/** Get the UV Layout Index defined in the Material. */
+	int32 GetImageUVLayoutFromMaterial(int32 ImageIndex) const;
+	
+	// Deprecated properties
+	/** Set all pins to Mutable mode. Even so, each pin can override its behaviour. */
+	UPROPERTY()
+	bool bDefaultPinModeMutable_DEPRECATED = false;
+	
+	UPROPERTY()
+	TArray<FCustomizableObjectNodeMaterialImage> Images_DEPRECATED;
+	
+	UPROPERTY()
+	TArray<FCustomizableObjectNodeMaterialVector> VectorParams_DEPRECATED;
+
+	UPROPERTY()
+	TArray<FCustomizableObjectNodeMaterialScalar> ScalarParams_DEPRECATED;
+
+	UPROPERTY()
+	TMap<FGuid, FEdGraphPinReference> PinsParameter_DEPRECATED;
+};
+
+
+/** Additional data for a Material Texture Parameter pin. */
+UCLASS()
+class CUSTOMIZABLEOBJECTEDITOR_API UCustomizableObjectNodeMaterialPinDataImage : public UCustomizableObjectNodeMaterialPinDataParameter
+{
+	GENERATED_BODY()
+
+	friend void UCustomizableObjectNodeMaterial::BackwardsCompatibleFixup(int32 CustomizableObjectCustomVersion);
+	
+public:
+	// UObject interface
+	virtual void PostEditChangeProperty(FPropertyChangedEvent& PropertyChangedEvent) override;
+	virtual bool CanEditChange(const FProperty* InProperty) const override;
+
+	// UCustomizableObjectNodePinData interface
+	virtual void BackwardsCompatibleFixup(int32 CustomizableObjectCustomVersion) override;
+
+	// NodePinDataParameter interface
+	/** Virtual function used to copy pin data when remapping pins. */
+	virtual void Copy(const UCustomizableObjectNodePinData& Other) override;
+	
+	// UCustomizableObjectNodeMaterialPinParameter interface
+	virtual bool IsDefault() const override;
+
+	// Own interface
+	/** Constructor parameters. Should always be called after a NewObject. */
+	void Init(UCustomizableObjectNodeMaterial& InNodeMaterial);
+
+	EPinMode GetPinMode() const;
+
+	void SetPinMode(EPinMode InPinMode);
+	
+private:
+	/** Image pin mode. If is not default, overrides the defined node behaviour. */
+	UPROPERTY(EditAnywhere, Category = NoCategory, DisplayName = "Texture Parameter Mode")
+	EPinMode PinMode = EPinMode::Default;
+
+public:
+	constexpr static int32 UV_LAYOUT_IGNORE = -1;
+
+	UPROPERTY(EditAnywhere, Category = NoCategory)
+	EUVLayoutMode UVLayoutMode = EUVLayoutMode::FromMaterial;
+	
+	/** Index of the UV channel that will be used with this image. It is necessary to apply the proper layout transformations to it. */
+	UPROPERTY(EditAnywhere, Category = NoCategory, meta = (EditCondition = "UVLayoutMode == EUVLayoutMode::Index", EditConditionHides))
+	int32 UVLayout = -2;
+
+	/** Reference Texture used to decide the texture properties of the mutable-generated textures
+	* connected to this material. If null, it will try to be guessed at compile time from
+	* the graph. */
+	UPROPERTY(EditAnywhere, Category = NoCategory) // Required to be EditAnywhere for the selector to work.
+	TObjectPtr<UTexture2D> ReferenceTexture = nullptr;
+
+private:
+	UPROPERTY()
+	TObjectPtr<UCustomizableObjectNodeMaterial> NodeMaterial = nullptr;
+};
+
+
+/** Additional data for a Material Vector Parameter pin. */
+UCLASS()
+class CUSTOMIZABLEOBJECTEDITOR_API UCustomizableObjectNodeMaterialPinDataVector : public UCustomizableObjectNodeMaterialPinDataParameter
+{
+	GENERATED_BODY()
+};
+
+
+/** Additional data for a Material Float Parameter pin. */
+UCLASS()
+class CUSTOMIZABLEOBJECTEDITOR_API UCustomizableObjectNodeMaterialPinDataScalar : public UCustomizableObjectNodeMaterialPinDataParameter
+{
+	GENERATED_BODY()
+};
+
+#if UE_ENABLE_INCLUDE_ORDER_DEPRECATED_IN_5_2
+#include "MaterialTypes.h"
+#include "MuCOE/CustomizableObjectEditor_Deprecated.h"
+#endif
diff --git a/Engine/Plugins/Mutable/Source/CustomizableObjectEditor/Public/MuCOE/Nodes/CustomizableObjectNodeMaterialVariation.h b/Engine/Plugins/Mutable/Source/CustomizableObjectEditor/Public/MuCOE/Nodes/CustomizableObjectNodeMaterialVariation.h
new file mode 100644
--- /dev/null
+++ b/Engine/Plugins/Mutable/Source/CustomizableObjectEditor/Public/MuCOE/Nodes/CustomizableObjectNodeMaterialVariation.h
@@ -0,0 +1,43 @@
+// Copyright Epic Games, Inc. All Rights Reserved.
+
+#pragma once
+
+#include "MuCOE/Nodes/CustomizableObjectNodeVariation.h"
+
+#include "CustomizableObjectNodeMaterialVariation.generated.h"
+
+struct FCustomizableObjectMaterialVariation;
+
+
+UENUM(BlueprintType)
+enum class ECustomizableObjectNodeMaterialVariationType : uint8
+{
+	Tag 		UMETA(DisplayName = "Tag"),
+	State 		UMETA(DisplayName = "State"),
+};
+
+
+UCLASS()
+class CUSTOMIZABLEOBJECTEDITOR_API UCustomizableObjectNodeMaterialVariation : public UCustomizableObjectNodeVariation
+{
+public:
+	GENERATED_BODY()
+
+	UPROPERTY(EditAnywhere, Category = CustomizableObject)
+	ECustomizableObjectNodeMaterialVariationType Type = ECustomizableObjectNodeMaterialVariationType::Tag;
+
+private:
+	// Deprecated properties
+	UPROPERTY()
+	TArray<FCustomizableObjectMaterialVariation> Variations_DEPRECATED;
+
+public:
+	// UCustomizableObjectNode interface
+	virtual void BackwardsCompatibleFixup(int32 CustomizableObjectCustomVersion) override;
+	virtual bool IsSingleOutputNode() const override;
+	
+	// UCustomizableObjectNodeVariation interface
+	virtual FName GetCategory() const override;
+	virtual bool IsInputPinArray() const override;
+};
+
diff --git a/Engine/Plugins/Mutable/Source/CustomizableObjectEditor/Public/MuCOE/Nodes/CustomizableObjectNodeMesh.h b/Engine/Plugins/Mutable/Source/CustomizableObjectEditor/Public/MuCOE/Nodes/CustomizableObjectNodeMesh.h
new file mode 100644
--- /dev/null
+++ b/Engine/Plugins/Mutable/Source/CustomizableObjectEditor/Public/MuCOE/Nodes/CustomizableObjectNodeMesh.h
@@ -0,0 +1,44 @@
+// Copyright Epic Games, Inc. All Rights Reserved.
+
+#pragma once
+
+#include "MuCOE/Nodes/CustomizableObjectNode.h"
+
+#include "CustomizableObjectNodeMesh.generated.h"
+
+class UCustomizableObjectLayout;
+class UEdGraphPin;
+class UObject;
+class UTexture2D;
+
+
+UCLASS(abstract)
+class CUSTOMIZABLEOBJECTEDITOR_API UCustomizableObjectNodeMesh : public UCustomizableObjectNode
+{
+public:
+	GENERATED_BODY()
+
+	// Own interface
+	/** Returns the multiple Layouts of a given Mesh pin. A pin can have multiple Layouts since it can have multiple UVs. Override. */
+	virtual TArray<UCustomizableObjectLayout*> GetLayouts(const UEdGraphPin& OutPin) const PURE_VIRTUAL(UCustomizableObjectNodeMesh::GetLayouts, return {};);
+
+	virtual UTexture2D* FindTextureForPin(const UEdGraphPin* Pin) const PURE_VIRTUAL(UCustomizableObjectNodeMesh::FindTextureForPin, return {};);
+
+	/** Returns the Unreal mesh (e.g., USkeletalMesh, UStaticMesh...). */
+	virtual UObject* GetMesh() const PURE_VIRTUAL(UCustomizableObjectNodeMesh::GetMesh, return {};);
+
+	/** Returns the output Mesh pin associated to the given LODIndex and SectionIndex. Override. */
+	virtual UEdGraphPin* GetMeshPin(int32 LOD, int32 SectionIndex) const PURE_VIRTUAL(UCustomizableObjectNodeMesh::GetMeshPin, return {};);
+
+	/** Given a pin, return the Section and Layout index.
+	 *
+	 * Will always return a valid result. The result can be valid but out of sync with respect the Unreal mesh asset.
+	 * In other words, the pin still represents a LOD 3 but the asset no longer has this third LOD.
+	 * 
+	 * @param Pin In. Node's owning pin.
+	 * @param LODIndex Out.
+	 * @param SectionIndex Out.
+	 * @param LayoutIndex Out. */
+	virtual void GetPinSection(const UEdGraphPin& Pin, int32& OutLODIndex, int32& OutSectionIndex, int32& OutLayoutIndex) const PURE_VIRTUAL(UCustomizableObjectNodeMesh::GetPinSection, );
+};
+
diff --git a/Engine/Plugins/Mutable/Source/CustomizableObjectEditor/Public/MuCOE/Nodes/CustomizableObjectNodeMeshGeometryOperation.h b/Engine/Plugins/Mutable/Source/CustomizableObjectEditor/Public/MuCOE/Nodes/CustomizableObjectNodeMeshGeometryOperation.h
new file mode 100644
--- /dev/null
+++ b/Engine/Plugins/Mutable/Source/CustomizableObjectEditor/Public/MuCOE/Nodes/CustomizableObjectNodeMeshGeometryOperation.h
@@ -0,0 +1,51 @@
+// Copyright Epic Games, Inc. All Rights Reserved.
+
+#pragma once
+
+#include "MuCOE/Nodes/CustomizableObjectNode.h"
+
+#include "CustomizableObjectNodeMeshGeometryOperation.generated.h"
+
+class UCustomizableObjectNodeRemapPins;
+class UEdGraphPin;
+class UObject;
+
+
+UCLASS()
+class CUSTOMIZABLEOBJECTEDITOR_API UCustomizableObjectNodeMeshGeometryOperation : public UCustomizableObjectNode
+{
+public:
+	GENERATED_BODY()
+
+	UCustomizableObjectNodeMeshGeometryOperation();
+
+	// Begin EdGraphNode interface
+	FText GetNodeTitle(ENodeTitleType::Type TitleType) const override;
+	FLinearColor GetNodeTitleColor() const override;
+	FText GetTooltipText() const override;
+
+	// UCustomizableObjectNode interface
+	void AllocateDefaultPins(UCustomizableObjectNodeRemapPins* RemapPins) override;
+
+	inline UEdGraphPin* MeshAPin() const
+	{
+		return FindPin(TEXT("Mesh A"), EGPD_Input);
+	}
+
+	inline UEdGraphPin* MeshBPin() const
+	{
+		return FindPin(TEXT("Mesh B"), EGPD_Input);
+	}
+
+	inline UEdGraphPin* ScalarAPin() const
+	{
+		return FindPin(TEXT("Scalar A"));
+	}
+
+	inline UEdGraphPin* ScalarBPin() const
+	{
+		return FindPin(TEXT("Scalar B"));
+	}
+	
+};
+
diff --git a/Engine/Plugins/Mutable/Source/CustomizableObjectEditor/Public/MuCOE/Nodes/CustomizableObjectNodeMeshMorph.h b/Engine/Plugins/Mutable/Source/CustomizableObjectEditor/Public/MuCOE/Nodes/CustomizableObjectNodeMeshMorph.h
new file mode 100644
--- /dev/null
+++ b/Engine/Plugins/Mutable/Source/CustomizableObjectEditor/Public/MuCOE/Nodes/CustomizableObjectNodeMeshMorph.h
@@ -0,0 +1,84 @@
+// Copyright Epic Games, Inc. All Rights Reserved.
+
+#pragma once
+
+#include "MuCOE/Nodes/CustomizableObjectNode.h"
+#include "MuCOE/Nodes/CustomizableObjectNodeMeshReshapeCommon.h"
+
+#include "CustomizableObjectNodeMeshMorph.generated.h"
+
+class UCustomizableObjectNodeRemapPins;
+class UEdGraphPin;
+class UObject;
+struct FMeshReshapeBoneReference;
+enum class EBoneDeformSelectionMethod : uint8;
+
+UCLASS()
+class CUSTOMIZABLEOBJECTEDITOR_API UCustomizableObjectNodeMeshMorph : public UCustomizableObjectNode
+{
+public:
+	GENERATED_BODY()
+
+	UCustomizableObjectNodeMeshMorph();
+
+	// Begin EdGraphNode interface
+	FText GetNodeTitle(ENodeTitleType::Type TitleType) const override;
+	FLinearColor GetNodeTitleColor() const override;
+	FText GetTooltipText() const override;
+
+	// UCustomizableObjectNode interface
+	void AllocateDefaultPins(UCustomizableObjectNodeRemapPins* RemapPins) override;
+	bool IsNodeOutDatedAndNeedsRefresh() override;
+	FString GetRefreshMessage() const override;
+	void PinConnectionListChanged(UEdGraphPin* Pin) override;
+
+
+	// Own interface
+	class UCustomizableObjectNodeSkeletalMesh* GetSourceSkeletalMesh() const;
+
+	UEdGraphPin* MeshPin() const
+	{
+		return FindPin(TEXT("Mesh"), EGPD_Input);
+	}
+
+	UEdGraphPin* FactorPin() const
+	{
+		return FindPin(TEXT("Factor"));
+	}
+
+	void Serialize(FArchive& Ar) override;
+
+	UPROPERTY(Category=CustomizableObject, EditAnywhere)
+	FString MorphTargetName;
+
+	/** Enable the deformation of the skeleton of the base mesh. */
+	UPROPERTY(EditAnywhere, Category = CustomizableObject)
+	bool bReshapeSkeleton = false;
+
+	/** Enable the deformation of physics volumes of the base mesh */
+    UPROPERTY(EditAnywhere, Category = CustomizableObject)
+    bool bReshapePhysicsVolumes = false;
+
+	/** Bone Reshape Selection Method */
+	UPROPERTY(EditAnywhere, Category = ReshapeBones, meta = (EditCondition = "bReshapeSkeleton"))
+	EBoneDeformSelectionMethod SelectionMethod;
+	
+	/** Array with selected bones that will be deformed */
+	UPROPERTY(EditAnywhere, Category = ReshapeBones, meta = (EditCondition = "bReshapeSkeleton && (SelectionMethod == EBoneDeformSelectionMethod::ONLY_SELECTED || SelectionMethod == EBoneDeformSelectionMethod::ALL_BUT_SELECTED)"))
+	TArray<FMeshReshapeBoneReference> BonesToDeform;
+
+	UPROPERTY(EditAnywhere, Category = ReshapePhysics, meta = (DisplayName = "Selection Method", EditCondition = "bReshapePhysicsVolumes"))
+	EBoneDeformSelectionMethod PhysicsSelectionMethod;
+
+	/** Array with bones with physics bodies that will be deformed */
+	UPROPERTY(EditAnywhere, Category = ReshapePhysics, meta = (EditCondition = "bReshapePhysicsVolumes && (PhysicsSelectionMethod == EBoneDeformSelectionMethod::ONLY_SELECTED || PhysicsSelectionMethod == EBoneDeformSelectionMethod::ALL_BUT_SELECTED)"))
+	TArray<FMeshReshapeBoneReference> PhysicsBodiesToDeform;
+
+	UPROPERTY()
+	bool bDeformAllBones_DEPRECATED = false;
+
+	UPROPERTY()
+	bool bDeformAllPhysicsBodies_DEPRECATED = false;
+
+};
+
diff --git a/Engine/Plugins/Mutable/Source/CustomizableObjectEditor/Public/MuCOE/Nodes/CustomizableObjectNodeMeshMorphStackApplication.h b/Engine/Plugins/Mutable/Source/CustomizableObjectEditor/Public/MuCOE/Nodes/CustomizableObjectNodeMeshMorphStackApplication.h
new file mode 100644
--- /dev/null
+++ b/Engine/Plugins/Mutable/Source/CustomizableObjectEditor/Public/MuCOE/Nodes/CustomizableObjectNodeMeshMorphStackApplication.h
@@ -0,0 +1,37 @@
+// Copyright Epic Games, Inc. All Rights Reserved.
+
+#pragma once
+
+#include "MuCOE/Nodes/CustomizableObjectNode.h"
+
+#include "CustomizableObjectNodeMeshMorphStackApplication.generated.h"
+
+class UCustomizableObjectNodeRemapPins;
+class UEdGraphPin;
+class UObject;
+
+
+UCLASS()
+class CUSTOMIZABLEOBJECTEDITOR_API UCustomizableObjectNodeMeshMorphStackApplication : public UCustomizableObjectNode
+{
+public:
+	GENERATED_BODY()
+
+	// EdGraphNode interface
+	virtual FText GetNodeTitle(ENodeTitleType::Type TittleType)const override;
+	virtual FLinearColor GetNodeTitleColor() const override;
+	virtual FText GetTooltipText() const override;
+
+	// UCustomizableObjectNode interface
+	virtual void AllocateDefaultPins(UCustomizableObjectNodeRemapPins* RemapPins) override;
+
+	// Own interface
+	/** Fills the list with all the morphs */
+	TArray<FString> GetMorphList() const;
+	
+	/** Returns the mesh pin. */
+	UEdGraphPin* GetMeshPin() const;
+
+	/** Returns the stack pin. */
+	UEdGraphPin* GetStackPin() const;
+};
diff --git a/Engine/Plugins/Mutable/Source/CustomizableObjectEditor/Public/MuCOE/Nodes/CustomizableObjectNodeMeshMorphStackDefinition.h b/Engine/Plugins/Mutable/Source/CustomizableObjectEditor/Public/MuCOE/Nodes/CustomizableObjectNodeMeshMorphStackDefinition.h
new file mode 100644
--- /dev/null
+++ b/Engine/Plugins/Mutable/Source/CustomizableObjectEditor/Public/MuCOE/Nodes/CustomizableObjectNodeMeshMorphStackDefinition.h
@@ -0,0 +1,51 @@
+// Copyright Epic Games, Inc. All Rights Reserved.
+
+#pragma once
+
+#include "MuCOE/Nodes/CustomizableObjectNode.h"
+
+#include "CustomizableObjectNodeMeshMorphStackDefinition.generated.h"
+
+class UCustomizableObjectNodeRemapPins;
+class UEdGraphPin;
+class UObject;
+
+
+UCLASS()
+class CUSTOMIZABLEOBJECTEDITOR_API UCustomizableObjectNodeMeshMorphStackDefinition : public UCustomizableObjectNode
+{
+public:
+	GENERATED_BODY()
+
+	// EdGraphNode interface
+	virtual FText GetNodeTitle(ENodeTitleType::Type TittleType)const override;
+	virtual FLinearColor GetNodeTitleColor() const override;
+	virtual FText GetTooltipText() const override;
+	virtual void PinConnectionListChanged(UEdGraphPin* Pin) override;
+	virtual bool IsNodeOutDatedAndNeedsRefresh() override;
+
+	// UCustomizableObjectNode interface
+	virtual void AllocateDefaultPins(UCustomizableObjectNodeRemapPins* RemapPins) override;
+	virtual void ReconstructNode(UCustomizableObjectNodeRemapPins* RemapPins) override;
+
+	// Own interface
+	/** Returns the mesh pin of the node. */
+	UEdGraphPin* GetMeshPin() const;
+
+	/** Returns the stack pin of the node. */
+	UEdGraphPin* GetStackPin() const;
+
+	/** Returns the morph pin at Index. */
+	UEdGraphPin* GetMorphPin(int32 Index) const;
+	
+	/** Returns the index of the next connected stack node. Returns -1 if there is none. */
+	int32 NextConnectedPin(int32 Index, TArray<FString> AvailableMorphs)const;
+
+private:
+	/** Fills the list with all the morphs. */
+	void UpdateMorphList();
+
+	/** List with all the morphs of the linked skeletal mesh. */
+	UPROPERTY()
+	TArray<FString> MorphNames;
+};
diff --git a/Engine/Plugins/Mutable/Source/CustomizableObjectEditor/Public/MuCOE/Nodes/CustomizableObjectNodeMeshReshape.h b/Engine/Plugins/Mutable/Source/CustomizableObjectEditor/Public/MuCOE/Nodes/CustomizableObjectNodeMeshReshape.h
new file mode 100644
--- /dev/null
+++ b/Engine/Plugins/Mutable/Source/CustomizableObjectEditor/Public/MuCOE/Nodes/CustomizableObjectNodeMeshReshape.h
@@ -0,0 +1,122 @@
+// Copyright Epic Games, Inc. All Rights Reserved.
+
+#pragma once
+
+#include "MuCOE/Nodes/CustomizableObjectNode.h"
+
+#include "CustomizableObjectNodeMeshReshape.generated.h"
+
+class UCustomizableObjectNodeRemapPins;
+class UEdGraphPin;
+class UObject;
+struct FMeshReshapeBoneReference;
+enum class EBoneDeformSelectionMethod : uint8;
+
+
+UENUM()
+enum class EMeshReshapeVertexColorChannelUsage
+{
+	None = 0,
+	RigidClusterId = 1,
+	MaskWeight = 2
+};
+
+USTRUCT()
+struct FMeshReshapeColorUsage
+{
+	GENERATED_BODY()
+
+	UPROPERTY(EditAnywhere, Category = CustomizableObject)
+	EMeshReshapeVertexColorChannelUsage R = EMeshReshapeVertexColorChannelUsage::None;
+
+	UPROPERTY(EditAnywhere, Category = CustomizableObject)
+	EMeshReshapeVertexColorChannelUsage G = EMeshReshapeVertexColorChannelUsage::None;
+
+	UPROPERTY(EditAnywhere, Category = CustomizableObject)
+	EMeshReshapeVertexColorChannelUsage B = EMeshReshapeVertexColorChannelUsage::None;
+
+	UPROPERTY(EditAnywhere, Category = CustomizableObject)
+	EMeshReshapeVertexColorChannelUsage A = EMeshReshapeVertexColorChannelUsage::None;
+};
+
+UCLASS()
+class CUSTOMIZABLEOBJECTEDITOR_API UCustomizableObjectNodeMeshReshape : public UCustomizableObjectNode
+{
+public:
+	GENERATED_BODY()
+
+	UCustomizableObjectNodeMeshReshape();
+
+	// Begin EdGraphNode interface
+	FText GetNodeTitle(ENodeTitleType::Type TitleType) const override;
+	FLinearColor GetNodeTitleColor() const override;
+	FText GetTooltipText() const override;
+
+	// UCustomizableObjectNode interface
+	void AllocateDefaultPins(UCustomizableObjectNodeRemapPins* RemapPins) override;
+	FString GetRefreshMessage() const override;
+
+	inline UEdGraphPin* BaseMeshPin() const
+	{
+		return FindPin(TEXT("Base Mesh"), EGPD_Input);
+	}
+
+	inline UEdGraphPin* BaseShapePin() const
+	{
+		return FindPin(TEXT("Base Shape"), EGPD_Input);
+	}
+
+	inline UEdGraphPin* TargetShapePin() const
+	{
+		return FindPin(TEXT("Target Shape"), EGPD_Input);
+	}
+
+	void Serialize(FArchive& Ar) override;
+
+	/** Enable the deformation of the vertices of the base mesh. */
+	UPROPERTY(EditAnywhere, Category = CustomizableObject)
+	bool bReshapeVertices = true;
+
+	UPROPERTY(EditAnywhere, Category = CustomizableObject)
+	bool bRecomputeNormals = false;
+
+	/** Enable laplacian smoothing to the result of the base mesh reshape. */
+	UPROPERTY(EditAnywhere, Category = CustomizableObject, meta = (EditCondition = "bReshapeVertices"))
+	bool bApplyLaplacianSmoothing = false;
+
+	/** Enable the deformation of the skeleton of the base mesh. */
+	UPROPERTY(EditAnywhere, Category = CustomizableObject)
+	bool bReshapeSkeleton = false;
+
+	/** Enable the deformation of physics volumes of the base mesh */
+    UPROPERTY(EditAnywhere, Category = CustomizableObject)
+    bool bReshapePhysicsVolumes = false;
+	
+	UPROPERTY()
+	bool bEnableRigidParts_DEPRECATED = false;
+
+	UPROPERTY(EditAnywhere, Category = CustomizableObject)
+	FMeshReshapeColorUsage VertexColorUsage;
+
+	/** Bone Reshape Selection Method */
+	UPROPERTY(EditAnywhere, Category = ReshapeBones, meta = (EditCondition = "bReshapeSkeleton"))
+	EBoneDeformSelectionMethod SelectionMethod;
+
+	/** Array with selected bones that will be deformed */
+	UPROPERTY(EditAnywhere, Category = ReshapeBones, meta = (EditCondition = "bReshapeSkeleton && (SelectionMethod == EBoneDeformSelectionMethod::ONLY_SELECTED || SelectionMethod == EBoneDeformSelectionMethod::ALL_BUT_SELECTED)"))
+	TArray<FMeshReshapeBoneReference> BonesToDeform;
+
+	UPROPERTY(EditAnywhere, Category = ReshapePhysics, meta = (DisplayName = "Selection Method", EditCondition = "bReshapePhysicsVolumes"))
+	EBoneDeformSelectionMethod PhysicsSelectionMethod;
+
+	/** Array with bones with physics bodies that will be deformed */
+	UPROPERTY(EditAnywhere, Category = ReshapePhysics, meta = (EditCondition = "bReshapePhysicsVolumes && (PhysicsSelectionMethod == EBoneDeformSelectionMethod::ONLY_SELECTED || PhysicsSelectionMethod == EBoneDeformSelectionMethod::ALL_BUT_SELECTED)"))
+	TArray<FMeshReshapeBoneReference> PhysicsBodiesToDeform;
+
+	UPROPERTY()
+	bool bDeformAllBones_DEPRECATED = false;
+
+	UPROPERTY()
+	bool bDeformAllPhysicsBodies_DEPRECATED = false;
+};
+
diff --git a/Engine/Plugins/Mutable/Source/CustomizableObjectEditor/Public/MuCOE/Nodes/CustomizableObjectNodeMeshReshapeCommon.h b/Engine/Plugins/Mutable/Source/CustomizableObjectEditor/Public/MuCOE/Nodes/CustomizableObjectNodeMeshReshapeCommon.h
new file mode 100644
--- /dev/null
+++ b/Engine/Plugins/Mutable/Source/CustomizableObjectEditor/Public/MuCOE/Nodes/CustomizableObjectNodeMeshReshapeCommon.h
@@ -0,0 +1,33 @@
+// Copyright Epic Games, Inc. All Rights Reserved.
+
+#pragma once
+
+
+#include "CustomizableObjectNodeMeshReshapeCommon.generated.h"
+
+UENUM()
+enum class EBoneDeformSelectionMethod : uint8
+{
+	// Only selected bones will be deform
+	ONLY_SELECTED = 0 UMETA(DisplayName = "Only Selected"),
+
+	// All bones will be deform except the selected ones
+	ALL_BUT_SELECTED = 1 UMETA(DisplayName = "All But Selected"),
+
+	// Deform only the bones of the reference skeleton
+	DEFORM_REF_SKELETON = 2 UMETA(DisplayName = "Deform Only Ref.Seleton Bones"),
+
+	// Deform only the bones that are not in the reference skeleton
+	DEFORM_NONE_REF_SKELETON = 3 UMETA(DisplayName = "Deform All But not Ref.Skeleton Bones")
+};
+
+
+USTRUCT()
+struct FMeshReshapeBoneReference
+{
+	GENERATED_USTRUCT_BODY()
+
+	/** Name of the bone that will be deformed */
+	UPROPERTY(EditAnywhere, Category = CustomizableObject)
+	FName BoneName;
+};
diff --git a/Engine/Plugins/Mutable/Source/CustomizableObjectEditor/Public/MuCOE/Nodes/CustomizableObjectNodeMeshSwitch.h b/Engine/Plugins/Mutable/Source/CustomizableObjectEditor/Public/MuCOE/Nodes/CustomizableObjectNodeMeshSwitch.h
new file mode 100644
--- /dev/null
+++ b/Engine/Plugins/Mutable/Source/CustomizableObjectEditor/Public/MuCOE/Nodes/CustomizableObjectNodeMeshSwitch.h
@@ -0,0 +1,22 @@
+// Copyright Epic Games, Inc. All Rights Reserved.
+
+#pragma once
+
+#include "MuCOE/Nodes/CustomizableObjectNodeSwitchBase.h"
+
+#include "CustomizableObjectNodeMeshSwitch.generated.h"
+
+
+UCLASS()
+class CUSTOMIZABLEOBJECTEDITOR_API UCustomizableObjectNodeMeshSwitch : public UCustomizableObjectNodeSwitchBase
+{
+public:
+	GENERATED_BODY()
+
+	// UCustomizableObjectNode interface
+	virtual void BackwardsCompatibleFixup(int32 CustomizableObjectCustomVersion) override;
+	
+	// UCustomizableObjectNodeSwitchBase interface
+	virtual FName GetCategory() const override;
+};
+
diff --git a/Engine/Plugins/Mutable/Source/CustomizableObjectEditor/Public/MuCOE/Nodes/CustomizableObjectNodeMeshVariation.h b/Engine/Plugins/Mutable/Source/CustomizableObjectEditor/Public/MuCOE/Nodes/CustomizableObjectNodeMeshVariation.h
new file mode 100644
--- /dev/null
+++ b/Engine/Plugins/Mutable/Source/CustomizableObjectEditor/Public/MuCOE/Nodes/CustomizableObjectNodeMeshVariation.h
@@ -0,0 +1,28 @@
+// Copyright Epic Games, Inc. All Rights Reserved.
+
+#pragma once
+
+#include "MuCOE/Nodes/CustomizableObjectNodeVariation.h"
+
+#include "CustomizableObjectNodeMeshVariation.generated.h"
+
+struct FCustomizableObjectMeshVariation;
+
+
+UCLASS()
+class CUSTOMIZABLEOBJECTEDITOR_API UCustomizableObjectNodeMeshVariation : public UCustomizableObjectNodeVariation
+{
+public:
+	GENERATED_BODY()
+
+	// Deprecated properties
+	UPROPERTY()
+	TArray<FCustomizableObjectMeshVariation> Variations_DEPRECATED;
+
+	// UCustomizableObjectNode interface
+	virtual void BackwardsCompatibleFixup(int32 CustomizableObjectCustomVersion) override;
+	
+	// UCustomizableObjectNodeVariation interface
+	virtual FName GetCategory() const override;
+};
+
diff --git a/Engine/Plugins/Mutable/Source/CustomizableObjectEditor/Public/MuCOE/Nodes/CustomizableObjectNodeModifierClipDeform.h b/Engine/Plugins/Mutable/Source/CustomizableObjectEditor/Public/MuCOE/Nodes/CustomizableObjectNodeModifierClipDeform.h
new file mode 100644
--- /dev/null
+++ b/Engine/Plugins/Mutable/Source/CustomizableObjectEditor/Public/MuCOE/Nodes/CustomizableObjectNodeModifierClipDeform.h
@@ -0,0 +1,51 @@
+// Copyright Epic Games, Inc. All Rights Reserved.
+
+#pragma once
+
+#include "MuCOE/Nodes/CustomizableObjectNodeModifierBase.h"
+#include "MuT/NodeModifier.h"
+
+#include "CustomizableObjectNodeModifierClipDeform.generated.h"
+
+class UCustomizableObjectNodeRemapPins;
+class UEdGraphPin;
+class UObject;
+
+UENUM()
+enum class EShapeBindingMethod : uint32
+{
+	ClosestProject = 0,
+	ClosestToSurface = 1,
+	NormalProject = 2
+};
+
+UCLASS()
+class CUSTOMIZABLEOBJECTEDITOR_API UCustomizableObjectNodeModifierClipDeform : public UCustomizableObjectNodeModifierBase
+{
+
+	GENERATED_BODY()
+
+public:
+
+	UPROPERTY()
+	TArray<FString> Tags_DEPRECATED;
+
+	UPROPERTY(EditAnywhere, Category = MeshClipDeform)
+	EShapeBindingMethod BindingMethod;
+	
+public:
+
+	// EdGraphNode interface
+	virtual FText GetNodeTitle(ENodeTitleType::Type TitleType) const override;
+	virtual FText GetTooltipText() const override;
+
+	// UCustomizableObjectNode interface
+	virtual void AllocateDefaultPins(UCustomizableObjectNodeRemapPins* RemapPins) override;
+	virtual void BackwardsCompatibleFixup(int32 CustomizableObjectCustomVersion) override;
+
+	// Own interface
+
+	/** Input pins. */
+	UEdGraphPin* ClipShapePin() const;
+};
+
diff --git a/Engine/Plugins/Mutable/Source/CustomizableObjectEditor/Public/MuCOE/Nodes/CustomizableObjectNodeModifierClipWithMesh.h b/Engine/Plugins/Mutable/Source/CustomizableObjectEditor/Public/MuCOE/Nodes/CustomizableObjectNodeModifierClipWithMesh.h
new file mode 100644
--- /dev/null
+++ b/Engine/Plugins/Mutable/Source/CustomizableObjectEditor/Public/MuCOE/Nodes/CustomizableObjectNodeModifierClipWithMesh.h
@@ -0,0 +1,66 @@
+// Copyright Epic Games, Inc. All Rights Reserved.
+
+#pragma once
+
+#include "MuCOE/Nodes/CustomizableObjectNodeModifierBase.h"
+#include "MuT/NodeModifier.h"
+
+#include "CustomizableObjectNodeModifierClipWithMesh.generated.h"
+
+namespace ENodeTitleType { enum Type : int; }
+
+class UCustomizableObject;
+class UCustomizableObjectNodeRemapPins;
+class UEdGraphPin;
+class UObject;
+struct FGuid;
+
+UCLASS()
+class CUSTOMIZABLEOBJECTEDITOR_API UCustomizableObjectNodeModifierClipWithMesh : public UCustomizableObjectNodeModifierBase
+{
+public:
+	GENERATED_BODY()
+
+	UCustomizableObjectNodeModifierClipWithMesh();
+
+	UPROPERTY()
+	TArray<FString> Tags_DEPRECATED;
+
+	//!< If assigned, then a material inside this CO will be clipped by this node.
+    //!< If several materials with the same name, all are considered (to cover all LOD levels)
+    UPROPERTY(EditAnywhere, Category = CustomizableObjectToClip)
+	TObjectPtr<UCustomizableObject> CustomizableObjectToClipWith;
+
+    //!< Array with the Guids of the nodes with the same material inside the CustomizableObjectToClipWith CO (if any is assigned)
+    UPROPERTY(EditAnywhere, Category = CustomizableObjectToClip)
+    TArray<FGuid> ArrayMaterialNodeToClipWithID;
+
+	/** Mesh Transform*/
+	UPROPERTY(EditAnywhere, Category = CustomizableObject)
+	FTransform Transform;
+
+	// Details view variables
+	// The clipping node uses tags
+	bool bUseTags;
+
+	// The clipping node uses a material name
+	bool bUseMaterials;
+
+	// EdGraphNode interface
+	virtual FText GetNodeTitle(ENodeTitleType::Type TitleType) const override;
+	virtual FText GetTooltipText() const override;
+
+	// UCustomizableObjectNode interface
+	virtual void AllocateDefaultPins(UCustomizableObjectNodeRemapPins* RemapPins) override;
+	virtual void ReconstructNode(UCustomizableObjectNodeRemapPins* RemapPinsMode) override;
+	virtual void PinConnectionListChanged(UEdGraphPin* Pin) override;
+	virtual void UpdateReferencedNodeId(const FGuid& NewGuid) override;
+	virtual void BeginPostDuplicate(bool bDuplicateForPIE) override;
+	virtual void BackwardsCompatibleFixup(int32 CustomizableObjectCustomVersion) override;
+
+	// Own interface
+	UEdGraphPin* OutputPin() const;
+
+	UEdGraphPin* ClipMeshPin() const;
+};
+
diff --git a/Engine/Plugins/Mutable/Source/CustomizableObjectEditor/Public/MuCOE/Nodes/CustomizableObjectNodeModifierEditMeshSection.h b/Engine/Plugins/Mutable/Source/CustomizableObjectEditor/Public/MuCOE/Nodes/CustomizableObjectNodeModifierEditMeshSection.h
new file mode 100644
--- /dev/null
+++ b/Engine/Plugins/Mutable/Source/CustomizableObjectEditor/Public/MuCOE/Nodes/CustomizableObjectNodeModifierEditMeshSection.h
@@ -0,0 +1,86 @@
+// Copyright Epic Games, Inc. All Rights Reserved.
+
+#pragma once
+
+#include "MuCOE/Nodes/CustomizableObjectNodeModifierEditLayoutBlocks.h"
+#include "MuCOE/Nodes/CustomizableObjectNodeMaterial.h"
+
+#include "CustomizableObjectNodeModifierEditMeshSection.generated.h"
+
+namespace ENodeTitleType { enum Type : int; }
+struct FCustomizableObjectNodeEditMaterialImage;
+struct FNodeMaterialParameterId;
+
+class FArchive;
+class FCustomizableObjectNodeParentedMaterial;
+class UCustomizableObject;
+class UCustomizableObjectNodeRemapPins;
+class UObject;
+
+
+/** Additional data for the Parent Texture Parameter to edit pin. */
+UCLASS()
+class CUSTOMIZABLEOBJECTEDITOR_API UCustomizableObjectNodeEditMaterialPinEditImageData : public UCustomizableObjectNodePinData
+{
+	GENERATED_BODY()
+
+public:
+
+	/** Parameter FGuid + Layer Index */
+	UPROPERTY()
+	FNodeMaterialParameterId ImageParamId;
+
+	UPROPERTY();
+	FEdGraphPinReference PinMask;
+
+	UPROPERTY()
+	FGuid ImageId_DEPRECATED;
+};
+
+
+UCLASS()
+class CUSTOMIZABLEOBJECTEDITOR_API UCustomizableObjectNodeModifierEditMeshSection : public UCustomizableObjectNodeModifierEditLayoutBlocks
+{
+
+	GENERATED_BODY()
+
+public:
+	
+	// EdGraphNode interface
+	virtual FText GetNodeTitle(ENodeTitleType::Type TitleType) const override;
+	virtual FText GetTooltipText() const override;
+
+	// UCustomizableObjectNode interface
+	virtual void BackwardsCompatibleFixup(int32 CustomizableObjectCustomVersion) override;
+	virtual void AllocateDefaultPins(UCustomizableObjectNodeRemapPins* RemapPins) override;
+	virtual FString GetRefreshMessage() const override;
+	virtual bool IsSingleOutputNode() const override;
+	virtual bool CustomRemovePin(UEdGraphPin& Pin) override;
+	virtual bool HasPinViewer() const override;
+	
+	// Own interface
+	
+	/** */
+	virtual void SetLayoutIndex(const int32 LayoutIndex);
+
+	/** Returns the Image mask pin of the given Image that will be edited.
+	 *
+	 * @returns Always returns a valid pin if EditsImage(const FGuid&) returns true. */
+	const UEdGraphPin* GetUsedImageMaskPin(const FNodeMaterialParameterId& ImageId) const;
+
+private:
+
+	// Deprecated properties
+	UPROPERTY()
+	TMap<FNodeMaterialParameterId, FEdGraphPinReference> PinsParameterMap_DEPRECATED;
+
+	UPROPERTY()
+	TMap<FGuid, FEdGraphPinReference> PinsParameter_DEPRECATED;
+	
+	UPROPERTY()
+	TArray<FCustomizableObjectNodeEditMaterialImage> Images_DEPRECATED;
+	
+	// Old layout blocks to patch. Now in parent class.
+	UPROPERTY()
+	TArray<int32> Blocks_DEPRECATED;
+};
diff --git a/Engine/Plugins/Mutable/Source/CustomizableObjectEditor/Public/MuCOE/Nodes/CustomizableObjectNodeModifierExtendMeshSection.h b/Engine/Plugins/Mutable/Source/CustomizableObjectEditor/Public/MuCOE/Nodes/CustomizableObjectNodeModifierExtendMeshSection.h
new file mode 100644
--- /dev/null
+++ b/Engine/Plugins/Mutable/Source/CustomizableObjectEditor/Public/MuCOE/Nodes/CustomizableObjectNodeModifierExtendMeshSection.h
@@ -0,0 +1,72 @@
+// Copyright Epic Games, Inc. All Rights Reserved.
+
+#pragma once
+
+#include "MuCOE/Nodes/CustomizableObjectNodeModifierWithMaterial.h"
+#include "MuCOE/CustomizableObjectLayout.h"
+
+#include "CustomizableObjectNodeModifierExtendMeshSection.generated.h"
+
+namespace ENodeTitleType { enum Type : int; }
+
+class FArchive;
+class UCustomizableObject;
+class UCustomizableObjectNode;
+class UCustomizableObjectNodeRemapPins;
+class UEdGraphPin;
+class UObject;
+struct FCustomizableObjectNodeExtendMaterialImage;
+struct FEdGraphPinReference;
+
+
+UCLASS()
+class CUSTOMIZABLEOBJECTEDITOR_API UCustomizableObjectNodeModifierExtendMeshSection : public UCustomizableObjectNodeModifierWithMaterial
+{
+	GENERATED_BODY()
+
+public:
+
+	/** Tags enabled when the modifier is applied. Also used to decide what modifiers get applied to the data added by this node.  */
+	UPROPERTY(EditAnywhere, BlueprintReadWrite, Category = EnableTags)
+	TArray<FString> Tags;
+
+public:
+	
+	// Begin EdGraphNode interface
+	virtual FText GetNodeTitle(ENodeTitleType::Type TitleType) const override;
+	virtual FText GetTooltipText() const override;
+
+	// UCustomizableObjectNode interface
+	virtual void BackwardsCompatibleFixup(int32 CustomizableObjectCustomVersion) override;
+	virtual void AllocateDefaultPins(UCustomizableObjectNodeRemapPins* RemapPins) override;
+	virtual FString GetRefreshMessage() const override;
+	virtual bool IsSingleOutputNode() const override;
+	virtual bool CustomRemovePin(UEdGraphPin& Pin) override;
+	virtual TArray<FString>* GetEnableTags() override;
+
+	// Own interface
+	UEdGraphPin* AddMeshPin() const;
+
+	TArray<UCustomizableObjectLayout*> GetLayouts() const;
+
+private:
+
+	// Deprecated properties
+
+	UPROPERTY()
+	TObjectPtr<UCustomizableObject> ParentMaterialObject_DEPRECATED = nullptr;
+
+	UPROPERTY()
+	FGuid ParentMaterialNodeId_DEPRECATED;
+	
+
+	UPROPERTY()
+	TMap<FNodeMaterialParameterId, FEdGraphPinReference> PinsParameterMap_DEPRECATED;
+
+	UPROPERTY()
+	TMap<FGuid, FEdGraphPinReference> PinsParameter_DEPRECATED;
+
+	UPROPERTY()
+	TArray<FCustomizableObjectNodeExtendMaterialImage> Images_DEPRECATED;
+};
+
diff --git a/Engine/Plugins/Mutable/Source/CustomizableObjectEditor/Public/MuCOE/Nodes/CustomizableObjectNodeModifierMorphMeshSection.h b/Engine/Plugins/Mutable/Source/CustomizableObjectEditor/Public/MuCOE/Nodes/CustomizableObjectNodeModifierMorphMeshSection.h
new file mode 100644
--- /dev/null
+++ b/Engine/Plugins/Mutable/Source/CustomizableObjectEditor/Public/MuCOE/Nodes/CustomizableObjectNodeModifierMorphMeshSection.h
@@ -0,0 +1,43 @@
+// Copyright Epic Games, Inc. All Rights Reserved.
+
+#pragma once
+
+#include "MuCOE/Nodes/CustomizableObjectNodeModifierEditMeshSectionBase.h"
+
+#include "CustomizableObjectNodeModifierMorphMeshSection.generated.h"
+
+namespace ENodeTitleType { enum Type : int; }
+
+class UCustomizableObjectNodeRemapPins;
+class UEdGraphPin;
+
+
+UCLASS()
+class CUSTOMIZABLEOBJECTEDITOR_API UCustomizableObjectNodeModifierMorphMeshSection: public UCustomizableObjectNodeModifierEditMeshSectionBase
+{
+
+	GENERATED_BODY()
+	
+public:
+
+	UPROPERTY(EditAnywhere, Category = CustomizableObject)
+	FString MorphTargetName;
+
+public:
+
+	// Begin EdGraphNode interface
+	FText GetNodeTitle(ENodeTitleType::Type TitleType) const override;
+	FText GetTooltipText() const override;
+
+	// UCustomizableObjectNode interface
+	virtual void AllocateDefaultPins(UCustomizableObjectNodeRemapPins* RemapPins) override;
+	virtual FString GetRefreshMessage() const override;
+	virtual bool IsSingleOutputNode() const override;
+
+	UEdGraphPin* FactorPin() const
+	{
+		return FindPin(TEXT("Factor"));
+	}
+
+};
+
diff --git a/Engine/Plugins/Mutable/Source/CustomizableObjectEditor/Public/MuCOE/Nodes/CustomizableObjectNodeModifierRemoveMesh.h b/Engine/Plugins/Mutable/Source/CustomizableObjectEditor/Public/MuCOE/Nodes/CustomizableObjectNodeModifierRemoveMesh.h
new file mode 100644
--- /dev/null
+++ b/Engine/Plugins/Mutable/Source/CustomizableObjectEditor/Public/MuCOE/Nodes/CustomizableObjectNodeModifierRemoveMesh.h
@@ -0,0 +1,36 @@
+// Copyright Epic Games, Inc. All Rights Reserved.
+
+#pragma once
+
+#include "MuCOE/Nodes/CustomizableObjectNodeModifierEditMeshSectionBase.h"
+
+#include "CustomizableObjectNodeModifierRemoveMesh.generated.h"
+
+namespace ENodeTitleType { enum Type : int; }
+
+class UCustomizableObjectNodeRemapPins;
+class UEdGraphPin;
+class UObject;
+
+UCLASS()
+class CUSTOMIZABLEOBJECTEDITOR_API UCustomizableObjectNodeModifierRemoveMesh : public UCustomizableObjectNodeModifierEditMeshSectionBase
+{
+public:
+	GENERATED_BODY()
+
+	// Begin EdGraphNode interface
+	FText GetNodeTitle(ENodeTitleType::Type TitleType) const override;
+	FText GetTooltipText() const override;
+
+	// UCustomizableObjectNode interface
+	virtual void AllocateDefaultPins(UCustomizableObjectNodeRemapPins* RemapPins) override;
+	virtual void PinConnectionListChanged(UEdGraphPin * Pin);
+
+	UEdGraphPin* RemoveMeshPin() const
+	{
+		return FindPin(TEXT("Remove Mesh"));
+	}
+
+	bool IsSingleOutputNode() const override;
+};
+
diff --git a/Engine/Plugins/Mutable/Source/CustomizableObjectEditor/Public/MuCOE/Nodes/CustomizableObjectNodeModifierRemoveMeshBlocks.h b/Engine/Plugins/Mutable/Source/CustomizableObjectEditor/Public/MuCOE/Nodes/CustomizableObjectNodeModifierRemoveMeshBlocks.h
new file mode 100644
--- /dev/null
+++ b/Engine/Plugins/Mutable/Source/CustomizableObjectEditor/Public/MuCOE/Nodes/CustomizableObjectNodeModifierRemoveMeshBlocks.h
@@ -0,0 +1,35 @@
+// Copyright Epic Games, Inc. All Rights Reserved.
+
+#pragma once
+
+#include "MuCOE/Nodes/CustomizableObjectNodeModifierEditLayoutBlocks.h"
+
+#include "CustomizableObjectNodeModifierRemoveMeshBlocks.generated.h"
+
+namespace ENodeTitleType { enum Type : int; }
+
+class FArchive;
+class UCustomizableObjectNodeRemapPins;
+class UEdGraphPin;
+class UObject;
+struct FPropertyChangedEvent;
+
+UCLASS()
+class CUSTOMIZABLEOBJECTEDITOR_API UCustomizableObjectNodeModifierRemoveMeshBlocks : public UCustomizableObjectNodeModifierEditLayoutBlocks
+{
+public:
+	GENERATED_BODY()
+
+	UPROPERTY()
+	TArray<int32> Blocks_DEPRECATED;
+
+	// Begin EdGraphNode interface
+	virtual FText GetNodeTitle(ENodeTitleType::Type TitleType) const override;
+	virtual FText GetTooltipText() const override;
+
+	// UCustomizableObjectNode interface
+	virtual void BackwardsCompatibleFixup(int32 CustomizableObjectCustomVersion) override;
+	virtual void AllocateDefaultPins(UCustomizableObjectNodeRemapPins* RemapPins) override;
+	virtual bool IsSingleOutputNode() const override;
+};
+
diff --git a/Engine/Plugins/Mutable/Source/CustomizableObjectEditor/Public/MuCOE/Nodes/CustomizableObjectNodeModifierWithMaterial.h b/Engine/Plugins/Mutable/Source/CustomizableObjectEditor/Public/MuCOE/Nodes/CustomizableObjectNodeModifierWithMaterial.h
new file mode 100644
--- /dev/null
+++ b/Engine/Plugins/Mutable/Source/CustomizableObjectEditor/Public/MuCOE/Nodes/CustomizableObjectNodeModifierWithMaterial.h
@@ -0,0 +1,56 @@
+// Copyright Epic Games, Inc. All Rights Reserved.
+
+#pragma once
+
+#include "MuCOE/Nodes/CustomizableObjectNodeModifierBase.h"
+#include "MuCOE/Nodes/CustomizableObjectNodeMaterialBase.h"
+
+#include "CustomizableObjectNodeModifierWithMaterial.generated.h"
+
+class UObject;
+
+
+UCLASS(abstract)
+class CUSTOMIZABLEOBJECTEDITOR_API UCustomizableObjectNodeModifierWithMaterial : public UCustomizableObjectNodeModifierBase
+{
+
+	GENERATED_BODY()
+
+public:
+
+	/** Reference material that defines the structure of the data to be extended.
+	* The Sections modified with this modifier are supposed to have the same texture parameters, but don't need to have exactly
+	* the ReferenceMaterial set.
+	*/
+	UPROPERTY(EditAnywhere, Category = Parameters)
+	TObjectPtr<UMaterialInterface> ReferenceMaterial = nullptr;
+
+	/** Relates a Parameter id (and layer if is a layered material) to a Pin. Only used to improve performance. */
+	UPROPERTY()
+	TMap<FNodeMaterialParameterId, FEdGraphPinReference> PinsParameterMap;
+
+public:
+
+	//
+	virtual void PostEditChangeProperty(FPropertyChangedEvent& PropertyChangedEvent) override;
+
+	// UCustomizableObjectNode interface
+	virtual bool IsNodeOutDatedAndNeedsRefresh() override;
+
+	// Own interface
+
+	/** */
+	bool UsesImage(const FNodeMaterialParameterId& ImageId) const;
+
+	/** */
+	const UEdGraphPin* GetUsedImagePin(const FNodeMaterialParameterId& ImageId) const;
+
+	/** */
+	int32 GetNumParameters(const EMaterialParameterType Type) const;
+	FNodeMaterialParameterId GetParameterId(const EMaterialParameterType Type, const int32 ParameterIndex) const;
+	FName GetParameterName(const EMaterialParameterType Type, const int32 ParameterIndex) const;
+	static int32 GetParameterLayerIndex(const UMaterialInterface* InMaterial, const EMaterialParameterType Type, const int32 ParameterIndex);
+	int32 GetParameterLayerIndex(const EMaterialParameterType Type, const int32 ParameterIndex) const;
+	FText GetParameterLayerName(const EMaterialParameterType Type, const int32 ParameterIndex) const;
+};
+
diff --git a/Engine/Plugins/Mutable/Source/CustomizableObjectEditor/Public/MuCOE/Nodes/CustomizableObjectNodeObject.h b/Engine/Plugins/Mutable/Source/CustomizableObjectEditor/Public/MuCOE/Nodes/CustomizableObjectNodeObject.h
new file mode 100644
--- /dev/null
+++ b/Engine/Plugins/Mutable/Source/CustomizableObjectEditor/Public/MuCOE/Nodes/CustomizableObjectNodeObject.h
@@ -0,0 +1,206 @@
+// Copyright Epic Games, Inc. All Rights Reserved.
+
+#pragma once
+
+#include "CustomizableObjectNodeComponentMesh.h"
+#include "CustomizableObjectNodeComponentMeshBase.h"
+#include "MuCOE/Nodes/CustomizableObjectNode.h"
+#include "MuR/System.h"
+
+#include "CustomizableObjectNodeObject.generated.h"
+
+namespace ENodeTitleType { enum Type : int; }
+
+class UCustomizableObject;
+class UCustomizableObjectNodeMaterialBase;
+class UCustomizableObjectNodeRemapPins;
+class UObject;
+struct FPropertyChangedEvent;
+struct FSoftObjectPath;
+struct FComponentSettings;
+
+
+USTRUCT()
+struct CUSTOMIZABLEOBJECTEDITOR_API FCustomizableObjectState
+{
+	GENERATED_USTRUCT_BODY()
+
+	friend UCustomizableObjectNodeObject;
+	
+	UPROPERTY(EditAnywhere, Category=CustomizableObject)
+	FString Name;
+
+	UPROPERTY(EditAnywhere, Category = CustomizableObject, meta = (ShowParameterOptions))
+	TArray<FString> RuntimeParameters;
+
+	/** Special treatment of texture compression for this state. */
+	UPROPERTY(EditAnywhere, Category = CustomizableObject)
+	ETextureCompressionStrategy TextureCompressionStrategy = ETextureCompressionStrategy::None;
+
+	/** If this is enabled, texture streaming won't be used for this state, and full images will be generated when an instance is first updated. */
+	UPROPERTY(EditAnywhere, Category = CustomizableObject)
+	bool bDisableTextureStreaming = false;
+
+	/** LiveUpdateMode will reuse instance temp. data between updates and speed up update times, but spend much more memory. Good for customization screens, not for actual gameplay modes. */
+	UPROPERTY(EditAnywhere, Category = CustomizableObject)
+	bool bLiveUpdateMode = false;
+
+	// Enables the reuse of all possible textures when the instance is updated without any changes in geometry or state (the first update after creation doesn't reuse any)
+	// It will only work if the textures aren't compressed, so set the instance to a Mutable state with texture compression disabled
+	// WARNING! If texture reuse is enabled, do NOT keep external references to the textures of the instance. The instance owns the textures.
+	UPROPERTY(EditAnywhere, Category = CustomizableObject)
+	bool bReuseInstanceTextures = false;
+
+	UPROPERTY(EditAnywhere, Category = CustomizableObject)
+	bool bBuildOnlyFirstLOD = false;
+
+	/** If there's an entry for a specific platform, when compiling for that platform Num LODs will be built after bBuildOnlyFirstLOD */
+	UPROPERTY(EditAnywhere, Category = CustomizableObject)
+	TMap<FString, int32> NumExtraLODsToBuildPerPlatform;
+
+	UPROPERTY(EditAnywhere, Category = CustomizableObject)
+	TMap<FString, FString> ForcedParameterValues;
+
+	UPROPERTY(EditAnywhere, Category = UI, meta = (DisplayName = "State UI Metadata"))
+	FMutableStateUIMetadata UIMetadata;
+
+private:
+	// Deprecated
+	UPROPERTY()
+	FMutableParamUIMetadata StateUIMetadata_DEPRECATED;
+	
+	/** This is now TextureCompressionStrategy.  */
+	UPROPERTY()
+	bool bDontCompressRuntimeTextures_DEPRECATED = false;
+};
+
+
+USTRUCT()
+struct FRealTimeMorphSelectionOverride
+{
+	GENERATED_BODY()
+
+    FRealTimeMorphSelectionOverride()
+    {
+    }
+
+    FRealTimeMorphSelectionOverride(const FName& InMorphName)
+        : MorphName(InMorphName)
+    {
+    }
+
+    UPROPERTY()
+    FName MorphName;
+
+    UPROPERTY()
+    ECustomizableObjectSelectionOverride SelectionOverride = 
+        ECustomizableObjectSelectionOverride::NoOverride;
+
+    UPROPERTY()
+    TArray<FName> SkeletalMeshesNames;
+
+    UPROPERTY()
+    TArray<ECustomizableObjectSelectionOverride> Override;
+};
+
+
+UCLASS()
+class CUSTOMIZABLEOBJECTEDITOR_API UCustomizableObjectNodeObject : public UCustomizableObjectNode
+{
+public:
+	GENERATED_BODY()
+
+	UCustomizableObjectNodeObject();
+
+	UPROPERTY(EditAnywhere, Category=CustomizableObject)
+	FString ObjectName;
+
+	UPROPERTY(EditAnywhere, Category = UI, meta = (DisplayName = "Parameter UI Metadata"))
+	FMutableParamUIMetadata ParamUIMetadata;
+
+	UPROPERTY()
+	int32 NumLODs_DEPRECATED = 1;
+
+	UPROPERTY()
+	ECustomizableObjectAutomaticLODStrategy AutoLODStrategy_DEPRECATED = ECustomizableObjectAutomaticLODStrategy::AutomaticFromMesh;
+
+	UPROPERTY()
+	int32 NumMeshComponents_DEPRECATED = 1;
+
+	UPROPERTY(EditAnywhere, Category=CustomizableObject ,meta = (TitleProperty = "Name"))
+	TArray<FCustomizableObjectState> States;
+
+	UPROPERTY(EditAnywhere, Category = AttachedToExternalObject)
+	TObjectPtr<UCustomizableObject> ParentObject;
+
+	UPROPERTY(EditAnywhere, Category = AttachedToExternalObject)
+	FGuid ParentObjectGroupId;
+
+	UPROPERTY()
+	FGuid Identifier;
+
+    UPROPERTY()
+    TArray<FRealTimeMorphSelectionOverride> RealTimeMorphSelectionOverrides;
+
+	UPROPERTY()
+	TArray<FComponentSettings> ComponentSettings_DEPRECATED;
+	
+	// UObject interface.
+	virtual void PostEditChangeProperty(FPropertyChangedEvent& PropertyChangedEvent) override;
+
+	// EdGraphNode interface
+	FText GetNodeTitle(ENodeTitleType::Type TitleType) const override;
+	FLinearColor GetNodeTitleColor() const override;
+	FText GetTooltipText() const override;
+	void PrepareForCopying() override;
+	void PostPasteNode() override;
+	void PostDuplicate(bool bDuplicateForPIE) override;
+
+	// UCustomizableObjectNode interface
+	virtual void BackwardsCompatibleFixup(int32 CustomizableObjectCustomVersion) override;
+	virtual void PostBackwardsCompatibleFixup() override;
+	void AllocateDefaultPins(UCustomizableObjectNodeRemapPins* RemapPins) override;
+
+	// Own Interface
+	void SetParentObject(UCustomizableObject* CustomizableParentObject);
+
+	// Own interface
+	UPROPERTY()
+	bool bIsBase;
+	
+	UEdGraphPin* ComponentsPin() const
+	{
+		return FindPin(ComponentsPinName);
+	}
+
+	UEdGraphPin* ModifiersPin() const
+	{
+		return FindPin(ModifiersPinName);
+	}
+
+	UEdGraphPin* ChildrenPin() const
+	{
+		return FindPin(ChildrenPinName);
+	}
+
+	UEdGraphPin* OutputPin() const
+	{
+		return FindPin(OutputPinName);
+	}
+
+	
+	virtual bool CanUserDeleteNode() const override;
+	virtual bool CanDuplicateNode() const override;
+	
+	bool IsSingleOutputNode() const override;
+	
+	// Array filled in the Details of the node to store all the parameter names of a CO graph (full tree)
+	TArray<FString> ParameterNames;
+
+private:
+	static const FName ChildrenPinName;
+	static const FName ComponentsPinName;
+	static const FName ModifiersPinName;
+	static const FName OutputPinName;
+};
+
diff --git a/Engine/Plugins/Mutable/Source/CustomizableObjectEditor/Public/MuCOE/Nodes/CustomizableObjectNodeObjectChild.h b/Engine/Plugins/Mutable/Source/CustomizableObjectEditor/Public/MuCOE/Nodes/CustomizableObjectNodeObjectChild.h
new file mode 100644
--- /dev/null
+++ b/Engine/Plugins/Mutable/Source/CustomizableObjectEditor/Public/MuCOE/Nodes/CustomizableObjectNodeObjectChild.h
@@ -0,0 +1,28 @@
+// Copyright Epic Games, Inc. All Rights Reserved.
+
+#pragma once
+
+#include "MuCOE/Nodes/CustomizableObjectNodeObject.h"
+
+#include "CustomizableObjectNodeObjectChild.generated.h"
+
+class UObject;
+
+
+UCLASS()
+class CUSTOMIZABLEOBJECTEDITOR_API UCustomizableObjectNodeObjectChild : public UCustomizableObjectNodeObject
+{
+public:
+	GENERATED_BODY()
+
+	UCustomizableObjectNodeObjectChild();
+
+public:
+	// UEdGraphNode interface
+	FText GetNodeTitle(ENodeTitleType::Type TitleType) const override;
+	void PrepareForCopying() override;
+	FText GetTooltipText() const override;
+	bool CanUserDeleteNode() const override;
+	bool CanDuplicateNode() const override;
+};
+
