[mutable] Moved the Mutable plugin out of Experimental status into Beta.

diff --git a/Engine/Plugins/Mutable/Source/CustomizableObjectEditor/Public/MuCOE/Nodes/CustomizableObjectNodeObjectGroup.h b/Engine/Plugins/Mutable/Source/CustomizableObjectEditor/Public/MuCOE/Nodes/CustomizableObjectNodeObjectGroup.h
new file mode 100644
--- /dev/null
+++ b/Engine/Plugins/Mutable/Source/CustomizableObjectEditor/Public/MuCOE/Nodes/CustomizableObjectNodeObjectGroup.h
@@ -0,0 +1,67 @@
+// Copyright Epic Games, Inc. All Rights Reserved.
+
+#pragma once
+
+#include "MuCOE/Nodes/CustomizableObjectNode.h"
+
+#include "CustomizableObjectNodeObjectGroup.generated.h"
+
+enum class ECustomizableObjectGroupType : uint8;
+namespace ENodeTitleType { enum Type : int; }
+
+class FArchive;
+class UCustomizableObjectNodeRemapPins;
+class UEdGraphPin;
+class UObject;
+struct FPropertyChangedEvent;
+
+
+UCLASS()
+class CUSTOMIZABLEOBJECTEDITOR_API UCustomizableObjectNodeObjectGroup : public UCustomizableObjectNode
+{
+public:
+	GENERATED_BODY()
+
+	UCustomizableObjectNodeObjectGroup();
+
+	UPROPERTY(EditAnywhere, Category=GroupInfo, meta = (DisplayName= "Name"))
+	FString GroupName;
+
+	UPROPERTY(EditAnywhere, Category=GroupInfo, meta = (DisplayName = "Type"))
+	ECustomizableObjectGroupType GroupType;
+	
+	UPROPERTY(EditAnywhere, Category=GroupInfo)
+	FString DefaultValue;
+
+	UPROPERTY(EditAnywhere, Category = UI, meta = (DisplayName = "Parameter UI Metadata"))
+	FMutableParamUIMetadata ParamUIMetadata;
+
+	// The sockets defined in meshes deriving from this node will inherit this socket priority. When in the generated merged mesh there
+	// are clashes with socket names, the one with higher priority will be kept and the other discarded.
+	UPROPERTY(EditAnywhere, Category = MeshSockets)
+	int32 SocketPriority = 0;
+
+	// UObject interface.
+	void Serialize(FArchive& Ar) override;
+	void PostEditChangeProperty(FPropertyChangedEvent& PropertyChangedEvent) override;
+
+	// EdGraphNode interface
+	FText GetNodeTitle(ENodeTitleType::Type TitleType) const override;
+	FLinearColor GetNodeTitleColor() const override;
+	FText GetTooltipText() const override;
+
+	// UCustomizableObjectNode interface
+	void AllocateDefaultPins(UCustomizableObjectNodeRemapPins* RemapPins) override;
+	bool IsSingleOutputNode() const override;
+
+	// Own interface
+	UEdGraphPin* ObjectsPin() const;
+
+	UEdGraphPin* GroupProjectorsPin() const;
+
+	UEdGraphPin* GroupPin() const;
+
+private:
+	FString LastGroupName;
+};
+
diff --git a/Engine/Plugins/Mutable/Source/CustomizableObjectEditor/Public/MuCOE/Nodes/CustomizableObjectNodeProjectorConstant.h b/Engine/Plugins/Mutable/Source/CustomizableObjectEditor/Public/MuCOE/Nodes/CustomizableObjectNodeProjectorConstant.h
new file mode 100644
--- /dev/null
+++ b/Engine/Plugins/Mutable/Source/CustomizableObjectEditor/Public/MuCOE/Nodes/CustomizableObjectNodeProjectorConstant.h
@@ -0,0 +1,92 @@
+// Copyright Epic Games, Inc. All Rights Reserved.
+
+#pragma once
+
+#include "MuCOE/Nodes/CustomizableObjectNode.h"
+
+#include "CustomizableObjectNodeProjectorConstant.generated.h"
+
+namespace ENodeTitleType { enum Type : int; }
+
+class UCustomizableObjectNodeRemapPins;
+class UEdGraphPin;
+class UObject;
+struct FPropertyChangedEvent;
+
+
+UCLASS()
+class CUSTOMIZABLEOBJECTEDITOR_API UCustomizableObjectNodeProjectorConstant : public UCustomizableObjectNode
+{
+public:
+	GENERATED_BODY()
+
+	UCustomizableObjectNodeProjectorConstant();
+
+	/**  */
+	UPROPERTY(EditAnywhere, Category = CustomizableObject, meta = (ShowOnlyInnerProperties))
+	FCustomizableObjectProjector Value;
+
+	UPROPERTY(EditAnywhere, Category = UI, meta = (DisplayName = "Parameter UI Metadata"))
+	FMutableParamUIMetadata ParamUIMetadata;
+
+	UPROPERTY(EditAnywhere, Category = CustomizableObject, Meta = (DisplayName = "Projection Angle (degrees)"))
+	float ProjectionAngle;
+
+	UPROPERTY(EditAnywhere, Category = ProjectorSnapToBone)
+	uint32 ReferenceSkeletonIndex = 0;
+
+	UPROPERTY(EditAnywhere, Category = ProjectorSnapToBone)
+	FName ProjectorBone;
+
+	/** Temporary variable where to put the location information for bone combo box selection changes (in FCustomizableObjectNodeProjectorParameterDetails)*/
+	FVector BoneComboBoxLocation;
+
+	/** Temporary variable where to put the direction information for bone combo box selection changes (in FCustomizableObjectNodeProjectorParameterDetails)*/
+	FVector BoneComboBoxForwardDirection;
+
+	/** Temporary variable where to put the direction information for bone combo box selection changes (in FCustomizableObjectNodeProjectorParameterDetails)*/
+	FVector BoneComboBoxUpDirection;
+
+	// UObject interface.
+	void PostEditChangeProperty(FPropertyChangedEvent& PropertyChangedEvent) override;
+
+	// Begin EdGraphNode interface
+	FText GetNodeTitle(ENodeTitleType::Type TitleType) const override;
+	FLinearColor GetNodeTitleColor() const override;
+	FText GetTooltipText() const override;
+
+	// UCustomizableObjectNode interface
+	void AllocateDefaultPins(UCustomizableObjectNodeRemapPins* RemapPins) override;
+	virtual void BackwardsCompatibleFixup(int32 CustomizableObjectCustomVersion) override;
+
+	// Own interface
+	UEdGraphPin* ValuePin();
+
+	ECustomizableObjectProjectorType GetProjectorType() const;
+
+	FVector GetProjectorPosition() const;
+
+	void SetProjectorPosition(const FVector& Position);
+
+	FVector GetProjectorDirection() const;
+
+	void SetProjectorDirection(const FVector& Direction);
+
+	FVector GetProjectorUp() const;
+
+	void SetProjectorUp(const FVector& Up);
+
+	FVector GetProjectorScale() const;
+
+	void SetProjectorScale(const FVector& Scale);
+
+	float GetProjectorAngle();
+
+	void SetProjectorAngle(float Angle);
+
+private:
+
+	UPROPERTY()
+	ECustomizableObjectProjectorType ProjectionType_DEPRECATED;
+};
+
diff --git a/Engine/Plugins/Mutable/Source/CustomizableObjectEditor/Public/MuCOE/Nodes/CustomizableObjectNodeProjectorParameter.h b/Engine/Plugins/Mutable/Source/CustomizableObjectEditor/Public/MuCOE/Nodes/CustomizableObjectNodeProjectorParameter.h
new file mode 100644
--- /dev/null
+++ b/Engine/Plugins/Mutable/Source/CustomizableObjectEditor/Public/MuCOE/Nodes/CustomizableObjectNodeProjectorParameter.h
@@ -0,0 +1,98 @@
+// Copyright Epic Games, Inc. All Rights Reserved.
+
+#pragma once
+
+#include "MuCOE/Nodes/CustomizableObjectNode.h"
+
+#include "CustomizableObjectNodeProjectorParameter.generated.h"
+
+namespace ENodeTitleType { enum Type : int; }
+
+class UCustomizableObjectNodeRemapPins;
+class UEdGraphPin;
+class UObject;
+struct FPropertyChangedEvent;
+
+
+UCLASS()
+class CUSTOMIZABLEOBJECTEDITOR_API UCustomizableObjectNodeProjectorParameter : public UCustomizableObjectNode
+{
+public:
+	GENERATED_BODY()
+
+	UPROPERTY(EditAnywhere, Category=CustomizableObject, meta = (ShowOnlyInnerProperties))
+	FCustomizableObjectProjector DefaultValue;
+
+	UPROPERTY(EditAnywhere, Category=CustomizableObject)
+	FString ParameterName = "Default Name";
+
+	UPROPERTY(EditAnywhere, Category = UI, meta = (DisplayName = "Parameter UI Metadata"))
+	FMutableParamUIMetadata ParamUIMetadata;
+
+	UPROPERTY(EditAnywhere, Category = CustomizableObject, Meta = (DisplayName = "Projection Angle (degrees)"))
+	float ProjectionAngle = 360.0f;
+
+	UPROPERTY(EditAnywhere, Category = ProjectorSnapToBone)
+	uint32 ReferenceSkeletonIndex = 0;
+
+	UPROPERTY(EditAnywhere, Category = ProjectorSnapToBone)
+	FName ProjectorBone;
+
+	/** Temporary variable where to put the location information for bone combo box selection changes (in FCustomizableObjectNodeProjectorParameterDetails)*/
+	FVector BoneComboBoxLocation = FVector::ZeroVector;
+
+	/** Temporary variable where to put the direction information for bone combo box selection changes (in FCustomizableObjectNodeProjectorParameterDetails)*/
+	FVector BoneComboBoxForwardDirection = FVector::ZeroVector;
+
+	/** Temporary variable where to put the direction information for bone combo box selection changes (in FCustomizableObjectNodeProjectorParameterDetails)*/
+	FVector BoneComboBoxUpDirection = FVector::ZeroVector;
+
+	// UObject interface.
+	void PostEditChangeProperty(FPropertyChangedEvent& PropertyChangedEvent) override;
+
+	// Begin EdGraphNode interface
+	FText GetNodeTitle(ENodeTitleType::Type TitleType) const override;
+	FLinearColor GetNodeTitleColor() const override;
+	FText GetTooltipText() const override;
+	virtual bool GetCanRenameNode() const override { return true; }
+	virtual void OnRenameNode(const FString& NewName) override;
+
+	// UCustomizableObjectNode interface
+	void AllocateDefaultPins(UCustomizableObjectNodeRemapPins* RemapPins) override;
+	virtual bool IsAffectedByLOD() const override { return false; }
+	virtual void BackwardsCompatibleFixup(int32 CustomizableObjectCustomVersion) override;
+
+	// Own interface
+	UEdGraphPin* ProjectorPin() const
+	{
+		return FindPin(TEXT("Value"));
+	}
+
+	ECustomizableObjectProjectorType GetProjectorType() const;
+
+	FVector GetProjectorDefaultPosition() const;
+	
+	void SetProjectorDefaultPosition(const FVector& Position);
+
+	FVector GetProjectorDefaultDirection() const;
+
+	void SetProjectorDefaultDirection(const FVector& Direction);
+
+	FVector GetProjectorDefaultUp() const;
+
+	void SetProjectorDefaultUp(const FVector& Up);
+
+	FVector GetProjectorDefaultScale() const;
+
+	void SetProjectorDefaultScale(const FVector& Scale);
+
+	float GetProjectorDefaultAngle() const;
+
+	void SetProjectorDefaultAngle(float Angle);
+
+private:
+
+	UPROPERTY()
+	ECustomizableObjectProjectorType ProjectionType_DEPRECATED;
+};
+
diff --git a/Engine/Plugins/Mutable/Source/CustomizableObjectEditor/Public/MuCOE/Nodes/CustomizableObjectNodeStaticMesh.h b/Engine/Plugins/Mutable/Source/CustomizableObjectEditor/Public/MuCOE/Nodes/CustomizableObjectNodeStaticMesh.h
new file mode 100644
--- /dev/null
+++ b/Engine/Plugins/Mutable/Source/CustomizableObjectEditor/Public/MuCOE/Nodes/CustomizableObjectNodeStaticMesh.h
@@ -0,0 +1,172 @@
+// Copyright Epic Games, Inc. All Rights Reserved.
+
+#pragma once
+
+#include "MuCOE/Nodes/CustomizableObjectNodeMesh.h"
+#include "MuCOE/Nodes/SCustomizableObjectNode.h"
+
+#include "CustomizableObjectNodeStaticMesh.generated.h"
+
+namespace ENodeTitleType { enum Type : int; }
+
+class FArchive;
+class FAssetThumbnail;
+class FAssetThumbnailPool;
+class SOverlay;
+class SVerticalBox;
+class UCustomizableObjectLayout;
+class UCustomizableObjectNodeRemapPins;
+class UCustomizableObjectNodeRemapPinsByName;
+class UMaterialInterface;
+class UObject;
+class UStaticMesh;
+class UTexture2D;
+struct FPropertyChangedEvent;
+struct FSlateBrush;
+
+
+// Class to render the Static Mesh thumbnail of a CustomizableObjectNodeStaticMesh
+class SGraphNodeStaticMesh : public SCustomizableObjectNode
+{
+public:
+
+	SLATE_BEGIN_ARGS(SGraphNodeStaticMesh) {}
+	SLATE_END_ARGS();
+
+	SGraphNodeStaticMesh() : SCustomizableObjectNode() {};
+
+	// Builds the SGraphNodeStaticMesh when needed
+	void Construct(const FArguments& InArgs, UEdGraphNode* InGraphNode);
+
+	// Calls the needed functions to build the SGraphNode widgets
+	void UpdateGraphNode();
+
+	// Overriden functions to build the SGraphNode widgets
+	virtual void SetDefaultTitleAreaWidget(TSharedRef<SOverlay> DefaultTitleAreaWidget) override;
+	virtual void CreateBelowPinControls(TSharedPtr<SVerticalBox> MainBox) override;
+	virtual bool ShouldAllowCulling() const override { return false; }
+
+	// Callbacks for the widget
+	void OnExpressionPreviewChanged(const ECheckBoxState NewCheckedState);
+	ECheckBoxState IsExpressionPreviewChecked() const;
+	const FSlateBrush* GetExpressionPreviewArrow() const;
+	EVisibility ExpressionPreviewVisibility() const;
+
+public:
+
+	// Single property that only draws the combo box widget of the static mesh
+	TSharedPtr<class ISinglePropertyView> StaticMeshSelector;
+
+	// Pointer to the NodeStaticMesh that owns this SGraphNode
+	class UCustomizableObjectNodeStaticMesh* NodeStaticMesh = nullptr;
+
+private:
+
+	// Classes needed to get and render the thumbnail of the Static Mesh
+	TSharedPtr<FAssetThumbnailPool> AssetThumbnailPool;
+	TSharedPtr<FAssetThumbnail> AssetThumbnail;
+
+	//This parameter defines the size of the thumbnail widget inside the Node
+	float WidgetSize;
+
+	// This parameter defines the resolution of the thumbnail
+	uint32 ThumbnailSize;
+
+};
+
+
+USTRUCT()
+struct CUSTOMIZABLEOBJECTEDITOR_API FCustomizableObjectNodeStaticMeshMaterial
+{
+	GENERATED_USTRUCT_BODY()
+
+	UPROPERTY()
+	FString Name;
+
+	UPROPERTY()
+	TObjectPtr<UEdGraphPin_Deprecated> MeshPin_DEPRECATED;
+
+	UPROPERTY()
+	TObjectPtr<UEdGraphPin_Deprecated> LayoutPin_DEPRECATED;
+
+	UPROPERTY()
+	TArray< TObjectPtr<UEdGraphPin_Deprecated>> ImagePins_DEPRECATED;
+
+	UPROPERTY()
+	FEdGraphPinReference MeshPinRef;
+
+	UPROPERTY()
+	FEdGraphPinReference LayoutPinRef;
+
+	UPROPERTY()
+	TArray<FEdGraphPinReference> ImagePinsRef;
+};
+
+
+USTRUCT()
+struct CUSTOMIZABLEOBJECTEDITOR_API FCustomizableObjectNodeStaticMeshLOD
+{
+	GENERATED_USTRUCT_BODY()
+
+	UPROPERTY()
+	TArray<FCustomizableObjectNodeStaticMeshMaterial> Materials;
+};
+
+
+UCLASS()
+class CUSTOMIZABLEOBJECTEDITOR_API UCustomizableObjectNodeStaticMesh : public UCustomizableObjectNodeMesh
+{ 
+public:
+	GENERATED_BODY()
+
+	UPROPERTY(EditAnywhere, Category = CustomizableObject)
+	TObjectPtr<UStaticMesh> StaticMesh;
+
+	/** Images */
+	UPROPERTY()
+	TArray<FCustomizableObjectNodeStaticMeshLOD> LODs;
+
+	/** Default pin when there is no mesh. */
+	UPROPERTY()
+	FEdGraphPinReference DefaultPin;
+	
+	// UObject interface.
+	virtual void PostLoad() override;
+	virtual void PostEditChangeProperty(FPropertyChangedEvent& PropertyChangedEvent) override;
+	virtual void Serialize(FArchive& Ar) override;
+
+	// UEdGraphNode interface
+	virtual FText GetNodeTitle(ENodeTitleType::Type TitleType) const override;
+	virtual FLinearColor GetNodeTitleColor() const override;
+	virtual FText GetTooltipText() const override;
+
+	// UCustomizableObjectNode interface
+	virtual void AllocateDefaultPins(UCustomizableObjectNodeRemapPins* RemapPins) override;
+	virtual bool ProvidesCustomPinRelevancyTest() const override { return true; }
+	virtual bool IsPinRelevant(const UEdGraphPin* Pin) const override;
+	virtual UCustomizableObjectNodeRemapPinsByName* CreateRemapPinsByName() const override;
+	virtual bool HasPinViewer() const override;
+	virtual bool IsNodeOutDatedAndNeedsRefresh() override;
+	virtual FString GetRefreshMessage() const override;
+
+	// UCustomizableObjectNodeMesh interface
+	virtual UTexture2D* FindTextureForPin(const UEdGraphPin* Pin) const override;
+	virtual TArray<UCustomizableObjectLayout*> GetLayouts(const UEdGraphPin& OutPin) const override;
+	virtual UObject* GetMesh() const override;
+	virtual UEdGraphPin* GetMeshPin(int32 LOD, int32 SectionIndex) const override;
+	virtual void GetPinSection(const UEdGraphPin& Pin, int32& OutLODIndex, int32& OutSectionIndex, int32& OutLayoutIndex) const override;
+
+	/** Returns the material associated to the given output pin. */
+	UMaterialInterface* GetMaterialFor(const UEdGraphPin* Pin) const;
+
+	// Creates the SGraph Node widget for the thumbnail
+	virtual TSharedPtr<SGraphNode> CreateVisualWidget() override;
+
+	// Determines if the Node is collapsed or not
+	bool bCollapsed = true;
+
+	// Pointer to the SGraphNodeStaticMesh
+	TWeakPtr< SGraphNodeStaticMesh > GraphNodeStaticMesh;
+
+};
+
diff --git a/Engine/Plugins/Mutable/Source/CustomizableObjectEditor/Public/MuCOE/Nodes/CustomizableObjectNodeSwitchBase.h b/Engine/Plugins/Mutable/Source/CustomizableObjectEditor/Public/MuCOE/Nodes/CustomizableObjectNodeSwitchBase.h
new file mode 100644
--- /dev/null
+++ b/Engine/Plugins/Mutable/Source/CustomizableObjectEditor/Public/MuCOE/Nodes/CustomizableObjectNodeSwitchBase.h
@@ -0,0 +1,88 @@
+// Copyright Epic Games, Inc. All Rights Reserved.
+
+#pragma once
+
+#include "MuCOE/Nodes/CustomizableObjectNode.h"
+
+#include "CustomizableObjectNodeSwitchBase.generated.h"
+
+namespace ENodeTitleType { enum Type : int; }
+
+class FArchive;
+class UCustomizableObjectNodeRemapPins;
+class UObject;
+struct FFrame;
+
+
+UCLASS(Abstract)
+class CUSTOMIZABLEOBJECTEDITOR_API UCustomizableObjectNodeSwitchBase : public UCustomizableObjectNode
+{
+public:
+	GENERATED_BODY()
+
+	// UObject interface
+	virtual void Serialize(FArchive& Ar) override;
+	
+	// Begin EdGraphNode interface
+	virtual FText GetNodeTitle(ENodeTitleType::Type TitleType) const override;
+	virtual FLinearColor GetNodeTitleColor() const override;
+	virtual FText GetTooltipText() const override;
+	virtual void PinConnectionListChanged(UEdGraphPin* Pin) override;
+	virtual void PostPasteNode() override;
+
+	// UCustomizableObjectNode interface
+	virtual void PostBackwardsCompatibleFixup() override;
+	virtual void AllocateDefaultPins(UCustomizableObjectNodeRemapPins* RemapPins) override;
+	virtual void ReconstructNode(UCustomizableObjectNodeRemapPins* RemapPins) override;
+
+	/** Get the output pin category. Override. */
+	virtual FName GetCategory() const PURE_VIRTUAL(UCustomizableObjectNodeSwitchBase::GetCategory, return {}; );
+
+	UEdGraphPin* OutputPin() const;
+
+	UEdGraphPin* SwitchParameter() const;
+
+	UEdGraphPin* GetElementPin(int32 Index) const
+	{
+		return FindPin(GetPinPrefix(Index));
+	}
+
+	int32 GetNumElements() const;
+
+	/** Links the PostEditChangeProperty delegate */
+	void LinkPostEditChangePropertyDelegate(const UEdGraphPin& Pin);
+
+	/** Get the ouput pin name. Override. */
+	virtual FString GetOutputPinName() const;
+
+private:
+	/** Get the pin prefix. Used for retrocompatibility. Override. */
+	virtual FString GetPinPrefix() const;
+
+protected:
+	UPROPERTY()
+	FEdGraphPinReference OutputPinReference;
+
+private:
+	/** Get the pin prefix with index. Used for retrocompatibility.*/
+	FString GetPinPrefix(int32 Index) const;
+
+	void ReloadEnumParam();
+
+	/** Last NodeEnumParameter connected. Used to remove the callback once disconnected. */
+	TWeakObjectPtr<UCustomizableObjectNode> LastNodeEnumParameterConnected;
+
+	/** NodeEnumParameter property changed callback function. Reconstructs the node. */
+	UFUNCTION()
+	void EnumParameterPostEditChangeProperty(FPostEditChangePropertyDelegateParameters& Parameters);
+
+	/** The node has to be reconstructed. PinConnectionListChanged(...) can not reconstruct the node, flag used to reconstruct the node on NodeConnectionListChanged(). */
+	bool bMarkReconstruct = false;
+
+	UPROPERTY()
+	TArray<FString> ReloadingElementsNames;
+
+	UPROPERTY()
+	FEdGraphPinReference SwitchParameterPinReference;
+};
+
diff --git a/Engine/Plugins/Mutable/Source/CustomizableObjectEditor/Public/MuCOE/Nodes/CustomizableObjectNodeTable.h b/Engine/Plugins/Mutable/Source/CustomizableObjectEditor/Public/MuCOE/Nodes/CustomizableObjectNodeTable.h
new file mode 100644
--- /dev/null
+++ b/Engine/Plugins/Mutable/Source/CustomizableObjectEditor/Public/MuCOE/Nodes/CustomizableObjectNodeTable.h
@@ -0,0 +1,411 @@
+// Copyright Epic Games, Inc. All Rights Reserved.
+
+#pragma once
+
+#include "Engine/DataTable.h"
+#include "Engine/Texture2DArray.h"
+#include "MuCOE/Nodes/CustomizableObjectNode.h"
+#include "MuCOE/RemapPins/CustomizableObjectNodeRemapPinsByName.h"
+
+#include "CustomizableObjectNodeTable.generated.h"
+
+namespace ENodeTitleType { enum Type : int; }
+
+class UCustomizableObjectLayout;
+class UCustomizableObjectNodeRemapPins;
+class UCustomizableObjectNodeRemapPinsByName;
+class UCustomizableObjectNodeTable;
+class UEdGraphPin;
+class UObject;
+class USkeletalMesh;
+class UTexture2D;
+class UTexture2DArray;
+class UAnimInstance;
+struct FGuid;
+struct FSkeletalMaterial;
+struct FSkelMeshSection;
+
+/** Enum class for the different types of image pins */
+UENUM()
+enum class ETableTextureType : uint8
+{
+	PASSTHROUGH_TEXTURE = 0 UMETA(DisplayName = "Passthrough"),
+	MUTABLE_TEXTURE = 1 UMETA(DisplayName = "Mutable")
+};
+
+
+/** Enum class for the different types of pin meshes */
+enum class ETableMeshPinType : uint8
+{
+	NONE = 0,
+	SKELETAL_MESH = 1,
+	STATIC_MESH = 2
+};
+
+
+/** Enum to decide where the data comes from: Struct or Data Table*/
+UENUM()
+enum class ETableDataGatheringSource : uint8
+{
+	/** Gathers the information from a data table */
+	ETDGM_DataTable = 0 UMETA(DisplayName = "Data Table"),
+
+	/** When compiling the CO, it uses the asset registry to gather and generate a data table. It uses all the data tables found in the specified paths that are references of the selected structure.*/
+	ETDGM_AssetRegistry = 1 UMETA(DisplayName = " Struct + Asset Registry")
+};
+
+
+USTRUCT()
+struct FTableNodeColumnData
+{
+	GENERATED_USTRUCT_BODY()
+
+	/** Anim Instance Column name related to this Mesh pin */
+	UPROPERTY()
+	FString AnimInstanceColumnName = "";
+
+	/** Anim Slot Column name related to this Mesh pin */
+	UPROPERTY()
+	FString AnimSlotColumnName = "";
+
+	/** Anim Tag Column name related to this Mesh pin */
+	UPROPERTY()
+	FString AnimTagColumnName = "";
+};
+
+
+/** Base class for all Table Pins. */
+UCLASS()
+class CUSTOMIZABLEOBJECTEDITOR_API UCustomizableObjectNodeTableObjectPinData : public UCustomizableObjectNodePinData
+{
+	GENERATED_BODY()
+
+public:
+
+	/** Id of the property associated to a struct column */
+	UPROPERTY()
+	FGuid StructColumnId;
+
+	/** Name of the data table column related to the pin */
+	UPROPERTY()
+	FString ColumnName = "";
+
+};
+
+
+/** Additional data for Image pins. */
+UCLASS()
+class CUSTOMIZABLEOBJECTEDITOR_API UCustomizableObjectNodeTableImagePinData : public UCustomizableObjectNodeTableObjectPinData
+{
+	GENERATED_BODY()
+
+public:
+	// UObject interface
+	virtual void BackwardsCompatibleFixup(int32 CustomizableObjectCustomVersion) override;
+	virtual void PostEditChangeProperty(FPropertyChangedEvent& PropertyChangedEvent) override;
+	virtual bool CanEditChange(const FProperty* InProperty) const override;
+
+	UPROPERTY(EditAnywhere, Category = NoCategory,  DisplayName = "Texture Parameter Mode")
+	ETableTextureType ImageMode = ETableTextureType::MUTABLE_TEXTURE;
+
+	UPROPERTY()
+	TObjectPtr<UCustomizableObjectNodeTable> NodeTable = nullptr;
+
+private:
+	// Replaced by the more general bIsNot2DTexture
+	UPROPERTY()
+	bool bIsArrayTexture_DEPRECATED = false;
+
+public:
+	UPROPERTY()
+	bool bIsNotTexture2D = false;
+};
+
+
+/** Additional data for Mesh pins. */
+UCLASS()
+class CUSTOMIZABLEOBJECTEDITOR_API UCustomizableObjectNodeTableMeshPinData : public UCustomizableObjectNodeTableObjectPinData
+{
+	GENERATED_BODY()
+
+public:
+
+	/** Anim Instance Column name related to this Mesh pin */
+	UPROPERTY()
+	FString AnimInstanceColumnName_DEPRECATED = "";
+
+	/** Anim Slot Column name related to this Mesh pin */
+	UPROPERTY()
+	FString AnimSlotColumnName_DEPRECATED = "";
+
+	/** Anim Tag Column name related to this Mesh pin */
+	UPROPERTY()
+	FString AnimTagColumnName_DEPRECATED = "";
+
+	UPROPERTY()
+	FString MutableColumnName = "";
+
+	/** LOD of the mesh related to this Mesh pin */
+	UPROPERTY()
+	int32 LOD = 0;
+
+	/** Section Index (Surface Index) of the mesh related to this Mesh pin */
+	UPROPERTY()
+	int32 Material = 0;
+
+	/** Layouts related to this Mesh pin */
+	UPROPERTY()
+	TArray< TObjectPtr<UCustomizableObjectLayout> > Layouts;
+};
+
+
+UCLASS()
+class UCustomizableObjectNodeTableRemapPins : public UCustomizableObjectNodeRemapPinsByName
+{
+	GENERATED_BODY()
+public:
+
+	// Specific method to decide when two pins are equal
+	virtual bool Equal(const UCustomizableObjectNode& Node, const UEdGraphPin& OldPin, const UEdGraphPin& NewPin) const override;
+
+	// Method to use in the RemapPins step of the node reconstruction process
+	virtual void RemapPins(const UCustomizableObjectNode& Node, const TArray<UEdGraphPin*>& OldPins, const TArray<UEdGraphPin*>& NewPins, TMap<UEdGraphPin*, UEdGraphPin*>& PinsToRemap, TArray<UEdGraphPin*>& PinsToOrphan) override;
+};
+
+
+UCLASS()
+class CUSTOMIZABLEOBJECTEDITOR_API UCustomizableObjectNodeTable : public UCustomizableObjectNode
+{
+public:
+	
+	GENERATED_BODY()
+
+	// Name of the property parameter
+	UPROPERTY(EditAnywhere, Category = TableProperties)
+	FString ParameterName = "Default Name";
+
+	/** If true, adds a "None" parameter option */
+	UPROPERTY(EditAnywhere, Category = TableProperties)
+	bool bAddNoneOption;
+
+	/** If there is a bool column in the table, checked rows will not be compiled */
+	UPROPERTY(EditAnywhere, Category = CompilationRestrictions)
+	bool bDisableCheckedRows = true;
+
+	/** Source where table gathers the data */
+	UPROPERTY(EditAnywhere, Category = TableProperties)
+	ETableDataGatheringSource TableDataGatheringMode = ETableDataGatheringSource::ETDGM_DataTable;
+
+	// Pointer to the Data Table Asset represented in this node
+	UPROPERTY(EditAnywhere, Category = TableProperties, meta = (DontUpdateWhileEditing, EditCondition = "TableDataGatheringMode == ETableDataGatheringSource::ETDGM_DataTable", EditConditionHides))
+	TObjectPtr<UDataTable> Table = nullptr;
+
+	// Pointer to the Struct Asset represented in this node
+	UPROPERTY(EditAnywhere, Category = TableProperties, meta = (DontUpdateWhileEditing, EditCondition = "TableDataGatheringMode == ETableDataGatheringSource::ETDGM_AssetRegistry", EditConditionHides))
+	TObjectPtr<UScriptStruct> Structure = nullptr;
+
+	UPROPERTY(EditAnywhere, Category = TableProperties, meta = (EditCondition = "TableDataGatheringMode == ETableDataGatheringSource::ETDGM_AssetRegistry", EditConditionHides))
+	TArray<FName> FilterPaths;
+	
+	/** Name of the column that contains the Version options. */
+	UPROPERTY(EditAnywhere, Category = CompilationRestrictions)
+	FName VersionColumn;
+
+	/** Name of the row that will be used as default value. */
+	UPROPERTY(EditAnywhere, Category = TableProperties)
+	FName DefaultRowName;
+
+	UPROPERTY(EditAnywhere, Category = UI, meta = (DisplayName = "Parameter UI Metadata"))
+	FMutableParamUIMetadata ParamUIMetadata;
+
+	/** Name of the column that contains the MutableUIMetadata of the row options. */
+	UPROPERTY(EditAnywhere, Category = UI)
+	FName ParamUIMetadataColumn;
+	
+	/** Name of the column that contains the asset to use its thumbnails as option thumbnails. */
+	UPROPERTY(EditAnywhere, Category = UI)
+	FName ThumbnailColumn;
+	
+	UPROPERTY(EditAnywhere, Category = UI, meta = (Tooltip = "Given a row, add all tags found in GameplayTag columns to its Parameter UI Metadata"))
+	bool bGatherTags = true;
+
+	/** Map to relate a Structure Column with its Data */
+	UPROPERTY()
+	TMap<FGuid, FTableNodeColumnData> ColumnDataMap;
+
+public:
+
+	// UObject interface
+	virtual void PostLoad() override;
+	virtual void PreEditChange(FProperty* PropertyAboutToChange) override;
+	virtual void PostEditChangeProperty(FPropertyChangedEvent& PropertyChangedEvent) override;
+	
+	// EdGraphNode interface
+	FText GetNodeTitle(ENodeTitleType::Type TitleType) const override;
+	FLinearColor GetNodeTitleColor() const override;
+	FText GetTooltipText() const override;
+	virtual void OnRenameNode(const FString& NewName) override;
+	virtual bool GetCanRenameNode() const override { return true; }
+	
+	// UCustomizableObjectNode interface
+	virtual void BackwardsCompatibleFixup(int32 CustomizableObjectCustomVersion) override;
+	void AllocateDefaultPins(UCustomizableObjectNodeRemapPins* RemapPins) override;
+	bool IsNodeOutDatedAndNeedsRefresh() override;
+	FString GetRefreshMessage() const override;
+	virtual bool ProvidesCustomPinRelevancyTest() const override;
+	virtual bool IsPinRelevant(const UEdGraphPin* Pin) const override;
+	UCustomizableObjectNodeTableRemapPins* CreateRemapPinsDefault() const;
+	virtual bool HasPinViewer() const override;
+	virtual TSharedPtr<IDetailsView> CustomizePinDetails(const UEdGraphPin& Pin) const override;
+	
+	/*** Allows to perform work when remapping the pin data. */
+	virtual void RemapPinsData(const TMap<UEdGraphPin*, UEdGraphPin*>& PinsToRemap) override;
+	
+	// Own interface	
+	// Returns the reference Texture parameter from a Material
+	UTexture2D* FindReferenceTextureParameter(const UEdGraphPin* Pin, FString ParameterImageName) const;
+
+	// Methods to get the UVs of the skeletal mesh 
+	UObject* GetDefaultMeshForLayout(const UCustomizableObjectLayout* Layout) const;
+
+	// Methods to provide the Layouts to the Layout block editors
+	TArray<UCustomizableObjectLayout*> GetLayouts(const UEdGraphPin* Pin) const;
+
+	// Returns the name of the table column related to a pin
+	FString GetColumnNameByPin(const UEdGraphPin* Pin) const;
+
+	// Returns the LOD of the mesh associated to the input pin
+	void GetPinLODAndSection(const UEdGraphPin* Pin, int32& LODIndex, int32& SectionIndex) const;
+
+	// Get the anim blueprint and anim slot columns related to a mesh
+	void GetAnimationColumns(const FGuid& ColumnId, FString& AnimBPColumnName, FString& AnimSlotColumnName, FString& AnimTagColumnName) const;
+
+	template<class T> 
+	T* GetColumnDefaultAssetByType(FString ColumnName) const
+	{
+		T* ObjectToReturn = nullptr;
+		const UScriptStruct* TableStruct = nullptr;
+
+		if (TableDataGatheringMode == ETableDataGatheringSource::ETDGM_AssetRegistry)
+		{
+			TableStruct = Structure;
+		}
+		else
+		{
+			if (Table)
+			{
+				TableStruct = Table->GetRowStruct();
+			}
+		}
+
+		if (TableStruct)
+		{
+			// Getting Default Struct Values
+			// A Script Struct always has at leaset one property
+			TArray<int8> DeafaultDataArray;
+			DeafaultDataArray.SetNumZeroed(TableStruct->GetStructureSize());
+			TableStruct->InitializeStruct(DeafaultDataArray.GetData());
+
+			FProperty* Property = FindTableProperty(TableStruct, FName(*ColumnName));
+
+			if (Property)
+			{
+				if (const FSoftObjectProperty* SoftObjectProperty = CastField<FSoftObjectProperty>(Property))
+				{
+					UObject* Object = nullptr;
+					
+					// Getting default UObject
+					uint8* CellData = SoftObjectProperty->ContainerPtrToValuePtr<uint8>(DeafaultDataArray.GetData());
+
+					if (CellData)
+					{
+						Object = SoftObjectProperty->GetPropertyValue(CellData).LoadSynchronous();
+					}
+
+					if (Object)
+					{
+						if (Object->IsA(T::StaticClass()))
+						{
+							ObjectToReturn = Cast <T>(Object);
+						}
+					}
+				}
+			}
+
+			// Cleaning Default Structure Pointer
+			TableStruct->DestroyStruct(DeafaultDataArray.GetData());
+		}
+
+		return ObjectToReturn;
+	}
+
+	template<class T>
+	T* GetColumnDefaultAssetByType(const UEdGraphPin* Pin) const
+	{
+		if (Pin)
+		{
+			FString ColumnName = GetColumnNameByPin(Pin);
+
+			if (ColumnName != FString())
+			{
+				return GetColumnDefaultAssetByType<T>(ColumnName);
+			}
+		}
+
+		return nullptr;
+	}
+
+	// Generation Mutable Source Methods
+	// We should do this in a template!
+	USkeletalMesh* GetSkeletalMeshAt(const UEdGraphPin* Pin, const UDataTable* DataTable, const FName& RowName) const;
+	TSoftClassPtr<UAnimInstance> GetAnimInstanceAt(const UEdGraphPin* Pin, const UDataTable* DataTable, const FName& RowName) const;
+	
+	// Returns the image mode of the column
+	ETableTextureType GetColumnImageMode(const FString& ColumnName) const;
+
+	// Returns the mesh type of the Pin
+	ETableMeshPinType GetPinMeshType(const UEdGraphPin* Pin) const;
+
+	// Functions to generate the names of a mutable table's column
+	FString GenerateSkeletalMeshMutableColumName(const FString& PinName, int32 LODIndex, int32 MaterialIndex) const;
+	FString GenerateStaticMeshMutableColumName(const FString& PinName, int32 MaterialIndex) const;
+
+	/** Returns the struct pointer used to gather data */
+	const UScriptStruct* GetTableNodeStruct() const;
+
+	// Methods from UDataTable but modified to support ScriptStructs
+	/** Get an array of all the column titles, using the friendly display name from the property */
+	TArray<FString> GetColumnTitles(const UScriptStruct* ScriptStruct) const;
+	
+	/** Returns the column property where PropertyName matches the name of the column property. Returns nullptr if no match is found or the match is not a supported table property */
+	FProperty* FindTableProperty(const UScriptStruct* ScriptStruct, const FName& PropertyName) const;
+
+	/** Returns the column id */
+	FGuid GetColumnIdByName(const FName& ColumnName) const;
+
+	/** Return the list of UDataTable that will be used to compose the final UDataTable. */
+	TArray<FAssetData> GetParentTables() const;
+
+	/** Returns the skeletal material associated to the given skeletal mesh output pin. */
+	FSkeletalMaterial* GetDefaultSkeletalMaterialFor(const UEdGraphPin& MeshPin) const;
+
+	/** Returns the index of the skeletal material associated to the given skeletal mesh output pin. */
+	int32 GetDefaultSkeletalMaterialIndexFor(const UEdGraphPin& MeshPin) const;
+
+	/** Returns the SkeletalMesh section associated to the given skeletal mesh output pin. */
+	const FSkelMeshSection* GetDefaultSkeletalMeshSectionFor(const UEdGraphPin& MeshPin) const;
+private:
+
+	/** Number of properties to know when the node needs an update */
+	UPROPERTY()
+	int32 NumProperties;
+	
+	FDelegateHandle OnTableChangedDelegateHandle;
+	
+	// Generates a mesh pin for each LOD and Material Surface of the reference SkeletalMesh
+	void GenerateMeshPins(UObject* Mesh, const FString& ColumnName, const FGuid& ColumnId);
+
+	// Checks if a pin already exists and if it has the same type as before the node refresh
+	bool CheckPinUpdated(const FString& PinName, const FName& PinType) const;
+
+};
diff --git a/Engine/Plugins/Mutable/Source/CustomizableObjectEditor/Public/MuCOE/Nodes/CustomizableObjectNodeTextureBinarise.h b/Engine/Plugins/Mutable/Source/CustomizableObjectEditor/Public/MuCOE/Nodes/CustomizableObjectNodeTextureBinarise.h
new file mode 100644
--- /dev/null
+++ b/Engine/Plugins/Mutable/Source/CustomizableObjectEditor/Public/MuCOE/Nodes/CustomizableObjectNodeTextureBinarise.h
@@ -0,0 +1,40 @@
+// Copyright Epic Games, Inc. All Rights Reserved.
+#pragma once
+
+#include "MuCOE/Nodes/CustomizableObjectNode.h"
+
+#include "CustomizableObjectNodeTextureBinarise.generated.h"
+
+namespace ENodeTitleType { enum Type : int; }
+
+class FArchive;
+class UCustomizableObjectNodeRemapPins;
+class UObject;
+
+UCLASS()
+class CUSTOMIZABLEOBJECTEDITOR_API UCustomizableObjectNodeTextureBinarise : public UCustomizableObjectNode
+{
+public:
+	GENERATED_BODY()
+
+	// Begin EdGraphNode interface
+	FText GetNodeTitle(ENodeTitleType::Type TittleType) const override;
+	FLinearColor GetNodeTitleColor() const override;
+	FText GetTooltipText() const override;
+	virtual void Serialize(FArchive& Ar) override;
+	
+	// UCustomizableObjectNode interface
+	virtual void AllocateDefaultPins(UCustomizableObjectNodeRemapPins* RemapPins) override;
+	virtual void BackwardsCompatibleFixup(int32 CustomizableObjectCustomVersion) override;
+
+	UEdGraphPin* GetBaseImagePin() const;
+
+	UEdGraphPin* GetThresholdPin() const
+	{
+		return FindPin(TEXT("Threshold"));
+	}
+
+private:
+	UPROPERTY()
+	FEdGraphPinReference BaseImagePinReference;
+};
diff --git a/Engine/Plugins/Mutable/Source/CustomizableObjectEditor/Public/MuCOE/Nodes/CustomizableObjectNodeTextureColourMap.h b/Engine/Plugins/Mutable/Source/CustomizableObjectEditor/Public/MuCOE/Nodes/CustomizableObjectNodeTextureColourMap.h
new file mode 100644
--- /dev/null
+++ b/Engine/Plugins/Mutable/Source/CustomizableObjectEditor/Public/MuCOE/Nodes/CustomizableObjectNodeTextureColourMap.h
@@ -0,0 +1,46 @@
+// Copyright Epic Games, Inc. All Rights Reserved.
+
+#pragma once
+
+#include "MuCOE/Nodes/CustomizableObjectNode.h"
+
+#include "CustomizableObjectNodeTextureColourMap.generated.h"
+
+namespace ENodeTitleType { enum Type : int; }
+
+class UCustomizableObjectNodeRemapPins;
+class UEdGraphPin;
+class UObject;
+
+
+UCLASS()
+class CUSTOMIZABLEOBJECTEDITOR_API UCustomizableObjectNodeTextureColourMap : public UCustomizableObjectNode
+{
+public:
+	GENERATED_BODY()
+
+	// Begin EdGraphNode interface
+	FText GetNodeTitle(ENodeTitleType::Type TitleType) const override;
+	FLinearColor GetNodeTitleColor() const override;
+	FText GetTooltipText() const override;
+
+	// UCustomizableObjectNode interface
+	virtual void AllocateDefaultPins(UCustomizableObjectNodeRemapPins* RemapPins) override;
+	virtual void BackwardsCompatibleFixup(int32 CustomizableObjectCustomVersion) override;
+
+	UEdGraphPin* GetMapPin() const
+	{
+		return FindPin(TEXT("Map"));
+	}
+
+	UEdGraphPin* GetBasePin() const
+	{
+		return FindPin(TEXT("Base"));
+	}
+
+	UEdGraphPin* GetMaskPin() const
+	{
+		return FindPin(TEXT("Mask"));
+	}
+
+};
diff --git a/Engine/Plugins/Mutable/Source/CustomizableObjectEditor/Public/MuCOE/Nodes/CustomizableObjectNodeTextureFromChannels.h b/Engine/Plugins/Mutable/Source/CustomizableObjectEditor/Public/MuCOE/Nodes/CustomizableObjectNodeTextureFromChannels.h
new file mode 100644
--- /dev/null
+++ b/Engine/Plugins/Mutable/Source/CustomizableObjectEditor/Public/MuCOE/Nodes/CustomizableObjectNodeTextureFromChannels.h
@@ -0,0 +1,51 @@
+// Copyright Epic Games, Inc. All Rights Reserved.
+
+#pragma once
+
+#include "MuCOE/Nodes/CustomizableObjectNode.h"
+
+#include "CustomizableObjectNodeTextureFromChannels.generated.h"
+
+namespace ENodeTitleType { enum Type : int; }
+
+class UCustomizableObjectNodeRemapPins;
+class UEdGraphPin;
+class UObject;
+
+
+UCLASS()
+class CUSTOMIZABLEOBJECTEDITOR_API UCustomizableObjectNodeTextureFromChannels : public UCustomizableObjectNode
+{
+public:
+	GENERATED_BODY()
+
+	// Begin EdGraphNode interface
+	FText GetNodeTitle(ENodeTitleType::Type TitleType) const override;
+	FLinearColor GetNodeTitleColor() const override;
+	FText GetTooltipText() const override;
+
+	// UCustomizableObjectNode interface
+	virtual void AllocateDefaultPins(UCustomizableObjectNodeRemapPins* RemapPins) override;
+	virtual void BackwardsCompatibleFixup(int32 CustomizableObjectCustomVersion) override;
+
+	UEdGraphPin* RPin() const
+	{
+		return FindPin(TEXT("R"));
+	}
+
+	UEdGraphPin* GPin() const
+	{
+		return FindPin(TEXT("G"));
+	}
+
+	UEdGraphPin* BPin() const
+	{
+		return FindPin(TEXT("B"));
+	}
+
+	UEdGraphPin* APin() const
+	{
+		return FindPin(TEXT("A"));
+	}
+};
+
diff --git a/Engine/Plugins/Mutable/Source/CustomizableObjectEditor/Public/MuCOE/Nodes/CustomizableObjectNodeTextureFromColor.h b/Engine/Plugins/Mutable/Source/CustomizableObjectEditor/Public/MuCOE/Nodes/CustomizableObjectNodeTextureFromColor.h
new file mode 100644
--- /dev/null
+++ b/Engine/Plugins/Mutable/Source/CustomizableObjectEditor/Public/MuCOE/Nodes/CustomizableObjectNodeTextureFromColor.h
@@ -0,0 +1,36 @@
+// Copyright Epic Games, Inc. All Rights Reserved.
+
+#pragma once
+
+#include "MuCOE/Nodes/CustomizableObjectNode.h"
+
+#include "CustomizableObjectNodeTextureFromColor.generated.h"
+
+namespace ENodeTitleType { enum Type : int; }
+
+class UCustomizableObjectNodeRemapPins;
+class UEdGraphPin;
+class UObject;
+
+
+UCLASS()
+class CUSTOMIZABLEOBJECTEDITOR_API UCustomizableObjectNodeTextureFromColor : public UCustomizableObjectNode
+{
+public:
+	GENERATED_BODY()
+
+	// Begin EdGraphNode interface
+	FText GetNodeTitle(ENodeTitleType::Type TitleType) const override;
+	FLinearColor GetNodeTitleColor() const override;
+	FText GetTooltipText() const override;
+
+	// UCustomizableObjectNode interface
+	virtual void AllocateDefaultPins(UCustomizableObjectNodeRemapPins* RemapPins) override;
+	virtual void BackwardsCompatibleFixup(int32 CustomizableObjectCustomVersion) override;
+
+	UEdGraphPin* ColorPin() const
+	{
+		return FindPin(TEXT("Color"));
+	}
+};
+
diff --git a/Engine/Plugins/Mutable/Source/CustomizableObjectEditor/Public/MuCOE/Nodes/CustomizableObjectNodeTextureInterpolate.h b/Engine/Plugins/Mutable/Source/CustomizableObjectEditor/Public/MuCOE/Nodes/CustomizableObjectNodeTextureInterpolate.h
new file mode 100644
--- /dev/null
+++ b/Engine/Plugins/Mutable/Source/CustomizableObjectEditor/Public/MuCOE/Nodes/CustomizableObjectNodeTextureInterpolate.h
@@ -0,0 +1,63 @@
+// Copyright Epic Games, Inc. All Rights Reserved.
+
+#pragma once
+
+#include "MuCOE/Nodes/CustomizableObjectNode.h"
+
+#include "CustomizableObjectNodeTextureInterpolate.generated.h"
+
+namespace ENodeTitleType { enum Type : int; }
+
+class UCustomizableObjectNodeRemapPins;
+class UObject;
+struct FPropertyChangedEvent;
+
+
+UCLASS()
+class CUSTOMIZABLEOBJECTEDITOR_API UCustomizableObjectNodeTextureInterpolate : public UCustomizableObjectNode
+{
+public:
+	GENERATED_BODY()
+
+	UCustomizableObjectNodeTextureInterpolate();
+
+	// UObject interface.
+	void PostEditChangeProperty(FPropertyChangedEvent& PropertyChangedEvent) override;
+
+	// Begin EdGraphNode interface
+	FText GetNodeTitle(ENodeTitleType::Type TitleType) const override;
+	FLinearColor GetNodeTitleColor() const override;
+	FText GetTooltipText() const override;
+
+	// UCustomizableObjectNode interface
+	void AllocateDefaultPins(UCustomizableObjectNodeRemapPins* RemapPins) override;
+
+	UEdGraphPin* FactorPin() const
+	{
+		return FindPin(TEXT("Factor"));
+	}
+
+	UEdGraphPin* Targets(int32 Index) const
+	{
+		return FindPin(FString::Printf(TEXT("Target %d"), Index));
+	}
+
+	int32 GetNumTargets() const
+	{
+		int32 Count = 0;
+
+		for (UEdGraphPin* Pin : GetAllNonOrphanPins())
+		{
+			if (Pin->GetName().StartsWith(TEXT("Target ")))
+			{
+				Count++;
+			}
+		}
+
+		return Count;
+	}
+
+	UPROPERTY(EditAnywhere, Category=CustomizableObject)
+	int32 NumTargets;
+};
+
diff --git a/Engine/Plugins/Mutable/Source/CustomizableObjectEditor/Public/MuCOE/Nodes/CustomizableObjectNodeTextureInvert.h b/Engine/Plugins/Mutable/Source/CustomizableObjectEditor/Public/MuCOE/Nodes/CustomizableObjectNodeTextureInvert.h
new file mode 100644
--- /dev/null
+++ b/Engine/Plugins/Mutable/Source/CustomizableObjectEditor/Public/MuCOE/Nodes/CustomizableObjectNodeTextureInvert.h
@@ -0,0 +1,35 @@
+// Copyright Epic Games, Inc. All Rights Reserved.
+#pragma once
+
+#include "MuCOE/Nodes/CustomizableObjectNode.h"
+
+#include "CustomizableObjectNodeTextureInvert.generated.h"
+
+namespace ENodeTitleType { enum Type : int; }
+
+class FArchive;
+class UCustomizableObjectNodeRemapPins;
+class UObject;
+
+UCLASS()
+class CUSTOMIZABLEOBJECTEDITOR_API UCustomizableObjectNodeTextureInvert : public UCustomizableObjectNode
+{
+public:
+	GENERATED_BODY()
+	
+	// EdGraphNode interface
+	virtual FText GetNodeTitle(ENodeTitleType::Type TittleType)const override;
+	virtual FLinearColor GetNodeTitleColor() const override;
+	virtual FText GetTooltipText() const override;
+	
+	// UCustomizableObjectNode interface
+	virtual void AllocateDefaultPins(UCustomizableObjectNodeRemapPins* RemapPins) override;
+	virtual void BackwardsCompatibleFixup(int32 CustomizableObjectCustomVersion) override;
+
+	// Own interface
+	UEdGraphPin* GetBaseImagePin() const;
+
+private:
+	UPROPERTY()
+	FEdGraphPinReference BaseImagePinReference;
+};
diff --git a/Engine/Plugins/Mutable/Source/CustomizableObjectEditor/Public/MuCOE/Nodes/CustomizableObjectNodeTextureLayer.h b/Engine/Plugins/Mutable/Source/CustomizableObjectEditor/Public/MuCOE/Nodes/CustomizableObjectNodeTextureLayer.h
new file mode 100644
--- /dev/null
+++ b/Engine/Plugins/Mutable/Source/CustomizableObjectEditor/Public/MuCOE/Nodes/CustomizableObjectNodeTextureLayer.h
@@ -0,0 +1,126 @@
+// Copyright Epic Games, Inc. All Rights Reserved.
+
+#pragma once
+
+#include "MuCOE/Nodes/CustomizableObjectNode.h"
+
+#include "CustomizableObjectNodeTextureLayer.generated.h"
+
+namespace ENodeTitleType { enum Type : int; }
+
+class FArchive;
+class UCustomizableObjectNodeRemapPins;
+class UEdGraph;
+class UObject;
+struct FPropertyChangedEvent;
+
+
+UENUM()
+enum ECustomizableObjectTextureLayerEffect : int
+{
+	COTLE_MODULATE			= 0 UMETA(DisplayName = "MODULATE"), 
+	COTLE_MULTIPLY				UMETA(DisplayName = "MULTIPLY"), 
+	COTLE_SOFTLIGHT				UMETA(DisplayName = "SOFTLIGHT"),
+	COTLE_HARDLIGHT				UMETA(DisplayName = "HARDLIGHT"),
+	COTLE_DODGE					UMETA(DisplayName = "DODGE"),
+	COTLE_BURN					UMETA(DisplayName = "BURN"),
+	COTLE_SCREEN				UMETA(DisplayName = "SCREEN"),
+	COTLE_OVERLAY				UMETA(DisplayName = "OVERLAY"),
+	COTLE_ALPHA_OVERLAY			UMETA(DisplayName = "LIGHTEN"),
+	COTLE_NORMAL_COMBINE		UMETA(DisplayName = "BLEND NORMALS")
+};
+
+
+USTRUCT()
+struct CUSTOMIZABLEOBJECTEDITOR_API FCustomizableObjectTextureLayer
+{
+	GENERATED_USTRUCT_BODY()
+
+	FCustomizableObjectTextureLayer()
+	{
+		Effect = COTLE_SOFTLIGHT;
+	}
+
+	UPROPERTY(EditAnywhere, Category=CustomizableObject)
+	TEnumAsByte<enum ECustomizableObjectTextureLayerEffect> Effect;
+};
+
+
+UCLASS()
+class CUSTOMIZABLEOBJECTEDITOR_API UCustomizableObjectNodeTextureLayer : public UCustomizableObjectNode
+{
+public:
+	GENERATED_BODY()
+
+	UCustomizableObjectNodeTextureLayer();
+
+	/**  */
+	UPROPERTY(EditAnywhere, Category=CustomizableObject)
+	TArray<FCustomizableObjectTextureLayer> Layers;
+
+	// UObject interface.
+	void PostEditChangeProperty(FPropertyChangedEvent& PropertyChangedEvent) override;
+	virtual void Serialize(FArchive& Ar) override;
+	
+	// Begin EdGraphNode interface
+	FText GetNodeTitle(ENodeTitleType::Type TitleType) const override;
+	FLinearColor GetNodeTitleColor() const override;
+	FText GetTooltipText() const override;
+	void PrepareForCopying() override;
+	void PostPasteNode() override;
+
+	// UCustomizableObjectNode interface
+	virtual void AllocateDefaultPins(UCustomizableObjectNodeRemapPins* RemapPins) override;
+	virtual void BackwardsCompatibleFixup(int32 CustomizableObjectCustomVersion) override;
+	virtual bool CanConnect(const UEdGraphPin* InOwnedInputPin, const UEdGraphPin* InOutputPin, bool& bOutIsOtherNodeBlocklisted, bool& bOutArePinsCompatible) const override;
+		
+	UEdGraphPin* BasePin() const
+	{
+		return FindPin(TEXT("Base"));
+	}
+
+	UEdGraphPin* OutputPin() const;
+
+	UEdGraphPin* LayerPin(int32 Index) const
+	{
+		FString PinName = FString::Printf(TEXT("Layer %d "), Index);
+		return FindPin(PinName);
+	}
+
+	UEdGraphPin* MaskPin(int32 Index) const
+	{
+		FString PinName = FString::Printf(TEXT("Mask %d "), Index);
+		return FindPin(PinName);
+	}
+
+	int32 GetNumLayers() const
+	{
+		int32 Count = 0;
+
+		for (UEdGraphPin* Pin : GetAllNonOrphanPins())
+		{
+			if (Pin->GetName().StartsWith(TEXT("Layer ")))
+			{
+				Count++;
+			}
+		}
+
+		return Count;
+	}
+protected:
+	// This struct will store the information of each textue layer node when its destroyed, to be able to recover the UCustomizableObjectNodeTextureLayer::Layers
+	// data in a post paste operation. Otherwise, this information is lost.
+	struct TextureLayerLayersData
+	{
+		FString Name; // Texture layer node name
+		TArray<FCustomizableObjectTextureLayer> Layers;
+		UEdGraph* Graph; // Graph that contained this node
+	};
+	static int32 GetIndexInArrayDestroyedNodes(UCustomizableObjectNodeTextureLayer* Node);
+	static TArray<TextureLayerLayersData> ArrayDestroyedNodes;
+
+private:
+	UPROPERTY()
+	FEdGraphPinReference OutputPinReference;
+};
+
diff --git a/Engine/Plugins/Mutable/Source/CustomizableObjectEditor/Public/MuCOE/Nodes/CustomizableObjectNodeTextureParameter.h b/Engine/Plugins/Mutable/Source/CustomizableObjectEditor/Public/MuCOE/Nodes/CustomizableObjectNodeTextureParameter.h
new file mode 100644
--- /dev/null
+++ b/Engine/Plugins/Mutable/Source/CustomizableObjectEditor/Public/MuCOE/Nodes/CustomizableObjectNodeTextureParameter.h
@@ -0,0 +1,56 @@
+// Copyright Epic Games, Inc. All Rights Reserved.
+
+#pragma once
+
+#include "MuCOE/Nodes/CustomizableObjectNode.h"
+
+#include "CustomizableObjectNodeTextureParameter.generated.h"
+
+namespace ENodeTitleType { enum Type : int; }
+
+class UCustomizableObjectNodeRemapPins;
+class UObject;
+class UTexture2D;
+
+
+UCLASS()
+class CUSTOMIZABLEOBJECTEDITOR_API UCustomizableObjectNodeTextureParameter : public UCustomizableObjectNode
+{
+public:
+	GENERATED_BODY()
+
+	/** Default value of the parameter. */
+	UPROPERTY(EditAnywhere, Category=CustomizableObject)
+	TObjectPtr<UTexture2D> DefaultValue;
+
+	/** Reference Texture where this parameter copies some properties from. */
+	UPROPERTY(EditAnywhere, Category=CustomizableObject)
+	TObjectPtr<UTexture2D> ReferenceValue;
+
+	UPROPERTY(EditAnywhere, Category=CustomizableObject)
+	FString ParameterName = "Default Name";
+
+	UPROPERTY(EditAnywhere, Category = UI, meta = (DisplayName = "Parameter UI Metadata"))
+	FMutableParamUIMetadata ParamUIMetadata;
+
+	/** Set the width of the Texture when there is no texture reference.*/
+	UPROPERTY(EditAnywhere, Category = CustomizableObject)
+	int32 TextureSizeX = 0;
+
+	/** Set the height of the Texture when there is no texture reference.*/
+	UPROPERTY(EditAnywhere, Category = CustomizableObject)
+	int32 TextureSizeY = 0;
+
+	// EdGraphNode interface
+	virtual FText GetNodeTitle(ENodeTitleType::Type TitleType) const override;
+	virtual FLinearColor GetNodeTitleColor() const override;
+	virtual FText GetTooltipText() const override;
+	virtual void OnRenameNode(const FString& NewName) override;
+	virtual bool GetCanRenameNode() const override { return true; }
+
+	// UCustomizableObjectNode interface
+	virtual void BackwardsCompatibleFixup(int32 CustomizableObjectCustomVersion) override;
+	virtual void AllocateDefaultPins(UCustomizableObjectNodeRemapPins* RemapPins) override;
+	virtual bool IsAffectedByLOD() const override { return false; }
+};
+
diff --git a/Engine/Plugins/Mutable/Source/CustomizableObjectEditor/Public/MuCOE/Nodes/CustomizableObjectNodeTextureSample.h b/Engine/Plugins/Mutable/Source/CustomizableObjectEditor/Public/MuCOE/Nodes/CustomizableObjectNodeTextureSample.h
new file mode 100644
--- /dev/null
+++ b/Engine/Plugins/Mutable/Source/CustomizableObjectEditor/Public/MuCOE/Nodes/CustomizableObjectNodeTextureSample.h
@@ -0,0 +1,51 @@
+// Copyright Epic Games, Inc. All Rights Reserved.
+
+#pragma once
+
+#include "MuCOE/Nodes/CustomizableObjectNode.h"
+
+#include "CustomizableObjectNodeTextureSample.generated.h"
+
+namespace ENodeTitleType { enum Type : int; }
+
+class UCustomizableObjectNodeRemapPins;
+class UEdGraphPin;
+class UObject;
+struct FPropertyChangedEvent;
+
+
+UCLASS()
+class CUSTOMIZABLEOBJECTEDITOR_API UCustomizableObjectNodeTextureSample : public UCustomizableObjectNode
+{
+public:
+	GENERATED_BODY()
+
+	UCustomizableObjectNodeTextureSample();
+
+	// UObject interface.
+	void PostEditChangeProperty(FPropertyChangedEvent& PropertyChangedEvent) override;
+
+	// Begin EdGraphNode interface
+	FText GetNodeTitle(ENodeTitleType::Type TitleType) const override;
+	FLinearColor GetNodeTitleColor() const override;
+	FText GetTooltipText() const override;
+
+	// UCustomizableObjectNode interface
+	void AllocateDefaultPins(UCustomizableObjectNodeRemapPins* RemapPins) override;
+
+	const UEdGraphPin* TexturePin() const
+	{
+		return FindPin(TEXT("Texture"));
+	}
+
+	const UEdGraphPin* XPin() const
+	{
+		return FindPin(TEXT("X"));
+	}
+
+	const UEdGraphPin* YPin() const
+	{
+		return FindPin(TEXT("Y"));
+	}
+};
+
diff --git a/Engine/Plugins/Mutable/Source/CustomizableObjectEditor/Public/MuCOE/Nodes/CustomizableObjectNodeTextureSwitch.h b/Engine/Plugins/Mutable/Source/CustomizableObjectEditor/Public/MuCOE/Nodes/CustomizableObjectNodeTextureSwitch.h
new file mode 100644
--- /dev/null
+++ b/Engine/Plugins/Mutable/Source/CustomizableObjectEditor/Public/MuCOE/Nodes/CustomizableObjectNodeTextureSwitch.h
@@ -0,0 +1,20 @@
+// Copyright Epic Games, Inc. All Rights Reserved.
+
+#pragma once
+
+#include "MuCOE/Nodes/CustomizableObjectNodeSwitchBase.h"
+
+#include "CustomizableObjectNodeTextureSwitch.generated.h"
+
+
+UCLASS()
+class CUSTOMIZABLEOBJECTEDITOR_API UCustomizableObjectNodeTextureSwitch : public UCustomizableObjectNodeSwitchBase
+{
+public:
+	GENERATED_BODY()
+	// UCustomizableObjectNode interface
+	virtual void BackwardsCompatibleFixup(int32 CustomizableObjectCustomVersion) override;
+	
+	// UCustomizableObjectNodeSwitchBase interface
+	virtual FName GetCategory() const override;
+};
diff --git a/Engine/Plugins/Mutable/Source/CustomizableObjectEditor/Public/MuCOE/Nodes/CustomizableObjectNodeTextureToChannels.h b/Engine/Plugins/Mutable/Source/CustomizableObjectEditor/Public/MuCOE/Nodes/CustomizableObjectNodeTextureToChannels.h
new file mode 100644
--- /dev/null
+++ b/Engine/Plugins/Mutable/Source/CustomizableObjectEditor/Public/MuCOE/Nodes/CustomizableObjectNodeTextureToChannels.h
@@ -0,0 +1,62 @@
+// Copyright Epic Games, Inc. All Rights Reserved.
+
+#pragma once
+
+#include "MuCOE/Nodes/CustomizableObjectNode.h"
+
+#include "CustomizableObjectNodeTextureToChannels.generated.h"
+
+namespace ENodeTitleType { enum Type : int; }
+
+class FArchive;
+class UCustomizableObjectNodeRemapPins;
+class UObject;
+struct FPropertyChangedEvent;
+
+
+UCLASS()
+class CUSTOMIZABLEOBJECTEDITOR_API UCustomizableObjectNodeTextureToChannels : public UCustomizableObjectNode
+{
+public:
+	GENERATED_BODY()
+	
+	// UObject interface.
+	void PostEditChangeProperty(FPropertyChangedEvent& PropertyChangedEvent) override;
+	virtual void Serialize(FArchive& Ar) override;
+	
+	// Begin EdGraphNode interface
+	FText GetNodeTitle(ENodeTitleType::Type TitleType) const override;
+	FLinearColor GetNodeTitleColor() const override;
+	FText GetTooltipText() const override;
+
+	// UCustomizableObjectNode interface
+	virtual void AllocateDefaultPins(UCustomizableObjectNodeRemapPins* RemapPins) override;
+	virtual void BackwardsCompatibleFixup(int32 CustomizableObjectCustomVersion) override;
+
+	UEdGraphPin* InputPin() const;
+
+	UEdGraphPin* RPin() const
+	{
+		return FindPin(TEXT("R"));
+	}
+
+	UEdGraphPin* GPin() const
+	{
+		return FindPin(TEXT("G"));
+	}
+
+	UEdGraphPin* BPin() const
+	{
+		return FindPin(TEXT("B"));
+	}
+
+	UEdGraphPin* APin() const
+	{
+		return FindPin(TEXT("A"));
+	}
+
+private:
+	UPROPERTY()
+	FEdGraphPinReference InputPinReference;
+};
+
diff --git a/Engine/Plugins/Mutable/Source/CustomizableObjectEditor/Public/MuCOE/Nodes/CustomizableObjectNodeTextureVariation.h b/Engine/Plugins/Mutable/Source/CustomizableObjectEditor/Public/MuCOE/Nodes/CustomizableObjectNodeTextureVariation.h
new file mode 100644
--- /dev/null
+++ b/Engine/Plugins/Mutable/Source/CustomizableObjectEditor/Public/MuCOE/Nodes/CustomizableObjectNodeTextureVariation.h
@@ -0,0 +1,69 @@
+// Copyright Epic Games, Inc. All Rights Reserved.
+
+#pragma once
+
+#include "MuCOE/Nodes/CustomizableObjectNode.h"
+
+#include "CustomizableObjectNodeTextureVariation.generated.h"
+
+namespace ENodeTitleType { enum Type : int; }
+
+class UCustomizableObjectNodeRemapPins;
+class UEdGraphPin;
+class UObject;
+struct FPropertyChangedEvent;
+
+
+USTRUCT()
+struct CUSTOMIZABLEOBJECTEDITOR_API FCustomizableObjectTextureVariation
+{
+	GENERATED_USTRUCT_BODY()
+
+	UPROPERTY(EditAnywhere, Category = CustomizableObject)
+	FString Tag;
+};
+
+
+UCLASS()
+class CUSTOMIZABLEOBJECTEDITOR_API UCustomizableObjectNodeTextureVariation : public UCustomizableObjectNode
+{
+public:
+	GENERATED_BODY()
+
+	UCustomizableObjectNodeTextureVariation();
+
+	UPROPERTY(EditAnywhere, Category = CustomizableObject)
+	TArray<FCustomizableObjectTextureVariation> Variations;
+
+	// UObject interface.
+	void PostEditChangeProperty(FPropertyChangedEvent& PropertyChangedEvent) override;
+
+	// EdGraphNode interface
+	virtual FText GetNodeTitle(ENodeTitleType::Type TitleType) const override;
+	FLinearColor GetNodeTitleColor() const override;
+	FText GetTooltipText() const override;
+	
+	// UCustomizableObjectNode interface
+	virtual void AllocateDefaultPins(UCustomizableObjectNodeRemapPins* RemapPins) override;
+
+	UEdGraphPin* OutputPin() const
+	{
+		return FindPin(TEXT("Texture"));
+	}
+
+	UEdGraphPin* DefaultPin() const
+	{
+		return FindPin(TEXT("Default"));
+	}
+
+	int32 GetNumVariations() const
+	{
+		return Variations.Num();
+	}
+
+	UEdGraphPin* VariationPin(int Index) const
+	{
+		return FindPin(FString::Printf(TEXT("Variation %d"),Index));
+	}
+};
+
diff --git a/Engine/Plugins/Mutable/Source/CustomizableObjectEditor/Public/MuCOE/Thumbnails/CustomizableObjectInstanceThumbnailRenderer.h b/Engine/Plugins/Mutable/Source/CustomizableObjectEditor/Public/MuCOE/Thumbnails/CustomizableObjectInstanceThumbnailRenderer.h
new file mode 100644
--- /dev/null
+++ b/Engine/Plugins/Mutable/Source/CustomizableObjectEditor/Public/MuCOE/Thumbnails/CustomizableObjectInstanceThumbnailRenderer.h
@@ -0,0 +1,40 @@
+// Copyright Epic Games, Inc. All Rights Reserved.
+
+/**
+ *
+ * This thumbnail renderer displays a given Customizable Object Instance
+ */
+
+#pragma once
+#include "ThumbnailRendering/DefaultSizedThumbnailRenderer.h"
+
+#include "CustomizableObjectInstanceThumbnailRenderer.generated.h"
+
+class FCanvas;
+class FRenderTarget;
+class UObject;
+
+
+UCLASS(config=Editor)
+class CUSTOMIZABLEOBJECTEDITOR_API UCustomizableObjectInstanceThumbnailRenderer : public UDefaultSizedThumbnailRenderer
+{
+public:
+	GENERATED_BODY()
+
+	UCustomizableObjectInstanceThumbnailRenderer();
+
+	// Begin UThumbnailRenderer Object
+	void Draw(UObject* Object, int32 X, int32 Y, uint32 Width, uint32 Height, FRenderTarget*, FCanvas* Canvas, bool bAdditionalViewFamily) override;
+	// End UThumbnailRenderer Object
+
+	// UObject implementation
+	void BeginDestroy() override;
+
+private:
+	class FSkeletalMeshThumbnailScene* ThumbnailScene;
+
+	UPROPERTY()
+	TObjectPtr<class UTexture2D> NoImage;
+
+};
+
diff --git a/Engine/Plugins/Mutable/Source/CustomizableObjectEditor/Public/MuCOE/Thumbnails/CustomizableObjectThumbnailRenderer.h b/Engine/Plugins/Mutable/Source/CustomizableObjectEditor/Public/MuCOE/Thumbnails/CustomizableObjectThumbnailRenderer.h
new file mode 100644
--- /dev/null
+++ b/Engine/Plugins/Mutable/Source/CustomizableObjectEditor/Public/MuCOE/Thumbnails/CustomizableObjectThumbnailRenderer.h
@@ -0,0 +1,42 @@
+// Copyright Epic Games, Inc. All Rights Reserved.
+
+/**
+ *
+ * This thumbnail renderer displays a given Customizable Object
+ */
+
+#pragma once
+#include "ThumbnailRendering/DefaultSizedThumbnailRenderer.h"
+
+#include "CustomizableObjectThumbnailRenderer.generated.h"
+
+class FCanvas;
+class FRenderTarget;
+class UObject;
+
+
+UCLASS(config=Editor)
+class CUSTOMIZABLEOBJECTEDITOR_API UCustomizableObjectThumbnailRenderer : public UDefaultSizedThumbnailRenderer
+{
+public:
+	GENERATED_BODY()
+
+	UCustomizableObjectThumbnailRenderer();
+
+	// Begin UThumbnailRenderer Object
+	virtual void Draw(UObject* Object, int32 X, int32 Y, uint32 Width, uint32 Height, FRenderTarget*, FCanvas* Canvas, bool bAdditionalViewFamily) override;
+	// End UThumbnailRenderer Object
+
+	// UObject implementation
+	virtual void BeginDestroy() override;
+
+private:
+	class FSkeletalMeshThumbnailScene* ThumbnailScene;
+
+	UPROPERTY()
+	TObjectPtr<class UCustomizableObjectInstance> CustomizableObjectInstance;
+
+	UPROPERTY()
+	TObjectPtr<class UTexture2D> NoImage;
+};
+
