[mutable] Moved the Mutable plugin out of Experimental status into Beta.

diff --git a/Engine/Plugins/Mutable/Source/MutableRuntime/Private/MuR/CodeRunner.cpp b/Engine/Plugins/Mutable/Source/MutableRuntime/Private/MuR/CodeRunner.cpp
new file mode 100644
--- /dev/null
+++ b/Engine/Plugins/Mutable/Source/MutableRuntime/Private/MuR/CodeRunner.cpp
@@ -0,0 +1,6811 @@
+// Copyright Epic Games, Inc. All Rights Reserved.
+
+#include "MuR/CodeRunner.h"
+
+#include "GenericPlatform/GenericPlatformMath.h"
+#include "HAL/UnrealMemory.h"
+#include "Logging/LogCategory.h"
+#include "Logging/LogMacros.h"
+#include "Math/IntPoint.h"
+#include "Math/UnrealMathSSE.h"
+#include "Math/UnrealMathUtility.h"
+#include "Math/Vector2D.h"
+#include "Misc/AssertionMacros.h"
+#include "MuR/ImagePrivate.h"
+#include "MuR/Instance.h"
+#include "MuR/InstancePrivate.h"
+#include "MuR/Mesh.h"
+#include "MuR/MeshBufferSet.h"
+#include "MuR/Model.h"
+#include "MuR/ModelPrivate.h"
+#include "MuR/MutableMath.h"
+#include "MuR/MutableString.h"
+#include "MuR/MutableTrace.h"
+#include "MuR/OpImageApplyComposite.h"
+#include "MuR/OpImageBinarise.h"
+#include "MuR/OpImageBlend.h"
+#include "MuR/OpImageColourMap.h"
+#include "MuR/OpImageDisplace.h"
+#include "MuR/OpImageGradient.h"
+#include "MuR/OpImageInterpolate.h"
+#include "MuR/OpImageInvert.h"
+#include "MuR/OpImageLuminance.h"
+#include "MuR/OpImageNormalCombine.h"
+#include "MuR/OpImageProject.h"
+#include "MuR/OpImageRasterMesh.h"
+#include "MuR/OpImageSaturate.h"
+#include "MuR/OpImageTransform.h"
+#include "MuR/OpLayoutPack.h"
+#include "MuR/OpLayoutRemoveBlocks.h"
+#include "MuR/OpMeshApplyLayout.h"
+#include "MuR/OpMeshApplyPose.h"
+#include "MuR/OpMeshBind.h"
+#include "MuR/OpMeshClipDeform.h"
+#include "MuR/OpMeshClipMorphPlane.h"
+#include "MuR/OpMeshClipWithMesh.h"
+#include "MuR/OpMeshDifference.h"
+#include "MuR/OpMeshExtractLayoutBlock.h"
+#include "MuR/OpMeshFormat.h"
+#include "MuR/OpMeshGeometryOperation.h"
+#include "MuR/OpMeshMerge.h"
+#include "MuR/OpMeshMorph.h"
+#include "MuR/OpMeshRemove.h"
+#include "MuR/OpMeshReshape.h"
+#include "MuR/OpMeshTransform.h"
+#include "MuR/OpMeshOptimizeSkinning.h"
+#include "MuR/Operations.h"
+#include "MuR/Parameters.h"
+#include "MuR/ParametersPrivate.h"
+#include "MuR/PhysicsBody.h"
+#include "MuR/Platform.h"
+#include "MuR/Skeleton.h"
+#include "MuR/SystemPrivate.h"
+#include "Templates/Tuple.h"
+#include "Trace/Detail/Channel.h"
+
+namespace
+{
+
+int32 ForcedProjectionMode = -1;
+static FAutoConsoleVariableRef CVarForceProjectionSamplingMode (
+	TEXT("mutable.ForceProjectionMode"),
+	ForcedProjectionMode,
+	TEXT("force mutable to use an specific projection mode, 0 = Point + None, 1 = Bilinear + TotalAreaHeuristic, -1 uses the values provided by the projector."),
+	ECVF_Default);
+
+float GlobalProjectionLodBias = 0.0f;
+static FAutoConsoleVariableRef CVarGlobalProjectionLodBias (
+	TEXT("mutable.GlobalProjectionLodBias"),
+	GlobalProjectionLodBias,
+	TEXT("Lod bias applied to the lod resulting form the best mip computation for ImageProject operations, only used if a min filter method different than None is used."),
+	ECVF_Default);
+
+bool bUseProjectionVectorImpl = true;
+static FAutoConsoleVariableRef CVarUseProjectionVectorImpl (
+	TEXT("mutable.UseProjectionVectorImpl"),
+	bUseProjectionVectorImpl,
+	TEXT("If set to true, enables the vectorized implementation of the projection pixel processing."),
+	ECVF_Default);
+
+float GlobalImageTransformLodBias = 0.0f;
+static FAutoConsoleVariableRef CVarGlobalImageTransformLodBias (
+	TEXT("mutable.GlobalImageTransformLodBias"),
+	GlobalImageTransformLodBias,
+		TEXT("Lod bias applied to the lod resulting form the best mip computation for ImageTransform operations"),
+	ECVF_Default);
+
+bool bUseImageTransformVectorImpl = true;
+static FAutoConsoleVariableRef CVarUseImageTransformVectorImpl (
+	TEXT("mutable.UseImageTransformVectorImpl"),
+	bUseImageTransformVectorImpl,
+	TEXT("If set to true, enables the vectorized implementation of the image transform pixel processing."),
+	ECVF_Default);
+}
+
+namespace mu
+{
+
+	TSharedRef<CodeRunner> CodeRunner::Create(
+		const Ptr<const Settings>& InSettings,
+		class System::Private* InSystem,
+		EExecutionStrategy InExecutionStrategy,
+		const TSharedPtr<const Model>& InModel,
+		const Parameters* InParams,
+		OP::ADDRESS At,
+		uint32 InLODMask, uint8 ExecutionOptions, int32 InImageLOD, FScheduledOp::EType Type)
+	{
+		return MakeShared<CodeRunner>(FPrivateToken {},
+				InSettings, InSystem, InExecutionStrategy, InModel, InParams, At, InLODMask, ExecutionOptions, InImageLOD, Type);
+	}
+
+    CodeRunner::CodeRunner(FPrivateToken PrivateToken, 
+		const Ptr<const Settings>& InSettings,
+		class System::Private* InSystem,
+		EExecutionStrategy InExecutionStrategy,
+		const TSharedPtr<const Model>& InModel,
+		const Parameters* InParams,
+		OP::ADDRESS at,
+		uint32 InLodMask, uint8 executionOptions, int32 InImageLOD, FScheduledOp::EType Type )
+		: m_pSettings(InSettings)
+		, RunnerCompletionEvent(TEXT("CodeRunnerCompletioneEventInit"))
+		, ExecutionStrategy(InExecutionStrategy)
+		, m_pSystem(InSystem)
+		, m_pModel(InModel)
+		, m_pParams(InParams)
+		, m_lodMask(InLodMask)
+	{
+		const FProgram& program = m_pModel->GetPrivate()->m_program;
+		ScheduledStagePerOp.resize(program.m_opAddress.Num());
+
+		// We will read this in the end, so make sure we keep it.
+   		if (Type == FScheduledOp::EType::Full)
+   		{
+			GetMemory().IncreaseHitCount(FCacheAddress(at, 0, executionOptions));
+		}
+    
+		// Start with a completed Event. This is checked at StartRun() to make sure StartRun is not called while there is 
+		// a Run in progress.
+		RunnerCompletionEvent.Trigger();
+
+		ImageLOD = InImageLOD;
+	
+		// Push the root operation
+		FScheduledOp rootOp;
+		rootOp.At = at;
+		rootOp.ExecutionOptions = executionOptions;
+		rootOp.Type = Type;
+		AddOp(rootOp);
+	}
+
+
+    //---------------------------------------------------------------------------------------------
+	FProgramCache& CodeRunner::GetMemory()
+    {
+		return *m_pSystem->WorkingMemoryManager.CurrentInstanceCache;
+	}
+
+
+	TTuple<UE::Tasks::FTask, TFunction<void()>> CodeRunner::LoadExternalImageAsync(FExternalResourceId Id, uint8 MipmapsToSkip, TFunction<void(Ptr<Image>)>& ResultCallback)
+    {
+		MUTABLE_CPUPROFILER_SCOPE(LoadExternalImageAsync);
+
+		check(m_pSystem);
+
+		if (m_pSystem->ExternalResourceProvider)
+		{
+			if (Id.ReferenceResourceId < 0)
+			{
+				// It's a parameter image
+				return m_pSystem->ExternalResourceProvider->GetImageAsync(Id.ParameterId, MipmapsToSkip, ResultCallback);
+			}
+			else
+			{
+				// It's an image reference
+				return m_pSystem->ExternalResourceProvider->GetReferencedImageAsync(m_pModel.Get(), Id.ReferenceResourceId, MipmapsToSkip, ResultCallback);
+			}
+		}
+		else
+		{
+			// Not found and there is no generator!
+			check(false);
+		}
+
+		return MakeTuple(UE::Tasks::MakeCompletedTask<void>(), []() -> void {});
+	}
+
+
+	TTuple<UE::Tasks::FTask, TFunction<void()>> CodeRunner::LoadExternalMeshAsync(FExternalResourceId Id, TFunction<void(Ptr<Mesh>)>& ResultCallback)
+	{
+		MUTABLE_CPUPROFILER_SCOPE(LoadExternalImageAsync);
+
+		check(m_pSystem);
+
+		if (m_pSystem->ExternalResourceProvider)
+		{
+			if (Id.ReferenceResourceId < 0)
+			{
+				// It's a parameter mesh
+				return m_pSystem->ExternalResourceProvider->GetMeshAsync(Id.ParameterId, ResultCallback);
+			}
+			else
+			{
+				// It's a mesh reference
+				return m_pSystem->ExternalResourceProvider->GetReferencedMeshAsync(m_pModel.Get(), Id.ReferenceResourceId, ResultCallback);
+			}
+		}
+		else
+		{
+			// Not found and there is no generator!
+			check(false);
+		}
+
+		return MakeTuple(UE::Tasks::MakeCompletedTask<void>(), []() -> void {});
+	}
+
+	
+    //---------------------------------------------------------------------------------------------
+	FImageDesc CodeRunner::GetExternalImageDesc(FName Id, uint8 MipmapsToSkip)
+	{
+		MUTABLE_CPUPROFILER_SCOPE(GetExternalImageDesc);
+
+		check(m_pSystem);
+
+		if (m_pSystem->ExternalResourceProvider)
+		{
+			return m_pSystem->ExternalResourceProvider->GetImageDesc(Id, MipmapsToSkip);
+		}
+		else
+		{
+			// Not found and there is no generator!
+			check(false);
+		}
+
+		return FImageDesc();
+	}
+
+	
+    //---------------------------------------------------------------------------------------------
+    void CodeRunner::RunCode_Conditional( const FScheduledOp& item, const Model* pModel )
+    {
+        MUTABLE_CPUPROFILER_SCOPE(RunCode_Conditional);
+
+		const FProgram& Program = m_pModel->GetPrivate()->m_program;
+
+		OP_TYPE type = Program.GetOpType(item.At);
+		OP::ConditionalArgs args = Program.GetOpArgs<OP::ConditionalArgs>(item.At);
+
+        // Conditionals have the following execution stages:
+        // 0: we need to run the condition
+        // 1: we need to run the branch
+        // 2: we need to fetch the result and store it in this op
+        switch( item.Stage )
+        {
+        case 0:
+        {
+            AddOp( FScheduledOp( item.At,item,1 ),
+                   FScheduledOp( args.condition, item ) );
+            break;
+        }
+
+        case 1:
+        {
+            // Get the condition result
+
+            // If there is no expression, we'll assume true.
+            bool value = true;
+            value = LoadBool( FCacheAddress(args.condition, item.ExecutionIndex, item.ExecutionOptions) );
+
+            OP::ADDRESS resultAt = value ? args.yes : args.no;
+
+            // Schedule the end of this instruction if necessary
+            AddOp( FScheduledOp( item.At, item, 2, (uint32)value),
+				FScheduledOp( resultAt, item) );
+
+            break;
+        }
+
+        case 2:
+        {
+            OP::ADDRESS resultAt = item.CustomState ? args.yes : args.no;
+
+            // Store the final result
+            FCacheAddress cat( item );
+            FCacheAddress rat( resultAt, item );
+            switch (GetOpDataType(type))
+            {
+            case DT_BOOL:       StoreBool( cat, LoadBool(rat) ); break;
+            case DT_INT:        StoreInt( cat, LoadInt(rat) ); break;
+            case DT_SCALAR:     StoreScalar( cat, LoadScalar(rat) ); break;
+			case DT_STRING:		StoreString( cat, LoadString( rat ) ); break;
+            case DT_COLOUR:		StoreColor( cat, LoadColor( rat ) ); break;
+            case DT_PROJECTOR:  StoreProjector( cat, LoadProjector(rat) ); break;
+            case DT_MESH:       StoreMesh( cat, LoadMesh(rat) ); break;
+            case DT_IMAGE:      StoreImage( cat, LoadImage(rat) ); break;
+            case DT_LAYOUT:     StoreLayout( cat, LoadLayout(rat) ); break;
+            case DT_INSTANCE:   StoreInstance( cat, LoadInstance(rat) ); break;
+			case DT_EXTENSION_DATA: StoreExtensionData( cat, LoadExtensionData(rat) ); break;
+            default:
+                // Not implemented
+                check( false );
+            }
+
+            break;
+        }
+
+        default:
+            check(false);
+        }
+    }
+
+
+	//---------------------------------------------------------------------------------------------
+	void CodeRunner::RunCode_Switch(const FScheduledOp& item, const Model* pModel )
+	{
+		const FProgram& Program = m_pModel->GetPrivate()->m_program;
+
+		OP_TYPE type = Program.GetOpType(item.At);
+
+		const uint8* data = Program.GetOpArgsPointer(item.At);
+
+		OP::ADDRESS VarAddress;
+		FMemory::Memcpy(&VarAddress, data, sizeof(OP::ADDRESS));
+		data += sizeof(OP::ADDRESS);
+
+		OP::ADDRESS DefAddress;
+		FMemory::Memcpy(&DefAddress, data, sizeof(OP::ADDRESS));
+		data += sizeof(OP::ADDRESS);
+
+		uint32 CaseCount;
+		FMemory::Memcpy(&CaseCount, data, sizeof(uint32));
+		data += sizeof(uint32);
+
+		switch (item.Stage)
+		{
+		case 0:
+		{
+			if (VarAddress)
+			{
+				AddOp(FScheduledOp(item.At, item, 1), FScheduledOp(VarAddress, item));
+			}
+			else
+			{
+				switch (GetOpDataType(type))
+				{
+				case DT_BOOL:       StoreBool(item, false); break;
+				case DT_INT:        StoreInt(item, 0); break;
+				case DT_SCALAR:		StoreScalar(item, 0.0f); break;
+				case DT_STRING:		StoreString(item, nullptr); break;
+				case DT_COLOUR:		StoreColor(item, FVector4f()); break;
+				case DT_PROJECTOR:  StoreProjector(item, FProjector()); break;
+				case DT_MESH:       StoreMesh(item, nullptr); break;
+				case DT_IMAGE:      StoreImage(item, nullptr); break;
+				case DT_LAYOUT:     StoreLayout(item, nullptr); break;
+				case DT_INSTANCE:   StoreInstance(item, nullptr); break;
+				case DT_EXTENSION_DATA: StoreExtensionData(item, new ExtensionData); break;
+				default:
+					// Not implemented
+					check(false);
+				}
+			}
+			break;
+		}
+
+		case 1:
+		{
+			// Get the variable result
+			int var = LoadInt(FCacheAddress(VarAddress, item));
+
+			OP::ADDRESS valueAt = DefAddress;
+			for (uint32 C = 0; C < CaseCount; ++C)
+			{
+				int32 Condition;
+				FMemory::Memcpy(&Condition, data, sizeof(int32));
+				data += sizeof(int32);
+
+				OP::ADDRESS At;
+				FMemory::Memcpy(&At, data, sizeof(OP::ADDRESS));
+				data += sizeof(OP::ADDRESS);
+
+				if (At && var == (int)Condition)
+				{
+					valueAt = At;
+					break; 
+				}
+			}
+
+            // Schedule the end of this instruction if necessary
+            AddOp( FScheduledOp( item.At, item, 2, valueAt ),
+				   FScheduledOp( valueAt, item ) );
+
+            break;
+        }
+
+        case 2:
+        {
+			OP::ADDRESS resultAt = OP::ADDRESS(item.CustomState);
+
+            // Store the final result
+            FCacheAddress cat( item );
+            FCacheAddress rat( resultAt, item );
+            switch (GetOpDataType(type))
+            {
+            case DT_BOOL:       StoreBool( cat, LoadBool(rat) ); break;
+            case DT_INT:        StoreInt( cat, LoadInt(rat) ); break;
+            case DT_SCALAR:     StoreScalar( cat, LoadScalar(rat) ); break;
+            case DT_STRING:		StoreString( cat, LoadString( rat ) ); break;
+            case DT_COLOUR:		StoreColor( cat, LoadColor( rat ) ); break;
+            case DT_PROJECTOR:  StoreProjector( cat, LoadProjector(rat) ); break;
+			case DT_MESH:       StoreMesh( cat, LoadMesh(rat) ); break;
+            case DT_IMAGE:      StoreImage( cat, LoadImage(rat) ); break;
+            case DT_LAYOUT:     StoreLayout( cat, LoadLayout(rat) ); break;
+            case DT_INSTANCE:   StoreInstance( cat, LoadInstance(rat) ); break;
+			case DT_EXTENSION_DATA: StoreExtensionData( cat, LoadExtensionData(rat) ); break;
+            default:
+                // Not implemented
+                check( false );
+            }
+
+            break;
+        }
+
+        default:
+            check(false);
+        }
+    }
+
+
+    //---------------------------------------------------------------------------------------------
+    void CodeRunner::RunCode_Instance(const FScheduledOp& item, const Model* pModel, uint32 lodMask )
+    {
+        MUTABLE_CPUPROFILER_SCOPE(RunCode_Instance);
+
+		const FProgram& Program = m_pModel->GetPrivate()->m_program;
+
+		OP_TYPE type = Program.GetOpType(item.At);
+        switch (type)
+        {
+
+        case OP_TYPE::IN_ADDVECTOR:
+        {
+			OP::InstanceAddArgs args = Program.GetOpArgs<OP::InstanceAddArgs>(item.At);
+
+            switch (item.Stage)
+            {
+            case 0:
+                    AddOp( FScheduledOp( item.At, item, 1),
+                           FScheduledOp( args.instance, item),
+                           FScheduledOp( args.value, item) );
+                break;
+
+            case 1:
+            {
+                InstancePtrConst pBase = LoadInstance( FCacheAddress(args.instance,item) );
+                InstancePtr pResult;
+                if (!pBase)
+                {
+                    pResult = new Instance();
+                }
+                else
+                {
+					pResult = mu::CloneOrTakeOver<Instance>(pBase.get());
+                }
+
+                if ( args.value )
+                {
+					FVector4f value = LoadColor( FCacheAddress(args.value,item) );
+
+                    OP::ADDRESS nameAd = args.name;
+                    check(  nameAd < (uint32)Program.m_constantStrings.Num() );
+                    const FString& Name = Program.m_constantStrings[ nameAd ];
+
+                    pResult->GetPrivate()->AddVector( 0, 0, 0, value, FName(Name) );
+                }
+                StoreInstance( item, pResult );
+                break;
+            }
+
+            default:
+                check(false);
+            }
+
+            break;
+        }
+
+        case OP_TYPE::IN_ADDSCALAR:
+        {
+			OP::InstanceAddArgs args = Program.GetOpArgs<OP::InstanceAddArgs>(item.At);
+            switch (item.Stage)
+            {
+            case 0:
+                    AddOp( FScheduledOp( item.At, item, 1),
+                           FScheduledOp( args.instance, item),
+                           FScheduledOp( args.value, item) );
+                break;
+
+            case 1:
+            {
+                InstancePtrConst pBase = LoadInstance( FCacheAddress(args.instance,item) );
+                InstancePtr pResult;
+                if (!pBase)
+                {
+                    pResult = new Instance();
+                }
+                else
+                {
+                    pResult = mu::CloneOrTakeOver<Instance>(pBase.get());
+                }
+
+                if ( args.value )
+                {
+                    float value = LoadScalar( FCacheAddress(args.value,item) );
+
+                    OP::ADDRESS nameAd = args.name;
+                    check(  nameAd < (uint32)Program.m_constantStrings.Num() );
+                    const FString& Name = Program.m_constantStrings[ nameAd ];
+
+                    pResult->GetPrivate()->AddScalar( 0, 0, 0, value, FName(Name));
+                }
+                StoreInstance( item, pResult );
+                break;
+            }
+
+            default:
+                check(false);
+            }
+
+            break;
+        }
+
+        case OP_TYPE::IN_ADDSTRING:
+        {
+			OP::InstanceAddArgs args = Program.GetOpArgs<OP::InstanceAddArgs>( item.At );
+            switch ( item.Stage )
+            {
+            case 0:
+                    AddOp( FScheduledOp( item.At, item, 1 ), FScheduledOp( args.instance, item ),
+                           FScheduledOp( args.value, item ) );
+                break;
+
+            case 1:
+            {
+                InstancePtrConst pBase =
+                    LoadInstance( FCacheAddress( args.instance, item ) );
+                InstancePtr pResult;
+                if ( !pBase )
+                {
+                    pResult = new Instance();
+                }
+                else
+                {
+					pResult = mu::CloneOrTakeOver<Instance>(pBase.get());
+				}
+
+                if ( args.value )
+                {
+                    Ptr<const String> value = LoadString( FCacheAddress( args.value, item ) );
+
+                    OP::ADDRESS nameAd = args.name;
+                    check( nameAd < (uint32)Program.m_constantStrings.Num() );
+                    const FString& Name = Program.m_constantStrings[nameAd];
+
+                    pResult->GetPrivate()->AddString( 0, 0, 0, value->GetValue(), FName(Name) );
+                }
+                StoreInstance( item, pResult );
+                break;
+            }
+
+            default:
+                check( false );
+            }
+
+            break;
+        }
+
+        case OP_TYPE::IN_ADDCOMPONENT:
+        {
+			OP::InstanceAddArgs args = Program.GetOpArgs<OP::InstanceAddArgs>(item.At);
+            switch (item.Stage)
+            {
+            case 0:
+                    AddOp( FScheduledOp( item.At, item, 1),
+                           FScheduledOp( args.instance, item),
+                           FScheduledOp( args.value, item) );
+                break;
+
+            case 1:
+            {
+				Ptr<const Instance> pBase = LoadInstance( FCacheAddress(args.instance,item) );
+				Ptr<Instance> pResult;
+                if (!pBase)
+                {
+                    pResult = new Instance();
+                }
+                else
+                {
+					pResult = mu::CloneOrTakeOver<Instance>(pBase.get());
+				}
+
+                if ( args.value )
+                {
+                    Ptr<const Instance> pComp = LoadInstance( FCacheAddress(args.value,item) );
+
+                    int32 NewComponentIndex = pResult->GetPrivate()->AddComponent();
+
+                    if ( !pComp->GetPrivate()->Components.IsEmpty() )
+                    {
+                        pResult->GetPrivate()->Components[NewComponentIndex] = pComp->GetPrivate()->Components[0];
+
+						// Id
+                    	pResult->GetPrivate()->Components[NewComponentIndex].Id = args.id;
+                    }
+                }
+                StoreInstance( item, pResult );
+                break;
+            }
+
+            default:
+                check(false);
+            }
+
+            break;
+        }
+
+        case OP_TYPE::IN_ADDSURFACE:
+        {
+			OP::InstanceAddArgs args = Program.GetOpArgs<OP::InstanceAddArgs>(item.At);
+            switch (item.Stage)
+            {
+            case 0:
+                    AddOp( FScheduledOp( item.At, item, 1),
+                           FScheduledOp( args.instance, item),
+                           FScheduledOp( args.value, item) );
+                break;
+
+            case 1:
+            {
+				Ptr<const Instance> pBase = LoadInstance( FCacheAddress(args.instance,item) );
+
+				Ptr<Instance> pResult;
+				if (pBase)
+				{
+					pResult = mu::CloneOrTakeOver<Instance>(pBase.get());
+				}
+				else
+				{
+					pResult = new Instance();
+				}
+
+                // Empty surfaces are ok, they still need to be created, because they may contain
+                // additional information like internal or external IDs
+                //if ( args.value )
+                {
+					Ptr<const Instance> pSurf = LoadInstance( FCacheAddress(args.value,item) );
+
+                    int sindex = pResult->GetPrivate()->AddSurface( 0, 0 );
+
+                    // Surface data
+                    if (pSurf
+                            &&
+                            pSurf->GetPrivate()->Components.Num()
+                            &&
+                            pSurf->GetPrivate()->Components[0].LODs.Num()
+                            &&
+                            pSurf->GetPrivate()->Components[0].LODs[0].Surfaces.Num())
+                    {
+                        pResult->GetPrivate()->Components[0].LODs[0].Surfaces[sindex] =
+                            pSurf->GetPrivate()->Components[0].LODs[0].Surfaces[0];
+                    }
+
+                    // Name
+                    OP::ADDRESS nameAd = args.name;
+                    check( nameAd < (uint32)Program.m_constantStrings.Num() );
+                    const FString& Name = Program.m_constantStrings[ nameAd ];
+                    pResult->GetPrivate()->SetSurfaceName( 0, 0, sindex, FName(Name) );
+
+                    // IDs
+                    pResult->GetPrivate()->Components[0].LODs[0].Surfaces[sindex].InternalId = args.id;
+                    pResult->GetPrivate()->Components[0].LODs[0].Surfaces[sindex].ExternalId = args.ExternalId;
+                    pResult->GetPrivate()->Components[0].LODs[0].Surfaces[sindex].SharedId = args.SharedSurfaceId;
+                }
+                StoreInstance( item, pResult );
+                break;
+            }
+
+            default:
+                check(false);
+            }
+
+            break;
+        }
+
+        case OP_TYPE::IN_ADDLOD:
+        {
+			const uint8* Data = Program.GetOpArgsPointer(item.At);
+
+			uint8 LODCount;
+			FMemory::Memcpy(&LODCount, Data, sizeof(uint8));
+			Data += sizeof(uint8);
+
+			switch (item.Stage)
+            {
+            case 0:
+            {                
+                TArray<FScheduledOp> deps;
+                for (uint8 LODIndex=0; LODIndex <LODCount; ++LODIndex)
+                {
+					OP::ADDRESS LODAddress;
+					FMemory::Memcpy(&LODAddress, Data, sizeof(OP::ADDRESS));
+					Data += sizeof(OP::ADDRESS);
+
+                    if (LODAddress)
+                    {
+                        bool bSelectedLod = ( (1<< LODIndex) & lodMask ) != 0;
+
+                        if (bSelectedLod)
+                        {
+                            deps.Emplace(LODAddress, item);
+                        }
+                    }
+                }
+
+                AddOp( FScheduledOp( item.At, item, 1), deps );
+
+                break;
+            }
+
+            case 1:
+            {
+                // Assemble result
+				Ptr<Instance> pResult = new Instance();
+				int32 ComponentIndex = pResult->GetPrivate()->AddComponent();
+
+                for (uint8 LODIndex = 0; LODIndex < LODCount; ++LODIndex)
+                {
+					OP::ADDRESS LODAddress;
+					FMemory::Memcpy(&LODAddress, Data, sizeof(OP::ADDRESS));
+					Data += sizeof(OP::ADDRESS);
+					
+					if ( LODAddress )
+                    {						
+						bool bIsSelectedLod = ( (1<<LODIndex) & lodMask ) != 0;
+
+						// Add an empty LOD even if not selected.
+						int32 InstanceLODIndex = pResult->GetPrivate()->AddLOD(ComponentIndex);
+						
+						if (bIsSelectedLod)
+                        {
+							Ptr<const Instance> pLOD = LoadInstance( FCacheAddress(LODAddress,item) );
+
+                            // In a degenerated case, the returned pLOD may not have an LOD inside
+ 							if (!pLOD->GetPrivate()->Components.IsEmpty()
+								&&
+								!pLOD->GetPrivate()->Components[0].LODs.IsEmpty())
+							{
+								pResult->GetPrivate()->Components[ComponentIndex].LODs[InstanceLODIndex] = pLOD->GetPrivate()->Components[0].LODs[0];
+							}
+						}
+                    }
+                }
+
+                StoreInstance( item, pResult );
+                break;
+            }
+
+            default:
+                check(false);
+            }
+
+            break;
+        }
+
+		case OP_TYPE::IN_ADDEXTENSIONDATA:
+		{
+			OP::InstanceAddExtensionDataArgs Args = Program.GetOpArgs<OP::InstanceAddExtensionDataArgs>(item.At);
+			switch (item.Stage)
+			{
+				case 0:
+				{
+					// Must pass in an Instance op and ExtensionData op
+					check(Args.Instance);
+					check(Args.ExtensionData);
+
+					TArray<FScheduledOp> Dependencies;
+					Dependencies.Emplace(Args.Instance, item);
+					Dependencies.Emplace(Args.ExtensionData, item);
+
+					AddOp(FScheduledOp(item.At, item, 1), Dependencies);
+
+					break;
+				}
+
+				case 1:
+				{
+					// Assemble result
+					InstancePtrConst InstanceOpResult = LoadInstance(FCacheAddress(Args.Instance, item));
+
+					InstancePtr Result = mu::CloneOrTakeOver<Instance>(InstanceOpResult.get());
+
+					if (ExtensionDataPtrConst ExtensionData = LoadExtensionData(FCacheAddress(Args.ExtensionData, item)))
+					{
+						const OP::ADDRESS NameAddress = Args.ExtensionDataName;
+						check(NameAddress < (uint32)Program.m_constantStrings.Num());
+						const FString& NameString = Program.m_constantStrings[NameAddress];
+
+						Result->GetPrivate()->AddExtensionData(ExtensionData, FName(NameString) );
+					}
+
+					StoreInstance(item, Result);
+					break;
+				}
+
+				default:
+					check(false);
+			}
+			
+			break;
+		}
+
+        default:
+                check(false);
+        }
+    }
+
+
+    //---------------------------------------------------------------------------------------------
+    void CodeRunner::RunCode_InstanceAddResource(const FScheduledOp& item, const TSharedPtr<const Model>& InModel, const Parameters* InParams )
+    {
+		MUTABLE_CPUPROFILER_SCOPE(RunCode_InstanceAddResource);
+
+		if (!InModel || !m_pSystem)
+		{
+			return;
+		}
+
+		const FProgram& Program = m_pModel->GetPrivate()->m_program;
+
+		OP_TYPE type = Program.GetOpType(item.At);
+        switch (type)
+        {
+
+        case OP_TYPE::IN_ADDMESH:
+        {
+			OP::InstanceAddArgs args = Program.GetOpArgs<OP::InstanceAddArgs>(item.At);
+            switch (item.Stage)
+            {
+            case 0:
+				// We don't build the resources when building instance: just store ids for them.
+				AddOp( FScheduledOp( item.At, item, 1), FScheduledOp( args.instance, item) );
+                break;
+
+            case 1:
+            {
+                InstancePtrConst pBase = LoadInstance( FCacheAddress(args.instance,item) );
+                InstancePtr pResult;
+                if (!pBase)
+                {
+                    pResult = new Instance();
+                }
+                else
+                {
+					pResult = mu::CloneOrTakeOver<Instance>(pBase.get());
+				}
+
+                if ( args.value )
+                {
+					FResourceID MeshId = m_pSystem->WorkingMemoryManager.GetResourceKey(InModel,InParams,args.relevantParametersListIndex, args.value);
+					OP::ADDRESS NameAd = args.name;
+					check(NameAd < (uint32)Program.m_constantStrings.Num());
+					const FString& Name = Program.m_constantStrings[NameAd];
+					pResult->GetPrivate()->SetMesh(0, 0, MeshId, FName(Name));
+                }
+                StoreInstance( item, pResult );
+                break;
+            }
+
+            default:
+                check(false);
+            }
+            break;
+        }
+
+        case OP_TYPE::IN_ADDIMAGE:
+        {
+			OP::InstanceAddArgs args = Program.GetOpArgs<OP::InstanceAddArgs>(item.At);
+            switch (item.Stage)
+            {
+            case 0:
+				// We don't build the resources when building instance: just store ids for them.
+				AddOp( FScheduledOp( item.At, item, 1), FScheduledOp( args.instance, item) );
+                break;
+
+            case 1:
+            {
+                InstancePtrConst pBase = LoadInstance( FCacheAddress(args.instance,item) );
+                InstancePtr pResult;
+                if (!pBase)
+                {
+                    pResult = new Instance();
+                }
+                else
+                {
+					pResult = mu::CloneOrTakeOver<Instance>(pBase.get());
+				}
+
+                if ( args.value )
+                {
+					FResourceID ImageId = m_pSystem->WorkingMemoryManager.GetResourceKey(InModel, InParams, args.relevantParametersListIndex, args.value);
+					OP::ADDRESS NameAd = args.name;
+					check(NameAd < (uint32)Program.m_constantStrings.Num());
+					const FString& Name = Program.m_constantStrings[NameAd];
+					pResult->GetPrivate()->AddImage(0, 0, 0, ImageId, FName(Name) );
+                }
+                StoreInstance( item, pResult );
+                break;
+            }
+
+            default:
+                check(false);
+            }
+
+            break;
+        }
+
+        default:
+			check(false);
+        }
+    }
+
+
+    //---------------------------------------------------------------------------------------------
+    bool CodeRunner::RunCode_ConstantResource(const FScheduledOp& item, const Model* pModel )
+    {
+		MUTABLE_CPUPROFILER_SCOPE(RunCode_Constant);
+
+		const FProgram& Program = m_pModel->GetPrivate()->m_program;
+
+		OP_TYPE type = Program.GetOpType(item.At);
+        switch (type)
+        {
+
+        case OP_TYPE::ME_CONSTANT:
+        {
+			OP::MeshConstantArgs args = Program.GetOpArgs<OP::MeshConstantArgs>(item.At);
+
+            OP::ADDRESS cat = args.value;
+
+            // Assume the ROM has been loaded previously
+            check(Program.ConstantMeshes[cat].Value)
+
+            Ptr<const Mesh> SourceConst;
+			Program.GetConstant(cat, SourceConst);
+
+			check(SourceConst);
+			Ptr<Mesh> Source = CreateMesh(SourceConst->GetDataSize());
+			Source->CopyFrom(*SourceConst);
+
+            // Set the separate skeleton if necessary
+            if (args.skeleton >= 0)
+            {
+                check(Program.m_constantSkeletons.Num() > size_t(args.skeleton));
+                Ptr<const Skeleton> pSkeleton = Program.m_constantSkeletons[args.skeleton];
+                Source->SetSkeleton(pSkeleton);
+            }
+
+			if (args.physicsBody >= 0)
+			{
+                check(Program.m_constantPhysicsBodies.Num() > size_t(args.physicsBody));
+                Ptr<const PhysicsBody> pPhysicsBody = Program.m_constantPhysicsBodies[args.physicsBody];
+                Source->SetPhysicsBody(pPhysicsBody);
+			}
+
+            StoreMesh(item, Source);
+			//UE_LOG(LogMutableCore, Log, TEXT("Set mesh constant %d."), item.At);
+            break;
+        }
+
+        case OP_TYPE::IM_CONSTANT:
+        {
+			OP::ResourceConstantArgs args = Program.GetOpArgs<OP::ResourceConstantArgs>(item.At);
+            OP::ADDRESS cat = args.value;
+
+			int32 MipsToSkip = item.ExecutionOptions;
+            Ptr<const Image> Source;
+			Program.GetConstant(cat, Source, MipsToSkip, [this](int32 x, int32 y, int32 m, EImageFormat f, EInitializationType i)
+				{
+					return CreateImage(x, y, m, f, i);
+				});
+
+			// Assume the ROM has been loaded previously in a task generated at IssueOp
+			if (!Source)
+			{
+				return false;
+			}
+
+            StoreImage( item, Source );
+			//UE_LOG(LogMutableCore, Log, TEXT("Set image constant %d."), item.At);
+			break;
+        }
+
+		case OP_TYPE::ED_CONSTANT:
+		{
+			OP::ResourceConstantArgs Args = Program.GetOpArgs<OP::ResourceConstantArgs>(item.At);
+
+			// Assume the ROM has been loaded previously
+			ExtensionDataPtrConst SourceConst;
+			Program.GetExtensionDataConstant(Args.value, SourceConst);
+
+			check(SourceConst);
+
+            StoreExtensionData(item, SourceConst);
+            break;
+		}
+
+        default:
+            if (type!=OP_TYPE::NONE)
+            {
+                // Operation not implemented
+                check( false );
+            }
+            break;
+        }
+
+		// Success
+		return true;
+    }
+
+    //---------------------------------------------------------------------------------------------
+	void CodeRunner::RunCode_Mesh(const FScheduledOp& item, const Model* pModel )
+    {
+		MUTABLE_CPUPROFILER_SCOPE(RunCode_Mesh);
+
+		const FProgram& Program = m_pModel->GetPrivate()->m_program;
+
+		OP_TYPE type = Program.GetOpType(item.At);
+
+        switch (type)
+        {
+
+		case OP_TYPE::ME_REFERENCE:
+		{
+			OP::ResourceReferenceArgs Args = Program.GetOpArgs<OP::ResourceReferenceArgs>(item.At);
+			switch (item.Stage)
+			{
+			case 0:
+			{
+				Ptr<Mesh> Result;
+				if (Args.ForceLoad)
+				{
+					// This should never be reached because it should have been caught as a Task in IssueOp
+					check(false);
+				}
+				else
+				{
+					Result = Mesh::CreateAsReference(Args.ID, false);
+				}
+				StoreMesh(item, Result);
+				break;
+			}
+
+			default:
+				check(false);
+			}
+
+			break;
+		}
+
+        case OP_TYPE::ME_APPLYLAYOUT:
+        {
+			OP::MeshApplyLayoutArgs args = Program.GetOpArgs<OP::MeshApplyLayoutArgs>(item.At);
+			switch (item.Stage)
+			{
+			case 0:
+			{
+				AddOp(FScheduledOp(item.At, item, 1),
+					FScheduledOp(args.mesh, item),
+					FScheduledOp(args.layout, item));
+				break;
+			}
+            case 1:
+            {
+           		MUTABLE_CPUPROFILER_SCOPE(ME_APPLYLAYOUT)
+            		
+                Ptr<const Mesh> pBase = LoadMesh(FCacheAddress(args.mesh, item));
+
+                if (pBase)
+                {
+					Ptr<Mesh> Result = CloneOrTakeOver(pBase);
+
+                    Ptr<const Layout> pLayout = LoadLayout(FCacheAddress(args.layout, item));
+                    int texCoordsSet = args.channel;
+
+                    MeshApplyLayout(Result.get(), pLayout.get(), texCoordsSet);
+					
+					StoreMesh(item, Result);
+                }
+				else
+				{
+					StoreMesh(item, nullptr);
+				}
+
+
+                break;
+            }
+
+            default:
+                check(false);
+            }
+
+            break;
+        }
+
+        case OP_TYPE::ME_DIFFERENCE:
+        {
+			const uint8* data = Program.GetOpArgsPointer(item.At);
+
+			OP::ADDRESS BaseAt = 0;
+			FMemory::Memcpy(&BaseAt, data, sizeof(OP::ADDRESS)); 
+			data += sizeof(OP::ADDRESS);
+
+			OP::ADDRESS TargetAt = 0;
+			FMemory::Memcpy(&TargetAt, data, sizeof(OP::ADDRESS)); 
+			data += sizeof(OP::ADDRESS);
+
+            switch (item.Stage)
+            {
+            case 0:
+			{
+				if (BaseAt && TargetAt)
+				{
+					AddOp(FScheduledOp(item.At, item, 1),
+						FScheduledOp(BaseAt, item),
+						FScheduledOp(TargetAt, item));
+				}
+				else
+				{
+					StoreMesh(item, nullptr);
+				}
+				break;
+			}
+            case 1:
+            {
+       	        MUTABLE_CPUPROFILER_SCOPE(ME_DIFFERENCE)
+
+                Ptr<const Mesh> pBase = LoadMesh(FCacheAddress(BaseAt,item));
+                Ptr<const Mesh> pTarget = LoadMesh(FCacheAddress(TargetAt,item));
+
+				TArray<EMeshBufferSemantic, TInlineAllocator<8>> Semantics;
+				TArray<int32, TInlineAllocator<8>> SemanticIndices;
+
+				uint8 bIgnoreTextureCoords = 0;
+				FMemory::Memcpy(&bIgnoreTextureCoords, data, sizeof(uint8)); 
+				data += sizeof(uint8);
+
+				uint8 NumChannels = 0;
+				FMemory::Memcpy(&NumChannels, data, sizeof(uint8)); 
+				data += sizeof(uint8);
+
+                for (uint8 i = 0; i < NumChannels; ++i)
+                {
+					uint8 Semantic = 0;
+					FMemory::Memcpy(&Semantic, data, sizeof(uint8)); 
+					data += sizeof(uint8);
+					
+					uint8 SemanticIndex = 0;
+					FMemory::Memcpy(&SemanticIndex, data, sizeof(uint8)); 
+					data += sizeof(uint8);
+
+					Semantics.Add(EMeshBufferSemantic(Semantic));
+					SemanticIndices.Add(SemanticIndex);
+                }
+
+				Ptr<Mesh> Result = CreateMesh();
+				bool bOutSuccess = false;
+                MeshDifference(Result.get(), pBase.get(), pTarget.get(),
+                               NumChannels, Semantics.GetData(), SemanticIndices.GetData(),
+                               bIgnoreTextureCoords != 0, bOutSuccess);
+				Release(pBase);
+				Release(pTarget);
+
+                StoreMesh(item, Result);
+                break;
+            }
+
+            default:
+                check(false);
+            }
+
+            break;
+        }
+
+        case OP_TYPE::ME_MORPH:
+        {
+			const uint8* Data = Program.GetOpArgsPointer(item.At);
+
+			OP::ADDRESS FactorAt = 0;
+			FMemory::Memcpy(&FactorAt, Data, sizeof(OP::ADDRESS)); 
+			Data += sizeof(OP::ADDRESS);
+			
+			OP::ADDRESS BaseAt = 0;
+			FMemory::Memcpy(&BaseAt, Data, sizeof(OP::ADDRESS)); 
+			Data += sizeof(OP::ADDRESS);
+
+			OP::ADDRESS TargetAt = 0;
+			FMemory::Memcpy(&TargetAt, Data, sizeof(OP::ADDRESS)); 
+			Data += sizeof(OP::ADDRESS);
+
+			switch (item.Stage)
+            {
+            case 0:
+                if (BaseAt)
+                {
+                    AddOp(FScheduledOp(item.At, item, 1),
+                           FScheduledOp(FactorAt, item));
+                }
+                else
+                {
+                    StoreMesh(item, nullptr);
+                }
+                break;
+
+            case 1:
+            {
+                MUTABLE_CPUPROFILER_SCOPE(ME_MORPH_1)
+
+                float Factor = LoadScalar(FCacheAddress(FactorAt, item));
+
+                // Factor goes from -1 to 1 across all targets. [0 - 1] represents positive morphs, while [-1, 0) represent negative morphs.
+				Factor = FMath::Clamp(Factor, -1.0f, 1.0f); // Is the factor not in range [-1, 1], it will index a non existing morph.
+
+                FScheduledOpData HeapData;
+				HeapData.Interpolate.Bifactor = Factor;
+				uint32 DataAddress = uint32(m_heapData.Add(HeapData));
+
+                // No morph
+				if (FMath::IsNearlyZero(Factor))
+                {                        
+                    AddOp(FScheduledOp(item.At, item, 2, DataAddress),
+						FScheduledOp(BaseAt, item));
+                }
+                // The Morph, partial or full
+                else
+                {
+                    // We will need the base again
+                    AddOp(FScheduledOp(item.At, item, 2, DataAddress),
+						FScheduledOp(BaseAt, item),
+						FScheduledOp(TargetAt, item));
+                }
+
+                break;
+            }
+
+            case 2:
+            {
+       		    MUTABLE_CPUPROFILER_SCOPE(ME_MORPH_2)
+
+                Ptr<const Mesh> BaseMesh = LoadMesh(FCacheAddress(BaseAt, item));
+
+                // Factor from 0 to 1 between the two targets
+                const FScheduledOpData& HeapData = m_heapData[item.CustomState];
+                float Factor = HeapData.Interpolate.Bifactor;
+
+                if (BaseMesh)
+                {
+					// No morph
+					if (FMath::IsNearlyZero(Factor))
+                    {
+						StoreMesh(item, BaseMesh);
+                    }
+					// The Morph, partial or full
+                    else 
+                    {
+                        Ptr<const Mesh> MorphMesh = LoadMesh(FCacheAddress(TargetAt, item));
+						
+						if (MorphMesh)
+						{
+							Ptr<Mesh> Result = CloneOrTakeOver(BaseMesh);
+							MeshMorph(Result.get(), MorphMesh.get(), Factor);
+
+							Release(MorphMesh);
+							StoreMesh(item, Result);
+						}
+						else
+						{
+							StoreMesh(item, BaseMesh);
+						}
+                    }
+                }
+				else
+				{
+					StoreMesh(item, nullptr);
+				}
+
+                break;
+            }
+
+            default:
+                check(false);
+            }
+
+            break;
+        }
+
+        case OP_TYPE::ME_MERGE:
+        {
+			OP::MeshMergeArgs args = Program.GetOpArgs<OP::MeshMergeArgs>(item.At);
+            switch (item.Stage)
+            {
+            case 0:
+			{
+				AddOp(FScheduledOp(item.At, item, 1),
+					FScheduledOp(args.base, item),
+					FScheduledOp(args.added, item));
+				break;
+			}
+            case 1:
+            {
+				MUTABLE_CPUPROFILER_SCOPE(ME_MERGE_1)
+
+                Ptr<const Mesh> pA = LoadMesh(FCacheAddress(args.base, item));
+                Ptr<const Mesh> pB = LoadMesh(FCacheAddress(args.added, item));
+
+                if (pA && pB && pA->GetVertexCount() && pB->GetVertexCount())
+                {
+					check(!pA->IsReference() && !pB->IsReference());
+
+					FMeshMergeScratchMeshes Scratch;
+					Scratch.FirstReformat = CreateMesh();
+					Scratch.SecondReformat = CreateMesh();
+
+					Ptr<Mesh> Result = CreateMesh(pA->GetDataSize() + pB->GetDataSize());
+
+					MeshMerge(Result.get(), pA.get(), pB.get(), !args.newSurfaceID, Scratch);
+
+					Release(Scratch.FirstReformat);
+					Release(Scratch.SecondReformat);
+
+                    if (args.newSurfaceID)
+                    {
+						check(pB->GetSurfaceCount() == 1);
+						Result->Surfaces.Last().Id = args.newSurfaceID;
+                    }
+
+					Release(pA);
+					Release(pB);
+					StoreMesh(item, Result);
+                }
+                else if (pA && (pA->GetVertexCount() || pA->IsReference()))
+                {
+					Release(pB);
+					StoreMesh(item, pA);
+                }
+                else if (pB && (pB->GetVertexCount() || pB->IsReference()))
+                {
+					Ptr<Mesh> Result = CloneOrTakeOver(pB);
+
+                    check(Result->IsReference() || (Result->GetSurfaceCount() == 1) );
+
+                    if (Result->GetSurfaceCount() > 0 && args.newSurfaceID)
+                    {
+                        Result->Surfaces.Last().Id = args.newSurfaceID;
+                    }
+
+					Release(pA);
+					StoreMesh(item, Result);
+                }
+                else
+                {
+					Release(pA);
+					Release(pB);
+					StoreMesh(item, CreateMesh());
+                }
+
+                break;
+            }
+
+            default:
+                check(false);
+            }
+
+            break;
+        }
+
+        case OP_TYPE::ME_INTERPOLATE:
+        {
+			OP::MeshInterpolateArgs Args = Program.GetOpArgs<OP::MeshInterpolateArgs>(item.At);
+            switch (item.Stage)
+            {
+            case 0:
+			{
+				if (Args.base)
+				{
+					AddOp(FScheduledOp(item.At, item, 1),
+						FScheduledOp(Args.base, item),
+						FScheduledOp(Args.factor, item));
+				}
+				else
+				{
+					StoreMesh(item, nullptr);
+				}
+				break;
+			}
+            case 1:
+            {
+            	MUTABLE_CPUPROFILER_SCOPE(ME_INTERPOLATE_1)
+
+                int32 Count = 1;
+                for (int32 TargetIndex = 0; TargetIndex < MUTABLE_OP_MAX_INTERPOLATE_COUNT-1 && Args.targets[TargetIndex]; ++TargetIndex)
+                {
+                    ++Count;
+                }
+
+                // Factor from 0 to 1 across all targets
+                float Factor = LoadScalar(FCacheAddress(Args.factor, item));
+
+                float Delta = 1.0f/(Count-1);
+                int32 Min = FMath::FloorToInt(Factor/Delta);
+                int32 Max = FMath::CeilToInt(Factor/Delta);
+
+                // Factor from 0 to 1 between the two targets
+                float Bifactor = Factor/Delta - Min;
+
+                FScheduledOpData Data;
+                Data.Interpolate.Bifactor = Bifactor;
+				Data.Interpolate.Min = FMath::Clamp(Min, 0, Count - 1);
+				Data.Interpolate.Max = FMath::Clamp(Max, 0, Count - 1);
+				uint32 DataAddress = uint32(m_heapData.Num());
+
+                // Just the first of the targets
+                if (Bifactor < UE_SMALL_NUMBER)
+                {
+                    if (Min == 0)
+                    {
+                        // Just the base
+						Ptr<const Mesh> BaseMesh = LoadMesh(FCacheAddress(Args.base, item));
+						StoreMesh(item, BaseMesh);
+					}
+                    else
+                    {
+                        // Base with one full morph
+                        m_heapData.Add(Data);
+						AddOp(FScheduledOp(item.At, item, 2, DataAddress),
+							FScheduledOp(Args.base, item),
+							FScheduledOp(Args.targets[Min-1], item));
+					}
+				}
+                // Just the second of the targets
+                else if (Bifactor > 1.0f-UE_SMALL_NUMBER)
+                {
+                    m_heapData.Add(Data);
+					AddOp(FScheduledOp(item.At, item, 2, DataAddress),
+						FScheduledOp(Args.base, item),
+						FScheduledOp(Args.targets[Max-1], item));
+				}
+                // Mix the first target on the base
+                else if (Min == 0)
+                {
+                    m_heapData.Add(Data);
+					AddOp(FScheduledOp(item.At, item, 2, DataAddress),
+						FScheduledOp(Args.base, item),
+						FScheduledOp(Args.targets[0], item));
+				}
+                // Mix two targets on the base
+                else
+                {
+                    m_heapData.Add(Data);
+					AddOp(FScheduledOp(item.At, item, 2, DataAddress),
+						FScheduledOp(Args.base, item),
+						FScheduledOp(Args.targets[Min-1], item),
+						FScheduledOp(Args.targets[Max-1], item));
+				}
+
+                break;
+            }
+
+            case 2:
+            {
+				MUTABLE_CPUPROFILER_SCOPE(ME_INTERPOLATE_2)
+
+                int32 Count = 1;
+                for (int32 TargetIndex = 0; TargetIndex < MUTABLE_OP_MAX_INTERPOLATE_COUNT-1 && Args.targets[TargetIndex]; ++TargetIndex)
+                {
+                    ++Count;
+                }
+
+                const FScheduledOpData& Data = m_heapData[ (size_t)item.CustomState ];
+
+                // Factor from 0 to 1 between the two targets
+                float Bifactor = Data.Interpolate.Bifactor;
+                int32 Min = Data.Interpolate.Min;
+                int32 Max = Data.Interpolate.Max;
+
+                Ptr<const Mesh> BaseMesh = LoadMesh(FCacheAddress(Args.base, item));
+
+                if (BaseMesh)
+                {
+                    // Just the first of the targets
+                    if (Bifactor < UE_SMALL_NUMBER)
+                    {
+                        if (Min == 0)
+                        {
+                            // Just the base. It should have been dealt with in the previous stage.
+                            check(false);
+                        }
+                        else
+                        {
+                            // Base with one full morph
+                            Ptr<const Mesh> MorphMesh = LoadMesh(FCacheAddress(Args.targets[Min-1], item));
+						
+							if (MorphMesh)
+							{
+								Ptr<Mesh> Result = CloneOrTakeOver(BaseMesh);
+
+								MeshMorph(Result.get(), MorphMesh.get());
+							
+								StoreMesh(item, Result);
+							}
+							else
+							{
+								StoreMesh(item, BaseMesh);
+							}
+							
+							Release(MorphMesh);
+                        }
+                    }
+                    // Just the second of the targets
+                    else if (Bifactor > 1.0f-UE_SMALL_NUMBER)
+                    {
+                        check(Max > 0);
+                        Ptr<const Mesh> MorphMesh = LoadMesh(FCacheAddress(Args.targets[Max-1], item));
+
+                        if (MorphMesh)
+                        {
+							Ptr<Mesh> Result = CloneOrTakeOver(BaseMesh);
+							
+                            MeshMorph(Result.get(), MorphMesh.get());
+
+							StoreMesh(item, Result);
+                        }
+                        else
+                        {
+							StoreMesh(item, BaseMesh);
+                        }
+						
+						Release(MorphMesh);
+                    }
+                    // Mix the first target on the base
+                    else if (Min == 0)
+                    {
+                        Ptr<const Mesh> MorphMesh = LoadMesh(FCacheAddress(Args.targets[0], item));
+                        if (MorphMesh)
+                        {
+							Ptr<Mesh> Result = CloneOrTakeOver(BaseMesh);
+
+                            MeshMorph(Result.get(), MorphMesh.get(), Bifactor);
+
+							StoreMesh(item, Result);
+                        }
+                        else
+                        {
+							StoreMesh(item, BaseMesh);
+                        }
+						
+						Release(MorphMesh);
+                    }
+                    // Mix two targets on the base
+                    else
+                    {
+                        Ptr<const Mesh> MinMesh = LoadMesh(FCacheAddress(Args.targets[Min-1], item));
+                        Ptr<const Mesh> MaxMesh = LoadMesh(FCacheAddress(Args.targets[Max-1], item));
+
+                        if (MinMesh && MaxMesh)
+                        {
+							Ptr<Mesh> Result = CloneOrTakeOver(BaseMesh);
+
+                            MeshMorph2(Result.get(), MinMesh.get(), MaxMesh.get(), Bifactor);
+
+							StoreMesh(item, Result);
+                        }
+                        else
+                        {
+							StoreMesh(item, BaseMesh);
+                        }
+
+						Release(MinMesh);
+						Release(MaxMesh);
+                    }
+                }
+				else
+				{
+					StoreMesh(item, BaseMesh);
+				}
+
+                break;
+            }
+
+            default:
+                check(false);
+            }
+
+            break;
+        }
+
+        case OP_TYPE::ME_MASKCLIPMESH:
+        {
+			OP::MeshMaskClipMeshArgs args = Program.GetOpArgs<OP::MeshMaskClipMeshArgs>(item.At);
+            switch (item.Stage)
+            {
+            case 0:
+			{
+				AddOp(FScheduledOp(item.At, item, 1),
+					FScheduledOp(args.source, item),
+					FScheduledOp(args.clip, item));
+				break;
+			}
+            case 1:
+            {
+            	MUTABLE_CPUPROFILER_SCOPE(ME_MASKCLIPMESH_1)
+            		
+                Ptr<const Mesh> Source = LoadMesh(FCacheAddress(args.source, item));
+                Ptr<const Mesh> pClip = LoadMesh(FCacheAddress(args.clip, item));
+
+                // Only if both are valid.
+                if (Source.get() && pClip.get())
+                {
+					Ptr<Mesh> Result = CreateMesh();
+
+					bool bOutSuccess = false;
+                    MeshMaskClipMesh(Result.get(), Source.get(), pClip.get(), bOutSuccess);
+					
+					Release(Source);
+					Release(pClip);
+					if (!bOutSuccess)
+					{
+						Release(Result);
+						StoreMesh(item, nullptr);
+					}
+					else
+					{
+						StoreMesh(item, Result);
+					}
+                }
+				else
+				{
+					Release(Source);
+					Release(pClip);
+					StoreMesh(item, nullptr);
+				}
+
+                break;
+            }
+
+            default:
+                check(false);
+            }
+
+            break;
+        }
+
+		case OP_TYPE::ME_MASKCLIPUVMASK:
+		{
+			OP::MeshMaskClipUVMaskArgs args = Program.GetOpArgs<OP::MeshMaskClipUVMaskArgs>(item.At);
+			switch (item.Stage)
+			{
+			case 0:
+			{
+				AddOp(FScheduledOp(item.At, item, 1),
+					FScheduledOp(args.Source, item),
+					FScheduledOp(args.UVSource, item),
+					FScheduledOp(args.MaskImage, item),
+					FScheduledOp(args.MaskLayout, item));
+				break;
+			}
+			case 1:
+			{
+				MUTABLE_CPUPROFILER_SCOPE(ME_MASKCLIPUVMASK_1);
+
+				Ptr<const Mesh> Source = LoadMesh(FCacheAddress(args.Source, item));
+				Ptr<const Mesh> UVSource = LoadMesh(FCacheAddress(args.UVSource, item));
+				Ptr<const Image> MaskImage = LoadImage(FCacheAddress(args.MaskImage, item));
+				Ptr<const Layout> MaskLayout = LoadLayout(FCacheAddress(args.MaskLayout, item));
+
+				// Only if both are valid.
+				if (Source.get() && MaskImage.get())
+				{
+					Ptr<Mesh> Result = CreateMesh();
+
+					bool bOutSuccess = false;
+					MakeMeshMaskFromUVMask(Result.get(), Source.get(), UVSource.get(), MaskImage.get(), args.LayoutIndex, bOutSuccess);
+
+					Release(Source);
+					Release(UVSource);
+					Release(MaskImage);
+					if (!bOutSuccess)
+					{
+						Release(Result);
+						StoreMesh(item, nullptr);
+					}
+					else
+					{
+						StoreMesh(item, Result);
+					}
+				}
+				else if (Source.get() && MaskLayout.get())
+				{					
+					Ptr<Mesh> Result = CreateMesh();
+
+					bool bOutSuccess = false;
+					MakeMeshMaskFromLayout(Result.get(), Source.get(), UVSource.get(), MaskLayout.get(), args.LayoutIndex, bOutSuccess);
+
+					Release(Source);
+					Release(UVSource);
+					Release(MaskImage);
+					if (!bOutSuccess)
+					{
+						Release(Result);
+						StoreMesh(item, nullptr);
+					}
+					else
+					{
+						StoreMesh(item, Result);
+					}
+				}
+				else
+				{
+					Release(Source);
+					Release(UVSource);
+					Release(MaskImage);
+					StoreMesh(item, nullptr);
+				}
+
+				break;
+			}
+
+			default:
+				check(false);
+			}
+
+			break;
+		}
+
+        case OP_TYPE::ME_MASKDIFF:
+        {
+			OP::MeshMaskDiffArgs args = Program.GetOpArgs<OP::MeshMaskDiffArgs>(item.At);
+            switch (item.Stage)
+            {
+            case 0:
+			{
+				AddOp(FScheduledOp(item.At, item, 1),
+					FScheduledOp(args.source, item),
+					FScheduledOp(args.fragment, item));
+				break;
+			}
+            case 1:
+            {
+           		MUTABLE_CPUPROFILER_SCOPE(ME_MASKDIFF_1)
+            		
+                Ptr<const Mesh> Source = LoadMesh(FCacheAddress(args.source, item));
+                Ptr<const Mesh> pClip = LoadMesh(FCacheAddress(args.fragment, item));
+
+                // Only if both are valid.
+                if (Source.get() && pClip.get())
+                {
+					Ptr<Mesh> Result = CreateMesh();
+
+					bool bOutSuccess = false;
+                    MeshMaskDiff(Result.get(), Source.get(), pClip.get(), bOutSuccess);
+
+					Release(Source);
+					Release(pClip);
+
+					if (!bOutSuccess)
+					{
+						Release(Result);
+						StoreMesh(item, nullptr);
+					}
+					else
+					{
+						StoreMesh(item, Result);
+					}
+                }
+				else
+				{
+					Release(Source);
+					Release(pClip);
+					StoreMesh(item, nullptr);
+				}
+
+                break;
+            }
+
+            default:
+                check(false);
+            }
+
+            break;
+        }
+
+        case OP_TYPE::ME_FORMAT:
+        {
+			OP::MeshFormatArgs args = Program.GetOpArgs<OP::MeshFormatArgs>(item.At);
+            switch (item.Stage)
+            {
+            case 0:
+			{
+				if (args.source && args.format)
+				{
+					AddOp(FScheduledOp(item.At, item, 1),
+						FScheduledOp(args.source, item),
+						FScheduledOp(args.format, item));
+				}
+				else
+				{
+					StoreMesh(item, nullptr);
+				}
+				break;
+			}
+            case 1:
+            {
+            	MUTABLE_CPUPROFILER_SCOPE(ME_FORMAT_1)
+            		
+                Ptr<const Mesh> Source = LoadMesh(FCacheAddress(args.source,item));
+                Ptr<const Mesh> Format = LoadMesh(FCacheAddress(args.format,item));
+
+				if (Source && Source->IsReference())
+				{
+					Release(Format);
+					StoreMesh(item, Source);
+				}
+				else if (Source)
+				{
+					uint8 Flags = args.Flags;
+					if (!Format && !(Flags & OP::MeshFormatArgs::ResetBufferIndices))
+					{
+						StoreMesh(item, Source);
+					}
+					else if (!Format)
+					{
+						Ptr<Mesh> Result = CloneOrTakeOver(Source);
+
+						if (Flags & OP::MeshFormatArgs::ResetBufferIndices)
+						{
+							Result->ResetBufferIndices();
+						}
+
+						StoreMesh(item, Result);
+					}
+					else
+					{
+						Ptr<Mesh> Result = CreateMesh();
+
+						bool bOutSuccess = false;
+						MeshFormat(Result.get(), Source.get(), Format.get(),
+							true,
+							(Flags & OP::MeshFormatArgs::Vertex) != 0,
+							(Flags & OP::MeshFormatArgs::Index) != 0,
+							(Flags & OP::MeshFormatArgs::IgnoreMissing) != 0,
+							bOutSuccess);
+
+						check(bOutSuccess);
+
+						if (Flags & OP::MeshFormatArgs::ResetBufferIndices)
+						{
+							Result->ResetBufferIndices();
+						}
+
+						if (Flags & OP::MeshFormatArgs::OptimizeBuffers)
+						{
+							MUTABLE_CPUPROFILER_SCOPE(MeshOptimizeBuffers)
+							MeshOptimizeBuffers(Result.get());
+						}
+
+						Release(Source);
+						Release(Format);
+						StoreMesh(item, Result);
+					}
+				}
+				else
+				{
+					Release(Format);
+					StoreMesh(item, nullptr);
+				}
+                break;
+            }
+
+            default:
+                check(false);
+            }
+
+            break;
+        }
+
+        case OP_TYPE::ME_EXTRACTLAYOUTBLOCK:
+        {
+            const uint8* data = Program.GetOpArgsPointer(item.At);
+
+            OP::ADDRESS source;
+            FMemory::Memcpy( &source, data, sizeof(OP::ADDRESS) );
+            data += sizeof(OP::ADDRESS);
+
+            uint16 layout;
+			FMemory::Memcpy( &layout, data, sizeof(uint16) );
+            data += sizeof(uint16);
+
+            uint16 blockCount;
+			FMemory::Memcpy( &blockCount, data, sizeof(uint16) );
+            data += sizeof(uint16);
+
+            switch (item.Stage)
+            {
+            case 0:
+			{
+				if (source)
+				{
+					AddOp(FScheduledOp(item.At, item, 1),
+						FScheduledOp(source, item));
+				}
+				else
+				{
+					StoreMesh(item, nullptr);
+				}
+				break;
+			}
+            case 1:
+            {
+				MUTABLE_CPUPROFILER_SCOPE(ME_EXTRACTLAYOUTBLOCK_1)
+
+                Ptr<const Mesh> Source = LoadMesh(FCacheAddress(source, item));
+
+                // Access with memcpy necessary for unaligned memory access issues.
+                uint64 blocks[512];
+				check(blockCount< 512);
+				FMemory::Memcpy(blocks, data, sizeof(uint64)*FMath::Min(512,int32(blockCount)));
+
+				if (Source)
+				{
+					Ptr<Mesh> Result = CreateMesh();
+					bool bOutSuccess;
+
+					if (blockCount > 0)
+					{
+						MeshExtractLayoutBlock(Result.get(), Source.get(), layout, blockCount, blocks, bOutSuccess);
+					}
+					else
+					{
+						MeshExtractLayoutBlock(Result.get(), Source.get(), layout, bOutSuccess);
+					}
+
+					if (!bOutSuccess)
+					{
+						Release(Result);
+						StoreMesh(item, Source);
+					}
+					else
+					{
+						Release(Source);
+						StoreMesh(item, Result);
+					}
+				}
+				else
+				{
+					StoreMesh(item, nullptr);
+				}
+                break;
+            }
+
+            default:
+                check(false);
+            }
+
+            break;
+        }
+
+        case OP_TYPE::ME_TRANSFORM:
+        {
+			OP::MeshTransformArgs args = Program.GetOpArgs<OP::MeshTransformArgs>(item.At);
+            switch (item.Stage)
+            {
+            case 0:
+			{
+				if (args.source)
+				{
+					AddOp(FScheduledOp(item.At, item, 1),
+						FScheduledOp(args.source, item));
+				}
+				else
+				{
+					StoreMesh(item, nullptr);
+				}
+				break;
+			}
+            case 1:
+            {
+           		MUTABLE_CPUPROFILER_SCOPE(ME_TRANSFORM_1)
+            		
+                Ptr<const Mesh> Source = LoadMesh(FCacheAddress(args.source,item));
+
+                const FMatrix44f& mat = Program.m_constantMatrices[args.matrix];
+
+				Ptr<Mesh> Result = CreateMesh(Source ? Source->GetDataSize() : 0);
+
+				bool bOutSuccess = false;
+                MeshTransform(Result.get(), Source.get(), mat, bOutSuccess);
+
+				if (!bOutSuccess)
+				{
+					Release(Result);
+					StoreMesh(item, Source);
+				}
+				else
+				{
+					Release(Source);
+					StoreMesh(item, Result);
+				}
+                break;
+            }
+
+            default:
+                check(false);
+            }
+
+            break;
+        }
+
+        case OP_TYPE::ME_CLIPMORPHPLANE:
+        {
+			OP::MeshClipMorphPlaneArgs args = Program.GetOpArgs<OP::MeshClipMorphPlaneArgs>(item.At);
+            switch (item.Stage)
+            {
+            case 0:
+			{
+				if (args.source)
+				{
+					AddOp(FScheduledOp(item.At, item, 1),
+						FScheduledOp(args.source, item));
+				}
+				else
+				{
+					StoreMesh(item, nullptr);
+				}
+				break;
+			}
+            case 1:
+            {
+           		MUTABLE_CPUPROFILER_SCOPE(ME_CLIPMORPHPLANE_1)
+            		
+                Ptr<const Mesh> Source = LoadMesh(FCacheAddress(args.source, item));
+
+                check(args.morphShape < (uint32)pModel->GetPrivate()->m_program.m_constantShapes.Num());
+
+                // Should be an ellipse
+                const FShape& morphShape = Program.m_constantShapes[args.morphShape];
+
+                const FVector3f& origin = morphShape.position;
+                const FVector3f& normal = morphShape.up;
+
+                if (args.vertexSelectionType == OP::MeshClipMorphPlaneArgs::VS_SHAPE)
+                {
+                    check(args.vertexSelectionShapeOrBone < (uint32)pModel->GetPrivate()->m_program.m_constantShapes.Num());
+
+                    // Should be None or an axis aligned box
+                    const FShape& selectionShape = Program.m_constantShapes[args.vertexSelectionShapeOrBone];
+
+					Ptr<Mesh> Result = CreateMesh(Source ? Source->GetDataSize() : 0);
+
+					bool bOutSuccess = false;
+					MeshClipMorphPlane(Result.get(), Source.get(), origin, normal, args.dist, args.factor, morphShape.size[0], morphShape.size[1], morphShape.size[2], selectionShape, bOutSuccess, nullptr, -1);
+					
+					if (!bOutSuccess)
+					{
+						Release(Result);
+						StoreMesh(item, Source);
+					}
+					else
+					{
+						Release(Source);
+						StoreMesh(item, Result);
+					}
+                }
+
+				else if (args.vertexSelectionType == OP::MeshClipMorphPlaneArgs::VS_BONE_HIERARCHY)
+				{
+					FShape selectionShape;
+					selectionShape.type = (uint8)FShape::Type::None;
+
+					Ptr<Mesh> Result = CreateMesh(Source->GetDataSize());
+
+					check(args.vertexSelectionShapeOrBone <= MAX_uint32);
+					const FBoneName Bone(args.vertexSelectionShapeOrBone);
+
+					bool bOutSuccess = false;
+					MeshClipMorphPlane(Result.get(), Source.get(), origin, normal, args.dist, args.factor, morphShape.size[0], morphShape.size[1], morphShape.size[2], selectionShape, bOutSuccess, &Bone, args.maxBoneRadius);
+
+					if (!bOutSuccess)
+					{
+						Release(Result);
+						StoreMesh(item, Source);
+					}
+					else
+					{
+						Release(Source);
+						StoreMesh(item, Result);
+					}
+                }
+                else
+                {
+                    // No vertex selection
+                    FShape selectionShape;
+                    selectionShape.type = (uint8)FShape::Type::None;
+
+					Ptr<Mesh> Result = CreateMesh(Source ? Source->GetDataSize() : 0);
+
+					bool bOutSuccess = false;
+					MeshClipMorphPlane(Result.get(), Source.get(), origin, normal, args.dist, args.factor, morphShape.size[0], morphShape.size[1], morphShape.size[2], selectionShape, bOutSuccess, nullptr, -1.0f);
+
+					if (!bOutSuccess)
+					{
+						Release(Result);
+						StoreMesh(item, Source);
+					}
+					else
+					{
+						Release(Source);
+						StoreMesh(item, Result);
+					}
+                }
+
+                break;
+            }
+
+            default:
+                check(false);
+            }
+
+            break;
+        }
+
+
+        case OP_TYPE::ME_CLIPWITHMESH:
+        {
+			OP::MeshClipWithMeshArgs args = Program.GetOpArgs<OP::MeshClipWithMeshArgs>(item.At);
+            switch (item.Stage)
+            {
+            case 0:
+			{
+				if (args.source)
+				{
+					AddOp(FScheduledOp(item.At, item, 1),
+						FScheduledOp(args.source, item),
+						FScheduledOp(args.clipMesh, item));
+				}
+				else
+				{
+					StoreMesh(item, nullptr);
+				}
+
+				break;
+			}
+            case 1:
+            {
+				MUTABLE_CPUPROFILER_SCOPE(ME_CLIPWITHMESH_1)
+
+                Ptr<const Mesh> Source = LoadMesh(FCacheAddress(args.source, item));
+                Ptr<const Mesh> pClip = LoadMesh(FCacheAddress(args.clipMesh, item));
+
+                // Only if both are valid.
+                if (Source && pClip)
+                {
+					Ptr<Mesh> Result = CreateMesh(Source->GetDataSize());
+
+					bool bOutSuccess = false;
+                    MeshClipWithMesh(Result.get(), Source.get(), pClip.get(), bOutSuccess);
+
+					Release(pClip);
+					if (!bOutSuccess)
+					{
+						Release(Result);
+						StoreMesh(item, Source);
+					}
+					else
+					{
+						Release(Source);
+						StoreMesh(item, Result);
+					}
+                }
+                else
+                {
+					Release(pClip);
+					StoreMesh(item, Source);
+                }
+
+                break;
+            }
+
+            default:
+                check(false);
+            }
+
+            break;
+        }
+		case OP_TYPE::ME_CLIPDEFORM:
+		{
+			OP::MeshClipDeformArgs args = Program.GetOpArgs<OP::MeshClipDeformArgs>(item.At);
+			switch (item.Stage)
+			{
+			case 0:
+			{
+				if (args.mesh)
+				{
+					AddOp(FScheduledOp(item.At, item, 1),
+						FScheduledOp(args.mesh, item),
+						FScheduledOp(args.clipShape, item));
+				}
+				else
+				{
+					StoreMesh(item, nullptr);
+				}
+				break;
+			}
+			case 1:
+			{
+				MUTABLE_CPUPROFILER_SCOPE(ME_CLIPDEFORM_1)
+
+				Ptr<const Mesh> BaseMesh = LoadMesh(FCacheAddress(args.mesh, item));
+				Ptr<const Mesh> ClipShape = LoadMesh(FCacheAddress(args.clipShape, item));
+
+				if (BaseMesh && ClipShape)
+				{
+					Ptr<Mesh> Result = CreateMesh(BaseMesh->GetDataSize());
+
+					bool bOutSuccess = false;
+					MeshClipDeform(Result.get(), BaseMesh.get(), ClipShape.get(), args.clipWeightThreshold, bOutSuccess);
+
+					Release(ClipShape);
+
+					if (!bOutSuccess)
+					{
+						Release(Result);
+						StoreMesh(item, BaseMesh);
+					}
+					else
+					{
+						Release(BaseMesh);
+						StoreMesh(item, Result);
+					}
+				}
+				else
+				{
+					Release(ClipShape);
+					StoreMesh(item, BaseMesh);
+				}
+
+				break;
+			}
+
+			default:
+				check(false);
+			}
+
+			break;
+		}
+
+        case OP_TYPE::ME_APPLYPOSE:
+        {
+			OP::MeshApplyPoseArgs args = Program.GetOpArgs<OP::MeshApplyPoseArgs>(item.At);
+            switch (item.Stage)
+            {
+            case 0:
+			{
+				if (args.base)
+				{
+					AddOp(FScheduledOp(item.At, item, 1),
+						FScheduledOp(args.base, item),
+						FScheduledOp(args.pose, item));
+				}
+				else
+				{
+					StoreMesh(item, nullptr);
+				}
+				break;
+			}
+            case 1:
+            {
+          		MUTABLE_CPUPROFILER_SCOPE(ME_APPLYPOSE_1)
+
+                Ptr<const Mesh> pBase = LoadMesh(FCacheAddress(args.base, item));
+                Ptr<const Mesh> pPose = LoadMesh(FCacheAddress(args.pose, item));
+
+                // Only if both are valid.
+                if (pBase && pPose)
+                {
+					Ptr<Mesh> Result = CreateMesh(pBase->GetSkeleton() ? pBase->GetDataSize() : 0);
+
+					bool bOutSuccess = false;
+					MeshApplyPose(Result.get(), pBase.get(), pPose.get(), bOutSuccess);
+
+					Release(pPose);
+					if (!bOutSuccess)
+					{
+						Release(Result);
+						StoreMesh(item, pBase);
+					}
+					else
+					{
+						Release(pBase);
+						StoreMesh(item, Result);
+					}
+                }
+                else
+                {
+					Release(pPose);
+					StoreMesh(item, pBase);
+                }
+
+                break;
+            }
+
+            default:
+                check(false);
+            }
+
+            break;
+        }
+
+
+		case OP_TYPE::ME_GEOMETRYOPERATION:
+		{
+			OP::MeshGeometryOperationArgs args = Program.GetOpArgs<OP::MeshGeometryOperationArgs>(item.At);
+			switch (item.Stage)
+			{
+			case 0:
+			{
+				if (args.meshA)
+				{
+					AddOp(FScheduledOp(item.At, item, 1),
+						FScheduledOp(args.meshA, item),
+						FScheduledOp(args.meshB, item),
+						FScheduledOp(args.scalarA, item),
+						FScheduledOp(args.scalarB, item));
+				}
+				else
+				{
+					StoreMesh(item, nullptr);
+				}
+				break;
+			}
+			case 1:
+			{
+				MUTABLE_CPUPROFILER_SCOPE(ME_GEOMETRYOPERATION_1)
+				
+				Ptr<const Mesh> MeshA = LoadMesh(FCacheAddress(args.meshA, item));
+				Ptr<const Mesh> MeshB = LoadMesh(FCacheAddress(args.meshB, item));
+				float ScalarA = LoadScalar(FCacheAddress(args.scalarA, item));
+				float ScalarB = LoadScalar(FCacheAddress(args.scalarB, item));
+
+				Ptr<Mesh> Result = CreateMesh(MeshA ? MeshA->GetDataSize() : 0);
+
+				bool bOutSuccess = false;
+				MeshGeometryOperation(Result.get(), MeshA.get(), MeshB.get(), ScalarA, ScalarB, bOutSuccess);
+
+				Release(MeshA);
+				Release(MeshB);
+
+				if (!bOutSuccess)
+				{
+					Release(Result);
+					StoreMesh(item, nullptr);
+				}
+				else
+				{
+					StoreMesh(item, Result);
+				}
+
+				break;
+			}
+
+			default:
+				check(false);
+			}
+
+			break;
+		}
+
+
+		case OP_TYPE::ME_BINDSHAPE:
+		{
+			OP::MeshBindShapeArgs Args = Program.GetOpArgs<OP::MeshBindShapeArgs>(item.At);
+			const uint8* Data = Program.GetOpArgsPointer(item.At);
+
+			switch (item.Stage)
+			{
+			case 0:
+			{
+				if (Args.mesh)
+				{
+					AddOp(FScheduledOp(item.At, item, 1),
+						FScheduledOp(Args.mesh, item),
+						FScheduledOp(Args.shape, item));
+				}
+				else
+				{
+					StoreMesh(item, nullptr);
+				}
+				break;
+			}
+			case 1:
+			{
+				MUTABLE_CPUPROFILER_SCOPE(ME_BINDSHAPE_1)
+				Ptr<const Mesh> BaseMesh = LoadMesh(FCacheAddress(Args.mesh, item));
+				Ptr<const Mesh> Shape = LoadMesh(FCacheAddress(Args.shape, item));
+				
+				EShapeBindingMethod BindingMethod = static_cast<EShapeBindingMethod>(Args.bindingMethod); 
+
+				if (BindingMethod == EShapeBindingMethod::ReshapeClosestProject)
+				{ 
+					// Bones are stored after the args
+					Data += sizeof(Args);
+
+					// Rebuilding array of bone names ----
+					int32 NumBones;
+					FMemory::Memcpy(&NumBones, Data, sizeof(int32)); 
+					Data += sizeof(int32);
+					
+					TArray<FBoneName> BonesToDeform;
+					BonesToDeform.SetNumUninitialized(NumBones);
+					FMemory::Memcpy(BonesToDeform.GetData(), Data, NumBones * sizeof(FBoneName));
+					Data += NumBones * sizeof(FBoneName);
+
+					int32 NumPhysicsBodies;
+					FMemory::Memcpy(&NumPhysicsBodies, Data, sizeof(int32)); 
+					Data += sizeof(int32);
+
+					TArray<FBoneName> PhysicsToDeform;
+					PhysicsToDeform.SetNumUninitialized(NumPhysicsBodies);
+					FMemory::Memcpy(PhysicsToDeform.GetData(), Data, NumPhysicsBodies * sizeof(FBoneName));
+					Data += NumPhysicsBodies * sizeof(FBoneName);
+
+					const EMeshBindShapeFlags BindFlags = static_cast<EMeshBindShapeFlags>(Args.flags);
+
+					FMeshBindColorChannelUsages ColorChannelUsages;
+					FMemory::Memcpy(&ColorChannelUsages, &Args.ColorUsage, sizeof(ColorChannelUsages));
+					static_assert(sizeof(ColorChannelUsages) == sizeof(Args.ColorUsage));
+
+					Ptr<Mesh> BindMeshResult = CreateMesh();
+
+					bool bOutSuccess = false;
+					MeshBindShapeReshape(BindMeshResult.get(), BaseMesh.get(), Shape.get(), BonesToDeform, PhysicsToDeform, BindFlags, ColorChannelUsages, bOutSuccess);
+				
+					Release(Shape);
+					// not success indicates nothing has bond so the base mesh can be reused.
+					if (!bOutSuccess)
+					{
+						Release(BindMeshResult);
+						StoreMesh(item, BaseMesh);
+					}
+					else
+					{
+						if (!EnumHasAnyFlags(BindFlags, EMeshBindShapeFlags::ReshapeVertices))
+						{
+							Ptr<Mesh> BindMeshNoVertsResult = CloneOrTakeOver(BaseMesh);
+							BindMeshNoVertsResult->AdditionalBuffers = MoveTemp(BindMeshResult->AdditionalBuffers);
+
+							Release(BaseMesh);
+							Release(BindMeshResult);
+							StoreMesh(item, BindMeshNoVertsResult);
+						}
+						else
+						{
+							Release(BaseMesh);
+							StoreMesh(item, BindMeshResult);
+						}
+					}
+				}	
+				else
+				{
+					Ptr<Mesh> Result = CreateMesh(BaseMesh ? BaseMesh->GetDataSize() : 0);
+
+					bool bOutSuccess = false;
+					MeshBindShapeClipDeform(Result.get(), BaseMesh.get(), Shape.get(), BindingMethod, bOutSuccess);
+
+					Release(Shape);
+					if (!bOutSuccess)
+					{
+						Release(Result);
+						StoreMesh(item, BaseMesh);
+					}
+					else
+					{
+						Release(BaseMesh);
+						StoreMesh(item, Result);
+					}
+				}
+
+				break;
+			}
+
+			default:
+				check(false);
+			}
+
+			break;
+		}
+
+
+		case OP_TYPE::ME_APPLYSHAPE:
+		{
+			OP::MeshApplyShapeArgs args = Program.GetOpArgs<OP::MeshApplyShapeArgs>(item.At);
+			switch (item.Stage)
+			{
+			case 0:
+			{
+				if (args.mesh)
+				{
+					AddOp(FScheduledOp(item.At, item, 1),
+						FScheduledOp(args.mesh, item),
+						FScheduledOp(args.shape, item));
+				}
+				else
+				{
+					StoreMesh(item, nullptr);
+				}
+				break;
+			}
+			case 1:
+			{
+				MUTABLE_CPUPROFILER_SCOPE(ME_APPLYSHAPE_1)
+					
+				Ptr<const Mesh> BaseMesh = LoadMesh(FCacheAddress(args.mesh, item));
+				Ptr<const Mesh> Shape = LoadMesh(FCacheAddress(args.shape, item));
+
+				const EMeshBindShapeFlags ReshapeFlags = static_cast<EMeshBindShapeFlags>(args.flags);
+				const bool bReshapeVertices = EnumHasAnyFlags(ReshapeFlags, EMeshBindShapeFlags::ReshapeVertices);
+
+				Ptr<Mesh> ReshapedMeshResult = CreateMesh(BaseMesh ? BaseMesh->GetDataSize() : 0);
+
+				bool bOutSuccess = false;
+				MeshApplyShape(ReshapedMeshResult.get(), BaseMesh.get(), Shape.get(), ReshapeFlags, bOutSuccess);
+
+				Release(Shape);
+				
+				if (!bOutSuccess)
+				{
+					Release(ReshapedMeshResult);
+					StoreMesh(item, BaseMesh);
+				}
+				else
+				{
+					if (!bReshapeVertices)
+					{
+						// Clone without Skeleton, Physics or Poses 
+						EMeshCopyFlags CopyFlags = ~(
+							EMeshCopyFlags::WithSkeleton |
+							EMeshCopyFlags::WithPhysicsBody |
+							EMeshCopyFlags::WithAdditionalPhysics |
+							EMeshCopyFlags::WithPoses);
+
+						Ptr<Mesh> NoVerticesReshpedMesh = CloneOrTakeOver(BaseMesh);
+
+						NoVerticesReshpedMesh->SetSkeleton(ReshapedMeshResult->GetSkeleton().get());
+						NoVerticesReshpedMesh->SetPhysicsBody(ReshapedMeshResult->GetPhysicsBody().get());
+						NoVerticesReshpedMesh->AdditionalPhysicsBodies = ReshapedMeshResult->AdditionalPhysicsBodies;
+						NoVerticesReshpedMesh->BonePoses = ReshapedMeshResult->BonePoses;
+
+						Release(BaseMesh);
+						Release(ReshapedMeshResult);
+						StoreMesh(item, NoVerticesReshpedMesh);
+					}
+					else
+					{
+						Release(BaseMesh);
+						StoreMesh(item, ReshapedMeshResult);
+					}
+				}
+				break;
+			}
+
+			default:
+				check(false);
+			}
+
+			break;
+		}
+
+		case OP_TYPE::ME_MORPHRESHAPE:
+		{
+			OP::MeshMorphReshapeArgs Args = Program.GetOpArgs<OP::MeshMorphReshapeArgs>(item.At);
+			switch(item.Stage)
+			{
+			case 0:
+			{
+				if (Args.Morph)
+				{
+					AddOp(FScheduledOp(item.At, item, 1), 
+						FScheduledOp(Args.Morph, item),
+						FScheduledOp(Args.Reshape, item));
+				}
+				else 
+				{
+					StoreMesh(item, nullptr);
+				}
+				break;
+			}
+			case 1:
+			{
+				MUTABLE_CPUPROFILER_SCOPE(ME_MORPHRESHAPE_1)
+					
+				Ptr<const Mesh> MorphedMesh = LoadMesh(FCacheAddress(Args.Morph, item));
+				Ptr<const Mesh> ReshapeMesh = LoadMesh(FCacheAddress(Args.Reshape, item));
+
+				if (ReshapeMesh && MorphedMesh)
+				{
+					// Copy without Skeleton, Physics or Poses 
+					EMeshCopyFlags CopyFlags = ~(
+							EMeshCopyFlags::WithSkeleton    | 
+							EMeshCopyFlags::WithPhysicsBody | 
+							EMeshCopyFlags::WithPoses);
+
+					Ptr<Mesh> Result = CreateMesh(MorphedMesh->GetDataSize());
+					Result->CopyFrom(*MorphedMesh, CopyFlags);
+
+					Result->SetSkeleton(ReshapeMesh->GetSkeleton().get());
+					Result->SetPhysicsBody(ReshapeMesh->GetPhysicsBody().get());
+					Result->BonePoses = ReshapeMesh->BonePoses;
+
+					Release(MorphedMesh);
+					Release(ReshapeMesh);
+					StoreMesh(item, Result);
+				}
+				else
+				{
+					StoreMesh(item, MorphedMesh);
+				}
+
+				break;
+			}
+
+			default:
+				check(false);
+			}
+
+			break;
+		}
+
+        case OP_TYPE::ME_SETSKELETON:
+        {
+			OP::MeshSetSkeletonArgs args = Program.GetOpArgs<OP::MeshSetSkeletonArgs>(item.At);
+            switch (item.Stage)
+            {
+            case 0:
+			{
+				if (args.source)
+				{
+					AddOp(FScheduledOp(item.At, item, 1),
+						FScheduledOp(args.source, item),
+						FScheduledOp(args.skeleton, item));
+				}
+				else
+				{
+					StoreMesh(item, nullptr);
+				}
+				break;
+			}
+            case 1:
+            {
+            	MUTABLE_CPUPROFILER_SCOPE(ME_SETSKELETON_1)
+            		
+                Ptr<const Mesh> Source = LoadMesh(FCacheAddress(args.source, item));
+                Ptr<const Mesh> pSkeleton = LoadMesh(FCacheAddress(args.skeleton, item));
+
+                // Only if both are valid.
+                if (Source && pSkeleton)
+                {
+                    if ( Source->GetSkeleton()
+                         &&
+                         Source->GetSkeleton()->GetBoneCount() > 0 )
+                    {
+                        // For some reason we already have bone data, so we can't just overwrite it
+                        // or the skinning may break. This may happen because of a problem in the
+                        // optimiser that needs investigation.
+                        // \TODO Be defensive, for now.
+                        UE_LOG(LogMutableCore, Warning, TEXT("Performing a MeshRemapSkeleton, instead of MeshSetSkeletonData because source mesh already has some skeleton."));
+
+						Ptr<Mesh> Result = CreateMesh(Source->GetDataSize());
+
+						bool bOutSuccess = false;
+                        MeshRemapSkeleton(Result.get(), Source.get(), pSkeleton->GetSkeleton().get(), bOutSuccess);
+
+						Release(pSkeleton);
+
+                        if (!bOutSuccess)
+                        {
+							Release(Result);
+							StoreMesh(item, Source);
+                        }
+						else
+						{
+							//Result->GetPrivate()->CheckIntegrity();
+							Release(Source);
+							StoreMesh(item, Result);
+						}
+                    }
+                    else
+                    {
+						Ptr<Mesh> Result = CloneOrTakeOver(Source);
+
+                        Result->SetSkeleton(pSkeleton->GetSkeleton().get());
+
+						//Result->GetPrivate()->CheckIntegrity();
+						Release(pSkeleton);
+						StoreMesh(item, Result);
+                    }
+                }
+                else
+                {
+					Release(pSkeleton);
+					StoreMesh(item, Source);
+                }
+
+                break;
+            }
+
+            default:
+                check(false);
+            }
+
+            break;
+        }
+
+        case OP_TYPE::ME_REMOVEMASK:
+        {
+       		MUTABLE_CPUPROFILER_SCOPE(ME_REMOVEMASK)
+        		
+            // Decode op
+            // TODO: Partial decode for each stage
+            const uint8* data = Program.GetOpArgsPointer(item.At);
+
+            OP::ADDRESS source;
+            FMemory::Memcpy(&source,data,sizeof(OP::ADDRESS)); 
+			data += sizeof(OP::ADDRESS);
+
+            TArray<FScheduledOp> conditions;
+			TArray<OP::ADDRESS> masks;
+
+            uint16 removes;
+			FMemory::Memcpy(&removes,data,sizeof(uint16)); 
+			data += sizeof(uint16);
+
+            for( uint16 r=0; r<removes; ++r)
+            {
+                OP::ADDRESS condition;
+				FMemory::Memcpy(&condition,data,sizeof(OP::ADDRESS)); 
+				data += sizeof(OP::ADDRESS);
+                
+				conditions.Emplace(condition, item);
+
+                OP::ADDRESS mask;
+				FMemory::Memcpy(&mask,data,sizeof(OP::ADDRESS)); 
+				data += sizeof(OP::ADDRESS);
+
+                masks.Add(mask);
+            }
+
+
+            // Schedule next stages
+            switch (item.Stage)
+            {
+            case 0:
+			{
+				if (source)
+				{
+					// Request the conditions
+					AddOp(FScheduledOp(item.At, item, 1), conditions);
+				}
+				else
+				{
+					StoreMesh(item, nullptr);
+				}
+				break;
+			}
+            case 1:
+            {
+        		MUTABLE_CPUPROFILER_SCOPE(ME_REMOVEMASK_1)
+
+                // Request the source and the necessary masks
+                // \todo: store condition values in heap?
+                TArray<FScheduledOp> deps;
+                deps.Emplace( source, item );
+                for( size_t r=0; source && r<conditions.Num(); ++r )
+                {
+                    // If there is no expression, we'll assume true.
+                    bool value = true;
+                    if (conditions[r].At)
+                    {
+                        value = LoadBool(FCacheAddress(conditions[r].At, item));
+                    }
+
+                    if (value)
+                    {
+                        deps.Emplace(masks[r], item);
+                    }
+                }
+
+                if (source)
+                {
+					AddOp(FScheduledOp(item.At, item, 2), deps);
+				}
+                break;
+            }
+
+            case 2:
+            {
+            	MUTABLE_CPUPROFILER_SCOPE(ME_REMOVEMASK_2)
+            	
+                // \todo: single remove operation with all masks?
+                Ptr<const Mesh> Source = LoadMesh(FCacheAddress(source, item));
+
+				if (Source)
+				{
+					Ptr<Mesh> Result = CreateMesh(Source->GetDataSize());
+					Result->CopyFrom(*Source);
+
+					Release(Source);
+
+					for (int32 r = 0; r < conditions.Num(); ++r)
+					{
+						// If there is no expression, we'll assume true.
+						bool value = true;
+						if (conditions[r].At)
+						{
+							value = LoadBool(FCacheAddress(conditions[r].At, item));
+						}
+
+						if (value)
+						{
+							Ptr<const Mesh> Mask = LoadMesh(FCacheAddress(masks[r], item));
+							if (Mask)
+							{
+								//MeshRemoveMask will make a copy of Result, try to make room for it. 
+								Ptr<Mesh> IterResult = CreateMesh(Result->GetDataSize());
+
+								bool bOutSuccess = false;
+								MeshRemoveMask(IterResult.get(), Result.get(), Mask.get(), bOutSuccess);
+
+								Release(Mask);
+
+								if (!bOutSuccess)
+								{
+									Release(IterResult);
+								}
+								else
+								{
+									Swap(Result, IterResult);
+									Release(IterResult);
+								}
+							}
+						}
+					}
+
+					StoreMesh(item, Result);
+				}
+				else
+				{
+					StoreMesh(item, nullptr);
+				}
+
+                break;
+            }
+
+            default:
+                check(false);
+            }
+
+            break;
+        }
+
+        case OP_TYPE::ME_ADDTAGS:
+		{
+			MUTABLE_CPUPROFILER_SCOPE(ME_ADDTAGS)
+
+			// Decode op
+			// TODO: Partial decode for each stage
+			const uint8* Data = Program.GetOpArgsPointer(item.At);
+
+			OP::ADDRESS Source;
+			FMemory::Memcpy(&Source, Data, sizeof(OP::ADDRESS));
+			Data += sizeof(OP::ADDRESS);
+
+			// Schedule next stages
+			switch (item.Stage)
+			{
+			case 0:
+			{
+				if (Source)
+				{
+					// Request the source
+					AddOp(FScheduledOp(item.At, item, 1), FScheduledOp(Source, item));
+				}
+				else
+				{
+					StoreMesh(item, nullptr);
+				}
+				break;
+			}
+
+			case 1:
+			{
+				MUTABLE_CPUPROFILER_SCOPE(ME_ADDTAGS_2)
+
+				Ptr<const Mesh> SourceMesh = LoadMesh(FCacheAddress(Source, item));
+
+				if (!SourceMesh)
+				{
+					StoreMesh(item, nullptr);
+				}
+				else
+				{
+					Ptr<Mesh> Result = CloneOrTakeOver(SourceMesh);
+
+					// Decode the tags
+					uint16 TagCount;
+					FMemory::Memcpy(&TagCount, Data, sizeof(uint16));
+					Data += sizeof(uint16);
+
+					int32 FirstMeshTagIndex = Result->Tags.Num();
+					Result->Tags.SetNum(FirstMeshTagIndex+TagCount);
+					for (uint16 TagIndex = 0; TagIndex < TagCount; ++TagIndex)
+					{
+						OP::ADDRESS TagConstant;
+						FMemory::Memcpy(&TagConstant, Data, sizeof(OP::ADDRESS));
+						Data += sizeof(OP::ADDRESS);
+
+						check(TagConstant < (uint32)pModel->GetPrivate()->m_program.m_constantStrings.Num());
+						const FString& Name = Program.m_constantStrings[TagConstant];
+						Result->Tags[FirstMeshTagIndex+TagIndex] = Name;
+					}
+
+					StoreMesh(item, Result);
+				}
+
+				break;
+			}
+
+			default:
+				check(false);
+			}
+
+			break;
+		}
+
+        case OP_TYPE::ME_PROJECT:
+        {
+			OP::MeshProjectArgs args = Program.GetOpArgs<OP::MeshProjectArgs>(item.At);
+            switch (item.Stage)
+            {
+            case 0:
+			{
+				if (args.mesh)
+				{
+					AddOp(FScheduledOp(item.At, item, 1),
+						FScheduledOp(args.mesh, item),
+						FScheduledOp(args.projector, item));
+				}
+				else
+				{
+					StoreMesh(item, nullptr);
+				}
+				break;
+			}
+            case 1:
+            {
+				MUTABLE_CPUPROFILER_SCOPE(ME_PROJECT_1)
+
+                Ptr<const Mesh> pMesh = LoadMesh(FCacheAddress(args.mesh,item));
+                const FProjector Projector = LoadProjector(FCacheAddress(args.projector, item));
+
+                // Only if both are valid.
+                if (pMesh && pMesh->GetVertexBuffers().GetBufferCount() > 0)
+                {
+					Ptr<Mesh> Result = CreateMesh();
+
+					bool bOutSuccess = false;
+					MeshProject(Result.get(), pMesh.get(), Projector, bOutSuccess);
+
+					if (!bOutSuccess)
+					{
+						Release(Result);
+						StoreMesh(item, pMesh);
+					}
+					else
+					{	
+//						Result->GetPrivate()->CheckIntegrity();
+						Release(pMesh);
+						StoreMesh(item, Result);
+					}
+                }
+				else
+				{
+					Release(pMesh);
+					StoreMesh(item, nullptr);
+				}
+
+                break;
+            }
+
+            default:
+                check(false);
+            }
+
+            break;
+        }
+
+		case OP_TYPE::ME_OPTIMIZESKINNING:
+		{
+			OP::MeshOptimizeSkinningArgs args = Program.GetOpArgs<OP::MeshOptimizeSkinningArgs>(item.At);
+			switch (item.Stage)
+			{
+			case 0:
+			{
+				if (args.source)
+				{
+					AddOp(FScheduledOp(item.At, item, 1), FScheduledOp(args.source, item));
+				}
+				else
+				{
+					StoreMesh(item, nullptr);
+				}
+				break;
+			}
+			case 1:
+			{
+				MUTABLE_CPUPROFILER_SCOPE(ME_OPTIMIZESKINNING_1)
+
+				Ptr<const Mesh> Source = LoadMesh(FCacheAddress(args.source, item));
+
+				if (Source && Source->IsReference())
+				{
+					StoreMesh(item, Source);
+				}
+
+				Ptr<Mesh> Result = CreateMesh();
+
+				bool bOutSuccess = false;
+				MeshOptimizeSkinning(Result.get(), Source.get(), bOutSuccess);
+
+				if (!bOutSuccess)
+				{
+					Release(Result);
+					StoreMesh(item, Source);
+				}
+				else
+				{
+					Release(Source);
+					StoreMesh(item, Result);
+				}
+
+				break;
+			}
+
+			default:
+				check(false);
+			}
+
+			break;
+		}
+
+        default:
+            if (type!=OP_TYPE::NONE)
+            {
+                // Operation not implemented
+                check( false );
+            }
+            break;
+        }
+    }
+
+    //---------------------------------------------------------------------------------------------
+    void CodeRunner::RunCode_Image(const FScheduledOp& item, const Parameters* pParams, const Model* pModel )
+    {
+		MUTABLE_CPUPROFILER_SCOPE(RunCode_Image);
+
+		FImageOperator ImOp = MakeImageOperator(this);
+
+		const FProgram& Program = m_pModel->GetPrivate()->m_program;
+
+		OP_TYPE type = Program.GetOpType(item.At);
+		switch (type)
+        {
+
+        case OP_TYPE::IM_LAYERCOLOUR:
+        {
+			OP::ImageLayerColourArgs args = Program.GetOpArgs<OP::ImageLayerColourArgs>(item.At);
+            switch (item.Stage)
+            {
+            case 0:
+                    AddOp( FScheduledOp( item.At, item, 1),
+                           FScheduledOp( args.base, item),
+                           FScheduledOp::FromOpAndOptions( args.colour, item, 0),
+                           FScheduledOp( args.mask, item) );
+                break;
+
+            case 1:
+				// This has been moved to a task. It should have been intercepted in IssueOp.
+				check(false);
+				break;
+
+            default:
+                check(false);
+            }
+
+            break;
+        }
+
+        case OP_TYPE::IM_LAYER:
+        {
+			OP::ImageLayerArgs args = Program.GetOpArgs<OP::ImageLayerArgs>(item.At);
+
+			if (ExecutionStrategy == EExecutionStrategy::MinimizeMemory)
+			{
+				switch (item.Stage)
+				{
+				case 0:
+					AddOp(FScheduledOp(item.At, item, 1),
+						FScheduledOp(args.base, item));
+					break;
+
+				case 1:
+					// Request the rest of the data.
+					AddOp(FScheduledOp(item.At, item, 2),
+						FScheduledOp(args.blended, item),
+						FScheduledOp(args.mask, item));
+					break;
+
+				case 2:
+					// This has been moved to a task. It should have been intercepted in IssueOp.
+					check(false);
+					break;
+
+				default:
+					check(false);
+				}
+			}
+			else
+			{
+				switch (item.Stage)
+				{
+				case 0:
+					AddOp(FScheduledOp(item.At, item, 1),
+						FScheduledOp(args.base, item),
+						FScheduledOp(args.blended, item),
+						FScheduledOp(args.mask, item));
+					break;
+
+				case 1:
+					// This has been moved to a task. It should have been intercepted in IssueOp.
+					check(false);
+					break;
+
+				default:
+					check(false);
+				}
+			}
+
+            break;
+        }
+
+        case OP_TYPE::IM_MULTILAYER:
+        {
+			OP::ImageMultiLayerArgs args = Program.GetOpArgs<OP::ImageMultiLayerArgs>(item.At);
+            switch (item.Stage)
+            {
+            case 0:
+                AddOp( FScheduledOp( item.At, item, 1),
+                       FScheduledOp( args.rangeSize, item ),
+					   FScheduledOp(args.base, item));
+                break;
+
+            case 1:
+            {
+           		MUTABLE_CPUPROFILER_SCOPE(IM_MULTILAYER_1)
+            		
+                // We now know the number of iterations
+                int32 Iterations = 0;
+                if (args.rangeSize)
+                {
+                    FCacheAddress RangeAddress(args.rangeSize,item);
+
+                    // We support both integers and scalars here, which is not common.
+                    // \todo: review if this is necessary or we can enforce it at compile time.
+                    DATATYPE RangeSizeType = GetOpDataType( pModel->GetPrivate()->m_program.GetOpType(args.rangeSize) );
+                    if (RangeSizeType == DT_INT)
+                    {
+						Iterations = LoadInt(RangeAddress);
+                    }
+                    else if (RangeSizeType == DT_SCALAR)
+                    {
+						Iterations = int32( LoadScalar(RangeAddress) );
+                    }
+                }
+
+				Ptr<const Image> Base = LoadImage(FCacheAddress(args.base, item));
+
+				if (Iterations <= 0)
+				{
+					// There are no layers: return the base
+					StoreImage(item, Base);
+				}
+				else
+				{
+					// Store the base
+					Ptr<Image> New = CloneOrTakeOver(Base);
+					EImageFormat InitialBaseFormat = New->GetFormat();
+
+					// Reset relevancy map.
+					New->m_flags &= ~Image::EImageFlags::IF_HAS_RELEVANCY_MAP;
+
+					// This shouldn't happen in optimised models, but it could happen in editors, etc.
+					// \todo: raise a performance warning?
+					EImageFormat BaseFormat = GetUncompressedFormat(New->GetFormat());
+					if (New->GetFormat() != BaseFormat)
+					{
+						Ptr<Image> Formatted = CreateImage( New->GetSizeX(), New->GetSizeY(), New->GetLODCount(), BaseFormat, EInitializationType::NotInitialized );
+
+						bool bSuccess = false;
+						ImOp.ImagePixelFormat(bSuccess, m_pSettings->ImageCompressionQuality, Formatted.get(), New.get());
+						check(bSuccess); // Decompression cannot fail
+
+						Release(New);
+						New = Formatted;
+					}
+
+					FScheduledOpData Data;
+					Data.Resource = New;
+					Data.MultiLayer.Iterations = Iterations;
+					Data.MultiLayer.OriginalBaseFormat = InitialBaseFormat;
+					Data.MultiLayer.bBlendOnlyOneMip = false;
+					int32 DataPos = m_heapData.Add(Data);
+
+					// Request the first layer
+					int32 CurrentIteration = 0;
+					FScheduledOp ItemCopy = item;
+					ExecutionIndex Index = GetMemory().GetRangeIndex(item.ExecutionIndex);
+					Index.SetFromModelRangeIndex(args.rangeId, CurrentIteration);
+					ItemCopy.ExecutionIndex = GetMemory().GetRangeIndexIndex(Index);
+					AddOp(FScheduledOp(item.At, item, 2, DataPos), FScheduledOp(args.base, item), FScheduledOp(args.blended, ItemCopy), FScheduledOp(args.mask, ItemCopy));
+				}
+
+                break;
+            }
+
+            default:
+            {
+            	MUTABLE_CPUPROFILER_SCOPE(IM_MULTILAYER_default)
+
+				FScheduledOpData& Data = m_heapData[item.CustomState];
+
+				int32 Iterations = Data.MultiLayer.Iterations;
+				int32 CurrentIteration = item.Stage - 2;
+				check(CurrentIteration >= 0 && CurrentIteration < 120);
+
+				TRACE_CPUPROFILER_EVENT_SCOPE_TEXT(*FString::Printf(TEXT("Layer %d of %d"), CurrentIteration, Iterations));
+
+				// Process the current layer
+
+				Ptr<Image> Base = static_cast<Image*>(Data.Resource.get());
+ 
+                FScheduledOp itemCopy = item;
+                ExecutionIndex index = GetMemory().GetRangeIndex( item.ExecutionIndex );
+				
+                {
+                    index.SetFromModelRangeIndex( args.rangeId, CurrentIteration);
+                    itemCopy.ExecutionIndex = GetMemory().GetRangeIndexIndex(index);
+					itemCopy.CustomState = 0;
+
+                    Ptr<const Image> Blended = LoadImage( FCacheAddress(args.blended,itemCopy) );
+
+                    // This shouldn't happen in optimised models, but it could happen in editors, etc.
+                    // \todo: raise a performance warning?
+                    if (Blended && Blended->GetFormat()!=Base->GetFormat() )
+                    {
+						MUTABLE_CPUPROFILER_SCOPE(ImageResize_BlendedReformat);
+
+						Ptr<Image> Formatted = CreateImage(Blended->GetSizeX(), Blended->GetSizeY(), Blended->GetLODCount(), Base->GetFormat(), EInitializationType::NotInitialized);
+
+						bool bSuccess = false;
+						ImOp.ImagePixelFormat(bSuccess, m_pSettings->ImageCompressionQuality, Formatted.get(), Blended.get());
+						check(bSuccess);
+
+						Release(Blended);
+						Blended = Formatted;
+                    }
+
+					// TODO: This shouldn't happen, but be defensive.
+					FImageSize ResultSize = Base->GetSize();
+					if (Blended && Blended->GetSize() != ResultSize)
+					{
+						MUTABLE_CPUPROFILER_SCOPE(ImageResize_BlendedFixForMultilayer);
+
+						Ptr<Image> Resized = CreateImage(ResultSize[0], ResultSize[1], Blended->GetLODCount(), Blended->GetFormat(), EInitializationType::NotInitialized);
+						ImOp.ImageResizeLinear(Resized.get(), 0, Blended.get());
+						Release(Blended);
+						Blended = Resized;
+					}
+
+					if (Blended->GetLODCount() < Base->GetLODCount())
+					{
+						Data.MultiLayer.bBlendOnlyOneMip = true;
+					}
+
+					bool bApplyColorBlendToAlpha = false;
+
+					bool bDone = false;
+
+					// This becomes true if we need to update the mips of the resulting image
+					// This could happen in the base image has mips, but one of the blended one doesn't.
+					bool bBlendOnlyOneMip = Data.MultiLayer.bBlendOnlyOneMip;
+					bool bUseBlendSourceFromBlendAlpha = false; // (Args.flags& OP::ImageLayerArgs::F_BLENDED_RGB_FROM_ALPHA) != 0;
+
+					if (!args.mask && args.bUseMaskFromBlended
+						&&
+						args.blendType == uint8(EBlendType::BT_BLEND)
+						&&
+						args.blendTypeAlpha == uint8(EBlendType::BT_LIGHTEN) )
+					{
+						// This is a frequent critical-path case because of multilayer projectors.
+						bDone = true;
+					
+						constexpr bool bUseVectorImpl = false;
+						if constexpr (bUseVectorImpl)
+						{
+							BufferLayerCompositeVector<VectorBlendChannelMasked, VectorLightenChannel, false>(Base.get(), Blended.get(), bBlendOnlyOneMip, args.BlendAlphaSourceChannel);
+						}
+						else
+						{
+							BufferLayerComposite<BlendChannelMasked, LightenChannel, false>(Base.get(), Blended.get(), bBlendOnlyOneMip, args.BlendAlphaSourceChannel);
+						}
+					}
+
+                    if (!bDone && args.mask)
+                    {
+                        Ptr<const Image> Mask = LoadImage( FCacheAddress(args.mask,itemCopy) );
+
+						// TODO: This shouldn't happen, but be defensive.
+						if (Mask && Mask->GetSize() != ResultSize)
+						{
+							MUTABLE_CPUPROFILER_SCOPE(ImageResize_MaskFixForMultilayer);
+
+							Ptr<Image> Resized = CreateImage(ResultSize[0], ResultSize[1], Mask->GetLODCount(), Mask->GetFormat(), EInitializationType::NotInitialized);
+							ImOp.ImageResizeLinear(Resized.get(), 0, Mask.get());
+							Release(Mask);
+							Mask = Resized;
+						}
+
+						// Not implemented yet
+						check(!bUseBlendSourceFromBlendAlpha);
+
+                        switch (EBlendType(args.blendType))
+                        {
+						case EBlendType::BT_NORMAL_COMBINE: check(false); break;
+                        case EBlendType::BT_SOFTLIGHT: BufferLayer<SoftLightChannelMasked, SoftLightChannel, false>(Base.get(), Base.get(), Mask.get(), Blended.get(), bApplyColorBlendToAlpha, bBlendOnlyOneMip); break;
+                        case EBlendType::BT_HARDLIGHT: BufferLayer<HardLightChannelMasked, HardLightChannel, false>(Base.get(), Base.get(), Mask.get(), Blended.get(), bApplyColorBlendToAlpha, bBlendOnlyOneMip); break;
+                        case EBlendType::BT_BURN: BufferLayer<BurnChannelMasked, BurnChannel, false>(Base.get(), Base.get(), Mask.get(), Blended.get(), bApplyColorBlendToAlpha, bBlendOnlyOneMip); break;
+                        case EBlendType::BT_DODGE: BufferLayer<DodgeChannelMasked, DodgeChannel, false>(Base.get(), Base.get(), Mask.get(), Blended.get(), bApplyColorBlendToAlpha, bBlendOnlyOneMip); break;
+                        case EBlendType::BT_SCREEN: BufferLayer<ScreenChannelMasked, ScreenChannel, false>(Base.get(), Base.get(), Mask.get(), Blended.get(), bApplyColorBlendToAlpha, bBlendOnlyOneMip); break;
+                        case EBlendType::BT_OVERLAY: BufferLayer<OverlayChannelMasked, OverlayChannel, false>(Base.get(), Base.get(), Mask.get(), Blended.get(), bApplyColorBlendToAlpha, bBlendOnlyOneMip); break;
+                        case EBlendType::BT_LIGHTEN: BufferLayer<LightenChannelMasked, LightenChannel, false>(Base.get(), Base.get(), Mask.get(), Blended.get(), bApplyColorBlendToAlpha, bBlendOnlyOneMip); break;
+                        case EBlendType::BT_MULTIPLY: BufferLayer<MultiplyChannelMasked, MultiplyChannel, false>(Base.get(), Base.get(), Mask.get(), Blended.get(), bApplyColorBlendToAlpha, bBlendOnlyOneMip); break;
+                        case EBlendType::BT_BLEND: BufferLayer<BlendChannelMasked, BlendChannel, false>(Base.get(), Base.get(), Mask.get(), Blended.get(), bApplyColorBlendToAlpha, bBlendOnlyOneMip); break;
+                        default: check(false);
+                        }
+
+						Release(Mask);
+                    }
+					else if (!bDone && args.bUseMaskFromBlended)
+					{
+						// Not implemented yet
+						check(!bUseBlendSourceFromBlendAlpha);
+
+						switch (EBlendType(args.blendType))
+						{
+						case EBlendType::BT_NORMAL_COMBINE: check(false); break;
+						case EBlendType::BT_SOFTLIGHT: BufferLayerEmbeddedMask<SoftLightChannelMasked, SoftLightChannel, false>(Base.get(), Base.get(), Blended.get(), bApplyColorBlendToAlpha, bBlendOnlyOneMip); break;
+						case EBlendType::BT_HARDLIGHT: BufferLayerEmbeddedMask<HardLightChannelMasked, HardLightChannel, false>(Base.get(), Base.get(), Blended.get(), bApplyColorBlendToAlpha, bBlendOnlyOneMip); break;
+						case EBlendType::BT_BURN: BufferLayerEmbeddedMask<BurnChannelMasked, BurnChannel, false>(Base.get(), Base.get(), Blended.get(), bApplyColorBlendToAlpha, bBlendOnlyOneMip); break;
+						case EBlendType::BT_DODGE: BufferLayerEmbeddedMask<DodgeChannelMasked, DodgeChannel, false>(Base.get(), Base.get(), Blended.get(), bApplyColorBlendToAlpha, bBlendOnlyOneMip); break;
+						case EBlendType::BT_SCREEN: BufferLayerEmbeddedMask<ScreenChannelMasked, ScreenChannel, false>(Base.get(), Base.get(), Blended.get(), bApplyColorBlendToAlpha, bBlendOnlyOneMip); break;
+						case EBlendType::BT_OVERLAY: BufferLayerEmbeddedMask<OverlayChannelMasked, OverlayChannel, false>(Base.get(), Base.get(), Blended.get(), bApplyColorBlendToAlpha, bBlendOnlyOneMip); break;
+						case EBlendType::BT_LIGHTEN: BufferLayerEmbeddedMask<LightenChannelMasked, LightenChannel, false>(Base.get(), Base.get(), Blended.get(), bApplyColorBlendToAlpha, bBlendOnlyOneMip); break;
+						case EBlendType::BT_MULTIPLY: BufferLayerEmbeddedMask<MultiplyChannelMasked, MultiplyChannel, false>(Base.get(), Base.get(), Blended.get(), bApplyColorBlendToAlpha, bBlendOnlyOneMip); break;
+						case EBlendType::BT_BLEND: BufferLayerEmbeddedMask<BlendChannelMasked, BlendChannel, false>(Base.get(), Base.get(), Blended.get(), bApplyColorBlendToAlpha, bBlendOnlyOneMip); break;
+						default: check(false);
+						}
+					}
+                    else if (!bDone)
+                    {
+                        switch (EBlendType(args.blendType))
+                        {
+						case EBlendType::BT_NORMAL_COMBINE: check(false); break;
+                        case EBlendType::BT_SOFTLIGHT: BufferLayer<SoftLightChannel, false>(Base.get(), Base.get(), Blended.get(), bApplyColorBlendToAlpha, bBlendOnlyOneMip, bUseBlendSourceFromBlendAlpha); break;
+                        case EBlendType::BT_HARDLIGHT: BufferLayer<HardLightChannel, false>(Base.get(), Base.get(), Blended.get(), bApplyColorBlendToAlpha, bBlendOnlyOneMip, bUseBlendSourceFromBlendAlpha); break;
+                        case EBlendType::BT_BURN: BufferLayer<BurnChannel, false>(Base.get(), Base.get(), Blended.get(), bApplyColorBlendToAlpha, bBlendOnlyOneMip, bUseBlendSourceFromBlendAlpha); break;
+                        case EBlendType::BT_DODGE: BufferLayer<DodgeChannel, false>(Base.get(), Base.get(), Blended.get(), bApplyColorBlendToAlpha, bBlendOnlyOneMip, bUseBlendSourceFromBlendAlpha); break;
+                        case EBlendType::BT_SCREEN: BufferLayer<ScreenChannel, false>(Base.get(), Base.get(),  Blended.get(), bApplyColorBlendToAlpha, bBlendOnlyOneMip, bUseBlendSourceFromBlendAlpha); break;
+                        case EBlendType::BT_OVERLAY: BufferLayer<OverlayChannel, false>(Base.get(), Base.get(), Blended.get(), bApplyColorBlendToAlpha, bBlendOnlyOneMip, bUseBlendSourceFromBlendAlpha); break;
+                        case EBlendType::BT_LIGHTEN: BufferLayer<LightenChannel, false>(Base.get(), Base.get(), Blended.get(), bApplyColorBlendToAlpha, bBlendOnlyOneMip, bUseBlendSourceFromBlendAlpha); break;
+                        case EBlendType::BT_MULTIPLY: BufferLayer<MultiplyChannel, false>(Base.get(), Base.get(), Blended.get(), bApplyColorBlendToAlpha, bBlendOnlyOneMip, bUseBlendSourceFromBlendAlpha); break;
+                        case EBlendType::BT_BLEND: BufferLayer<BlendChannel, false>(Base.get(), Base.get(), Blended.get(), bApplyColorBlendToAlpha, bBlendOnlyOneMip, bUseBlendSourceFromBlendAlpha); break;
+                        default: check(false);
+                        }
+                    }
+
+					// Apply the separate blend operation for alpha
+					if (!bDone && !bApplyColorBlendToAlpha && args.blendTypeAlpha != uint8(EBlendType::BT_NONE) )
+					{
+						// Separate alpha operation ignores the mask.
+						switch (EBlendType(args.blendTypeAlpha))
+						{
+						case EBlendType::BT_SOFTLIGHT: BufferLayerInPlace<SoftLightChannel, false, 1>(Base.get(), Blended.get(), bBlendOnlyOneMip, 3, args.BlendAlphaSourceChannel); break;
+						case EBlendType::BT_HARDLIGHT: BufferLayerInPlace<HardLightChannel, false, 1>(Base.get(), Blended.get(), bBlendOnlyOneMip, 3, args.BlendAlphaSourceChannel); break;
+						case EBlendType::BT_BURN: BufferLayerInPlace<BurnChannel, false, 1>(Base.get(), Blended.get(), bBlendOnlyOneMip, 3, args.BlendAlphaSourceChannel); break;
+						case EBlendType::BT_DODGE: BufferLayerInPlace<DodgeChannel, false, 1>(Base.get(), Blended.get(), bBlendOnlyOneMip, 3, args.BlendAlphaSourceChannel); break;
+						case EBlendType::BT_SCREEN: BufferLayerInPlace<ScreenChannel, false, 1>(Base.get(), Blended.get(), bBlendOnlyOneMip, 3, args.BlendAlphaSourceChannel); break;
+						case EBlendType::BT_OVERLAY: BufferLayerInPlace<OverlayChannel, false, 1>(Base.get(), Blended.get(), bBlendOnlyOneMip, 3, args.BlendAlphaSourceChannel); break;
+						case EBlendType::BT_LIGHTEN: BufferLayerInPlace<LightenChannel, false, 1>(Base.get(), Blended.get(), bBlendOnlyOneMip, 3, args.BlendAlphaSourceChannel); break;
+						case EBlendType::BT_MULTIPLY: BufferLayerInPlace<MultiplyChannel, false, 1>(Base.get(), Blended.get(), bBlendOnlyOneMip, 3, args.BlendAlphaSourceChannel); break;
+						case EBlendType::BT_BLEND: BufferLayerInPlace<BlendChannel, false, 1>(Base.get(), Blended.get(), bBlendOnlyOneMip, 3, args.BlendAlphaSourceChannel); break;
+						default: check(false);
+						}
+					}
+
+					Release(Blended);
+				}
+
+				// Are we done?
+				if (CurrentIteration + 1 == Iterations)
+				{
+					if (Data.MultiLayer.bBlendOnlyOneMip)
+					{
+						MUTABLE_CPUPROFILER_SCOPE(ImageLayer_MipFix);
+						FMipmapGenerationSettings DummyMipSettings{};
+						ImageMipmapInPlace(m_pSettings->ImageCompressionQuality, Base.get(), DummyMipSettings);
+					}
+
+					// TODO: Reconvert to OriginalBaseFormat if necessary?
+
+					Data.Resource = nullptr;
+					StoreImage(item, Base);
+					break;
+				}
+				else
+				{
+					// Request a new layer
+					++CurrentIteration;
+					FScheduledOp ItemCopy = item;
+					ExecutionIndex Index = GetMemory().GetRangeIndex(item.ExecutionIndex);
+					Index.SetFromModelRangeIndex(args.rangeId, CurrentIteration);
+					ItemCopy.ExecutionIndex = GetMemory().GetRangeIndexIndex(Index);
+					AddOp(FScheduledOp(item.At, item, 2+CurrentIteration, item.CustomState), FScheduledOp(args.blended, ItemCopy), FScheduledOp(args.mask, ItemCopy));
+
+				}
+
+                break;
+            }
+
+            } // switch stage
+
+            break;
+        }
+
+		case OP_TYPE::IM_NORMALCOMPOSITE:
+		{
+			OP::ImageNormalCompositeArgs args = Program.GetOpArgs<OP::ImageNormalCompositeArgs>(item.At);
+			switch (item.Stage)
+			{
+			case 0:
+				if (args.base && args.normal)
+				{
+					AddOp(FScheduledOp(item.At, item, 1),
+							FScheduledOp(args.base, item),
+							FScheduledOp(args.normal, item));
+				}
+				else
+				{
+					StoreImage(item, nullptr);
+				}
+				break;
+
+			case 1:
+			{
+				MUTABLE_CPUPROFILER_SCOPE(IM_NORMALCOMPOSITE_1)
+
+				Ptr<const Image> Base = LoadImage(FCacheAddress(args.base, item));
+				Ptr<const Image> Normal = LoadImage(FCacheAddress(args.normal, item));
+
+				if (Normal->GetLODCount() < Base->GetLODCount())
+				{
+					MUTABLE_CPUPROFILER_SCOPE(ImageNormalComposite_EmergencyFix);
+
+					int32 StartLevel = Normal->GetLODCount() - 1;
+					int32 LevelCount = Base->GetLODCount();
+					
+					Ptr<Image> NormalFix = CloneOrTakeOver(Normal);
+
+					FMipmapGenerationSettings MipSettings{};
+					ImOp.ImageMipmap(m_pSettings->ImageCompressionQuality, NormalFix.get(), NormalFix.get(), StartLevel, LevelCount, MipSettings);
+
+					Normal = NormalFix;
+				}
+
+
+                Ptr<Image> Result = CreateImage(Base->GetSizeX(), Base->GetSizeY(), Base->GetLODCount(), Base->GetFormat(), EInitializationType::NotInitialized);
+				ImageNormalComposite(Result.get(), Base.get(), Normal.get(), args.mode, args.power);
+
+				Release(Base);
+				Release(Normal);
+				StoreImage(item, Result);
+				break;
+			}
+
+			default:
+				check(false);
+			}
+
+			break;
+		}
+
+        case OP_TYPE::IM_PIXELFORMAT:
+        {
+			OP::ImagePixelFormatArgs args = Program.GetOpArgs<OP::ImagePixelFormatArgs>(item.At);
+            switch (item.Stage)
+            {
+            case 0:
+                AddOp( FScheduledOp( item.At, item, 1), FScheduledOp( args.source, item) );
+                break;
+
+            case 1:
+            {
+				// This has been moved to a task. It should have been intercepted in IssueOp.
+				check(false);
+				break;
+            }
+
+            default:
+                check(false);
+            }
+
+            break;
+        }
+
+        case OP_TYPE::IM_MIPMAP:
+        {
+			OP::ImageMipmapArgs args = Program.GetOpArgs<OP::ImageMipmapArgs>(item.At);
+            switch (item.Stage)
+            {
+            case 0:
+				AddOp( FScheduledOp( item.At, item, 1), FScheduledOp( args.source, item) );
+                break;
+
+            case 1:
+				// This has been moved to a task. It should have been intercepted in IssueOp.
+				check(false);
+				break;
+
+            default:
+                check(false);
+            }
+
+            break;
+        }
+
+        case OP_TYPE::IM_RESIZE:
+        {
+			OP::ImageResizeArgs args = Program.GetOpArgs<OP::ImageResizeArgs>(item.At);
+            switch (item.Stage)
+            {
+            case 0:
+                AddOp( FScheduledOp( item.At, item, 1), FScheduledOp( args.source, item) );
+                break;
+
+            case 1:
+            {
+				// This has been moved to a task. It should have been intercepted in IssueOp.
+				check(false);
+            }
+
+            default:
+                check(false);
+            }
+
+            break;
+        }
+
+        case OP_TYPE::IM_RESIZELIKE:
+        {
+			OP::ImageResizeLikeArgs args = Program.GetOpArgs<OP::ImageResizeLikeArgs>(item.At);
+            switch (item.Stage)
+            {
+            case 0:
+                AddOp(FScheduledOp(item.At, item, 1),
+                      	FScheduledOp(args.source, item),
+                        FScheduledOp(args.sizeSource, item));
+                break;
+
+            case 1:
+            {
+            	MUTABLE_CPUPROFILER_SCOPE(IM_RESIZELIKE_1)
+            	
+                Ptr<const Image> Base = LoadImage( FCacheAddress(args.source,item) );
+                Ptr<const Image> SizeBase = LoadImage( FCacheAddress(args.sizeSource,item) );
+				FImageSize DestSize = SizeBase->GetSize();
+				Release(SizeBase);
+
+                if (Base->GetSize() != DestSize)
+                {
+					int32 BaseLODCount = Base->GetLODCount();
+					Ptr<Image> Result = CreateImage(DestSize[0], DestSize[1], BaseLODCount, Base->GetFormat(), EInitializationType::NotInitialized);
+					ImOp.ImageResizeLinear(Result.get(), m_pSettings->ImageCompressionQuality, Base.get());
+					Release(Base);
+
+                    // If the source image had mips, generate them as well for the resized image.
+                    // This shouldn't happen often since "ResizeLike" should be usually optimised out
+                    // during model compilation. The mipmap generation below is not very precise with
+                    // the number of mips that are needed and will probably generate too many
+                    bool bSourceHasMips = BaseLODCount > 1;
+                    
+					if (bSourceHasMips)
+                    {
+						int32 LevelCount = Image::GetMipmapCount(Result->GetSizeX(), Result->GetSizeY());	
+						Result->DataStorage.SetNumLODs(LevelCount);
+
+						FMipmapGenerationSettings MipSettings{};
+						ImOp.ImageMipmap(m_pSettings->ImageCompressionQuality, Result.get(), Result.get(), 0, LevelCount, MipSettings);
+                    }				
+
+					StoreImage(item, Result);
+				}
+                else
+                {
+					StoreImage(item, Base);
+				}
+				
+                break;
+            }
+
+            default:
+                check(false);
+            }
+
+            break;
+        }
+
+        case OP_TYPE::IM_RESIZEREL:
+        {
+			OP::ImageResizeRelArgs args = Program.GetOpArgs<OP::ImageResizeRelArgs>(item.At);
+            switch (item.Stage)
+            {
+            case 0:
+                AddOp( FScheduledOp( item.At, item, 1), FScheduledOp( args.source, item) );
+                break;
+
+            case 1:
+            {
+				// This has been moved to a task. It should have been intercepted in IssueOp.
+				check(false);
+            }
+
+            default:
+                check(false);
+            }
+
+
+            break;
+        }
+        case OP_TYPE::IM_BLANKLAYOUT:
+        {
+			OP::ImageBlankLayoutArgs Args = Program.GetOpArgs<OP::ImageBlankLayoutArgs>(item.At);
+            switch (item.Stage)
+            {
+            case 0:
+                AddOp( FScheduledOp( item.At, item, 1), FScheduledOp::FromOpAndOptions(Args.layout, item, 0));
+                break;
+
+            case 1:
+            {
+            	MUTABLE_CPUPROFILER_SCOPE(IM_BLANKLAYOUT_1)
+            		
+                Ptr<const Layout> pLayout = LoadLayout(FScheduledOp::FromOpAndOptions(Args.layout, item, 0));
+
+                FIntPoint SizeInBlocks = pLayout->GetGridSize();
+
+				FIntPoint BlockSizeInPixels(Args.blockSize[0], Args.blockSize[1]);
+
+				// Image size if we don't skip any mipmap
+				FIntPoint FullImageSizeInPixels = SizeInBlocks * BlockSizeInPixels;
+				int32 FullImageMipCount = Image::GetMipmapCount(FullImageSizeInPixels.X, FullImageSizeInPixels.Y);
+
+				FIntPoint ImageSizeInPixels = FullImageSizeInPixels;
+				int32 MipsToSkip = item.ExecutionOptions;
+				MipsToSkip = FMath::Min(MipsToSkip, FullImageMipCount);
+				if (MipsToSkip > 0)
+				{
+					//FIntPoint ReducedBlockSizeInPixels;
+
+					// This method tries to reduce only the block size, but it fails if the image is still too big
+					// If we want to generate only a subset of mipmaps, reduce the layout block size accordingly.
+					//ReducedBlockSizeInPixels.X = BlockSizeInPixels.X >> MipsToSkip;
+					//ReducedBlockSizeInPixels.Y = BlockSizeInPixels.Y >> MipsToSkip;
+					//const FImageFormatData& FormatData = GetImageFormatData((EImageFormat)args.format);
+					//int MinBlockSize = FMath::Max(FormatData.PixelsPerBlockX, FormatData.PixelsPerBlockY);
+					//ReducedBlockSizeInPixels.X = FMath::Max<int32>(ReducedBlockSizeInPixels.X, FormatData.PixelsPerBlockX);
+					//ReducedBlockSizeInPixels.Y = FMath::Max<int32>(ReducedBlockSizeInPixels.Y, FormatData.PixelsPerBlockY);
+					//FIntPoint ReducedImageSizeInPixels = SizeInBlocks * ReducedBlockSizeInPixels;
+
+					// This method simply reduces the size and assumes all the other operations will handle degeenrate cases.
+					ImageSizeInPixels = FullImageSizeInPixels / (1 << MipsToSkip);
+					
+					//if (ReducedImageSizeInPixels!= ImageSizeInPixels)
+					//{
+					//	check(false);
+					//}
+				}
+
+                int32 MipsToGenerate = 1;
+                if (Args.generateMipmaps)
+                {
+                    if (Args.mipmapCount == 0)
+                    {
+						MipsToGenerate = Image::GetMipmapCount(ImageSizeInPixels.X, ImageSizeInPixels.Y);
+                    }
+                    else
+                    {
+						MipsToGenerate = FMath::Max(Args.mipmapCount - MipsToSkip, 1);
+                    }
+                }
+
+				// It needs to be initialized in case it has gaps.
+                Ptr<Image> New = CreateImage(ImageSizeInPixels.X, ImageSizeInPixels.Y, MipsToGenerate, EImageFormat(Args.format), EInitializationType::Black );
+                StoreImage(item, New);
+                break;
+            }
+
+            default:
+                check(false);
+            }
+
+
+            break;
+        }
+
+        case OP_TYPE::IM_COMPOSE:
+        {
+			OP::ImageComposeArgs Args = Program.GetOpArgs<OP::ImageComposeArgs>(item.At);
+
+			if (ExecutionStrategy == EExecutionStrategy::MinimizeMemory)
+			{
+            	switch (item.Stage)
+				{
+				case 0:
+					AddOp(FScheduledOp(item.At, item, 1), FScheduledOp::FromOpAndOptions(Args.layout, item, 0));
+					break;
+				case 1:
+				{
+					Ptr<const Layout> ComposeLayout = 
+							LoadLayout(FCacheAddress(Args.layout, FScheduledOp::FromOpAndOptions(Args.layout, item, 0)));
+
+					FScheduledOpData Data;
+					Data.Resource = const_cast<Layout*>(ComposeLayout.get());
+					int32 DataPos = m_heapData.Add(Data);
+
+					int32 RelBlockIndex = ComposeLayout->FindBlock(Args.BlockId);
+
+					if (RelBlockIndex >= 0)
+					{
+						AddOp(FScheduledOp(item.At, item, 2, DataPos), FScheduledOp(Args.base, item));
+					}
+					else
+					{
+						// Jump directly to stage 3, no need to load mask or blockImage.
+						AddOp(FScheduledOp(item.At, item, 3, DataPos), FScheduledOp(Args.base, item));
+					}
+
+					break;
+				}
+				case 2:
+				{
+					AddOp(FScheduledOp(item.At, item, 3, item.CustomState),
+						  FScheduledOp(Args.blockImage, item),
+						  FScheduledOp(Args.mask, item));
+					break;
+				}
+
+				case 3:
+					// This has been moved to a task. It should have been intercepted in IssueOp.
+					check(false);
+					break;
+
+				default:
+					check(false);
+				}
+			}
+			else
+			{
+            	switch (item.Stage)
+				{
+				case 0:
+					AddOp(FScheduledOp(item.At, item, 1), FScheduledOp::FromOpAndOptions(Args.layout, item, 0));
+					break;
+
+				case 1:
+				{	
+					Ptr<const Layout> ComposeLayout = 
+							LoadLayout(FCacheAddress(Args.layout, FScheduledOp::FromOpAndOptions(Args.layout, item, 0)));
+
+					FScheduledOpData Data;
+					Data.Resource = const_cast<Layout*>(ComposeLayout.get());
+					int32 DataPos = m_heapData.Add(Data);
+
+					int32 RelBlockIndex = ComposeLayout->FindBlock(Args.BlockId);
+					if (RelBlockIndex >= 0)
+					{
+						AddOp(FScheduledOp(item.At, item, 2, DataPos),
+							  FScheduledOp(Args.base, item),
+							  FScheduledOp(Args.blockImage, item),
+							  FScheduledOp(Args.mask, item));
+					}
+					else
+					{
+						AddOp(FScheduledOp(item.At, item, 2, DataPos), FScheduledOp(Args.base, item));
+					}
+					break;
+				}
+
+				case 2:
+					// This has been moved to a task. It should have been intercepted in IssueOp.
+					check(false);
+					break;
+
+				default:
+					check(false);
+				}
+			}
+
+            break;
+        }
+
+        case OP_TYPE::IM_INTERPOLATE:
+        {
+			OP::ImageInterpolateArgs args = Program.GetOpArgs<OP::ImageInterpolateArgs>(item.At);
+            switch (item.Stage)
+            {
+            case 0:
+                AddOp( FScheduledOp( item.At, item, 1),
+                       FScheduledOp( args.factor, item) );
+                break;
+
+            case 1:
+            {
+           		MUTABLE_CPUPROFILER_SCOPE(IM_INTERPOLATE_1)
+            	
+                // Targets must be consecutive
+                int count = 0;
+                for ( int i=0
+                    ; i<MUTABLE_OP_MAX_INTERPOLATE_COUNT && args.targets[i]
+                    ; ++i )
+                {
+                    count++;
+                }
+
+                float factor = LoadScalar( FCacheAddress(args.factor,item) );
+
+                float delta = 1.0f/(count-1);
+                int min = (int)floorf( factor/delta );
+                int max = (int)ceilf( factor/delta );
+
+                float bifactor = factor/delta - min;
+
+                FScheduledOpData data;
+                data.Interpolate.Bifactor = bifactor;
+				data.Interpolate.Min = FMath::Clamp(min, 0, count - 1);
+				data.Interpolate.Max = FMath::Clamp(max, 0, count - 1);
+				uint32 dataPos = uint32(m_heapData.Add(data));
+
+                if ( bifactor < UE_SMALL_NUMBER )
+                {
+                    AddOp( FScheduledOp( item.At, item, 2, dataPos),
+                            FScheduledOp( args.targets[min], item) );
+                }
+                else if ( bifactor > 1.0f-UE_SMALL_NUMBER )
+                {
+                    AddOp( FScheduledOp( item.At, item, 2, dataPos),
+                            FScheduledOp( args.targets[max], item) );
+                }
+                else
+                {
+                    AddOp( FScheduledOp( item.At, item, 2, dataPos),
+                            FScheduledOp( args.targets[min], item),
+                            FScheduledOp( args.targets[max], item) );
+                }
+                break;
+            }
+
+            case 2:
+            {
+           		MUTABLE_CPUPROFILER_SCOPE(IM_INTERPOLATE_2)
+            		
+                // Targets must be consecutive
+                int count = 0;
+                for ( int i=0
+                    ; i<MUTABLE_OP_MAX_INTERPOLATE_COUNT && args.targets[i]
+                    ; ++i )
+                {
+                    count++;
+                }
+
+                // Factor from 0 to 1 between the two targets
+                const FScheduledOpData& Data = m_heapData[(size_t)item.CustomState];
+                float Bifactor = Data.Interpolate.Bifactor;
+                int32 Min = Data.Interpolate.Min;
+                int32 Max = Data.Interpolate.Max;
+
+                if (Bifactor < UE_SMALL_NUMBER)
+                {
+                    Ptr<const Image> Source = LoadImage(FCacheAddress(args.targets[Min], item));
+					StoreImage(item, Source);
+				}
+                else if (Bifactor > 1.0f - UE_SMALL_NUMBER)
+                {
+                    Ptr<const Image> Source = LoadImage(FCacheAddress(args.targets[Max], item));
+					StoreImage(item, Source);
+				}
+                else
+                {
+					Ptr<const Image> pMin = LoadImage(FCacheAddress(args.targets[Min], item));
+                    Ptr<const Image> pMax = LoadImage(FCacheAddress(args.targets[Max], item));
+
+                    if (pMin && pMax)
+                    {						
+						Ptr<Image> pNew = CloneOrTakeOver(pMin);
+
+						// Be defensive: ensure image sizes match.
+						if (pNew->GetSize() != pMax->GetSize())
+						{
+							MUTABLE_CPUPROFILER_SCOPE(ImageResize_ForInterpolate);
+							Ptr<Image> Resized = CreateImage(pNew->GetSizeX(), pNew->GetSizeY(), pMax->GetLODCount(), pMax->GetFormat(), EInitializationType::NotInitialized);
+							ImOp.ImageResizeLinear(Resized.get(), 0, pMax.get());
+							Release(pMax);
+							pMax = Resized;
+						}
+
+						// Be defensive: ensure format matches.
+						if (pNew->GetFormat() != pMax->GetFormat())
+						{
+							MUTABLE_CPUPROFILER_SCOPE(Format_ForInterpolate);
+
+							Ptr<Image> Formatted = CreateImage(pMax->GetSizeX(), pMax->GetSizeY(), pMax->GetLODCount(), pNew->GetFormat(), EInitializationType::NotInitialized);
+							
+							bool bSuccess = false;
+							ImOp.ImagePixelFormat(bSuccess, m_pSettings->ImageCompressionQuality, Formatted.get(), pMax.get());
+							check(bSuccess);
+							
+							Release(pMax);
+							pMax = Formatted;
+						}
+
+						int32 LevelCount = FMath::Max(pNew->GetLODCount(), pMax->GetLODCount());
+
+						if (pNew->GetLODCount() != LevelCount)
+						{
+							MUTABLE_CPUPROFILER_SCOPE(Mipmap_ForInterpolate);
+						
+							int32 StartLevel = pNew->GetLODCount() - 1;
+							// pNew is local owned, no need to CloneOrTakeOver.
+							pNew->DataStorage.SetNumLODs(LevelCount);
+
+							FMipmapGenerationSettings Settings{};
+							ImOp.ImageMipmap(m_pSettings->ImageCompressionQuality, pNew.get(), pNew.get(), StartLevel, LevelCount, Settings);
+
+						}
+
+						if (pMax->GetLODCount() != LevelCount)
+						{
+							MUTABLE_CPUPROFILER_SCOPE(Mipmap_ForInterpolate);
+
+							int32 StartLevel = pMax->GetLODCount() - 1;
+
+							Ptr<Image> MaxFix = CloneOrTakeOver(pMax);
+							MaxFix->DataStorage.SetNumLODs(LevelCount);
+							
+							FMipmapGenerationSettings Settings{};
+							ImOp.ImageMipmap(m_pSettings->ImageCompressionQuality, MaxFix.get(), MaxFix.get(), StartLevel, LevelCount, Settings);
+
+							pMax = MaxFix;
+						}
+
+                        ImageInterpolate(pNew.get(), pMax.get(), Bifactor);
+
+						Release(pMax);
+						StoreImage(item, pNew);
+					}
+                    else if (pMin)
+                    {
+						StoreImage(item, pMin);
+					}
+                    else if (pMax)
+                    {
+						StoreImage(item, pMax);
+					}
+				}
+
+                break;
+            }
+
+            default:
+                check(false);
+            }
+
+            break;
+        }
+
+        case OP_TYPE::IM_SATURATE:
+        {
+			OP::ImageSaturateArgs args = Program.GetOpArgs<OP::ImageSaturateArgs>(item.At);
+            switch (item.Stage)
+            {
+            case 0:
+                AddOp( FScheduledOp( item.At, item, 1),
+                        FScheduledOp( args.base, item ),
+                        FScheduledOp::FromOpAndOptions( args.factor, item, 0 ));
+                break;
+
+            case 1:
+            {
+				// This has been moved to a task. It should have been intercepted in IssueOp.
+				check(false);
+            }
+
+            default:
+                check(false);
+            }
+
+            break;
+        }
+
+        case OP_TYPE::IM_LUMINANCE:
+        {
+			OP::ImageLuminanceArgs args = Program.GetOpArgs<OP::ImageLuminanceArgs>(item.At);
+            switch (item.Stage)
+            {
+            case 0:
+                AddOp( FScheduledOp( item.At, item, 1),
+                        FScheduledOp( args.base, item ) );
+                break;
+
+            case 1:
+            {
+           		MUTABLE_CPUPROFILER_SCOPE(IM_LUMINANCE_1)
+            		
+                Ptr<const Image> Base = LoadImage( FCacheAddress(args.base,item) );
+
+				Ptr<Image> Result = CreateImage(Base->GetSizeX(), Base->GetSizeY(), Base->GetLODCount(), EImageFormat::IF_L_UBYTE, EInitializationType::NotInitialized);
+                ImageLuminance( Result.get(),Base.get() );
+
+				Release(Base);
+				StoreImage( item, Result );
+                break;
+            }
+
+            default:
+                check(false);
+            }
+
+            break;
+        }
+
+        case OP_TYPE::IM_SWIZZLE:
+        {
+			OP::ImageSwizzleArgs args = Program.GetOpArgs<OP::ImageSwizzleArgs>(item.At);
+            switch (item.Stage)
+            {
+            case 0:
+                AddOp( FScheduledOp( item.At, item, 1),
+                        FScheduledOp( args.sources[0], item ),
+                        FScheduledOp( args.sources[1], item ),
+                        FScheduledOp( args.sources[2], item ),
+                        FScheduledOp( args.sources[3], item ) );
+                break;
+
+            case 1:
+				// This has been moved to a task. It should have been intercepted in IssueOp.
+				check(false);
+				break;
+
+            default:
+                check(false);
+            }
+
+            break;
+        }
+
+        case OP_TYPE::IM_COLOURMAP:
+        {
+			OP::ImageColourMapArgs args = Program.GetOpArgs<OP::ImageColourMapArgs>(item.At);
+            switch (item.Stage)
+            {
+            case 0:
+                    AddOp( FScheduledOp( item.At, item, 1),
+                           FScheduledOp( args.base, item ),
+                           FScheduledOp( args.mask, item ),
+                           FScheduledOp( args.map, item ) );
+                break;
+
+            case 1:
+            {
+            	MUTABLE_CPUPROFILER_SCOPE(IM_COLOURMAP_1)
+            		
+                Ptr<const Image> Source = LoadImage( FCacheAddress(args.base,item) );
+                Ptr<const Image> Mask = LoadImage( FCacheAddress(args.mask,item) );
+                Ptr<const Image> Map = LoadImage( FCacheAddress(args.map,item) );
+
+				bool bOnlyOneMip = (Mask->GetLODCount() < Source->GetLODCount());
+
+				// Be defensive: ensure image sizes match.
+				if (Mask->GetSize() != Source->GetSize())
+				{
+					MUTABLE_CPUPROFILER_SCOPE(ImageResize_ForColourmap);
+					Ptr<Image> Resized = CreateImage(Source->GetSizeX(), Source->GetSizeY(), 1, Mask->GetFormat(), EInitializationType::NotInitialized);
+					ImOp.ImageResizeLinear(Resized.get(), 0, Mask.get());
+					Release(Mask);
+					Mask = Resized;
+				}
+
+				Ptr<Image> Result = CreateImage(Source->GetSizeX(), Source->GetSizeY(), Source->GetLODCount(), Source->GetFormat(), EInitializationType::NotInitialized);
+				ImageColourMap( Result.get(), Source.get(), Mask.get(), Map.get(), bOnlyOneMip);
+
+				if (bOnlyOneMip)
+				{
+					MUTABLE_CPUPROFILER_SCOPE(ImageColourMap_MipFix);
+					FMipmapGenerationSettings DummyMipSettings{};
+					ImageMipmapInPlace(m_pSettings->ImageCompressionQuality, Result.get(), DummyMipSettings);
+				}
+
+				Release(Source);
+				Release(Mask);
+				Release(Map);
+				StoreImage( item, Result );
+                break;
+            }
+
+            default:
+                check(false);
+            }
+
+            break;
+        }
+
+        case OP_TYPE::IM_GRADIENT:
+        {
+			OP::ImageGradientArgs args = Program.GetOpArgs<OP::ImageGradientArgs>(item.At);
+            switch (item.Stage)
+            {
+            case 0:
+                AddOp( FScheduledOp( item.At, item, 1),
+                        FScheduledOp::FromOpAndOptions( args.colour0, item, 0 ),
+                        FScheduledOp::FromOpAndOptions( args.colour1, item, 0 ) );
+                break;
+
+            case 1:
+            {
+				MUTABLE_CPUPROFILER_SCOPE(IM_GRADIENT_1)
+            		
+				FVector4f colour0 = LoadColor(FScheduledOp::FromOpAndOptions(args.colour0, item, 0));
+				FVector4f colour1 = LoadColor(FScheduledOp::FromOpAndOptions(args.colour1, item, 0));
+
+				ImagePtr pResult = CreateImage(args.size[0], args.size[1], 1, EImageFormat::IF_RGB_UBYTE, EInitializationType::NotInitialized);
+                ImageGradient( pResult.get(), colour0, colour1 );
+
+				StoreImage( item, pResult );
+                break;
+            }
+
+            default:
+                check(false);
+            }
+
+            break;
+        }
+
+        case OP_TYPE::IM_BINARISE:
+        {
+			OP::ImageBinariseArgs args = Program.GetOpArgs<OP::ImageBinariseArgs>(item.At);
+            switch (item.Stage)
+            {
+            case 0:
+                AddOp( FScheduledOp( item.At, item, 1),
+                        FScheduledOp( args.base, item ),
+                        FScheduledOp::FromOpAndOptions( args.threshold, item, 0 ) );
+                break;
+
+            case 1:
+            {
+            	MUTABLE_CPUPROFILER_SCOPE(IM_BINARISE_1)
+            		
+                Ptr<const Image> pA = LoadImage( FCacheAddress(args.base,item) );
+
+                float c = LoadScalar(FScheduledOp::FromOpAndOptions(args.threshold, item, 0));
+
+                Ptr<Image> Result = CreateImage(pA->GetSizeX(), pA->GetSizeY(), pA->GetLODCount(), EImageFormat::IF_L_UBYTE, EInitializationType::NotInitialized);
+				ImageBinarise( Result.get(), pA.get(), c );
+
+				Release(pA);
+				StoreImage( item, Result );
+                break;
+            }
+
+            default:
+                check(false);
+            }
+
+            break;
+        }
+
+		case OP_TYPE::IM_INVERT:
+		{
+			OP::ImageInvertArgs args = Program.GetOpArgs<OP::ImageInvertArgs>(item.At);
+			switch (item.Stage)
+			{
+			case 0:
+				AddOp(FScheduledOp(item.At, item, 1), FScheduledOp(args.base, item));
+				break;
+
+			case 1:
+			{
+				// This has been moved to a task. It should have been intercepted in IssueOp.
+				check(false);
+			}
+
+			default:
+				check(false);
+			}
+
+			break;
+		}
+
+        case OP_TYPE::IM_PLAINCOLOUR:
+        {
+			OP::ImagePlainColourArgs args = Program.GetOpArgs<OP::ImagePlainColourArgs>(item.At);
+            switch (item.Stage)
+            {
+            case 0:
+				AddOp( FScheduledOp( item.At, item, 1), FScheduledOp::FromOpAndOptions( args.colour, item, 0 ) );
+                break;
+
+            case 1:
+            {
+            	MUTABLE_CPUPROFILER_SCOPE(IM_PLAINCOLOUR_1)
+            		
+				FVector4f c = LoadColor(FScheduledOp::FromOpAndOptions(args.colour, item, 0));
+
+				uint16 SizeX = args.size[0];
+				uint16 SizeY = args.size[1];
+				int32 LODs = args.LODs;
+				
+				// This means all the mip chain
+				if (LODs == 0)
+				{
+					LODs = FMath::CeilLogTwo(FMath::Max(SizeX,SizeY));
+				}
+
+				for (int l=0; l<item.ExecutionOptions; ++l)
+				{
+					SizeX = FMath::Max(uint16(1), FMath::DivideAndRoundUp(SizeX, uint16(2)));
+					SizeY = FMath::Max(uint16(1), FMath::DivideAndRoundUp(SizeY, uint16(2)));
+					--LODs;
+				}
+
+                ImagePtr pA = CreateImage( SizeX, SizeY, FMath::Max(LODs,1), EImageFormat(args.format), EInitializationType::NotInitialized );
+
+				ImOp.FillColor(pA.get(), c);
+
+				StoreImage( item, pA );
+                break;
+            }
+
+            default:
+                check(false);
+            }
+
+            break;
+        }
+
+		case OP_TYPE::IM_REFERENCE:
+		{
+			OP::ResourceReferenceArgs Args = Program.GetOpArgs<OP::ResourceReferenceArgs>(item.At);
+			switch (item.Stage)
+			{
+			case 0:
+			{
+				Ptr<Image> Result;
+				if (Args.ForceLoad)
+				{
+					// This should never be reached because it should have been caught as a Task in IssueOp
+					check(false);
+				}
+				else
+				{
+					Result = Image::CreateAsReference(Args.ID, Args.ImageDesc, false);
+				}
+				StoreImage(item, Result);
+				break;
+			}
+
+			default:
+				check(false);
+			}
+
+			break;
+		}
+
+        case OP_TYPE::IM_CROP:
+        {
+			OP::ImageCropArgs args = Program.GetOpArgs<OP::ImageCropArgs>(item.At);
+            switch (item.Stage)
+            {
+            case 0:
+                AddOp( FScheduledOp( item.At, item, 1), FScheduledOp( args.source, item ) );
+                break;
+
+            case 1:
+            {
+            	MUTABLE_CPUPROFILER_SCOPE(IM_CROP_1)
+            		
+                Ptr<const Image> pA = LoadImage( FCacheAddress(args.source,item) );
+
+                box< UE::Math::TIntVector2<int32> > rect;
+                rect.min[0] = args.minX;
+                rect.min[1] = args.minY;
+                rect.size[0] = args.sizeX;
+                rect.size[1] = args.sizeY;
+
+				// Apply ther mipmap reduction to the crop rectangle.
+				int32 MipsToSkip = item.ExecutionOptions;
+				while ( MipsToSkip>0 && rect.size[0]>0 && rect.size[1]>0 )
+				{
+					rect.min[0] /= 2;
+					rect.min[1] /= 2;
+					rect.size[0] /= 2;
+					rect.size[1] /= 2;
+					MipsToSkip--;
+				}
+
+				ImagePtr pResult;
+				if (!rect.IsEmpty())
+				{
+					pResult = CreateImage( rect.size[0], rect.size[1], 1, pA->GetFormat(), EInitializationType::NotInitialized);
+					ImOp.ImageCrop(pResult.get(), m_pSettings->ImageCompressionQuality, pA.get(), rect);
+				}
+
+				Release(pA);
+				StoreImage( item, pResult );
+                break;
+            }
+
+            default:
+                check(false);
+            }
+
+            break;
+        }
+
+        case OP_TYPE::IM_PATCH:
+        {
+			// TODO: This is optimized for memory-usage but base and patch could be requested at the same time
+			OP::ImagePatchArgs args = Program.GetOpArgs<OP::ImagePatchArgs>(item.At);
+            switch (item.Stage)
+            {
+			case 0:
+				AddOp(FScheduledOp(item.At, item, 1), FScheduledOp(args.base, item));
+				break;
+
+			case 1:
+				AddOp(FScheduledOp(item.At, item, 2), FScheduledOp(args.patch, item));
+				break;
+
+			case 2:
+            {
+                MUTABLE_CPUPROFILER_SCOPE(IM_PATCH_1)
+
+                Ptr<const Image> pA = LoadImage( FCacheAddress(args.base,item) );
+                Ptr<const Image> pB = LoadImage( FCacheAddress(args.patch,item) );
+
+				// Failsafe
+				if (!pA || !pB)
+				{
+					Release(pB);
+					StoreImage(item, pA);
+					break;
+				}
+
+				// Apply the mipmap reduction to the crop rectangle.
+				int32 MipsToSkip = item.ExecutionOptions;
+				box<FIntVector2> rect;
+				rect.min[0] = args.minX / (1 << MipsToSkip);
+				rect.min[1] = args.minY / (1 << MipsToSkip);
+				rect.size[0] = pB->GetSizeX();
+				rect.size[1] = pB->GetSizeY();
+
+                ImagePtr pResult = CloneOrTakeOver(pA);
+
+				bool bApplyPatch = !rect.IsEmpty();
+				if (bApplyPatch)
+				{
+					// Change the block image format if it doesn't match the composed image
+					// This is usually enforced at object compilation time.
+					if (pResult->GetFormat() != pB->GetFormat())
+					{
+						MUTABLE_CPUPROFILER_SCOPE(ImagPatchReformat);
+
+						EImageFormat format = GetMostGenericFormat(pResult->GetFormat(), pB->GetFormat());
+
+						const FImageFormatData& finfo = GetImageFormatData(format);
+						if (finfo.PixelsPerBlockX == 0)
+						{
+							format = GetUncompressedFormat(format);
+						}
+
+						if (pResult->GetFormat() != format)
+						{
+							Ptr<Image> Formatted = CreateImage(pResult->GetSizeX(), pResult->GetSizeY(), pResult->GetLODCount(), format, EInitializationType::NotInitialized);
+							bool bSuccess = false;
+							ImOp.ImagePixelFormat(bSuccess, m_pSettings->ImageCompressionQuality, Formatted.get(), pResult.get());
+							check(bSuccess);
+							Release(pResult);
+							pResult = Formatted;
+						}
+						if (pB->GetFormat() != format)
+						{
+							Ptr<Image> Formatted = CreateImage(pB->GetSizeX(), pB->GetSizeY(), pB->GetLODCount(), format, EInitializationType::NotInitialized);
+							bool bSuccess = false;
+							ImOp.ImagePixelFormat(bSuccess, m_pSettings->ImageCompressionQuality, Formatted.get(), pB.get());
+							check(bSuccess);
+							Release(pB);
+							pB = Formatted;
+						}
+					}
+
+					// Don't patch if below the image compression block size.
+					const FImageFormatData& finfo = GetImageFormatData(pResult->GetFormat());
+					bApplyPatch =
+						(rect.min[0] % finfo.PixelsPerBlockX == 0) &&
+						(rect.min[1] % finfo.PixelsPerBlockY == 0) &&
+						(rect.size[0] % finfo.PixelsPerBlockX == 0) &&
+						(rect.size[1] % finfo.PixelsPerBlockY == 0) &&
+						(rect.min[0] + rect.size[0]) <= pResult->GetSizeX() &&
+						(rect.min[1] + rect.size[1]) <= pResult->GetSizeY()
+						;
+				}
+
+				if (bApplyPatch)
+				{
+					ImOp.ImageCompose(pResult.get(), pB.get(), rect);
+					pResult->m_flags = 0;
+				}
+				else
+				{
+					// This happens very often when skipping mips, and floods the log.
+					//UE_LOG( LogMutableCore, Verbose, TEXT("Skipped patch operation for image not fitting the block compression size. Small image? Patch rect is (%d, %d), (%d, %d), base is (%d, %d)"),
+					//	rect.min[0], rect.min[1], rect.size[0], rect.size[1], pResult->GetSizeX(), pResult->GetSizeY());
+				}
+
+				Release(pB);
+				StoreImage( item, pResult );
+                break;
+            }
+
+            default:
+                check(false);
+            }
+
+            break;
+        }
+
+        case OP_TYPE::IM_RASTERMESH:
+        {
+			OP::ImageRasterMeshArgs args = Program.GetOpArgs<OP::ImageRasterMeshArgs>(item.At);
+            switch (item.Stage)
+            {
+            case 0:
+				if (args.image)
+				{
+					AddOp(FScheduledOp(item.At, item, 1),
+						FScheduledOp::FromOpAndOptions(args.mesh, item, 0),
+						FScheduledOp::FromOpAndOptions(args.projector, item, 0));
+				}
+				else
+				{
+					AddOp(FScheduledOp(item.At, item, 1),
+						FScheduledOp::FromOpAndOptions(args.mesh, item, 0));
+				}
+                break;
+
+			case 1:
+			{
+				MUTABLE_CPUPROFILER_SCOPE(IM_RASTERMESH_1)
+
+				Ptr<const Mesh> pMesh = LoadMesh(FScheduledOp::FromOpAndOptions(args.mesh, item, 0));
+
+				// If no image, we are generating a flat mesh UV raster. This is the final stage in this case.
+				if (!args.image)
+				{
+					uint16 SizeX = args.sizeX;
+					uint16 SizeY = args.sizeY;
+					UE::Math::TIntVector2<uint16> CropMin(args.CropMinX, args.CropMinY);
+					UE::Math::TIntVector2<uint16> UncroppedSize(args.UncroppedSizeX, args.UncroppedSizeY);
+
+					// Drop mips while possible
+					int32 MipsToDrop = item.ExecutionOptions;
+					bool bUseCrop = UncroppedSize[0] > 0;
+					while (MipsToDrop && !(SizeX % 2) && !(SizeY % 2))
+					{
+						SizeX = FMath::Max(uint16(1), FMath::DivideAndRoundUp(SizeX, uint16(2)));
+						SizeY = FMath::Max(uint16(1), FMath::DivideAndRoundUp(SizeY, uint16(2)));
+						if (bUseCrop)
+						{
+							CropMin[0] = FMath::DivideAndRoundUp(CropMin[0], uint16(2));
+							CropMin[1] = FMath::DivideAndRoundUp(CropMin[1], uint16(2));
+							UncroppedSize[0] = FMath::Max(uint16(1), FMath::DivideAndRoundUp(UncroppedSize[0], uint16(2)));
+							UncroppedSize[1] = FMath::Max(uint16(1), FMath::DivideAndRoundUp(UncroppedSize[1], uint16(2)));
+						}
+						--MipsToDrop;
+					}
+
+                    // Flat mesh UV raster
+					Ptr<Image> ResultImage = CreateImage(SizeX, SizeY, 1, EImageFormat::IF_L_UBYTE, EInitializationType::Black);
+					if (pMesh)
+					{
+						ImageRasterMesh(pMesh.get(), ResultImage.get(), args.LayoutIndex, args.BlockId, CropMin, UncroppedSize);
+						Release(pMesh);
+					}
+
+					// Stop execution.
+					StoreImage(item, ResultImage);
+					break;
+				}
+
+				const int32 MipsToSkip = item.ExecutionOptions;
+				int32 ProjectionMip = MipsToSkip;
+
+				FScheduledOpData Data;
+				Data.RasterMesh.Mip = ProjectionMip;
+				Data.RasterMesh.MipValue = static_cast<float>(ProjectionMip);
+				FProjector Projector = LoadProjector(FScheduledOp::FromOpAndOptions(args.projector, item, 0));
+
+				EMinFilterMethod MinFilterMethod = Invoke([&]() -> EMinFilterMethod
+				{
+					if (ForcedProjectionMode == 0)
+					{
+						return EMinFilterMethod::None;
+					}
+					else if (ForcedProjectionMode == 1)
+					{
+						return EMinFilterMethod::TotalAreaHeuristic;
+					}
+						
+					return static_cast<EMinFilterMethod>(args.MinFilterMethod);
+				});
+
+				if (MinFilterMethod == EMinFilterMethod::TotalAreaHeuristic)
+				{
+					FVector2f TargetImageSizeF = FVector2f(
+						FMath::Max(args.sizeX >> MipsToSkip, 1),
+						FMath::Max(args.sizeY >> MipsToSkip, 1));
+					FVector2f SourceImageSizeF = FVector2f(args.SourceSizeX, args.SourceSizeY);
+						
+					if (pMesh)
+					{ 
+						const float ComputedMip = ComputeProjectedFootprintBestMip(pMesh.get(), Projector, TargetImageSizeF, SourceImageSizeF);
+
+						Data.RasterMesh.MipValue = FMath::Max(0.0f, ComputedMip + GlobalProjectionLodBias);
+						Data.RasterMesh.Mip = static_cast<uint8>(FMath::FloorToInt32(Data.RasterMesh.MipValue));
+					}
+				}
+		
+				const int32 DataHeapAddress = m_heapData.Add(Data);
+
+				// pMesh is need again in the next stage, store it in the heap.
+				m_heapData[DataHeapAddress].Resource = const_cast<Mesh*>(pMesh.get());
+
+				AddOp(FScheduledOp(item.At, item, 2, DataHeapAddress),
+					FScheduledOp::FromOpAndOptions(args.projector, item, 0),
+					FScheduledOp::FromOpAndOptions(args.image, item, Data.RasterMesh.Mip),
+					FScheduledOp(args.mask, item),
+					FScheduledOp::FromOpAndOptions(args.angleFadeProperties, item, 0));
+
+				break;
+			}
+
+            case 2:
+            {
+				MUTABLE_CPUPROFILER_SCOPE(IM_RASTERMESH_2)
+
+				if (!args.image)
+				{
+					// This case is treated at the previous stage.
+					check(false);
+					StoreImage(item, nullptr);
+					break;
+				}
+
+				FScheduledOpData& Data = m_heapData[item.CustomState];
+
+				// Unsafe downcast, should be fine as it is known to be a Mesh.
+				Ptr<const Mesh> pMesh = static_cast<Mesh*>(Data.Resource.get());
+				Data.Resource = nullptr;
+
+				if (!pMesh)
+				{
+					check(false);
+					StoreImage(item, nullptr);
+					break;
+				}
+
+				uint16 SizeX = args.sizeX;
+				uint16 SizeY = args.sizeY;
+				UE::Math::TIntVector2<uint16> CropMin(args.CropMinX, args.CropMinY);
+				UE::Math::TIntVector2<uint16> UncroppedSize(args.UncroppedSizeX, args.UncroppedSizeY);
+
+				// Drop mips while possible
+				int32 MipsToDrop = item.ExecutionOptions;
+				bool bUseCrop = UncroppedSize[0] > 0;
+				while (MipsToDrop && !(SizeX % 2) && !(SizeY % 2))
+				{
+					SizeX = FMath::Max(uint16(1),FMath::DivideAndRoundUp(SizeX, uint16(2)));
+					SizeY = FMath::Max(uint16(1),FMath::DivideAndRoundUp(SizeY, uint16(2)));
+					if (bUseCrop)
+					{
+						CropMin[0] = FMath::DivideAndRoundUp(CropMin[0], uint16(2));
+						CropMin[1] = FMath::DivideAndRoundUp(CropMin[1], uint16(2));
+						UncroppedSize[0] = FMath::Max(uint16(1), FMath::DivideAndRoundUp(UncroppedSize[0], uint16(2)));
+						UncroppedSize[1] = FMath::Max(uint16(1), FMath::DivideAndRoundUp(UncroppedSize[1], uint16(2)));
+					}
+					--MipsToDrop;
+				}
+
+				// Raster with projection
+				Ptr<const Image> Source = LoadImage(FCacheAddress(args.image, item.ExecutionIndex, Data.RasterMesh.Mip));
+
+				Ptr<const Image> Mask = nullptr;
+				if (args.mask)
+				{
+					Mask = LoadImage(FCacheAddress(args.mask, item));
+
+					// TODO: This shouldn't happen, but be defensive.
+					FImageSize ResultSize(SizeX, SizeY);
+					if (Mask && Mask->GetSize()!= ResultSize)
+					{
+						MUTABLE_CPUPROFILER_SCOPE(ImageResize_MaskFixForProjection);
+
+						Ptr<Image> Resized = CreateImage(SizeX, SizeY, Mask->GetLODCount(), Mask->GetFormat(), EInitializationType::NotInitialized);
+						ImOp.ImageResizeLinear(Resized.get(), 0, Mask.get());
+						Release(Mask);
+						Mask = Resized;
+					}
+				}
+
+				float fadeStart = 180.0f;
+				float fadeEnd = 180.0f;
+				if ( args.angleFadeProperties )
+				{
+					FVector4f fadeProperties = LoadColor(FScheduledOp::FromOpAndOptions(args.angleFadeProperties, item, 0));
+					fadeStart = fadeProperties[0];
+					fadeEnd = fadeProperties[1];
+				}
+				const float FadeStartRad = FMath::DegreesToRadians(fadeStart);
+				const float FadeEndRad = FMath::DegreesToRadians(fadeEnd);
+
+				EImageFormat Format = Source ? GetUncompressedFormat(Source->GetFormat()) : EImageFormat::IF_L_UBYTE;
+
+				if (Source && Source->GetFormat()!=Format)
+				{
+					MUTABLE_CPUPROFILER_SCOPE(RunCode_RasterMesh_ReformatSource);
+					Ptr<Image> Formatted = CreateImage(Source->GetSizeX(), Source->GetSizeY(), Source->GetLODCount(), Format, EInitializationType::NotInitialized);
+					bool bSuccess = false;
+					ImOp.ImagePixelFormat(bSuccess, m_pSettings->ImageCompressionQuality, Formatted.get(), Source.get());
+					check(bSuccess); 
+					Release(Source);
+					Source = Formatted;
+				}
+
+				// Allocate memory for the temporary buffers
+				FScratchImageProject Scratch;
+				Scratch.Vertices.SetNum(pMesh->GetVertexCount());
+				Scratch.CulledVertex.SetNum(pMesh->GetVertexCount());
+
+				ESamplingMethod SamplingMethod = Invoke([&]() -> ESamplingMethod
+				{
+					if (ForcedProjectionMode == 0)
+					{
+						return ESamplingMethod::Point;
+					}
+					else if (ForcedProjectionMode == 1)
+					{
+						return ESamplingMethod::BiLinear;
+					}
+					
+					return static_cast<ESamplingMethod>(args.SamplingMethod);
+				});
+
+				if (SamplingMethod == ESamplingMethod::BiLinear)
+				{
+					if (Source->GetLODCount() < 2 && Source->GetSizeX() > 1 && Source->GetSizeY() > 1)
+					{
+						MUTABLE_CPUPROFILER_SCOPE(RunCode_RasterMesh_BilinearMipGen);
+
+						Ptr<Image> OwnedSource = CloneOrTakeOver(Source);
+
+						OwnedSource->DataStorage.SetNumLODs(2);
+						ImageMipmapInPlace(0, OwnedSource.get(), FMipmapGenerationSettings{});
+
+						Source = OwnedSource;
+					}
+				}
+
+				// Allocate new image after bilinear mip generation to reduce operation memory peak.
+				Ptr<Image> New = CreateImage(SizeX, SizeY, 1, Format, EInitializationType::Black);
+
+				if (args.projector && Source && Source->GetSizeX() > 0 && Source->GetSizeY() > 0)
+				{
+					FProjector Projector = LoadProjector(FScheduledOp::FromOpAndOptions(args.projector, item, 0));
+
+					switch (Projector.type)
+					{
+					case PROJECTOR_TYPE::PLANAR:
+						ImageRasterProjectedPlanar(pMesh.get(), New.get(),
+							Source.get(), Mask.get(),
+							args.bIsRGBFadingEnabled, args.bIsAlphaFadingEnabled,
+							SamplingMethod,
+							FadeStartRad, FadeEndRad, FMath::Frac(Data.RasterMesh.MipValue),
+							args.LayoutIndex, args.BlockId,
+							CropMin, UncroppedSize,
+							&Scratch, bUseProjectionVectorImpl);
+						break;
+
+					case PROJECTOR_TYPE::WRAPPING:
+						ImageRasterProjectedWrapping(pMesh.get(), New.get(),
+							Source.get(), Mask.get(),
+							args.bIsRGBFadingEnabled, args.bIsAlphaFadingEnabled,
+							SamplingMethod,
+							FadeStartRad, FadeEndRad, FMath::Frac(Data.RasterMesh.MipValue),
+							args.LayoutIndex, args.BlockId,
+							CropMin, UncroppedSize,
+							&Scratch, bUseProjectionVectorImpl);
+						break;
+
+					case PROJECTOR_TYPE::CYLINDRICAL:
+						ImageRasterProjectedCylindrical(pMesh.get(), New.get(),
+							Source.get(), Mask.get(),
+							args.bIsRGBFadingEnabled, args.bIsAlphaFadingEnabled,
+							SamplingMethod,
+							FadeStartRad, FadeEndRad, FMath::Frac(Data.RasterMesh.MipValue),
+							args.LayoutIndex,
+							Projector.projectionAngle,
+							CropMin, UncroppedSize,
+							&Scratch, bUseProjectionVectorImpl);
+						break;
+
+					default:
+						check(false);
+						break;
+					}
+				}
+
+				Release(pMesh);
+				Release(Source);
+				Release(Mask);
+				StoreImage(item, New);
+
+                break;
+            }
+
+            default:
+                check(false);
+            }
+
+            break;
+        }
+
+        case OP_TYPE::IM_MAKEGROWMAP:
+        {
+			OP::ImageMakeGrowMapArgs args = Program.GetOpArgs<OP::ImageMakeGrowMapArgs>(item.At);
+            switch (item.Stage)
+            {
+            case 0:
+                AddOp( FScheduledOp( item.At, item, 1), FScheduledOp( args.mask, item) );
+                break;
+
+            case 1:
+            {
+                MUTABLE_CPUPROFILER_SCOPE(IM_MAKEGROWMAP_1)
+
+                Ptr<const Image> Mask = LoadImage( FCacheAddress(args.mask,item) );
+
+                Ptr<Image> Result = CreateImage( Mask->GetSizeX(), Mask->GetSizeY(), Mask->GetLODCount(), EImageFormat::IF_L_UBYTE, EInitializationType::NotInitialized);
+
+                ImageMakeGrowMap(Result.get(), Mask.get(), args.border );
+				Result->m_flags |= Image::IF_CANNOT_BE_SCALED;
+
+				Release(Mask);
+                StoreImage( item, Result);
+                break;
+            }
+
+            default:
+                check(false);
+            }
+
+            break;
+        }
+
+        case OP_TYPE::IM_DISPLACE:
+        {
+			OP::ImageDisplaceArgs args = Program.GetOpArgs<OP::ImageDisplaceArgs>(item.At);
+            switch (item.Stage)
+            {
+            case 0:
+                AddOp( FScheduledOp( item.At, item, 1),
+                        FScheduledOp( args.source, item ),
+                        FScheduledOp( args.displacementMap, item ) );
+				break;
+
+            case 1:
+            {
+                MUTABLE_CPUPROFILER_SCOPE(IM_DISPLACE_1)
+
+                Ptr<const Image> Source = LoadImage( FCacheAddress(args.source,item) );
+                Ptr<const Image> pMap = LoadImage( FCacheAddress(args.displacementMap,item) );
+
+				if (!Source)
+				{
+					Release(pMap);
+					StoreImage(item, nullptr);
+					break;
+				}
+
+				// TODO: This shouldn't happen: displacement maps cannot be scaled because their information
+				// is resolution sensitive (pixel offsets). If the size doesn't match, scale the source, apply 
+				// displacement and then unscale it.
+				FImageSize OriginalSourceScale = Source->GetSize();
+				if (OriginalSourceScale[0]>0 && OriginalSourceScale[1]>0 && OriginalSourceScale != pMap->GetSize())
+				{
+					MUTABLE_CPUPROFILER_SCOPE(ImageResize_EmergencyHackForDisplacementStep1);
+
+					Ptr<Image> Resized = CreateImage(pMap->GetSizeX(), pMap->GetSizeY(), Source->GetLODCount(), Source->GetFormat(), EInitializationType::NotInitialized);
+					ImOp.ImageResizeLinear(Resized.get(), 0, Source.get());
+					Release(Source);
+					Source = Resized;
+				}
+
+				// This works based on the assumption that displacement maps never read from a position they actually write to.
+				// Since they are used for UV border expansion, this should always be the case.
+				Ptr<Image> Result = CloneOrTakeOver(Source);
+
+				if (OriginalSourceScale[0] > 0 && OriginalSourceScale[1] > 0)
+				{
+					ImageDisplace(Result.get(), Result.get(), pMap.get());
+
+					if (OriginalSourceScale != Result->GetSize())
+					{
+						MUTABLE_CPUPROFILER_SCOPE(ImageResize_EmergencyHackForDisplacementStep2);
+						Ptr<Image> Resized = CreateImage(OriginalSourceScale[0], OriginalSourceScale[1], Result->GetLODCount(), Result->GetFormat(), EInitializationType::NotInitialized);
+						ImOp.ImageResizeLinear(Resized.get(), 0, Result.get());
+						Release(Result);
+						Result = Resized;
+					}
+				}
+
+				Release(pMap);
+                StoreImage( item, Result);
+                break;
+            }
+
+            default:
+                check(false);
+            }
+
+            break;
+        }
+
+        case OP_TYPE::IM_TRANSFORM:
+        {
+            const OP::ImageTransformArgs Args = Program.GetOpArgs<OP::ImageTransformArgs>(item.At);
+
+            switch (item.Stage)
+            {
+            case 0:
+			{
+				const TArray<FScheduledOp, TFixedAllocator<2>> Deps = 
+				{
+					FScheduledOp(Args.ScaleX, item),
+					FScheduledOp(Args.ScaleY, item),
+				};
+
+                AddOp(FScheduledOp(item.At, item, 1), Deps);
+
+				break;
+			}
+			case 1:
+			{
+            	MUTABLE_CPUPROFILER_SCOPE(IM_TRANSFORM_1)
+
+				FVector2f Scale = FVector2f(
+                        Args.ScaleX ? LoadScalar(FCacheAddress(Args.ScaleX, item)) : 1.0f,
+                        Args.ScaleY ? LoadScalar(FCacheAddress(Args.ScaleY, item)) : 1.0f);
+	
+				using FUint16Vector2 = UE::Math::TIntVector2<uint16>;
+				const FUint16Vector2 DestSizeI = Invoke([&]() 
+				{
+					int32 MipsToDrop = item.ExecutionOptions;
+					
+					FUint16Vector2 Size = FUint16Vector2(
+							Args.SizeX > 0 ? Args.SizeX : Args.SourceSizeX, 
+							Args.SizeY > 0 ? Args.SizeY : Args.SourceSizeY); 
+
+					while (MipsToDrop && !(Size.X % 2) && !(Size.Y % 2))
+					{
+						Size.X = FMath::Max(uint16(1), FMath::DivideAndRoundUp(Size.X, uint16(2)));
+						Size.Y = FMath::Max(uint16(1), FMath::DivideAndRoundUp(Size.Y, uint16(2)));
+						--MipsToDrop;
+					}
+
+					return FUint16Vector2(FMath::Max(Size.X, uint16(1)), FMath::Max(Size.Y, uint16(1)));
+				});
+
+				const FVector2f DestSize   = FVector2f(DestSizeI.X, DestSizeI.Y);
+				const FVector2f SourceSize = FVector2f(FMath::Max(Args.SourceSizeX, uint16(1)), FMath::Max(Args.SourceSizeY, uint16(1)));
+
+				FVector2f AspectCorrectionScale = FVector2f(1.0f, 1.0f);
+				if (Args.bKeepAspectRatio)
+				{
+					const float DestAspectOverSrcAspect = (DestSize.X * SourceSize.Y) / (DestSize.Y * SourceSize.X);
+
+					AspectCorrectionScale = DestAspectOverSrcAspect > 1.0f 
+										  ? FVector2f(1.0f/DestAspectOverSrcAspect, 1.0f) 
+										  : FVector2f(1.0f, DestAspectOverSrcAspect); 
+				}
+			
+				const FTransform2f Transform = FTransform2f(FVector2f(-0.5f))
+					.Concatenate(FTransform2f(FScale2f(Scale)))
+					.Concatenate(FTransform2f(FScale2f(AspectCorrectionScale)))
+					.Concatenate(FTransform2f(FVector2f(0.5f)));
+
+				FBox2f NormalizedCropRect(ForceInit);
+				NormalizedCropRect += Transform.TransformPoint(FVector2f(0.0f, 0.0f));
+				NormalizedCropRect += Transform.TransformPoint(FVector2f(1.0f, 0.0f));
+				NormalizedCropRect += Transform.TransformPoint(FVector2f(0.0f, 1.0f));
+				NormalizedCropRect += Transform.TransformPoint(FVector2f(1.0f, 1.0f));
+
+				const FVector2f ScaledSourceSize = NormalizedCropRect.GetSize() * DestSize;
+
+				const float BestMip = 
+					FMath::Log2(FMath::Max(1.0f, FMath::Square(SourceSize.GetMin()))) * 0.5f - 
+				    FMath::Log2(FMath::Max(1.0f, FMath::Square(ScaledSourceSize.GetMin()))) * 0.5f;
+
+				FScheduledOpData HeapData;
+				HeapData.ImageTransform.SizeX = DestSizeI.X;
+				HeapData.ImageTransform.SizeY = DestSizeI.Y;
+				FPlatformMath::StoreHalf(&HeapData.ImageTransform.ScaleXEncodedHalf, Scale.X);
+				FPlatformMath::StoreHalf(&HeapData.ImageTransform.ScaleYEncodedHalf, Scale.Y);
+				HeapData.ImageTransform.MipValue = BestMip + GlobalImageTransformLodBias;
+
+				const int32 HeapDataAddress = m_heapData.Add(HeapData);
+
+				const uint8 Mip = static_cast<uint8>(FMath::Max(0, FMath::FloorToInt(HeapData.ImageTransform.MipValue)));
+				const TArray<FScheduledOp, TFixedAllocator<4>> Deps = 
+				{
+					FScheduledOp::FromOpAndOptions(Args.Base, item, Mip),
+					FScheduledOp(Args.OffsetX,  item),
+					FScheduledOp(Args.OffsetY,  item),
+					FScheduledOp(Args.Rotation, item) 
+				};
+				
+                AddOp(FScheduledOp(item.At, item, 2, HeapDataAddress), Deps);
+
+				break;
+			}
+            case 2:
+            {
+				MUTABLE_CPUPROFILER_SCOPE(IM_TRANSFORM_2);
+			
+				const FScheduledOpData HeapData = m_heapData[item.CustomState];
+
+				const uint8 Mip = static_cast<uint8>(FMath::Max(0, FMath::FloorToInt(HeapData.ImageTransform.MipValue)));
+				Ptr<const Image> Source = LoadImage(FCacheAddress(Args.Base, item.ExecutionIndex, Mip));
+
+				const FVector2f Offset = FVector2f(
+                        Args.OffsetX ? LoadScalar(FCacheAddress(Args.OffsetX, item)) : 0.0f,
+                        Args.OffsetY ? LoadScalar(FCacheAddress(Args.OffsetY, item)) : 0.0f);
+
+                FVector2f Scale = FVector2f(
+						FPlatformMath::LoadHalf(&HeapData.ImageTransform.ScaleXEncodedHalf),
+						FPlatformMath::LoadHalf(&HeapData.ImageTransform.ScaleYEncodedHalf));
+
+				FVector2f AspectCorrectionScale = FVector2f(1.0f, 1.0f);
+				if (Args.bKeepAspectRatio)
+				{
+					const FVector2f DestSize   = FVector2f(HeapData.ImageTransform.SizeX, HeapData.ImageTransform.SizeY);
+					const FVector2f SourceSize = FVector2f(FMath::Max(Args.SourceSizeX, uint16(1)), FMath::Max(Args.SourceSizeY, uint16(1)));
+					
+					const float DestAspectOverSrcAspect = (DestSize.X * SourceSize.Y) / (DestSize.Y * SourceSize.X);
+					
+					AspectCorrectionScale = DestAspectOverSrcAspect > 1.0f 
+										  ? FVector2f(1.0f/DestAspectOverSrcAspect, 1.0f) 
+										  : FVector2f(1.0f, DestAspectOverSrcAspect); 
+				}
+
+				// Map Range [0..1] to a full rotation
+                const float RotationRad = LoadScalar(FCacheAddress(Args.Rotation, item)) * UE_TWO_PI;
+	
+				EImageFormat SourceFormat = Source->GetFormat();
+				EImageFormat Format = GetUncompressedFormat(SourceFormat);
+
+				if (Format != SourceFormat)
+				{
+					MUTABLE_CPUPROFILER_SCOPE(RunCode_ImageTransform_FormatFixup);	
+					Ptr<Image> Formatted = CreateImage(Source->GetSizeX(), Source->GetSizeY(), Source->GetLODCount(), Format, EInitializationType::NotInitialized);
+					bool bSuccess = false;
+					ImOp.ImagePixelFormat(bSuccess, m_pSettings->ImageCompressionQuality, Formatted.get(), Source.get());
+					check(bSuccess); 
+
+					Release(Source);
+					Source = Formatted;
+				}
+
+				if (Source->GetLODCount() < 2 && Source->GetSizeX() > 1 && Source->GetSizeY() > 1)
+				{
+					MUTABLE_CPUPROFILER_SCOPE(RunCode_ImageTransform_BilinearMipGen);
+
+					Ptr<Image> OwnedSource = CloneOrTakeOver(Source);
+					OwnedSource->DataStorage.SetNumLODs(2);
+
+					ImageMipmapInPlace(0, OwnedSource.get(), FMipmapGenerationSettings{});
+
+					Source = OwnedSource;
+				}
+
+				Scale.X = FMath::IsNearlyZero(Scale.X, UE_KINDA_SMALL_NUMBER) ? UE_KINDA_SMALL_NUMBER : Scale.X;
+				Scale.Y = FMath::IsNearlyZero(Scale.Y, UE_KINDA_SMALL_NUMBER) ? UE_KINDA_SMALL_NUMBER : Scale.Y;
+
+				AspectCorrectionScale.X = FMath::IsNearlyZero(AspectCorrectionScale.X, UE_KINDA_SMALL_NUMBER) 
+									    ? UE_KINDA_SMALL_NUMBER 
+										: AspectCorrectionScale.X;
+
+				AspectCorrectionScale.Y = FMath::IsNearlyZero(AspectCorrectionScale.Y, UE_KINDA_SMALL_NUMBER) 
+										? UE_KINDA_SMALL_NUMBER 
+										: AspectCorrectionScale.Y;
+
+				const FTransform2f Transform = FTransform2f(FVector2f(-0.5f))
+						.Concatenate(FTransform2f(FScale2f(Scale)))
+						.Concatenate(FTransform2f(FQuat2f(RotationRad)))
+						.Concatenate(FTransform2f(FScale2f(AspectCorrectionScale)))
+						.Concatenate(FTransform2f(Offset + FVector2f(0.5f)));
+
+				const EAddressMode AddressMode = static_cast<EAddressMode>(Args.AddressMode);
+
+				const EInitializationType InitType = AddressMode == EAddressMode::ClampToBlack 
+											       ? EInitializationType::Black
+											       : EInitializationType::NotInitialized;
+
+				Ptr<Image> Result = CreateImage(
+						HeapData.ImageTransform.SizeX, HeapData.ImageTransform.SizeY, 1, Format, InitType);
+
+				const float MipFactor = FMath::Frac(FMath::Max(0.0f, HeapData.ImageTransform.MipValue));
+				ImageTransform(Result.get(), Source.get(), Transform, MipFactor, AddressMode, bUseImageTransformVectorImpl);
+
+				Release(Source);
+				StoreImage(item, Result);
+
+                break;
+            }
+
+            default:
+                check(false);
+            }
+
+			break;
+		}
+
+        default:
+            if (type!=OP_TYPE::NONE)
+            {
+                // Operation not implemented
+                check( false );
+            }
+            break;
+        }
+    }	
+	
+	//---------------------------------------------------------------------------------------------
+    Ptr<RangeIndex> CodeRunner::BuildCurrentOpRangeIndex( const FScheduledOp& item, const Parameters* pParams, const Model* pModel, int32 parameterIndex )
+    {
+        if (!item.ExecutionIndex)
+        {
+            return nullptr;
+        }
+
+        // \todo: optimise to avoid allocating the index here, we could access internal
+        // data directly.
+		Ptr<RangeIndex> index = pParams->NewRangeIndex( parameterIndex );
+        if (!index)
+        {
+            return nullptr;
+        }
+
+		const FProgram& Program = m_pModel->GetPrivate()->m_program;
+		const FParameterDesc& paramDesc = Program.m_parameters[ parameterIndex ];
+        for( size_t rangeIndexInParam=0;
+             rangeIndexInParam<paramDesc.m_ranges.Num();
+             ++rangeIndexInParam )
+        {
+            uint32 rangeIndexInModel = paramDesc.m_ranges[rangeIndexInParam];
+            const ExecutionIndex& currentIndex = GetMemory().GetRangeIndex( item.ExecutionIndex );
+            int position = currentIndex.GetFromModelRangeIndex(rangeIndexInModel);
+            index->GetPrivate()->m_values[rangeIndexInParam] = position;
+        }
+
+        return index;
+    }
+
+
+    //---------------------------------------------------------------------------------------------
+    void CodeRunner::RunCode_Bool(const FScheduledOp& item, const Parameters* pParams, const Model* pModel )
+    {
+        MUTABLE_CPUPROFILER_SCOPE(RunCode_Bool);
+
+		const FProgram& Program = m_pModel->GetPrivate()->m_program;
+		OP_TYPE type = Program.GetOpType(item.At);
+        switch (type)
+        {
+
+        case OP_TYPE::BO_CONSTANT:
+        {
+			OP::BoolConstantArgs args = Program.GetOpArgs<OP::BoolConstantArgs>(item.At);
+            bool result = args.value;
+            StoreBool( item, result );
+            break;
+        }
+
+        case OP_TYPE::BO_PARAMETER:
+        {
+			OP::ParameterArgs args = Program.GetOpArgs<OP::ParameterArgs>(item.At);
+            bool result = false;
+			Ptr<RangeIndex> index = BuildCurrentOpRangeIndex( item, pParams, pModel, args.variable );
+            result = pParams->GetBoolValue( args.variable, index );
+            StoreBool( item, result );
+            break;
+        }
+
+        case OP_TYPE::BO_LESS:
+        {
+			OP::BoolLessArgs args = Program.GetOpArgs<OP::BoolLessArgs>(item.At);
+            switch (item.Stage)
+            {
+            case 0:
+                AddOp( FScheduledOp( item.At, item, 1),
+                        FScheduledOp( args.a, item),
+                        FScheduledOp( args.b, item) );
+                break;
+
+            case 1:
+            {
+                float a = LoadScalar( FCacheAddress(args.a,item) );
+                float b = LoadScalar( FCacheAddress(args.b,item) );
+                bool result = a<b;
+                StoreBool( item, result );
+                break;
+            }
+
+            default:
+                check(false);
+            }
+
+            break;
+        }
+
+        case OP_TYPE::BO_AND:
+        {
+			OP::BoolBinaryArgs args = Program.GetOpArgs<OP::BoolBinaryArgs>(item.At);
+            switch (item.Stage)
+            {
+            case 0:
+                {
+                    // Try to avoid the op entirely if we have some children cached
+                    bool skip = false;
+                    if ( args.a && GetMemory().IsValid( FCacheAddress(args.a,item) ) )
+                    {
+                         bool a = LoadBool( FCacheAddress(args.a,item) );
+                         if (!a)
+                         {
+                            StoreBool( item, false );
+                            skip=true;
+                         }
+                    }
+
+                    if ( !skip && args.b && GetMemory().IsValid( FCacheAddress(args.b,item) ) )
+                    {
+                         bool b = LoadBool( FCacheAddress(args.b,item) );
+                         if (!b)
+                         {
+                            StoreBool( item, false );
+                            skip=true;
+                         }
+                    }
+
+                    if (!skip)
+                    {
+                        AddOp( FScheduledOp( item.At, item, 1),
+                               FScheduledOp( args.a, item));
+                    }
+				break;
+                }
+
+            case 1:
+            {
+                bool a = args.a ? LoadBool( FCacheAddress(args.a,item) ) : true;
+                if (!a)
+                {
+                    StoreBool( item, false );
+                }
+                else
+                {
+                    AddOp( FScheduledOp( item.At, item, 2),
+                           FScheduledOp( args.b, item));
+                }
+                break;
+            }
+
+            case 2:
+            {
+                // We arrived here because a is true
+                bool b = args.b ? LoadBool( FCacheAddress(args.b,item) ) : true;
+                StoreBool( item, b );
+                break;
+            }
+
+            default:
+                check(false);
+            }
+            break;
+        }
+
+        case OP_TYPE::BO_OR:
+        {
+			OP::BoolBinaryArgs args = Program.GetOpArgs<OP::BoolBinaryArgs>(item.At);
+            switch (item.Stage)
+            {
+            case 0:
+                {
+                    // Try to avoid the op entirely if we have some children cached
+                    bool skip = false;
+                    if ( args.a && GetMemory().IsValid( FCacheAddress(args.a,item) ) )
+                    {
+                         bool a = LoadBool( FCacheAddress(args.a,item) );
+                         if (a)
+                         {
+                            StoreBool( item, true );
+                            skip=true;
+                         }
+                    }
+
+                    if ( !skip && args.b && GetMemory().IsValid( FCacheAddress(args.b,item) ) )
+                    {
+                         bool b = LoadBool( FCacheAddress(args.b,item) );
+                         if (b)
+                         {
+                            StoreBool( item, true );
+                            skip=true;
+                         }
+                    }
+
+                    if (!skip)
+                    {
+                        AddOp( FScheduledOp( item.At, item, 1),
+                               FScheduledOp( args.a, item));
+                    }
+				break;
+                }
+
+            case 1:
+            {
+                bool a = args.a ? LoadBool( FCacheAddress(args.a,item) ) : false;
+                if (a)
+                {
+                    StoreBool( item, true );
+                }
+                else
+                {
+                    AddOp( FScheduledOp( item.At, item, 2),
+                           FScheduledOp( args.b, item));
+                }
+                break;
+            }
+
+            case 2:
+            {
+                // We arrived here because a is false
+                bool b = args.b ? LoadBool( FCacheAddress(args.b,item) ) : false;
+                StoreBool( item, b );
+                break;
+            }
+
+            default:
+                check(false);
+            }
+            break;
+        }
+
+        case OP_TYPE::BO_NOT:
+        {
+			OP::BoolNotArgs args = Program.GetOpArgs<OP::BoolNotArgs>(item.At);
+            switch (item.Stage)
+            {
+            case 0:
+                    AddOp( FScheduledOp( item.At, item, 1),
+                           FScheduledOp( args.source, item) );
+                break;
+
+            case 1:
+            {
+                bool result = !LoadBool( FCacheAddress(args.source,item) );
+                StoreBool( item, result );
+                break;
+            }
+
+            default:
+                check(false);
+            }
+            break;
+        }
+
+        case OP_TYPE::BO_EQUAL_INT_CONST:
+        {
+			OP::BoolEqualScalarConstArgs args = Program.GetOpArgs<OP::BoolEqualScalarConstArgs>(item.At);
+            switch (item.Stage)
+            {
+            case 0:
+                    AddOp( FScheduledOp( item.At, item, 1),
+                           FScheduledOp( args.value, item) );
+                break;
+
+            case 1:
+            {
+                int a = LoadInt( FCacheAddress(args.value,item) );
+                bool result = a == args.constant;
+                StoreBool( item, result );
+                break;
+            }
+
+            default:
+                check(false);
+            }
+            break;
+        }
+
+        default:
+            check( false );
+            break;
+        }
+    }
+
+    //---------------------------------------------------------------------------------------------
+    void CodeRunner::RunCode_Int(const FScheduledOp& item, const Parameters* pParams, const Model* pModel )
+    {
+        MUTABLE_CPUPROFILER_SCOPE(RunCode_Int);
+
+		const FProgram& Program = m_pModel->GetPrivate()->m_program;
+
+		OP_TYPE type = Program.GetOpType(item.At);
+        switch (type)
+        {
+
+        case OP_TYPE::NU_CONSTANT:
+        {
+			OP::IntConstantArgs args = Program.GetOpArgs<OP::IntConstantArgs>(item.At);
+            int result = args.value;
+            StoreInt( item, result );
+            break;
+        }
+
+        case OP_TYPE::NU_PARAMETER:
+        {
+			OP::ParameterArgs args = Program.GetOpArgs<OP::ParameterArgs>(item.At);
+			Ptr<RangeIndex> index = BuildCurrentOpRangeIndex( item, pParams, pModel, args.variable );
+            int result = pParams->GetIntValue( args.variable, index );
+
+            // Check that the value is actually valid. Otherwise set the default.
+            if ( pParams->GetIntPossibleValueCount( args.variable ) )
+            {
+                bool valid = false;
+                for ( int i=0;
+                      (!valid) && i<pParams->GetIntPossibleValueCount( args.variable );
+                      ++i )
+                {
+                    if ( result == pParams->GetIntPossibleValue( args.variable, i ) )
+                    {
+                        valid = true;
+                    }
+                }
+
+                if (!valid)
+                {
+                    result = pParams->GetIntPossibleValue( args.variable, 0 );
+                }
+            }
+
+            StoreInt( item, result );
+            break;
+        }
+
+        default:
+            check( false );
+            break;
+        }
+    }
+
+    //---------------------------------------------------------------------------------------------
+    void CodeRunner::RunCode_Scalar(const FScheduledOp& item, const Parameters* pParams, const Model* pModel )
+    {
+        MUTABLE_CPUPROFILER_SCOPE(RunCode_Scalar);
+
+		const FProgram& Program = m_pModel->GetPrivate()->m_program;
+
+		OP_TYPE type = Program.GetOpType(item.At);
+        switch (type)
+        {
+
+        case OP_TYPE::SC_CONSTANT:
+        {
+			OP::ScalarConstantArgs args = Program.GetOpArgs<OP::ScalarConstantArgs>(item.At);
+            float result = args.value;
+            StoreScalar( item, result );
+            break;
+        }
+
+        case OP_TYPE::SC_PARAMETER:
+        {
+			OP::ParameterArgs args = Program.GetOpArgs<OP::ParameterArgs>(item.At);
+			Ptr<RangeIndex> index = BuildCurrentOpRangeIndex( item, pParams, pModel, args.variable );
+            float result = pParams->GetFloatValue( args.variable, index );
+            StoreScalar( item, result );
+            break;
+        }
+
+        case OP_TYPE::SC_CURVE:
+        {
+			OP::ScalarCurveArgs args = Program.GetOpArgs<OP::ScalarCurveArgs>(item.At);
+            switch (item.Stage)
+            {
+            case 0:
+                AddOp( FScheduledOp( item.At, item, 1),
+                        FScheduledOp( args.time, item) );
+                break;
+
+            case 1:
+            {
+                float time = LoadScalar( FCacheAddress(args.time,item) );
+
+                const FRichCurve& Curve = Program.ConstantCurves[args.curve];
+                float Result = Curve.Eval(time);
+
+                StoreScalar( item, Result );
+                break;
+            }
+
+            default:
+                check(false);
+            }
+            break;
+        }
+
+        case OP_TYPE::SC_MULTIPLYADD:
+            // \TODO
+            check( false );
+            break;
+
+        case OP_TYPE::SC_ARITHMETIC:
+        {
+			OP::ArithmeticArgs args = Program.GetOpArgs<OP::ArithmeticArgs>(item.At);
+            switch (item.Stage)
+            {
+            case 0:
+                    AddOp( FScheduledOp( item.At, item, 1),
+                           FScheduledOp( args.a, item),
+                           FScheduledOp( args.b, item) );
+                break;
+
+            case 1:
+            {
+                float a = LoadScalar( FCacheAddress(args.a,item) );
+                float b = LoadScalar( FCacheAddress(args.b,item) );
+
+                float result = 1.0f;
+                switch (args.operation)
+                {
+                case OP::ArithmeticArgs::ADD:
+                    result = a + b;
+                    break;
+
+                case OP::ArithmeticArgs::MULTIPLY:
+                    result = a * b;
+                    break;
+
+                case OP::ArithmeticArgs::SUBTRACT:
+                    result = a - b;
+                    break;
+
+                case OP::ArithmeticArgs::DIVIDE:
+                    result = a / b;
+                    break;
+
+                default:
+                    checkf(false, TEXT("Arithmetic operation not implemented."));
+                    break;
+                }
+
+                StoreScalar( item, result );
+                break;
+            }
+
+            default:
+                check(false);
+            }
+            break;
+        }
+
+        default:
+            check( false );
+            break;
+        }
+    }
+
+    //---------------------------------------------------------------------------------------------
+    void CodeRunner::RunCode_String(const FScheduledOp& item, const Parameters* pParams, const Model* pModel )
+    {
+        MUTABLE_CPUPROFILER_SCOPE(RunCode_String );
+
+		const FProgram& Program = m_pModel->GetPrivate()->m_program;
+
+		OP_TYPE type = Program.GetOpType( item.At );
+        switch ( type )
+        {
+
+        case OP_TYPE::ST_CONSTANT:
+        {
+			OP::ResourceConstantArgs args = Program.GetOpArgs<OP::ResourceConstantArgs>( item.At );
+            check( args.value < (uint32)pModel->GetPrivate()->m_program.m_constantStrings.Num() );
+
+            const FString& result = Program.m_constantStrings[args.value];
+            StoreString( item, new String(result) );
+
+            break;
+        }
+
+        case OP_TYPE::ST_PARAMETER:
+        {
+			OP::ParameterArgs args = Program.GetOpArgs<OP::ParameterArgs>( item.At );
+			Ptr<RangeIndex> index = BuildCurrentOpRangeIndex( item, pParams, pModel, args.variable );
+			FString result;
+			pParams->GetStringValue(args.variable, result, index);
+            StoreString( item, new String(result) );
+            break;
+        }
+
+        default:
+            check( false );
+            break;
+        }
+    }
+
+
+    //---------------------------------------------------------------------------------------------
+    void CodeRunner::RunCode_Colour(const FScheduledOp& item, const Parameters* pParams, const Model* pModel )
+    {
+		MUTABLE_CPUPROFILER_SCOPE(RunCode_Colour);
+
+		const FProgram& Program = m_pModel->GetPrivate()->m_program;
+
+		OP_TYPE type = Program.GetOpType(item.At);
+
+        switch ( type )
+        {
+
+        case OP_TYPE::CO_CONSTANT:
+        {
+			OP::ColourConstantArgs args = Program.GetOpArgs<OP::ColourConstantArgs>(item.At);
+			FVector4f result;
+            result[0] = args.value[0];
+            result[1] = args.value[1];
+            result[2] = args.value[2];
+            result[3] = args.value[3];
+            StoreColor( item, result );
+            break;
+        }
+
+        case OP_TYPE::CO_PARAMETER:
+        {
+			OP::ParameterArgs args = Program.GetOpArgs<OP::ParameterArgs>(item.At);
+			Ptr<RangeIndex> index = BuildCurrentOpRangeIndex( item, pParams, pModel, args.variable );
+            float R = 0.f;
+            float G = 0.f;
+			float B = 0.f;
+			float A = 0.f;
+            pParams->GetColourValue( args.variable, &R, &G, &B, &A, index );
+            StoreColor( item, FVector4f(R, G, B, A) );
+            break;
+        }
+
+        case OP_TYPE::CO_SAMPLEIMAGE:
+        {
+			OP::ColourSampleImageArgs args = Program.GetOpArgs<OP::ColourSampleImageArgs>(item.At);
+            switch (item.Stage)
+            {
+            case 0:
+                    AddOp( FScheduledOp( item.At, item, 1),
+                           FScheduledOp( args.x, item),
+                           FScheduledOp( args.y, item),
+						   // Don't skip mips for the texture to sample
+                           FScheduledOp::FromOpAndOptions( args.image, item, 0) );
+                break;
+
+            case 1:
+            {
+                float x = args.x ? LoadScalar( FCacheAddress(args.x,item) ) : 0.5f;
+                float y = args.y ? LoadScalar( FCacheAddress(args.y,item) ) : 0.5f;
+
+                Ptr<const Image> pImage = LoadImage(FScheduledOp::FromOpAndOptions(args.image, item, 0));
+
+				FVector4f result;
+                if (pImage)
+                {
+                    if (args.filter)
+                    {
+                        // TODO
+                        result = pImage->Sample(FVector2f(x, y));
+                    }
+                    else
+                    {
+                        result = pImage->Sample(FVector2f(x, y));
+                    }
+                }
+                else
+                {
+                    result = FVector4f();
+                }
+
+				Release(pImage);
+                StoreColor( item, result );
+                break;
+            }
+
+            default:
+                check(false);
+            }
+
+            break;
+        }
+
+        case OP_TYPE::CO_SWIZZLE:
+        {
+			OP::ColourSwizzleArgs args = Program.GetOpArgs<OP::ColourSwizzleArgs>(item.At);
+            switch (item.Stage)
+            {
+            case 0:
+                    AddOp( FScheduledOp( item.At, item, 1),
+                           FScheduledOp( args.sources[0], item),
+                           FScheduledOp( args.sources[1], item),
+                           FScheduledOp( args.sources[2], item),
+                           FScheduledOp( args.sources[3], item) );
+                break;
+
+            case 1:
+            {
+				FVector4f result;
+
+                for (int t=0;t<MUTABLE_OP_MAX_SWIZZLE_CHANNELS;++t)
+                {
+                    if ( args.sources[t] )
+                    {
+                        FVector4f p = LoadColor( FCacheAddress(args.sources[t],item) );
+                        result[t] = p[ args.sourceChannels[t] ];
+                    }
+                }
+
+                StoreColor( item, result );
+                break;
+            }
+
+            default:
+                check(false);
+            }
+
+            break;
+        }
+
+        case OP_TYPE::CO_FROMSCALARS:
+        {
+			OP::ColourFromScalarsArgs args = Program.GetOpArgs<OP::ColourFromScalarsArgs>(item.At);
+            switch (item.Stage)
+            {
+            case 0:
+                    AddOp( FScheduledOp( item.At, item, 1),
+                           FScheduledOp( args.v[0], item),
+                           FScheduledOp( args.v[1], item),
+                           FScheduledOp( args.v[2], item),
+                           FScheduledOp( args.v[3], item));
+                break;
+
+            case 1:
+            {
+				FVector4f Result = FVector4f(0, 0, 0, 1);
+
+				for (int32 t = 0; t < MUTABLE_OP_MAX_SWIZZLE_CHANNELS; ++t)
+				{
+					if (args.v[t])
+					{
+						Result[t] = LoadScalar(FCacheAddress(args.v[t], item));
+					}
+				}
+
+                StoreColor( item, Result );
+                break;
+            }
+
+            default:
+                check(false);
+            }
+
+            break;
+        }
+
+        case OP_TYPE::CO_ARITHMETIC:
+        {
+			OP::ArithmeticArgs args = Program.GetOpArgs<OP::ArithmeticArgs>(item.At);
+            switch (item.Stage)
+            {
+            case 0:
+                    AddOp( FScheduledOp( item.At, item, 1),
+                           FScheduledOp( args.a, item),
+                           FScheduledOp( args.b, item));
+                break;
+
+            case 1:
+            {
+				OP_TYPE otype = Program.GetOpType( args.a );
+                DATATYPE dtype = GetOpDataType( otype );
+                check( dtype == DT_COLOUR );
+                otype = Program.GetOpType( args.b );
+                dtype = GetOpDataType( otype );
+                check( dtype == DT_COLOUR );
+				FVector4f a = args.a ? LoadColor( FCacheAddress( args.a, item ) )
+                                 : FVector4f( 0, 0, 0, 0 );
+				FVector4f b = args.b ? LoadColor( FCacheAddress( args.b, item ) )
+                                 : FVector4f( 0, 0, 0, 0 );
+
+				FVector4f result = FVector4f(0,0,0,0);
+                switch (args.operation)
+                {
+                case OP::ArithmeticArgs::ADD:
+                    result = a + b;
+                    break;
+
+                case OP::ArithmeticArgs::MULTIPLY:
+                    result = a * b;
+                    break;
+
+                case OP::ArithmeticArgs::SUBTRACT:
+                    result = a - b;
+                    break;
+
+                case OP::ArithmeticArgs::DIVIDE:
+                    result = a / b;
+                    break;
+
+                default:
+                    checkf(false, TEXT("Arithmetic operation not implemented."));
+                    break;
+                }
+
+                StoreColor( item, result );
+                break;
+            }
+
+            default:
+                check(false);
+            }
+
+            break;
+        }
+
+        default:
+            check( false );
+            break;
+        }
+    }
+
+
+    //---------------------------------------------------------------------------------------------
+    void CodeRunner::RunCode_Projector(const FScheduledOp& item, const Parameters* pParams, const Model* pModel )
+    {
+        MUTABLE_CPUPROFILER_SCOPE(RunCode_Projector);
+
+		const FProgram& Program = m_pModel->GetPrivate()->m_program;
+
+		OP_TYPE type = Program.GetOpType(item.At);
+        switch (type)
+        {
+
+        case OP_TYPE::PR_CONSTANT:
+        {
+			OP::ResourceConstantArgs args = Program.GetOpArgs<OP::ResourceConstantArgs>(item.At);
+            FProjector Result = Program.m_constantProjectors[args.value];
+            StoreProjector( item, Result );
+            break;
+        }
+
+        case OP_TYPE::PR_PARAMETER:
+        {
+			OP::ParameterArgs args = Program.GetOpArgs<OP::ParameterArgs>(item.At);
+			Ptr<RangeIndex> index = BuildCurrentOpRangeIndex( item, pParams, pModel, args.variable );
+            FProjector Result = pParams->GetPrivate()->GetProjectorValue(args.variable,index);
+
+            // The type cannot be changed, take it from the default value
+            const FProjector& def = Program.m_parameters[args.variable].m_defaultValue.Get<ParamProjectorType>();
+            Result.type = def.type;
+
+            StoreProjector( item, Result );
+            break;
+        }
+
+        default:
+            check( false );
+            break;
+        }
+    }
+
+
+    //---------------------------------------------------------------------------------------------
+    void CodeRunner::RunCode_Layout(const FScheduledOp& item, const Model* pModel )
+    {
+        //MUTABLE_CPUPROFILER_SCOPE(RunCode_Layout);
+
+		const FProgram& Program = m_pModel->GetPrivate()->m_program;
+
+		OP_TYPE type = Program.GetOpType(item.At);
+        switch (type)
+        {
+
+        case OP_TYPE::LA_CONSTANT:
+        {
+			OP::ResourceConstantArgs args = Program.GetOpArgs<OP::ResourceConstantArgs>(item.At);
+            check( args.value < (uint32)pModel->GetPrivate()->m_program.m_constantLayouts.Num() );
+
+            Ptr<const Layout> pResult = Program.m_constantLayouts
+                    [ args.value ];
+            StoreLayout( item, pResult );
+            break;
+        }
+
+        case OP_TYPE::LA_MERGE:
+        {
+			OP::LayoutMergeArgs args = Program.GetOpArgs<OP::LayoutMergeArgs>(item.At);
+            switch (item.Stage)
+            {
+            case 0:
+                    AddOp( FScheduledOp( item.At, item, 1),
+                           FScheduledOp( args.Base, item),
+                           FScheduledOp( args.Added, item) );
+                break;
+
+            case 1:
+            {
+                Ptr<const Layout> pA = LoadLayout( FCacheAddress(args.Base,item) );
+                Ptr<const Layout> pB = LoadLayout( FCacheAddress(args.Added,item) );
+
+                Ptr<const Layout> pResult;
+
+                if (pA && pB)
+                {
+					pResult = LayoutMerge(pA.get(),pB.get());
+                }
+                else if (pA)
+                {
+                    pResult = pA->Clone();
+                }
+                else if (pB)
+                {
+                    pResult = pB->Clone();
+                }
+
+                StoreLayout( item, pResult );
+                break;
+            }
+
+            default:
+                check(false);
+            }
+
+            break;
+        }
+
+        case OP_TYPE::LA_PACK:
+        {
+			OP::LayoutPackArgs args = Program.GetOpArgs<OP::LayoutPackArgs>(item.At);
+            switch (item.Stage)
+            {
+            case 0:
+                    AddOp( FScheduledOp( item.At, item, 1),
+                           FScheduledOp( args.Source, item) );
+                break;
+
+            case 1:
+            {
+                Ptr<const Layout> Source = LoadLayout( FCacheAddress(args.Source,item) );
+
+				Ptr<Layout> Result;
+
+				if (Source)
+				{
+					Result = Source->Clone();
+
+					LayoutPack3(Result.get(), Source.get() );
+				}
+
+                StoreLayout( item, Result );
+                break;
+            }
+
+            default:
+                check(false);
+            }
+
+            break;
+        }
+
+		case OP_TYPE::LA_FROMMESH:
+		{
+			OP::LayoutFromMeshArgs args = Program.GetOpArgs<OP::LayoutFromMeshArgs>(item.At);
+			switch (item.Stage)
+			{
+			case 0:
+				AddOp(FScheduledOp(item.At, item, 1),
+					FScheduledOp(args.Mesh, item));
+				break;
+
+			case 1:
+			{
+				Ptr<const Mesh> Mesh = LoadMesh(FCacheAddress(args.Mesh, item));
+
+				Ptr<const Layout> Result = LayoutFromMesh_RemoveBlocks(Mesh.get(), args.LayoutIndex);
+
+				Release(Mesh);
+				StoreLayout(item, Result);
+				break;
+			}
+
+			default:
+				check(false);
+			}
+
+			break;
+		}
+
+		case OP_TYPE::LA_REMOVEBLOCKS:
+		{
+			OP::LayoutRemoveBlocksArgs args = Program.GetOpArgs<OP::LayoutRemoveBlocksArgs>(item.At);
+			switch (item.Stage)
+			{
+			case 0:
+				AddOp(FScheduledOp(item.At, item, 1),
+					FScheduledOp(args.Source, item),
+					FScheduledOp(args.ReferenceLayout, item));
+				break;
+
+			case 1:
+			{
+				Ptr<const Layout> Source = LoadLayout(FCacheAddress(args.Source, item));
+				Ptr<const Layout> ReferenceLayout = LoadLayout(FCacheAddress(args.ReferenceLayout, item));
+
+				Ptr<const Layout> pResult;
+
+				if (Source && ReferenceLayout)
+				{
+					pResult = LayoutRemoveBlocks(Source.get(), ReferenceLayout.get());
+				}
+				else if (Source)
+				{
+					pResult = Source;
+				}
+
+				StoreLayout(item, pResult);
+				break;
+			}
+
+			default:
+				check(false);
+			}
+
+			break;
+		}
+
+        default:
+            // Operation not implemented
+            check( false );
+            break;
+        }
+    }
+
+
+    //---------------------------------------------------------------------------------------------
+    void CodeRunner::RunCode( const FScheduledOp& item, const Parameters* pParams, const TSharedPtr<const Model>& InModel, uint32 lodMask)
+    {
+		//UE_LOG( LogMutableCore, Log, TEXT("Running :%5d , %d "), item.At, item.Stage );
+		check( item.Type == FScheduledOp::EType::Full );
+
+		const Model* pModel = InModel.Get();
+		
+		const FProgram& Program = pModel->GetPrivate()->m_program;
+
+		OP_TYPE type = Program.GetOpType(item.At);
+		//UE_LOG(LogMutableCore, Log, TEXT("Running :%5d , %d, of type %d "), item.At, item.Stage, type);
+
+		// Very spammy, for debugging purposes.
+		//if (m_pSystem)
+		//{
+		//	m_pSystem->WorkingMemoryManager.LogWorkingMemory( this );
+		//}
+
+		switch ( type )
+        {
+        case OP_TYPE::NONE:
+            break;
+
+        case OP_TYPE::NU_CONDITIONAL:
+        case OP_TYPE::SC_CONDITIONAL:
+        case OP_TYPE::CO_CONDITIONAL:
+        case OP_TYPE::IM_CONDITIONAL:
+        case OP_TYPE::ME_CONDITIONAL:
+        case OP_TYPE::LA_CONDITIONAL:
+        case OP_TYPE::IN_CONDITIONAL:
+		case OP_TYPE::ED_CONDITIONAL:
+            RunCode_Conditional(item, pModel);
+            break;
+
+        case OP_TYPE::ME_CONSTANT:
+		case OP_TYPE::IM_CONSTANT:
+		case OP_TYPE::ED_CONSTANT:
+            RunCode_ConstantResource(item, pModel);
+            break;
+
+        case OP_TYPE::NU_SWITCH:
+        case OP_TYPE::SC_SWITCH:
+        case OP_TYPE::CO_SWITCH:
+        case OP_TYPE::IM_SWITCH:
+        case OP_TYPE::ME_SWITCH:
+        case OP_TYPE::LA_SWITCH:
+        case OP_TYPE::IN_SWITCH:
+		case OP_TYPE::ED_SWITCH:
+            RunCode_Switch(item, pModel);
+            break;
+
+        case OP_TYPE::IN_ADDMESH:
+        case OP_TYPE::IN_ADDIMAGE:
+            RunCode_InstanceAddResource(item, InModel, pParams);
+            break;
+
+		default:
+		{
+			DATATYPE DataType = GetOpDataType(type);
+			switch (DataType)
+			{
+			case DT_INSTANCE:
+				RunCode_Instance(item, pModel, lodMask);
+				break;
+
+			case DT_MESH:
+				RunCode_Mesh(item, pModel);
+				break;
+
+			case DT_IMAGE:
+				RunCode_Image(item, pParams, pModel);
+				break;
+
+			case DT_LAYOUT:
+				RunCode_Layout(item, pModel);
+				break;
+
+			case DT_BOOL:
+				RunCode_Bool(item, pParams, pModel);
+				break;
+
+			case DT_SCALAR:
+				RunCode_Scalar(item, pParams, pModel);
+				break;
+
+			case DT_STRING:
+				RunCode_String(item, pParams, pModel);
+				break;
+
+			case DT_INT:
+				RunCode_Int(item, pParams, pModel);
+				break;
+
+			case DT_PROJECTOR:
+				RunCode_Projector(item, pParams, pModel);
+				break;
+
+			case DT_COLOUR:
+				RunCode_Colour(item, pParams, pModel);
+				break;
+
+			default:
+				check(false);
+				break;
+			}
+			break;
+		}
+
+        }
+    }
+
+	//---------------------------------------------------------------------------------------------
+	void CodeRunner::RunCodeImageDesc(const FScheduledOp& item, const Parameters* pParams, const Model* pModel,  uint32 lodMask )
+	{
+		MUTABLE_CPUPROFILER_SCOPE(RunCodeImageDesc);
+
+		check(item.Type == FScheduledOp::EType::ImageDesc);
+
+		// Ensure there is room for the result in the heap.
+		if (item.CustomState >= uint32(m_heapData.Num()))
+		{
+			m_heapImageDesc.SetNum(item.CustomState+1);
+		}
+
+
+		const FProgram& Program = m_pModel->GetPrivate()->m_program;
+
+		OP_TYPE type = Program.GetOpType(item.At);
+		switch (type)
+		{
+
+		case OP_TYPE::IM_CONSTANT:
+		{
+			check(item.Stage == 0);
+			OP::ResourceConstantArgs args = Program.GetOpArgs<OP::ResourceConstantArgs>(item.At);
+			int32 ImageIndex = args.value;
+
+			FImageDesc& Result = m_heapImageDesc[item.CustomState];
+			Result.m_format = Program.ConstantImages[ImageIndex].ImageFormat;
+			Result.m_size[0] = Program.ConstantImages[ImageIndex].ImageSizeX;
+			Result.m_size[1] = Program.ConstantImages[ImageIndex].ImageSizeY;
+			Result.m_lods = Program.ConstantImages[ImageIndex].LODCount;
+			StoreValidDesc(item);
+			break;
+		}
+
+		case OP_TYPE::IM_PARAMETER:
+		{
+			check(item.Stage == 0);
+			OP::ParameterArgs args = Program.GetOpArgs<OP::ParameterArgs>(item.At);
+			FName Id = pParams->GetImageValue(args.variable);
+			uint8 MipsToSkip = item.ExecutionOptions;
+			m_heapImageDesc[item.CustomState] = GetExternalImageDesc(Id, MipsToSkip);
+			StoreValidDesc(item);
+			break;
+		}
+
+		case OP_TYPE::IM_REFERENCE:
+		{
+			check(item.Stage == 0);
+			OP::ResourceReferenceArgs Args = Program.GetOpArgs<OP::ResourceReferenceArgs>(item.At);
+			FImageDesc& Result = m_heapImageDesc[item.CustomState];
+			Result = Args.ImageDesc;
+			StoreValidDesc(item);
+			break;
+		}
+
+		case OP_TYPE::IM_CONDITIONAL:
+		{
+			OP::ConditionalArgs args = Program.GetOpArgs<OP::ConditionalArgs>(item.At);
+			switch (item.Stage)
+			{
+			case 0:
+			{
+				// We need to run the full condition result
+				FScheduledOp FullConditionOp(args.condition, item);
+				FullConditionOp.Type = FScheduledOp::EType::Full;
+				AddOp(FScheduledOp(item.At, item, 1), FullConditionOp);
+				break;
+			}
+
+			case 1:
+			{
+				bool value = LoadBool(FCacheAddress(args.condition, item.ExecutionIndex, item.ExecutionOptions));
+				OP::ADDRESS resultAt = value ? args.yes : args.no;
+				AddOp(FScheduledOp(item.At, item, 2), FScheduledOp(resultAt, item));
+				break;
+			}
+
+			case 2: StoreValidDesc(item); break;
+			default: check(false);
+			}
+			break;
+		}
+
+		case OP_TYPE::IM_SWITCH:
+		{
+			const uint8* data = Program.GetOpArgsPointer(item.At);
+		
+			OP::ADDRESS VarAddress;
+			FMemory::Memcpy( &VarAddress, data, sizeof(OP::ADDRESS));
+			data += sizeof(OP::ADDRESS);
+
+			OP::ADDRESS DefAddress;
+			FMemory::Memcpy( &DefAddress, data, sizeof(OP::ADDRESS));
+			data += sizeof(OP::ADDRESS);
+
+			uint32 CaseCount;
+			FMemory::Memcpy( &CaseCount, data, sizeof(uint32));
+			data += sizeof(uint32);
+	
+			switch (item.Stage)
+			{
+			case 0:
+			{
+				if (VarAddress)
+				{
+					// We need to run the full condition result
+					FScheduledOp FullVariableOp(VarAddress, item);
+					FullVariableOp.Type = FScheduledOp::EType::Full;
+					AddOp(FScheduledOp(item.At, item, 1), FullVariableOp);
+				}
+				else
+				{
+					StoreValidDesc(item);
+				}
+				break;
+			}
+
+			case 1:
+			{
+				// Get the variable result
+				int var = LoadInt(FCacheAddress(VarAddress, item));
+
+				OP::ADDRESS valueAt = DefAddress;
+				for (uint32 C = 0; C < CaseCount; ++C)
+				{
+					int32 Condition;
+					FMemory::Memcpy( &Condition, data, sizeof(int32) );
+					data += sizeof(int32);
+					
+					OP::ADDRESS At;
+					FMemory::Memcpy( &At, data, sizeof(OP::ADDRESS) );
+					data += sizeof(OP::ADDRESS);
+
+					if (At && var == (int)Condition)
+					{
+						valueAt = At;
+						break;
+					}
+				}
+
+				AddOp(FScheduledOp(item.At, item, 2, valueAt),
+					  FScheduledOp(valueAt, item));
+
+				break;
+			}
+
+			case 2: StoreValidDesc(item); break;
+			default: check(false); break;
+			}
+			break;
+		}
+
+		case OP_TYPE::IM_LAYERCOLOUR:
+		{
+			OP::ImageLayerColourArgs args = Program.GetOpArgs<OP::ImageLayerColourArgs>(item.At);
+			switch (item.Stage)
+			{
+			case 0: AddOp(FScheduledOp(item.At, item, 1), FScheduledOp(args.base, item)); break;
+			case 1: StoreValidDesc(item); break;
+			default: check(false);
+			}
+			break;
+		}
+
+		case OP_TYPE::IM_LAYER:
+		{
+			OP::ImageLayerArgs args = Program.GetOpArgs<OP::ImageLayerArgs>(item.At);
+			switch (item.Stage)
+			{
+			case 0: AddOp(FScheduledOp(item.At, item, 1), FScheduledOp(args.base, item)); break;
+			case 1: StoreValidDesc(item); break;
+			default: check(false);
+			}
+			break;
+		}
+
+		case OP_TYPE::IM_MULTILAYER:
+		{
+			OP::ImageMultiLayerArgs args = Program.GetOpArgs<OP::ImageMultiLayerArgs>(item.At);
+			switch (item.Stage)
+			{
+			case 0: AddOp(FScheduledOp(item.At, item, 1), FScheduledOp(args.base, item)); break;
+			case 1: StoreValidDesc(item); break;
+			default: check(false);
+			}
+			break;
+		}
+
+		case OP_TYPE::IM_NORMALCOMPOSITE:
+		{
+			OP::ImageNormalCompositeArgs args = Program.GetOpArgs<OP::ImageNormalCompositeArgs>(item.At);
+			switch (item.Stage)
+			{
+			case 0: AddOp(FScheduledOp(item.At, item, 1), FScheduledOp(args.base, item)); break;
+			case 1: StoreValidDesc(item); break;
+			default: check(false);
+			}
+			break;
+		}
+
+		case OP_TYPE::IM_PIXELFORMAT:
+		{
+			OP::ImagePixelFormatArgs args = Program.GetOpArgs<OP::ImagePixelFormatArgs>(item.At);
+			switch (item.Stage)
+			{
+			case 0:
+				AddOp(FScheduledOp(item.At, item, 1), FScheduledOp(args.source, item));
+				break;
+
+			case 1:
+			{
+				// Update directly in the heap
+				EImageFormat OldFormat = m_heapImageDesc[item.CustomState].m_format;
+				EImageFormat NewFormat = args.format;
+				if (args.formatIfAlpha != EImageFormat::IF_NONE
+					&&
+					GetImageFormatData(OldFormat).Channels > 3)
+				{
+					NewFormat = args.formatIfAlpha;
+				}
+				m_heapImageDesc[item.CustomState].m_format = NewFormat;				
+				StoreValidDesc(item);
+				break;
+			}
+
+			default:
+				check(false);
+			}
+
+			break;
+		}
+
+		case OP_TYPE::IM_MIPMAP:
+		{
+			OP::ImageMipmapArgs args = Program.GetOpArgs<OP::ImageMipmapArgs>(item.At);
+			switch (item.Stage)
+			{
+			case 0:
+				AddOp(FScheduledOp(item.At, item, 1), FScheduledOp(args.source, item));
+				break;
+
+			case 1:
+			{
+				// Somewhat synched with Full op execution code.
+				FImageDesc BaseDesc = m_heapImageDesc[item.CustomState];
+				int32 LevelCount = args.levels;
+				int32 MaxLevelCount = Image::GetMipmapCount(BaseDesc.m_size[0], BaseDesc.m_size[1]);
+				if (LevelCount == 0)
+				{
+					LevelCount = MaxLevelCount;
+				}
+				else if (LevelCount > MaxLevelCount)
+				{
+					// If code generation is smart enough, this should never happen.
+					// \todo But apparently it does, sometimes.
+					LevelCount = MaxLevelCount;
+				}
+
+				// At least keep the levels we already have.
+				int32 StartLevel = BaseDesc.m_lods;
+				LevelCount = FMath::Max(StartLevel, LevelCount);
+
+				// Update result.
+				m_heapImageDesc[item.CustomState].m_lods = LevelCount;
+				StoreValidDesc(item);
+				break;
+			}
+
+			default:
+				check(false);
+			}
+			break;
+		}
+
+		case OP_TYPE::IM_RESIZE:
+		{
+			OP::ImageResizeArgs args = Program.GetOpArgs<OP::ImageResizeArgs>(item.At);
+			switch (item.Stage)
+			{
+			case 0:
+				AddOp(FScheduledOp(item.At, item, 1), FScheduledOp(args.source, item));
+				break;
+
+			case 1:
+				m_heapImageDesc[item.CustomState].m_size[0] = args.size[0];
+				m_heapImageDesc[item.CustomState].m_size[1] = args.size[1];
+				StoreValidDesc(item);
+				break;
+
+			default:
+				check(false);
+			}
+			break;
+		}
+
+		case OP_TYPE::IM_RESIZELIKE:
+		{
+			OP::ImageResizeLikeArgs args = Program.GetOpArgs<OP::ImageResizeLikeArgs>(item.At);
+			switch (item.Stage)
+			{
+			case 0:
+			{
+				int32 ResultAndBaseDesc = item.CustomState;
+				int32 SourceDescAddress = m_heapImageDesc.Add({});
+				FScheduledOpData Data;
+				Data.ResizeLike.ResultDescAt = ResultAndBaseDesc;
+				Data.ResizeLike.SourceDescAt = SourceDescAddress;
+				int32 SecondStageData = m_heapData.Add(Data);
+				AddOp(FScheduledOp(item.At, item, 1, SecondStageData),
+					FScheduledOp(args.source, item, 0, ResultAndBaseDesc),
+					FScheduledOp(args.sizeSource, item, 0, SourceDescAddress));
+				break;
+			}
+
+			case 1:
+			{
+				const FScheduledOpData& SecondStageData = m_heapData[ item.CustomState ];
+				FImageDesc& ResultAndBaseDesc = m_heapImageDesc[SecondStageData.ResizeLike.ResultDescAt];
+				const FImageDesc& SourceDesc = m_heapImageDesc[SecondStageData.ResizeLike.SourceDescAt];
+				ResultAndBaseDesc.m_size = SourceDesc.m_size;
+				StoreValidDesc(item);
+				break;
+			}
+
+			default:
+				check(false);
+			}
+
+			break;
+		}
+
+		case OP_TYPE::IM_RESIZEREL:
+		{
+			OP::ImageResizeRelArgs args = Program.GetOpArgs<OP::ImageResizeRelArgs>(item.At);
+			switch (item.Stage)
+			{
+			case 0:
+				AddOp(FScheduledOp(item.At, item, 1), FScheduledOp(args.source, item));
+				break;
+
+			case 1:
+			{
+				FImageDesc& ResultAndBaseDesc = m_heapImageDesc[item.CustomState];
+				FImageSize destSize(
+					uint16(ResultAndBaseDesc.m_size[0] * args.factor[0] + 0.5f),
+					uint16(ResultAndBaseDesc.m_size[1] * args.factor[1] + 0.5f));
+				ResultAndBaseDesc.m_size = destSize;
+				StoreValidDesc(item);
+				break;
+			}
+
+			default:
+				check(false);
+			}
+
+			break;
+		}
+
+		case OP_TYPE::IM_BLANKLAYOUT:
+		{
+			OP::ImageBlankLayoutArgs args = Program.GetOpArgs<OP::ImageBlankLayoutArgs>(item.At);
+			switch (item.Stage)
+			{
+			case 0:
+			{
+				// We need to run the full layout
+				FScheduledOp FullLayoutOp(args.layout, item);
+				FullLayoutOp.Type = FScheduledOp::EType::Full;
+				AddOp(FScheduledOp(item.At, item, 1), FullLayoutOp);
+				break;
+			}
+
+			case 1:
+			{
+				Ptr<const Layout> pLayout = LoadLayout(FCacheAddress(args.layout, item));
+
+				FIntPoint SizeInBlocks = pLayout->GetGridSize();
+				FIntPoint BlockSizeInPixels(args.blockSize[0], args.blockSize[1]);
+				FIntPoint ImageSizeInPixels = SizeInBlocks * BlockSizeInPixels;
+
+				FImageDesc& ResultAndBaseDesc = m_heapImageDesc[item.CustomState];
+				FImageSize destSize(uint16(ImageSizeInPixels.X), uint16(ImageSizeInPixels.Y));
+				ResultAndBaseDesc.m_size = destSize;
+				ResultAndBaseDesc.m_format = args.format;
+				
+				if (args.generateMipmaps)
+				{
+					if (args.mipmapCount == 0)
+					{
+						ResultAndBaseDesc.m_lods = Image::GetMipmapCount(ImageSizeInPixels.X, ImageSizeInPixels.Y);
+					}
+					else
+					{
+						ResultAndBaseDesc.m_lods = args.mipmapCount;
+					}
+				}
+				StoreValidDesc(item);
+				break;
+			}
+
+			default:
+				check(false);
+			}
+
+			break;
+		}
+
+		case OP_TYPE::IM_COMPOSE:
+		{
+			OP::ImageComposeArgs args = Program.GetOpArgs<OP::ImageComposeArgs>(item.At);
+			switch (item.Stage)
+			{
+			case 0: AddOp(FScheduledOp(item.At, item, 1), FScheduledOp(args.base, item)); break;
+			case 1: StoreValidDesc(item); break;
+			default: check(false);
+			}
+			break;
+		}
+
+		case OP_TYPE::IM_INTERPOLATE:
+		{
+			OP::ImageInterpolateArgs args = Program.GetOpArgs<OP::ImageInterpolateArgs>(item.At);
+			switch (item.Stage)
+			{
+			case 0: AddOp(FScheduledOp(item.At, item, 1), FScheduledOp(args.targets[0], item)); break;
+			case 1: StoreValidDesc(item); break;
+			default: check(false);
+			}
+			break;
+		}
+
+		case OP_TYPE::IM_SATURATE:
+		{
+			OP::ImageSaturateArgs args = Program.GetOpArgs<OP::ImageSaturateArgs>(item.At);
+			switch (item.Stage)
+			{
+			case 0: AddOp(FScheduledOp(item.At, item, 1), FScheduledOp(args.base, item)); break;
+			case 1: StoreValidDesc(item); break;
+			default: check(false);
+			}
+			break;
+		}
+
+		case OP_TYPE::IM_LUMINANCE:
+		{
+			OP::ImageLuminanceArgs args = Program.GetOpArgs<OP::ImageLuminanceArgs>(item.At);
+			switch (item.Stage)
+			{
+			case 0:
+				AddOp(FScheduledOp(item.At, item, 1), FScheduledOp(args.base, item));
+				break;
+
+			case 1:
+				m_heapImageDesc[item.CustomState].m_format = EImageFormat::IF_L_UBYTE;
+				StoreValidDesc(item);
+				break;
+
+			default:
+				check(false);
+			}
+
+			break;
+		}
+
+		case OP_TYPE::IM_SWIZZLE:
+		{
+			OP::ImageSwizzleArgs args = Program.GetOpArgs<OP::ImageSwizzleArgs>(item.At);
+			switch (item.Stage)
+			{
+			case 0:
+				AddOp(FScheduledOp(item.At, item, 1), FScheduledOp(args.sources[0], item));
+				break;
+
+			case 1:
+				m_heapImageDesc[item.CustomState].m_format = args.format;
+				StoreValidDesc(item);
+				break;
+
+			default:
+				check(false);
+			}
+
+			break;
+		}
+
+		case OP_TYPE::IM_COLOURMAP:
+		{
+			OP::ImageColourMapArgs args = Program.GetOpArgs<OP::ImageColourMapArgs>(item.At);
+			switch (item.Stage)
+			{
+			case 0: AddOp(FScheduledOp(item.At, item, 1), FScheduledOp(args.base, item)); break;
+			case 1: StoreValidDesc(item); break;
+			default: check(false);
+			}
+			break;
+		}
+
+		case OP_TYPE::IM_GRADIENT:
+		{
+			OP::ImageGradientArgs args = Program.GetOpArgs<OP::ImageGradientArgs>(item.At);
+			m_heapImageDesc[item.CustomState].m_size[0] = args.size[0];
+			m_heapImageDesc[item.CustomState].m_size[1] = args.size[1];
+			m_heapImageDesc[item.CustomState].m_lods = 1;
+			m_heapImageDesc[item.CustomState].m_format = EImageFormat::IF_RGB_UBYTE;
+			StoreValidDesc(item);
+			break;
+		}
+
+		case OP_TYPE::IM_BINARISE:
+		{
+			OP::ImageBinariseArgs args = Program.GetOpArgs<OP::ImageBinariseArgs>(item.At);
+			switch (item.Stage)
+			{
+			case 0:
+				AddOp(FScheduledOp(item.At, item, 1), FScheduledOp(args.base, item));
+				break;
+
+			case 1:
+				m_heapImageDesc[item.CustomState].m_format = EImageFormat::IF_L_UBYTE;
+				StoreValidDesc(item);
+				break;
+
+			default:
+				check(false);
+			}
+			break;
+		}
+
+		case OP_TYPE::IM_INVERT:
+		{
+			OP::ImageInvertArgs args = Program.GetOpArgs<OP::ImageInvertArgs>(item.At);
+			switch (item.Stage)
+			{
+			case 0: AddOp(FScheduledOp(item.At, item, 1), FScheduledOp(args.base, item)); break;
+			case 1: StoreValidDesc(item); break;
+			default: check(false);
+			}
+			break;
+		}
+
+		case OP_TYPE::IM_PLAINCOLOUR:
+		{
+			OP::ImagePlainColourArgs args = Program.GetOpArgs<OP::ImagePlainColourArgs>(item.At);
+			m_heapImageDesc[item.CustomState].m_size[0] = args.size[0];
+			m_heapImageDesc[item.CustomState].m_size[1] = args.size[1];
+			m_heapImageDesc[item.CustomState].m_lods = args.LODs;
+			m_heapImageDesc[item.CustomState].m_format = args.format;
+			StoreValidDesc(item);
+			break;
+		}
+
+		case OP_TYPE::IM_CROP:
+		{
+			OP::ImageCropArgs args = Program.GetOpArgs<OP::ImageCropArgs>(item.At);
+			switch (item.Stage)
+			{
+			case 0:
+				AddOp(FScheduledOp(item.At, item, 1), FScheduledOp(args.source, item));
+				break;
+
+			case 1:
+				m_heapImageDesc[item.CustomState].m_size[0] = args.sizeX;
+				m_heapImageDesc[item.CustomState].m_size[1] = args.sizeY;
+				m_heapImageDesc[item.CustomState].m_lods = 1;
+				StoreValidDesc(item);
+				break;
+
+			default:
+				check(false);
+			}
+			break;
+		}
+
+		case OP_TYPE::IM_PATCH:
+		{
+			OP::ImagePatchArgs args = Program.GetOpArgs<OP::ImagePatchArgs>(item.At);
+			switch (item.Stage)
+			{
+			case 0: AddOp(FScheduledOp(item.At, item, 1), FScheduledOp(args.base, item)); break;
+			case 1: StoreValidDesc(item); break;
+			default: check(false);
+			}
+			break;
+		}
+
+		case OP_TYPE::IM_RASTERMESH:
+		{
+			OP::ImageRasterMeshArgs args = Program.GetOpArgs<OP::ImageRasterMeshArgs>(item.At);
+			m_heapImageDesc[item.CustomState].m_size[0] = args.sizeX;
+			m_heapImageDesc[item.CustomState].m_size[1] = args.sizeY;
+			m_heapImageDesc[item.CustomState].m_lods = 1;
+			m_heapImageDesc[item.CustomState].m_format = EImageFormat::IF_L_UBYTE;
+			StoreValidDesc(item);
+			break;
+		}
+
+		case OP_TYPE::IM_MAKEGROWMAP:
+		{
+			OP::ImageMakeGrowMapArgs args = Program.GetOpArgs<OP::ImageMakeGrowMapArgs>(item.At);
+			switch (item.Stage)
+			{
+			case 0:
+				AddOp(FScheduledOp(item.At, item, 1), FScheduledOp(args.mask, item));
+				break;
+
+			case 1:
+				m_heapImageDesc[item.CustomState].m_format = EImageFormat::IF_L_UBYTE;
+				m_heapImageDesc[item.CustomState].m_lods = 1;
+				StoreValidDesc(item);
+				break;
+
+			default:
+				check(false);
+			}
+
+			break;
+		}
+
+		case OP_TYPE::IM_DISPLACE:
+		{
+			OP::ImageDisplaceArgs args = Program.GetOpArgs<OP::ImageDisplaceArgs>(item.At);
+			switch (item.Stage)
+			{
+			case 0: AddOp(FScheduledOp(item.At, item, 1), FScheduledOp(args.source, item)); break;
+			case 1: StoreValidDesc(item); break;
+			default: check(false);
+			}
+			break;
+		}
+
+        case OP_TYPE::IM_TRANSFORM:
+        {
+
+			OP::ImageTransformArgs Args = Program.GetOpArgs<OP::ImageTransformArgs>(item.At);
+
+            switch (item.Stage)
+            {
+            case 0:
+			{
+				AddOp(FScheduledOp(item.At, item, 1), FScheduledOp(Args.Base, item));	
+                break;
+			}
+            case 1:
+            {
+				m_heapImageDesc[item.CustomState].m_lods = 1;
+				m_heapImageDesc[item.CustomState].m_format = GetUncompressedFormat(m_heapImageDesc[item.CustomState].m_format);
+				
+				if (!(Args.SizeX == 0 && Args.SizeY == 0))
+				{
+					m_heapImageDesc[item.CustomState].m_size[0] = Args.SizeX;
+					m_heapImageDesc[item.CustomState].m_size[1] = Args.SizeY;
+				}
+
+				StoreValidDesc(item);
+                break;
+            }
+
+            default:
+                check(false);
+            }
+
+			break;
+		}
+
+		default:
+			if (type != OP_TYPE::NONE)
+			{
+				// Operation not implemented
+				check(false);
+				m_heapImageDesc[item.CustomState] = FImageDesc();
+			}
+			break;
+		}
+	}
+}
