[mutable] Moved the Mutable plugin out of Experimental status into Beta.

diff --git a/Engine/Plugins/Mutable/Source/MutableRuntime/Private/MuR/CodeRunner.h b/Engine/Plugins/Mutable/Source/MutableRuntime/Private/MuR/CodeRunner.h
new file mode 100644
--- /dev/null
+++ b/Engine/Plugins/Mutable/Source/MutableRuntime/Private/MuR/CodeRunner.h
@@ -0,0 +1,779 @@
+// Copyright Epic Games, Inc. All Rights Reserved.
+
+#pragma once
+
+#include "Containers/Array.h"
+#include "Containers/ContainerAllocationPolicies.h"
+#include "HAL/PlatformCrt.h"
+#include "HAL/PlatformMath.h"
+#include "MuR/ExtensionDataStreamer.h"
+#include "MuR/Image.h"
+#include "MuR/Layout.h"
+#include "MuR/MutableMemory.h"
+#include "MuR/Operations.h"
+#include "MuR/Ptr.h"
+#include "MuR/RefCounted.h"
+#include "MuR/Serialisation.h"
+#include "MuR/Settings.h"
+#include "MuR/System.h"
+#include "MuR/SystemPrivate.h"
+#include "MuR/Types.h"
+#include "Tasks/Task.h"
+#include "Templates/RefCounting.h"
+#include "Templates/SharedPointer.h"
+#include "Templates/Tuple.h"
+
+namespace mu::MemoryCounters
+{
+	struct MUTABLERUNTIME_API FStreamingMemoryCounter
+	{
+		alignas(8) static inline std::atomic<SSIZE_T> Counter {0};
+	};
+}
+
+namespace  mu
+{
+	class Model;
+	class Parameters;
+	class RangeIndex;
+
+    /** Code execution of the mutable virtual machine. */
+    class CodeRunner : public TSharedFromThis<CodeRunner>
+    {
+		// The private token allows only members or friends to call MakeShared.
+		struct FPrivateToken { explicit FPrivateToken() = default; };
+
+	public:
+		static TSharedRef<CodeRunner> Create(
+				const Ptr<const Settings>&, 
+				class System::Private*, 
+				EExecutionStrategy,
+				const TSharedPtr<const Model>&, 
+				const Parameters* pParams,
+				OP::ADDRESS at, uint32 lodMask, uint8 executionOptions, int32 InImageLOD, FScheduledOp::EType);
+
+		// Private constructor to prevent stack allocation. In general we can not call AsShared() if the lifetime is
+		// bounded.
+		explicit CodeRunner(FPrivateToken, 
+			const Ptr<const Settings>&, 
+			class System::Private*, 
+			EExecutionStrategy,
+			const TSharedPtr<const Model>&, 
+			const Parameters* pParams,
+			OP::ADDRESS at, uint32 lodMask, uint8 executionOptions, int32 InImageLOD, FScheduledOp::EType);
+
+    protected:
+		struct FProfileContext
+		{
+			uint32 NumRunOps = 0;
+			uint32 RunOpsPerType[int32(OP_TYPE::COUNT)] = {};
+		};
+
+        /** Type of data sometimes stored in the code runner heap to pass info between operation stages. */
+        struct FScheduledOpData
+        {
+			union
+			{
+				struct
+				{
+					float Bifactor;
+					int32 Min, Max;
+				} Interpolate;
+
+				struct
+				{
+					int32 Iterations;
+					EImageFormat OriginalBaseFormat;
+					bool bBlendOnlyOneMip;
+				} MultiLayer;
+
+				struct
+				{
+					int32 ResultDescAt;
+					int32 SourceDescAt;
+				} ResizeLike;
+
+				struct
+				{
+					uint8 Mip;
+					float MipValue;
+				} RasterMesh;
+
+				struct
+				{
+					uint16 SizeX;
+					uint16 SizeY;
+					uint16 ScaleXEncodedHalf;
+					uint16 ScaleYEncodedHalf;
+					float MipValue;
+				} ImageTransform;
+			};
+            Ptr<RefCounted> Resource;
+        };
+
+		// Assertion to know when FScheduledOpData size changes. It is ok to modifiy if needed. 
+		static_assert(sizeof(FScheduledOpData) == 4*4 + sizeof(Ptr<RefCounted>), "FScheduledOpData size changed.");
+
+
+        Ptr<RangeIndex> BuildCurrentOpRangeIndex( const FScheduledOp&, const Parameters*, const Model*, int32 ParameterIndex );
+
+        void RunCode( const FScheduledOp&, const Parameters*, const TSharedPtr<const Model>&, uint32 LodMask);
+
+        void RunCode_Conditional(const FScheduledOp&, const Model* );
+        void RunCode_Switch(const FScheduledOp&, const Model* );
+        void RunCode_Instance(const FScheduledOp&, const Model*, uint32 LodMask );
+        void RunCode_InstanceAddResource(const FScheduledOp&, const TSharedPtr<const Model>& Model, const Parameters* );
+
+		/** Return false incase of failure. */
+        bool RunCode_ConstantResource(const FScheduledOp&, const Model* );
+
+        void RunCode_Mesh(const FScheduledOp&, const Model* );
+        void RunCode_Image(const FScheduledOp&, const Parameters*, const Model* );
+        void RunCode_Layout(const FScheduledOp&, const Model* );
+        void RunCode_Bool(const FScheduledOp&, const Parameters*, const Model* );
+        void RunCode_Int(const FScheduledOp&, const Parameters*, const Model* );
+        void RunCode_Scalar(const FScheduledOp&, const Parameters*, const Model* );
+        void RunCode_String(const FScheduledOp&, const Parameters*, const Model* );
+        void RunCode_Colour(const FScheduledOp&, const Parameters*, const Model* );
+        void RunCode_Projector(const FScheduledOp&, const Parameters*, const Model* );
+
+		void RunCodeImageDesc(const FScheduledOp&, const Parameters*, const Model*, uint32 LodMask);
+
+    public:
+
+		struct FExternalResourceId
+		{
+			/** If it is an image or mesh reference. */
+			int32 ReferenceResourceId = -1;
+
+			/** If it is an image or mesh parameter.*/
+			FName ParameterId; 
+		};
+
+		/** Load an external image asynchronously, returns an event to wait for complition and a cleanup function that must be called once the event has completed. */
+		TTuple<UE::Tasks::FTask, TFunction<void()>> LoadExternalImageAsync(FExternalResourceId Id, uint8 MipmapsToSkip, TFunction<void(Ptr<Image>)>& ResultCallback);
+ 	    mu::FImageDesc GetExternalImageDesc(FName Id, uint8 MipmapsToSkip);
+
+		/** Load an external mesh asynchronously, returns an event to wait for complition and a cleanup function that must be called once the event has completed. */
+		TTuple<UE::Tasks::FTask, TFunction<void()>> LoadExternalMeshAsync(FExternalResourceId Id, TFunction<void(Ptr<Mesh>)>& ResultCallback);
+
+		/** Settings that may affect the execution of some operations, like image conversion quality. */
+		Ptr<const Settings> m_pSettings;
+
+    protected:
+        //! Heap of intermediate data pushed by some instructions and referred by others.
+        //! It is not released until no operations are pending.
+		TArray<FScheduledOpData> m_heapData;
+		TArray<FImageDesc> m_heapImageDesc;
+
+		/** Only used for correct mip skipping with external images. It is the LOD for which the image is build. */
+		int32 ImageLOD;
+
+		UE::Tasks::FTaskEvent RunnerCompletionEvent;
+
+		void Run(TUniquePtr<FProfileContext>&& ProfileContext, bool bForceInlineExecution);
+		void AbortRun();
+	public:
+
+		UE::Tasks::FTask StartRun(bool bForceInlineExecution);
+
+		//!
+		void GetImageDescResult(FImageDesc& OutDesc);
+
+		//!
+		FProgramCache& GetMemory();
+
+
+		struct FTask
+		{
+			FTask() {}
+			FTask(const FScheduledOp& InOp) : Op(InOp) {}
+			FTask(const FScheduledOp& InOp, const FScheduledOp& InDep0) : Op(InOp) 
+			{ 
+				if (InDep0.At) Deps.Add(InDep0); 
+			}
+			FTask(const FScheduledOp& InOp, const FScheduledOp& InDep0, const FScheduledOp& InDep1) : Op(InOp) 
+			{ 
+				if (InDep0.At) Deps.Add(InDep0); 
+				if (InDep1.At) Deps.Add(InDep1); 
+			}
+			FTask(const FScheduledOp& InOp, const FScheduledOp& InDep0, const FScheduledOp& InDep1, const FScheduledOp& InDep2) : Op(InOp) 
+			{ 
+				if (InDep0.At) Deps.Add(InDep0); 
+				if (InDep1.At) Deps.Add(InDep1); 
+				if (InDep2.At) Deps.Add(InDep2); 
+			}
+			FTask(const FScheduledOp& InOp, const FScheduledOp& InDep0, const FScheduledOp& InDep1, const FScheduledOp& InDep2, const FScheduledOp& InDep3) : Op(InOp) 
+			{ 
+				if (InDep0.At) Deps.Add(InDep0); 
+				if (InDep1.At) Deps.Add(InDep1); 
+				if (InDep2.At) Deps.Add(InDep2); 
+				if (InDep3.At) Deps.Add(InDep3); 
+			}
+			FTask(const FScheduledOp& InOp, const FScheduledOp& InDep0, const FScheduledOp& InDep1, const FScheduledOp& InDep2, const FScheduledOp& InDep3, const FScheduledOp& InDep4) : Op(InOp) 
+			{ 
+				if (InDep0.At) Deps.Add(InDep0); 
+				if (InDep1.At) Deps.Add(InDep1); 
+				if (InDep2.At) Deps.Add(InDep2); 
+				if (InDep3.At) Deps.Add(InDep3); 
+				if (InDep4.At) Deps.Add(InDep4); 
+			}
+
+			FScheduledOp Op;
+			TArray<FCacheAddress, TInlineAllocator<3>> Deps;
+		};
+
+		class FIssuedTask
+		{
+		public:
+			const FScheduledOp Op;
+
+			UE::Tasks::FTask Event = {};
+
+			FIssuedTask(const FScheduledOp& InOp) : Op(InOp) {}
+			virtual ~FIssuedTask() {}
+
+			/** */
+			virtual bool Prepare(CodeRunner*, bool& bOutFailed) { bOutFailed = false; return true; }
+			virtual void DoWork() {}
+
+			/** Return true if succeeded. */
+			virtual bool Complete(CodeRunner*) = 0;
+
+			/** Return true if the task has been completed. */
+			virtual bool IsComplete(CodeRunner*)
+			{ 
+				return !Event.IsValid() || Event.IsCompleted(); 
+			}
+		};
+
+		struct FRomLoadOp
+		{
+			using StreamingDataContainerType = TArray<uint8, FDefaultMemoryTrackingAllocator<MemoryCounters::FStreamingMemoryCounter>>;
+
+			int32 RomIndex = -1;
+			ModelReader::OPERATION_ID m_streamID = -1;
+			StreamingDataContainerType m_streamBuffer;
+			UE::Tasks::FTask Event;
+		};
+
+		class FLoadMeshRomTask : public CodeRunner::FIssuedTask
+		{
+		public:
+			FLoadMeshRomTask( const FScheduledOp& InOp, int32 InRomIndex)
+				: FIssuedTask(InOp)
+			{
+				RomIndex = InRomIndex;
+			}
+
+			// FIssuedTask interface
+			virtual bool Prepare(CodeRunner*, bool& bOutFailed) override;
+			virtual bool Complete(CodeRunner*) override;
+
+		private:
+			int32 RomIndex = -1;
+		};
+
+		class FLoadExtensionDataTask : public CodeRunner::FIssuedTask
+		{
+		public:
+			FLoadExtensionDataTask(const FScheduledOp& InOp, const ExtensionDataPtrConst& InExtensionData, int32 InModelConstantIndex)
+				: FIssuedTask(InOp)
+				, Data(InExtensionData)
+				, ModelConstantIndex(InModelConstantIndex)
+			{
+				check(Data.get());
+				check(Data->Origin == ExtensionData::EOrigin::ConstantStreamed);
+			}
+
+			// FIssuedTask interface
+			virtual bool Prepare(CodeRunner*, bool& bOutFailed) override;
+			virtual bool Complete(CodeRunner*) override;
+			virtual bool IsComplete(CodeRunner*) override;
+
+		private:
+			ExtensionDataPtrConst Data;
+			// The index of this constant in FProgram::m_constantExtensionData
+			int32 ModelConstantIndex;
+
+			TSharedPtr<const FExtensionDataLoadHandle> LoadHandle;
+		};
+
+		class FLoadImageRomsTask : public CodeRunner::FIssuedTask
+		{
+		public:
+			FLoadImageRomsTask(const FScheduledOp& InOp, int32 InLODIndexIndex, int32 InLODIndexCount )
+				: FIssuedTask(InOp)
+			{
+				LODIndexIndex = InLODIndexIndex;
+				LODIndexCount = InLODIndexCount;
+			}
+
+			// FIssuedTask interface
+			virtual bool Prepare(CodeRunner*, bool& bOutFailed) override;
+			virtual bool Complete(CodeRunner*) override;
+
+		private:
+ 			int32 LODIndexIndex = -1;
+			int32 LODIndexCount = -1;
+
+			TArray<int32> RomIndices;
+		};
+
+		void AddOp(const FScheduledOp& op)
+		{
+			// It has no dependencies, so add it directly to the open tasks list.
+			OpenTasks.Add(op);
+			ScheduledStagePerOp[op] = op.Stage + 1;
+		}
+
+		void AddOp(const FScheduledOp& op,
+			const FScheduledOp& dep0)
+		{
+			ClosedTasks.Add(FTask(op, dep0));
+			ScheduledStagePerOp[op] = op.Stage + 1;
+			AddChildren(dep0);
+		}
+
+		void AddOp(const FScheduledOp& op,
+			const FScheduledOp& dep0,
+			const FScheduledOp& dep1)
+		{
+			ClosedTasks.Add(FTask(op, dep0, dep1));
+			ScheduledStagePerOp[op] = op.Stage + 1;
+			AddChildren(dep0);
+			AddChildren(dep1);
+		}
+
+		void AddOp(const FScheduledOp& op,
+			const FScheduledOp& dep0,
+			const FScheduledOp& dep1,
+			const FScheduledOp& dep2)
+		{
+			ClosedTasks.Add(FTask(op, dep0, dep1, dep2));
+			ScheduledStagePerOp[op] = op.Stage + 1;
+			AddChildren(dep0);
+			AddChildren(dep1);
+			AddChildren(dep2);
+		}
+
+		void AddOp(const FScheduledOp& op,
+			const FScheduledOp& dep0,
+			const FScheduledOp& dep1,
+			const FScheduledOp& dep2,
+			const FScheduledOp& dep3)
+		{
+			ClosedTasks.Add(FTask(op, dep0, dep1, dep2, dep3));
+			ScheduledStagePerOp[op] = op.Stage + 1;
+			AddChildren(dep0);
+			AddChildren(dep1);
+			AddChildren(dep2);
+			AddChildren(dep3);
+		}
+
+		void AddOp(const FScheduledOp& op,
+			const FScheduledOp& dep0,
+			const FScheduledOp& dep1,
+			const FScheduledOp& dep2,
+			const FScheduledOp& dep3,
+			const FScheduledOp& dep4)
+		{
+			ClosedTasks.Add(FTask(op, dep0, dep1, dep2, dep3, dep4));
+			ScheduledStagePerOp[op] = op.Stage + 1;
+			AddChildren(dep0);
+			AddChildren(dep1);
+			AddChildren(dep2);
+			AddChildren(dep3);
+			AddChildren(dep4);
+		}
+
+		template<class RangeType>
+		void AddOp(const FScheduledOp& Op, const RangeType& Deps)
+		{
+			FTask Task(Op);
+
+			for (const FScheduledOp& D : Deps)
+			{
+				Task.Deps.Add(D);
+			}
+
+			ClosedTasks.Add(Task);
+
+			ScheduledStagePerOp[Op] = Op.Stage + 1;
+
+			for (const FScheduledOp& D : Deps)
+			{
+				AddChildren(D);
+			}
+		}
+
+    	/** Calculate an approximation of memory used by streaming buffers in this class. */
+		int32 GetStreamingMemoryBytes() const
+    	{
+			return RomLoadOps.GetAllocatedSize();			
+    	}
+    	
+		/** Calculate an approximation of memory used by manging structures in this class. */
+		int32 GetInternalMemoryBytes() const
+		{
+			return sizeof(CodeRunner) 
+				+ m_heapData.GetAllocatedSize() + m_heapImageDesc.GetAllocatedSize()
+				+ ClosedTasks.GetAllocatedSize() + OpenTasks.GetAllocatedSize() + ScheduledStagePerOp.GetAllocatedSize()
+				// this contains smart pointers, approximate size like this:
+				+ IssuedTasks.Max() * ( sizeof(FIssuedTask) + 16);
+		}
+
+	protected:
+
+		/** Strategy to choose the order of execution of operations. */
+		EExecutionStrategy ExecutionStrategy = EExecutionStrategy::None;
+
+		/** If this flag is enabled, issued operation stage that use tasks will be executed in the mutable thread instead of in a generic worker thread. */
+		bool bForceSerialTaskExecution = false;
+
+		/** List of pending operations that we don't know if they cannot be run yet because of dependencies. */
+		TArray< FTask > ClosedTasks;
+
+		/** List of tasks that can be run because they don't have any unmet dependency. */
+		TArray< FScheduledOp > OpenTasks;
+
+		/** For every op, up to what stage it has been scheduled to run. */
+		CodeContainer<uint8> ScheduledStagePerOp;
+
+		/** List of tasks that are ready to run concurrently. */
+		TArray< TSharedPtr<FIssuedTask> > IssuedTasksOnHold;
+
+		/** List of tasks that have been set to run concurrently and their completion is unknown. */
+		TArray< TSharedPtr<FIssuedTask> > IssuedTasks;
+
+	public:
+
+		inline bool LoadBool(const FCacheAddress& From)
+		{
+			return m_pSystem->WorkingMemoryManager.CurrentInstanceCache->GetBool(From);
+		}
+
+		inline float LoadInt(const FCacheAddress& From)
+		{
+			return m_pSystem->WorkingMemoryManager.CurrentInstanceCache->GetInt(From);
+		}
+
+		inline float LoadScalar(const FCacheAddress& From)
+		{
+			return m_pSystem->WorkingMemoryManager.CurrentInstanceCache->GetScalar(From);
+		}
+
+		inline FVector4f LoadColor(const FCacheAddress& From)
+		{
+			return m_pSystem->WorkingMemoryManager.CurrentInstanceCache->GetColour(From);
+		}
+
+		inline Ptr<const String> LoadString(const FCacheAddress& From)
+		{
+			return m_pSystem->WorkingMemoryManager.CurrentInstanceCache->GetString(From);
+		}
+
+		inline FProjector LoadProjector(const FCacheAddress& From)
+		{
+			return m_pSystem->WorkingMemoryManager.CurrentInstanceCache->GetProjector(From);
+		}
+
+		inline Ptr<const Mesh> LoadMesh(const FCacheAddress& From)
+		{
+			return m_pSystem->WorkingMemoryManager.LoadMesh(From);
+		}
+
+		inline Ptr<const Image> LoadImage(const FCacheAddress& From)
+		{
+			return m_pSystem->WorkingMemoryManager.LoadImage(From);
+		}
+
+		inline Ptr<const Layout> LoadLayout(const FCacheAddress& From)
+		{
+			return m_pSystem->WorkingMemoryManager.CurrentInstanceCache->GetLayout(From);
+		}
+
+		inline Ptr<const Instance> LoadInstance(const FCacheAddress& From)
+		{
+			return m_pSystem->WorkingMemoryManager.CurrentInstanceCache->GetInstance(From);
+		}
+
+		inline Ptr<const ExtensionData> LoadExtensionData(const FCacheAddress& From)
+		{
+			return m_pSystem->WorkingMemoryManager.CurrentInstanceCache->GetExtensionData(From);
+		}
+
+		inline void StoreValidDesc(const FCacheAddress& To)
+		{
+			m_pSystem->WorkingMemoryManager.CurrentInstanceCache->SetValidDesc(To);
+		}
+
+		inline void StoreBool(const FCacheAddress& To, bool Value)
+		{
+			m_pSystem->WorkingMemoryManager.CurrentInstanceCache->SetBool(To, Value);
+		}
+
+		inline void StoreInt(const FCacheAddress& To, int32 Value)
+		{
+			m_pSystem->WorkingMemoryManager.CurrentInstanceCache->SetInt(To, Value);
+		}
+
+		inline void StoreScalar(const FCacheAddress& To, float Value)
+		{
+			m_pSystem->WorkingMemoryManager.CurrentInstanceCache->SetScalar(To, Value);
+		}
+
+		inline void StoreString(const FCacheAddress& To, Ptr<const String> Value)
+		{
+			m_pSystem->WorkingMemoryManager.CurrentInstanceCache->SetString(To, Value);
+		}
+
+		inline void StoreColor(const FCacheAddress& To, const FVector4f& Value)
+		{
+			m_pSystem->WorkingMemoryManager.CurrentInstanceCache->SetColour(To, Value);
+		}
+
+		inline void StoreProjector(const FCacheAddress& To, const FProjector& Value)
+		{
+			m_pSystem->WorkingMemoryManager.CurrentInstanceCache->SetProjector(To, Value);
+		}
+
+		inline void StoreMesh(const FCacheAddress& To, Ptr<const Mesh> Resource)
+		{
+			m_pSystem->WorkingMemoryManager.StoreMesh(To, Resource);
+		}
+
+		inline void StoreImage(const FCacheAddress& To, Ptr<const Image> Resource)
+		{
+			m_pSystem->WorkingMemoryManager.StoreImage(To, Resource);
+		}
+
+		inline void StoreLayout(const FCacheAddress& To, Ptr<const Layout> Resource)
+		{
+			m_pSystem->WorkingMemoryManager.CurrentInstanceCache->SetLayout(To, Resource);
+		}
+
+		inline void StoreInstance(const FCacheAddress& To, Ptr<const Instance> Resource)
+		{
+			m_pSystem->WorkingMemoryManager.CurrentInstanceCache->SetInstance(To, Resource);
+		}
+
+		inline void StoreExtensionData(const FCacheAddress& To, Ptr<const ExtensionData> Resource)
+		{
+			m_pSystem->WorkingMemoryManager.CurrentInstanceCache->SetExtensionData(To, Resource);
+		}
+
+		inline Ptr<Image> CreateImage(uint32 SizeX, uint32 SizeY, uint32 Lods, EImageFormat Format, EInitializationType Init)
+		{
+			Ptr<Image> Result = m_pSystem->WorkingMemoryManager.CreateImage(SizeX, SizeY, Lods, Format, Init);
+			return MoveTemp(Result);
+		}
+
+		Ptr<Image> CreateImageLike(const Image* Ref, EInitializationType Init)
+		{
+			Ptr<Image> Result = m_pSystem->WorkingMemoryManager.CreateImage(Ref->GetSizeX(), Ref->GetSizeY(), Ref->GetLODCount(), Ref->GetFormat(), Init);
+			return MoveTemp(Result);
+		}
+
+		/** Ref will be nulled and relesed in any case. */
+		inline Ptr<Image> CloneOrTakeOver(Ptr<const Image>& Ref)
+		{
+			Ptr<Image> Result = m_pSystem->WorkingMemoryManager.CloneOrTakeOver(Ref);
+			return MoveTemp(Result);
+		}
+
+		inline void Release(Ptr<const Image>& Resource)
+		{
+			return m_pSystem->WorkingMemoryManager.Release(Resource);
+		}
+
+		inline void Release(Ptr<Image>& Resource)
+		{
+			return m_pSystem->WorkingMemoryManager.Release(Resource);
+		}
+
+		[[nodiscard]] inline Ptr<Mesh> CreateMesh(int32 BudgetReserveSize = 0)
+		{
+			return m_pSystem->WorkingMemoryManager.CreateMesh(BudgetReserveSize);
+		}
+
+		[[nodiscard]] inline Ptr<Mesh> CloneOrTakeOver(Ptr<const Mesh>& Ref)
+		{
+			return m_pSystem->WorkingMemoryManager.CloneOrTakeOver(Ref);
+		}
+
+		inline void Release(Ptr<const Mesh>& Resource)
+		{
+			m_pSystem->WorkingMemoryManager.Release(Resource);
+		}
+
+		inline void Release(Ptr<Mesh>& Resource)
+		{
+			m_pSystem->WorkingMemoryManager.Release(Resource);
+		}
+
+	// TODO: protect this.
+	public:
+		
+		/** This flag is turned on when a streaming error or similar happens. Results are not usable.
+		* This should only happen in-editor.
+		*/
+		bool bUnrecoverableError = false;
+
+		System::Private* m_pSystem = nullptr;
+		TSharedPtr<const Model> m_pModel = nullptr;
+		const Parameters* m_pParams = nullptr;
+		uint32 m_lodMask = 0;
+
+    private:
+    	struct FRomLoadOps
+    	{
+   		private:
+    		/** Rom read operations already in progress. */
+    		TArray<FRomLoadOp> RomLoadOps;
+
+			CodeRunner* Runner;
+
+    	public:
+
+			FRomLoadOps(CodeRunner& InRunner)
+			{
+				Runner = &InRunner;
+			}
+
+    		FRomLoadOp* Find(const int32 RomIndex)
+    		{
+    			for (FRomLoadOp& RomLoadOp : RomLoadOps)
+    			{
+    				if (RomLoadOp.RomIndex == RomIndex)
+    				{
+    					return &RomLoadOp;
+    				}
+    			}
+
+    			return nullptr;
+    		}
+    	
+    		FRomLoadOp& Create(const int32 RomIndex)
+    		{
+    			for (FRomLoadOp& RomLoadOp : RomLoadOps)
+    			{
+    				if (RomLoadOp.RomIndex == -1)
+    				{
+    					RomLoadOp.RomIndex = RomIndex;
+    					return RomLoadOp;
+    				}
+    			}
+
+    			FRomLoadOp& RomLoadOp = RomLoadOps.AddDefaulted_GetRef();
+    			RomLoadOp.RomIndex = RomIndex;
+
+    			return RomLoadOp;
+    		}
+
+    		void Remove(FRomLoadOp& RomLoadOp)
+    		{
+    			RomLoadOp.RomIndex = -1;
+                RomLoadOp.m_streamBuffer.Empty();
+                RomLoadOp.Event = {};
+            }
+
+    		int32 GetAllocatedSize() const
+    		{
+    			int32 Result = 0;
+			
+    			for (const FRomLoadOp& RomLoadOp : RomLoadOps)
+    			{
+    				Result += RomLoadOp.m_streamBuffer.GetAllocatedSize();
+    			}
+
+    			return Result;
+    		}
+    	};
+    	
+    	FRomLoadOps RomLoadOps = FRomLoadOps(*this);
+    	    	
+		inline void AddChildren(const FScheduledOp& dep)
+		{
+			FCacheAddress at(dep);
+			if (dep.At && !GetMemory().IsValid(at))
+			{
+				if (ScheduledStagePerOp.get(at) <= dep.Stage)
+				{
+					OpenTasks.Add(dep);
+					ScheduledStagePerOp[at] = dep.Stage + 1;
+				}
+			}
+
+			if (dep.Type == FScheduledOp::EType::Full)
+			{
+				m_pSystem->WorkingMemoryManager.CurrentInstanceCache->IncreaseHitCount(at);
+			}
+		}
+
+		/** Try to create a concurrent task for the given op. Return null if not possible. */
+		TSharedPtr<FIssuedTask> IssueOp(FScheduledOp item);
+
+		/** Calculate an heuristic to select op execution based on memory usage. */
+		int32 GetOpEstimatedMemoryDelta(const FScheduledOp& Candidate, const FProgram& Program);
+
+		/** Update debug stats. */
+		void UpdateTraces();
+
+		/** */
+		inline bool ShouldIssueTask() const
+		{
+			// Can we afford to delay issued tasks?
+			bool bCanDelayTasks = IssuedTasks.Num() > 0 || OpenTasks.Num() > 0;
+			if (!bCanDelayTasks)
+			{
+				return true;
+			}
+			else
+			{
+				// We could wait. Let's see if we have enough memory to issue tasks anyway.
+				bool bHaveEnoughMemory = !m_pSystem->WorkingMemoryManager.IsMemoryBudgetFull();
+				if (bHaveEnoughMemory)
+				{
+					return  true;
+				}
+			}
+
+			return false;
+		}
+
+		/** */
+		void LaunchIssuedTask(const TSharedPtr<FIssuedTask>& TaskToIssue, bool& bOutFailed);
+
+    };
+
+
+	/** Helper function to create the memory-tracked image operator. */
+	inline FImageOperator MakeImageOperator(CodeRunner* Runner)
+	{
+		return FImageOperator(
+			// Create
+			[Runner](int32 x, int32 y, int32 m, EImageFormat f, EInitializationType i)
+			{
+				return Runner->CreateImage(x, y, m, f, i);
+			},
+
+			// Release
+			[Runner](Ptr<Image>& i)
+			{
+				Runner->Release(i);
+			},
+
+			// Clone
+			[Runner](const Image* i)
+			{
+				Ptr<Image> New = Runner->CreateImage(i->GetSizeX(), i->GetSizeY(), i->GetLODCount(), i->GetFormat(), EInitializationType::NotInitialized);
+				New->Copy(i);
+				return New;
+			},
+
+			Runner->m_pSystem->ImagePixelFormatOverride
+		);
+	}
+
+}
