[mutable] Moved the Mutable plugin out of Experimental status into Beta.

diff --git a/Engine/Plugins/Mutable/Source/MutableRuntime/Private/MuR/CodeRunnerTasked.cpp b/Engine/Plugins/Mutable/Source/MutableRuntime/Private/MuR/CodeRunnerTasked.cpp
new file mode 100644
--- /dev/null
+++ b/Engine/Plugins/Mutable/Source/MutableRuntime/Private/MuR/CodeRunnerTasked.cpp
@@ -0,0 +1,2872 @@
+// Copyright Epic Games, Inc. All Rights Reserved.
+
+#include "Algo/Find.h"
+#include "Async/TaskGraphInterfaces.h"
+#include "Containers/Array.h"
+#include "Containers/UnrealString.h"
+#include "HAL/PlatformCrt.h"
+#include "HAL/UnrealMemory.h"
+#include "Logging/LogCategory.h"
+#include "Logging/LogMacros.h"
+#include "Math/IntPoint.h"
+#include "Math/UnrealMathSSE.h"
+#include "Misc/AssertionMacros.h"
+#include "MuR/CodeRunner.h"
+#include "MuR/System.h"
+#include "MuR/ExtensionDataStreamer.h"
+#include "MuR/Image.h"
+#include "MuR/ImagePrivate.h"
+#include "MuR/Layout.h"
+#include "MuR/Mesh.h"
+#include "MuR/Model.h"
+#include "MuR/ModelPrivate.h"
+#include "MuR/MutableMath.h"
+#include "MuR/MutableTrace.h"
+#include "MuR/OpImageBlend.h"
+#include "MuR/OpImageNormalCombine.h"
+#include "MuR/OpImageSaturate.h"
+#include "MuR/OpImageInvert.h"
+#include "MuR/Operations.h"
+#include "MuR/Platform.h"
+#include "MuR/Ptr.h"
+#include "MuR/RefCounted.h"
+#include "MuR/Serialisation.h"
+#include "MuR/SerialisationPrivate.h"
+#include "MuR/Settings.h"
+#include "MuR/SystemPrivate.h"
+#include "MuR/MutableRuntimeModule.h"
+#include "Stats/Stats2.h"
+#include "Templates/RefCounting.h"
+#include "Templates/SharedPointer.h"
+#include "Templates/Tuple.h"
+#include "Trace/Detail/Channel.h"
+#include "ProfilingDebugging/CountersTrace.h"
+#include "Async/Fundamental/Scheduler.h"
+
+
+DECLARE_CYCLE_STAT(TEXT("MutableCoreTask"), STAT_MutableCoreTask, STATGROUP_Game);
+
+namespace mu
+{
+
+	//---------------------------------------------------------------------------------------------
+	TRACE_DECLARE_INT_COUNTER(MutableRuntime_OpenTask, TEXT("MutableRuntime/OpenTask"));
+	TRACE_DECLARE_INT_COUNTER(MutableRuntime_ClosedTasks, TEXT("MutableRuntime/ClosedTasks"));
+	TRACE_DECLARE_INT_COUNTER(MutableRuntime_IssuedTasks, TEXT("MutableRuntime/IssuedTasks"));
+	TRACE_DECLARE_INT_COUNTER(MutableRuntime_IssuedTasksOnHold, TEXT("MutableRuntime/IssuedHoldTasks"));
+
+	void CodeRunner::UpdateTraces()
+	{
+		// Code Runner status
+		TRACE_COUNTER_SET(MutableRuntime_OpenTask, OpenTasks.Num());
+		TRACE_COUNTER_SET(MutableRuntime_ClosedTasks, ClosedTasks.Num());
+		TRACE_COUNTER_SET(MutableRuntime_IssuedTasks, IssuedTasks.Num());
+		TRACE_COUNTER_SET(MutableRuntime_IssuedTasksOnHold, IssuedTasksOnHold.Num());
+	}
+
+
+	//---------------------------------------------------------------------------------------------
+	/** This table encodes an heuristic for every execution stage of every operation  that tries to 
+	* guess the memory delta of that stage: 
+	* - if negative: after the stage the working memory may have reduced
+	* - if positive: after the stage the working memory may have increased
+	* The actual number is just a unit-less weight that represents "amount of additional memory that will be used".
+	* Some operations don't have a constant memory delta and are dealt with specially, but still need 
+	* to be in this table.
+	*/
+	static constexpr int32 sMemoryWeightsStageCount = 4;
+	static const int8 sMemoryWeights[int32(OP_TYPE::COUNT)][sMemoryWeightsStageCount] =
+	{ 
+		{   0,   0,   0,   0 },	// NONE
+		    
+		{   0,   0,   0,   0 },	// BO_CONSTANT
+		{   0,   0,   0,   0 },	// NU_CONSTANT
+		{   0,   0,   0,   0 },	// SC_CONSTANT
+		{   0,   0,   0,   0 },	// CO_CONSTANT
+		{  20,   0,   0,   0 },	// IM_CONSTANT
+		{  20,   0,   0,   0 },	// ME_CONSTANT
+		{   0,   0,   0,   0 },	// LA_CONSTANT
+		{   0,   0,   0,   0 },	// PR_CONSTANT
+		{   0,   0,   0,   0 },	// ST_CONSTANT
+		    
+		{   0,   0,   0,   0 },	// BO_PARAMETER
+		{   0,   0,   0,   0 },	// NU_PARAMETER
+		{   0,   0,   0,   0 },	// SC_PARAMETER
+		{   0,   0,   0,   0 },	// CO_PARAMETER
+		{   0,   0,   0,   0 },	// PR_PARAMETER
+		{  20,   0,   0,   0 },	// IM_PARAMETER
+		{   0,   0,   0,   0 },	// ST_PARAMETER
+		    
+		{   0,   0,   0,   0 },	// IM_REFERENCE
+		{   0,   0,   0,   0 },	// ME_REFERENCE
+
+		{   0,   0,   0,   0 },	// NU_CONDITIONAL
+		{   0,   0,   0,   0 },	// SC_CONDITIONAL
+		{   0,   0,   0,   0 },	// CO_CONDITIONAL
+		{   0,   0,   0,   0 },	// IM_CONDITIONAL
+		{   0,   0,   0,   0 },	// ME_CONDITIONAL
+		{   0,   0,   0,   0 },	// LA_CONDITIONAL
+		{   0,   0,   0,   0 },	// IN_CONDITIONAL
+		{   0,   0,   0,   0 },	// ED_CONDITIONAL
+		    
+		{   0,   0,   0,   0 },	// NU_SWITCH
+		{   0,   0,   0,   0 },	// SC_SWITCH
+		{   0,   0,   0,   0 },	// CO_SWITCH
+		{   0,   0,   0,   0 },	// IM_SWITCH
+		{   0,   0,   0,   0 },	// ME_SWITCH
+		{   0,   0,   0,   0 },	// LA_SWITCH
+		{   0,   0,   0,   0 },	// IN_SWITCH
+		{   0,   0,   0,   0 },	// ED_SWITCH
+		    
+		{   0,   0,   0,   0 },	// BO_LESS
+		{   0,   0,   0,   0 },	// BO_EQUAL_SC_CONST
+		{   0,   0,   0,   0 },	// BO_AND
+		{   0,   0,   0,   0 },	// BO_OR
+		{   0,   0,   0,   0 },	// BO_NOT
+		    
+		{   0,   0,   0,   0 },	// SC_MULTIPLYADD
+		{   0,   0,   0,   0 },	// SC_ARITHMETIC
+		{   0,   0,   0,   0 },	// SC_CURVE
+		    
+		{   0,   0,   0,   0 },	// CO_SAMPLEIMAGE
+		{   0,   0,   0,   0 },	// CO_SWIZZLE
+		{   0,   0,   0,   0 },	// CO_FROMSCALARS
+		{   0,   0,   0,   0 },	// CO_ARITHMETIC
+		    
+		{   0, -20, -20,   0 },	// IM_LAYER
+		{   0,   0,   0,   0 },	// IM_LAYERCOLOUR
+		{   0,   0,   0,   0 },	// IM_PIXELFORMAT	(special case)
+		{   0,  10,   0,   0 },	// IM_MIPMAP
+		{   0,   0,   0,   0 },	// IM_RESIZE		(special case)
+		{   0,   0,   0,   0 },	// IM_RESIZELIKE	(to be deprecated?)
+		{   0,   0,   0,   0 },	// IM_RESIZEREL		(special case)
+		{   0,  20,   0,   0 },	// IM_BLANKLAYOUT
+		{   0, -20, -20,   0 },	// IM_COMPOSE
+		{   0,   0, -20,   0 },	// IM_INTERPOLATE
+		{   0,   0,   0,   0 },	// IM_SATURATE
+		{   0,   0,   0,   0 },	// IM_LUMINANCE
+		{   0,   0,   0,   0 },	// IM_SWIZZLE
+		{   0,   0,   0,   0 },	// IM_COLOURMAP
+		{   0,   5,   0,   0 },	// IM_GRADIENT
+		{   0,   0,   0,   0 },	// IM_BINARISE
+		{   0,  20,   0,   0 },	// IM_PLAINCOLOUR
+		{   0, -10,   0,   0 },	// IM_CROP
+		{   0, -10,   0,   0 },	// IM_PATCH
+		{   0,  10,  10,   0 },	// IM_RASTERMESH
+		{   0,   0,   0,   0 },	// IM_MAKEGROWMAP
+		{   0, -10,   0,   0 },	// IM_DISPLACE
+		{   0,   0, -20, -20 },	// IM_MULTILAYER
+		{   0,   0,   0,   0 },	// IM_INVERT
+		{   0,   0,   0,   0 },	// IM_NORMALCOMPOSITE
+		{   0,   0,   0,   0 },	// IM_TRANSFORM
+		    
+		{   0,   0,   0,   0 },	// ME_APPLYLAYOUT
+		{   0, -20,   0,   0 },	// ME_DIFFERENCE
+		{   0,   0, -10,   0 },	// ME_MORPH2
+		{   0,   0,   0,   0 },	// ME_MERGE
+		{   0,   0, -20,   0 },	// ME_INTERPOLATE
+		{   0, -10,   0,   0 },	// ME_MASKCLIPMESH
+		{   0, -10,   0,   0 },	// ME_MASKCLIPUVMASK
+		{   0, -10,   0,   0 },	// ME_MASKDIFF
+		{   0,   0, -10,   0 },	// ME_REMOVEMASK
+		{   0,   0,   0,   0 },	// ME_FORMAT
+		{   0, -10,   0,   0 },	// ME_EXTRACTLAYOUTBLOCK
+		{   0,   0,   0,   0 },	// ME_TRANSFORM
+		{   0, -10,   0,   0 },	// ME_CLIPMORPHPLANE
+		{   0, -10,   0,   0 },	// ME_CLIPWITHMESH
+		{   0,   0,   0,   0 },	// ME_SETSKELETON
+		{   0,   0,   0,   0 },	// ME_PROJECT
+		{   0,   0,   0,   0 },	// ME_APPLYPOSE
+		{   0,   0,   0,   0 },	// ME_GEOMETRYOPERATION
+		{   0, -20,   0,   0 },	// ME_BINDSHAPE
+		{   0, -10,   0,   0 },	// ME_APPLYSHAPE
+		{   0, -10,   0,   0 },	// ME_CLIPDEFORM
+		{   0, -10,   0,   0 },	// ME_MORPHRESHAPE
+		{   0,   0,   0,   0 },	// ME_OPTIMIZESKINNING
+		{	0,   0,	  0,   0 },	// ME_ADDTAGS
+
+		{   0,   0,   0,   0 },	// IN_ADDMESH
+		{   0,   0,   0,   0 },	// IN_ADDIMAGE
+		{   0,   0,   0,   0 },	// IN_ADDVECTOR
+		{   0,   0,   0,   0 },	// IN_ADDSCALAR
+		{   0,   0,   0,   0 },	// IN_ADDSTRING
+		{   0,   0,   0,   0 },	// IN_ADDSURFACE
+		{   0,   0,   0,   0 },	// IN_ADDCOMPONENT
+		{	0,   0,   0,   0 },	// IN_ADDLOD
+		{	0,   0,   0,   0 },	// IN_ADDEXTENSIONDATA
+
+		{   0,   0,   0,   0 },	// LA_PACK
+		{   0,   0,   0,   0 },	// LA_MERGE
+		{   0,  -1,   0,   0 },	// LA_REMOVEBLOCKS
+		{   0,   0,   0,   0 },	// LA_FROMMESH
+	};
+
+	static_assert(sizeof(sMemoryWeights)/sMemoryWeightsStageCount == int32(OP_TYPE::COUNT));
+
+
+	int32 CodeRunner::GetOpEstimatedMemoryDelta( const FScheduledOp& Candidate, const FProgram& Program )
+	{
+		int32 OpDelta = 0;
+		int32 Stage = FMath::Min(int32(Candidate.Stage), sMemoryWeightsStageCount - 1);
+		OP_TYPE OpType = Program.GetOpType(Candidate.At);
+
+		if (OpType == OP_TYPE::IM_PIXELFORMAT && Stage == 1)
+		{
+			// TODO: We should get the actual format from the arguments for a more precise calculation.
+			OP::ImagePixelFormatArgs args = Program.GetOpArgs<OP::ImagePixelFormatArgs>(Candidate.At);
+			bool bIsCompressed = GetUncompressedFormat(args.format) != args.format;
+			bool bIsSmallish = GetImageFormatData(args.format).BytesPerBlock < 2;
+			if (bIsCompressed || bIsSmallish)
+			{
+				OpDelta = -10;
+			}
+			else
+			{
+				OpDelta = 10;
+			}
+		}
+		else
+		{
+			OpDelta = sMemoryWeights[int32(OpType)][Stage];
+		}
+
+		return OpDelta;
+	}
+
+
+	//---------------------------------------------------------------------------------------------
+	void CodeRunner::LaunchIssuedTask( const TSharedPtr<FIssuedTask>& TaskToIssue, bool& bOutFailed )
+	{
+		bool bFailed = false;
+		bool bHasWork = TaskToIssue->Prepare(this, bFailed);
+		if (bFailed)
+		{
+			bUnrecoverableError = true;
+			return;
+		}
+
+		// Launch it
+		if (bHasWork)
+		{
+			if (bForceSerialTaskExecution)
+			{
+				TaskToIssue->Event = {};
+				TaskToIssue->DoWork();
+			}
+			else
+			{
+				TaskToIssue->Event = UE::Tasks::Launch(TEXT("MutableCore_Task"),
+					[TaskToIssue]() { TaskToIssue->DoWork(); },
+					UE::Tasks::ETaskPriority::Inherit);
+			}
+		}
+
+		// Remember it for later processing.
+		IssuedTasks.Add(TaskToIssue);
+	}
+
+	UE::Tasks::FTask CodeRunner::StartRun(bool bForceInlineExecution)
+	{
+		check(RunnerCompletionEvent.IsCompleted());
+
+		bUnrecoverableError = false;
+
+		m_heapData.SetNum(0, EAllowShrinking::No);
+		m_heapImageDesc.SetNum(0, EAllowShrinking::No);
+		
+		RunnerCompletionEvent = UE::Tasks::FTaskEvent(TEXT("CodeRunnerCompletionEvent"));
+		
+		bool bProfile = false;
+		
+		TUniquePtr<FProfileContext> ProfileContext = bProfile ? MakeUnique<FProfileContext>() : nullptr;
+		Run(MoveTemp(ProfileContext), bForceInlineExecution);
+
+		check(!bForceInlineExecution || RunnerCompletionEvent.IsCompleted());
+
+		return RunnerCompletionEvent;
+	}
+
+	void CodeRunner::AbortRun()
+	{
+		bUnrecoverableError = true;
+		RunnerCompletionEvent.Trigger();
+	}
+
+    void CodeRunner::Run(TUniquePtr<FProfileContext>&& ProfileContext, bool bForceInlineExecution)
+    {
+		MUTABLE_CPUPROFILER_SCOPE(CodeRunner_Run);
+
+		check(!RunnerCompletionEvent.IsCompleted());
+
+		// TODO: Move MaxAllowedTime somewhere else more accessible, maybe a cvar.
+		const FTimespan MaxAllowedTime = FTimespan::FromMilliseconds(2.0); 
+		const FTimespan TimeOut = FTimespan::FromSeconds(FPlatformTime::Seconds()) + MaxAllowedTime;
+
+		bool bSuccess = true;
+
+        while(!OpenTasks.IsEmpty() || !ClosedTasks.IsEmpty() || !IssuedTasks.IsEmpty())
+        {
+			UpdateTraces();
+			// Debug: log the amount of tasks that we'd be able to run concurrently:
+			//{
+			//	int32 ClosedReady = ClosedTasks.Num();
+			//	for (int Index = ClosedTasks.Num() - 1; Index >= 0; --Index)
+			//	{
+			//		for (const FCacheAddress& Dep : ClosedTasks[Index].Deps)
+			//		{
+			//			if (Dep.at && !GetMemory().IsValid(Dep))
+			//			{
+			//				--ClosedReady;
+			//				continue;
+			//			}
+			//		}
+			//	}
+
+			//	UE_LOG(LogMutableCore, Log, TEXT("Tasks: %5d open, %5d issued, %5d closed, %d closed ready"), OpenTasks.Num(), IssuedTasks.Num(), ClosedTasks.Num(), ClosedReady);
+			//}
+
+			for (int32 Index = 0; bSuccess && Index < IssuedTasks.Num(); )
+			{
+				check(IssuedTasks[Index]);
+
+				bool bWorkDone = IssuedTasks[Index]->IsComplete(this);
+				if (bWorkDone)
+				{
+					const FScheduledOp& item = IssuedTasks[Index]->Op;
+					bSuccess = IssuedTasks[Index]->Complete(this);
+
+					if (ScheduledStagePerOp[item] == item.Stage + 1)
+					{
+						// We completed everything that was requested, clear it otherwise if needed
+						// again it is not going to be rebuilt.
+						// \TODO: track rebuilds.
+						ScheduledStagePerOp[item] = 0;
+					}
+
+					IssuedTasks.RemoveAt(Index, EAllowShrinking::No); // with swap? changes order of execution.
+				}
+				else
+				{
+					++Index;
+				}
+			}
+
+			if (!bSuccess)
+			{
+				return AbortRun();
+			}
+
+			while (!OpenTasks.IsEmpty())
+			{
+				// Get a new task to run
+				FScheduledOp item;
+				switch (ExecutionStrategy)
+				{
+				//case EExecutionStrategy::MinimizeMemory:
+				//{
+				//	// TODO: This should be done when an operation is added to the OpenTasks array instead of every time.
+				//	int32 BestOp = 0;
+				//	int8 BestDelta = TNumericLimits<int8>::Max();
+				//	int32 OpenOpCount = OpenTasks.Num();
+				//	const FProgram& Program = m_pModel->GetPrivate()->m_program;
+				//	for (int32 OpIndex = 0; OpIndex < OpenOpCount; ++OpIndex)
+				//	{
+				//		const FScheduledOp& Candidate = OpenTasks[OpIndex];
+				//		int32 OpDelta = GetOpEstimatedMemoryDelta(Candidate, Program);
+
+				//		if (OpDelta < BestDelta)
+				//		{
+				//			BestOp = OpIndex;
+				//			BestDelta = OpDelta;
+
+				//			// Shortcut: If we are freeing memory, we don't care how much: it is already the best op
+				//			if (BestDelta < 0)
+				//			{
+				//				break;
+				//			}
+				//		}
+				//	}
+
+				//	item = OpenTasks[BestOp];
+				//	OpenTasks.RemoveAtSwap(BestOp,EAllowShrinking::No);
+				//	break;
+				//}
+
+				case EExecutionStrategy::None:
+				default:
+					// Just get one.
+					item = OpenTasks.Pop(EAllowShrinking::No);
+					break;
+
+				}
+
+				// Special processing in case it is an ImageDesc operation
+				if (item.Type == FScheduledOp::EType::ImageDesc)
+				{
+					RunCodeImageDesc(item, m_pParams, m_pModel.Get(), m_lodMask);
+					continue;
+				}
+
+				// Don't run it if we already have the result.
+				FCacheAddress cat(item);
+				if (GetMemory().IsValid(cat))
+				{
+					continue;
+				}
+
+				// See if we can schedule this item concurrently
+				TSharedPtr<FIssuedTask> IssuedTask = IssueOp(item);
+				if (IssuedTask)
+				{
+					if (ShouldIssueTask())
+					{
+						bool bFailed = false;
+						LaunchIssuedTask(IssuedTask, bFailed);
+						if (bFailed)
+						{
+							return AbortRun();
+						}
+					}
+					else
+					{
+						IssuedTasksOnHold.Add(IssuedTask);
+					}
+				}
+				else
+				{
+					// Run immediately
+					RunCode(item, m_pParams, m_pModel, m_lodMask);
+
+					if (ScheduledStagePerOp[item] == item.Stage + 1)
+					{
+						// We completed everything that was requested, clear it. Otherwise if needed again it is not going to be rebuilt.
+						// \TODO: track operations that are run more than once?.
+						ScheduledStagePerOp[item] = 0;
+					}
+				}
+
+				if (ProfileContext)
+				{
+					++ProfileContext->NumRunOps;
+					++ProfileContext->RunOpsPerType[int32(m_pModel->GetPrivate()->m_program.GetOpType(item.At))];
+				}
+			}
+
+			UpdateTraces();
+
+			// Look for tasks on hold and see if we can launch them
+			while (IssuedTasksOnHold.Num() && ShouldIssueTask())
+			{
+				TSharedPtr<FIssuedTask> TaskToIssue = IssuedTasksOnHold.Pop(EAllowShrinking::No);
+
+				bool bFailed = false;
+				LaunchIssuedTask(TaskToIssue, bFailed);
+				if (bFailed)
+				{
+					return AbortRun();
+				}
+			}
+
+			// Look for a closed task with dependencies satisfied and move them to the open task list.
+			bool bSomeWasReady = false;
+			for (int32 Index = 0; Index < ClosedTasks.Num(); )
+			{
+				bool Ready = true;
+				for ( const FCacheAddress& Dep: ClosedTasks[Index].Deps )
+				{
+					bool bDependencyFailed = false;
+					bDependencyFailed = Dep.At && !GetMemory().IsValid(Dep);
+					
+					if (bDependencyFailed)
+					{
+						Ready = false;
+						break;
+					}
+				}
+
+				if (Ready)
+				{
+					bSomeWasReady = true;
+					FTask Task = ClosedTasks[Index];
+					ClosedTasks.RemoveAt(Index, EAllowShrinking::No); // with swap? would change order of execution.
+					OpenTasks.Push(Task.Op);
+				}
+				else
+				{
+					++Index;
+				}
+
+			}
+
+			UpdateTraces();
+
+			// Debug: Did we dead-lock?
+			bool bDeadLock = !(OpenTasks.Num() || IssuedTasks.Num() || !ClosedTasks.Num() || bSomeWasReady);
+			if (bDeadLock)
+			{
+				// Log the task graph
+				for (int32 Index = 0; Index < ClosedTasks.Num(); ++Index)
+				{
+					FString TaskDesc = FString::Printf(TEXT("Closed task %d-%d-%d depends on : "), ClosedTasks[Index].Op.At, ClosedTasks[Index].Op.ExecutionIndex, ClosedTasks[Index].Op.Stage );
+					for (const FCacheAddress& Dep : ClosedTasks[Index].Deps)
+					{
+						if (Dep.At && !GetMemory().IsValid(Dep))
+						{
+							TaskDesc += FString::Printf(TEXT("%d-%d, "), Dep.At, Dep.ExecutionIndex);
+						}
+					}
+
+					UE_LOG(LogMutableCore, Log, TEXT("%s"), *TaskDesc);
+				}
+				check(false);
+
+				// This should never happen but if it does, abort the code execution.
+				return AbortRun();
+			}
+
+			// If at this point there is no open op and we haven't finished, we need to wait for an issued op to complete.
+			if (OpenTasks.IsEmpty() && !IssuedTasks.IsEmpty())
+			{
+				if (!bForceInlineExecution)
+				{
+					TArray<UE::Tasks::FTask, TInlineAllocator<8>> IssuedTasksCompletionEvents;
+					IssuedTasksCompletionEvents.Reserve(IssuedTasks.Num());
+
+					for (TSharedPtr<FIssuedTask>& IssuedTask : IssuedTasks)
+					{
+						if (IssuedTask->Event.IsValid())
+						{
+							IssuedTasksCompletionEvents.Add(IssuedTask->Event);
+						}
+					}
+
+					m_pSystem->WorkingMemoryManager.InvalidateRunnerThread();
+
+					UE::Tasks::Launch(TEXT("CodeRunnerFromIssuedTasksTask"),
+						[Runner = AsShared(), ProfileContext = MoveTemp(ProfileContext)]() mutable
+						{
+							Runner->m_pSystem->WorkingMemoryManager.ResetRunnerThread();
+
+							constexpr bool bForceInlineExecution = false;
+							Runner->Run(MoveTemp(ProfileContext), bForceInlineExecution);
+						},
+						UE::Tasks::Prerequisites(UE::Tasks::Any(IssuedTasksCompletionEvents)),
+						UE::Tasks::ETaskPriority::Inherit);
+					
+					return;
+				}	
+				else
+				{
+					MUTABLE_CPUPROFILER_SCOPE(CodeRunner_WaitIssued);
+					for (int32 IssuedIndex = 0; IssuedIndex < IssuedTasks.Num(); ++IssuedIndex)
+					{
+						if (IssuedTasks[IssuedIndex]->Event.IsValid())
+						{
+							IssuedTasks[IssuedIndex]->Event.Wait();
+
+							break;
+						}
+					}
+				}
+			}
+
+			if (!bForceInlineExecution)
+			{
+				if (FTimespan::FromSeconds(FPlatformTime::Seconds()) > TimeOut)
+				{
+					m_pSystem->WorkingMemoryManager.InvalidateRunnerThread();
+
+					UE::Tasks::Launch(TEXT("CodeRunnerFromTimeoutTask"),
+						[Runner = AsShared(), ProfileContext = MoveTemp(ProfileContext)]() mutable
+						{
+							Runner->m_pSystem->WorkingMemoryManager.ResetRunnerThread();
+
+							constexpr bool bForceInlineExecution = false;
+							Runner->Run(MoveTemp(ProfileContext), bForceInlineExecution);
+						},
+						UE::Tasks::ETaskPriority::Inherit);
+					
+					return;
+				}
+			}
+		}
+
+		if (ProfileContext)
+		{
+			UE_LOG(LogMutableCore, Log, TEXT("Mutable Heap Bytes: %d"), m_heapData.Num()* m_heapData.GetTypeSize());
+			UE_LOG(LogMutableCore, Log, TEXT("Ran ops : %5d "), ProfileContext->NumRunOps);
+
+			constexpr int32 HistogramSize = 8;
+			int32 MostCommonOps[HistogramSize] = {};
+			for (int32 OpIndex = 0; OpIndex < int32(OP_TYPE::COUNT); ++OpIndex)
+			{
+				for (int32 HistIndex = 0; HistIndex < HistogramSize; ++HistIndex)
+				{
+					if (ProfileContext->RunOpsPerType[OpIndex] > ProfileContext->RunOpsPerType[MostCommonOps[HistIndex]])
+					{
+						// Displace others
+						int32 ElementsToMove = HistogramSize - HistIndex - 1;
+						if (ElementsToMove > 0)
+						{
+							FMemory::Memcpy(&MostCommonOps[HistIndex + 1], &MostCommonOps[HistIndex], sizeof(int32)*ElementsToMove);
+						}
+						// Set new value
+						MostCommonOps[HistIndex] = OpIndex;
+						break;
+					}
+				}
+			}
+
+			for (int32 HistIndex = 0; HistIndex < HistogramSize; ++HistIndex)
+			{
+				UE_LOG(LogMutableCore, Log, TEXT("    op %4d, %4d times."), MostCommonOps[HistIndex], ProfileContext->RunOpsPerType[MostCommonOps[HistIndex]]);
+			}
+		}
+
+		RunnerCompletionEvent.Trigger();
+    }
+
+	
+	//---------------------------------------------------------------------------------------------
+	void CodeRunner::GetImageDescResult(FImageDesc& OutDesc)
+	{
+		check( m_heapImageDesc.Num()>0 );
+		OutDesc = m_heapImageDesc[0];
+	}
+
+	//---------------------------------------------------------------------------------------------
+	//---------------------------------------------------------------------------------------------
+	//---------------------------------------------------------------------------------------------
+	class FImageLayerTask : public CodeRunner::FIssuedTask
+	{
+	public:
+		FImageLayerTask(const FScheduledOp& InOp, const OP::ImageLayerArgs& InArgs)
+			: FIssuedTask(InOp), Args(InArgs)
+		{}
+
+		// FIssuedTask interface
+		virtual bool Prepare(CodeRunner*, bool& bOutFailed) override;
+		virtual void DoWork() override;
+		virtual bool Complete(CodeRunner*) override;
+
+	private:
+		int32 ImageCompressionQuality = 0;
+		OP::ImageLayerArgs Args;
+		Ptr<const Image> Blended;
+		Ptr<const Image> Mask;
+		Ptr<Image> Result;
+		EImageFormat InitialFormat = EImageFormat::IF_NONE;
+	};
+
+
+	bool FImageLayerTask::Prepare(CodeRunner* Runner, bool& bOutFailed)
+	{
+		// This runs in the mutable Runner thread
+		MUTABLE_CPUPROFILER_SCOPE(FImageLayerTask_Prepare);
+		bOutFailed = false;
+
+		ImageCompressionQuality = Runner->m_pSettings->ImageCompressionQuality;
+
+		Ptr<const Image> Base = Runner->LoadImage({ Args.base, Op.ExecutionIndex, Op.ExecutionOptions });
+		check(!Base || Base->GetFormat() < EImageFormat::IF_COUNT);
+
+		Blended = Runner->LoadImage({ Args.blended, Op.ExecutionIndex, Op.ExecutionOptions });
+		if (Args.mask)
+		{
+			Mask = Runner->LoadImage({ Args.mask, Op.ExecutionIndex, Op.ExecutionOptions });
+			check(!Mask || Mask->GetFormat() < EImageFormat::IF_COUNT);
+		}
+
+		// Shortcuts
+		if (!Base)
+		{
+			Runner->Release(Blended);
+			Runner->Release(Mask);
+			Runner->StoreImage(Op, Base);
+			return false;
+		}
+
+		bool bValid = Base->GetSizeX() > 0 && Base->GetSizeY() > 0;
+		if (!bValid || !Blended)
+		{
+			Runner->Release(Blended);
+			Runner->Release(Mask);
+			Runner->StoreImage(Op, Base);
+			return false;
+		}
+
+		FImageOperator ImOp = MakeImageOperator(Runner);
+
+		// Input data fixes
+		InitialFormat = Base->GetFormat();
+
+		if (IsCompressedFormat(InitialFormat))
+		{
+			EImageFormat UncompressedFormat = GetUncompressedFormat(InitialFormat);
+			Ptr<Image> Formatted = Runner->CreateImage( Base->GetSizeX(), Base->GetSizeY(), Base->GetLODCount(), UncompressedFormat, EInitializationType::NotInitialized );
+			bool bSuccess = false;
+			ImOp.ImagePixelFormat(bSuccess, ImageCompressionQuality, Formatted.get(), Base.get());
+			check(bSuccess); // Decompression cannot fail
+			Runner->Release(Base);
+			Base = Formatted;
+		}
+
+		bool bMustHaveSameFormat = !(Args.flags & (OP::ImageLayerArgs::F_BASE_RGB_FROM_ALPHA | OP::ImageLayerArgs::F_BLENDED_RGB_FROM_ALPHA));
+		if (Blended && InitialFormat != Blended->GetFormat() && bMustHaveSameFormat)
+		{
+			Ptr<Image> Formatted = Runner->CreateImage(Blended->GetSizeX(), Blended->GetSizeY(), Blended->GetLODCount(), Base->GetFormat(), EInitializationType::NotInitialized);
+			bool bSuccess = false;
+			ImOp.ImagePixelFormat(bSuccess, ImageCompressionQuality, Formatted.get(), Blended.get());
+			check(bSuccess); // Decompression cannot fail
+			Runner->Release(Blended);
+			Blended = Formatted;
+		}
+
+		if (Base->GetSize() != Blended->GetSize())
+		{
+			MUTABLE_CPUPROFILER_SCOPE(ImageResize_EmergencyFix);
+			Ptr<Image> Resized = Runner->CreateImage(Base->GetSizeX(), Base->GetSizeY(), 1, Blended->GetFormat(), EInitializationType::NotInitialized);
+			ImOp.ImageResizeLinear(Resized.get(), ImageCompressionQuality, Blended.get());
+			Runner->Release(Blended);
+			Blended = Resized;
+
+		}
+
+		if (Mask)
+		{
+			if (Base->GetSize() != Mask->GetSize())
+			{
+				MUTABLE_CPUPROFILER_SCOPE(ImageResize_EmergencyFix);
+				Ptr<Image> Resized = Runner->CreateImage(Base->GetSizeX(), Base->GetSizeY(), 1, Mask->GetFormat(), EInitializationType::NotInitialized);
+				ImOp.ImageResizeLinear(Resized.get(), ImageCompressionQuality, Mask.get());
+				Runner->Release(Mask);
+				Mask = Resized;
+
+			}
+
+			// emergy fix c36adf47-e40d-490f-b709-41142bafad78
+			if (Mask->GetLODCount() < Base->GetLODCount())
+			{
+				MUTABLE_CPUPROFILER_SCOPE(ImageLayer_EmergencyFix);
+
+				int32 StartLevel = Mask->GetLODCount() - 1;
+				int32 LevelCount = Base->GetLODCount();
+
+				// Uncompress mask to avoid excessive RLE-compression and decompression in the following ops.
+				Ptr<Image> UncompressedMask = Runner->CreateImage(Mask->GetSizeX(), Mask->GetSizeY(), Mask->GetLODCount(), GetUncompressedFormat(Mask->GetFormat()), EInitializationType::NotInitialized);
+				bool bSuccess = false;
+				ImOp.ImagePixelFormat(bSuccess, ImageCompressionQuality, UncompressedMask.get(), Mask.get());
+
+				Ptr<Image> MaskFix = UncompressedMask;
+				MaskFix->DataStorage.SetNumLODs(LevelCount);
+
+				FMipmapGenerationSettings Settings{};
+				ImOp.ImageMipmap(ImageCompressionQuality, MaskFix.get(), MaskFix.get(), StartLevel, LevelCount, Settings);
+
+				Runner->Release(Mask);
+				Mask = MaskFix;
+			}
+		}
+
+		// Create destination data
+		Result = Runner->CloneOrTakeOver(Base);
+		return true;
+	}
+
+
+	void FImageLayerTask::DoWork()
+	{
+		MUTABLE_CPUPROFILER_SCOPE(FImageLayerTask);
+
+		// This runs in a worker thread.
+
+		bool bOnlyOneMip = false;
+		if (Blended->GetLODCount() < Result->GetLODCount())
+		{
+			bOnlyOneMip = true;
+		}
+
+		bool bDone = false;
+
+		if (!Mask 
+			&& 
+			// Flags have to match exactly for this optimize case. Other flags are not supported.
+			Args.flags == OP::ImageLayerArgs::F_USE_MASK_FROM_BLENDED
+			&&
+			Args.blendType == uint8(EBlendType::BT_BLEND)
+			&&
+			Args.blendTypeAlpha == uint8(EBlendType::BT_LIGHTEN))
+		{
+			MUTABLE_CPUPROFILER_SCOPE(ImageLayerTask_Optimized);
+
+			// This is a frequent critical-path case because of multilayer projectors.
+			bDone = true;
+
+			constexpr bool bUseVectorImplementation = false;	
+			if constexpr (bUseVectorImplementation)
+			{
+				BufferLayerCompositeVector<VectorBlendChannelMasked, VectorLightenChannel, false>(Result.get(), Blended.get(), bOnlyOneMip, Args.BlendAlphaSourceChannel);
+			}
+			else
+			{
+				BufferLayerComposite<BlendChannelMasked, LightenChannel, false>(Result.get(), Blended.get(), bOnlyOneMip, Args.BlendAlphaSourceChannel);
+			}
+		}
+
+		bool bApplyColorBlendToAlpha = (Args.flags & OP::ImageLayerArgs::F_APPLY_TO_ALPHA) != 0;
+		bool bUseBlendSourceFromBlendAlpha = (Args.flags & OP::ImageLayerArgs::F_BLENDED_RGB_FROM_ALPHA) != 0;
+		bool bUseMaskFromBlendAlpha = (Args.flags & OP::ImageLayerArgs::F_USE_MASK_FROM_BLENDED);
+
+		if (!bDone && Mask)
+		{
+			// Not implemented yet
+			check(!bUseBlendSourceFromBlendAlpha);
+
+			// TODO: in-place
+
+			switch (EBlendType(Args.blendType))
+			{
+			case EBlendType::BT_NORMAL_COMBINE: ImageNormalCombine(Result.get(), Result.get(), Mask.get(), Blended.get(), bOnlyOneMip); break;
+			case EBlendType::BT_SOFTLIGHT: BufferLayer<SoftLightChannelMasked, SoftLightChannel, true>(Result.get(), Result.get(), Mask.get(), Blended.get(), bApplyColorBlendToAlpha, bOnlyOneMip); break;
+			case EBlendType::BT_HARDLIGHT: BufferLayer<HardLightChannelMasked, HardLightChannel, true>(Result.get(), Result.get(), Mask.get(), Blended.get(), bApplyColorBlendToAlpha, bOnlyOneMip); break;
+			case EBlendType::BT_BURN: BufferLayer<BurnChannelMasked, BurnChannel, true>(Result.get(), Result.get(), Mask.get(), Blended.get(), bApplyColorBlendToAlpha, bOnlyOneMip); break;
+			case EBlendType::BT_DODGE: BufferLayer<DodgeChannelMasked, DodgeChannel, true>(Result.get(), Result.get(), Mask.get(), Blended.get(), bApplyColorBlendToAlpha, bOnlyOneMip); break;
+			case EBlendType::BT_SCREEN: BufferLayer<ScreenChannelMasked, ScreenChannel, true>(Result.get(), Result.get(), Mask.get(), Blended.get(), bApplyColorBlendToAlpha, bOnlyOneMip); break;
+			case EBlendType::BT_OVERLAY: BufferLayer<OverlayChannelMasked, OverlayChannel, true>(Result.get(), Result.get(), Mask.get(), Blended.get(), bApplyColorBlendToAlpha, bOnlyOneMip); break;
+			case EBlendType::BT_LIGHTEN: BufferLayer<LightenChannelMasked, LightenChannel, true>(Result.get(), Result.get(), Mask.get(), Blended.get(), bApplyColorBlendToAlpha, bOnlyOneMip); break;
+			case EBlendType::BT_MULTIPLY: BufferLayer<MultiplyChannelMasked, MultiplyChannel, true>(Result.get(), Result.get(), Mask.get(), Blended.get(), bApplyColorBlendToAlpha, bOnlyOneMip); break;
+			case EBlendType::BT_BLEND: BufferLayer<BlendChannelMasked, BlendChannel, true>(Result.get(), Result.get(), Mask.get(), Blended.get(), bApplyColorBlendToAlpha, bOnlyOneMip); break;
+			case EBlendType::BT_NONE: break;
+			default: check(false);
+			}
+
+		}
+		else if (!bDone && bUseMaskFromBlendAlpha)
+		{
+			// Not implemented yet
+			check(!bUseBlendSourceFromBlendAlpha);
+
+			// Apply blend without to RGB using mask in blended alpha
+			switch (EBlendType(Args.blendType))
+			{
+			case EBlendType::BT_NORMAL_COMBINE: check(false); break;
+			case EBlendType::BT_SOFTLIGHT: BufferLayerEmbeddedMask<SoftLightChannelMasked, SoftLightChannel, false>(Result.get(), Result.get(), Blended.get(), bApplyColorBlendToAlpha, bOnlyOneMip); break;
+			case EBlendType::BT_HARDLIGHT: BufferLayerEmbeddedMask<HardLightChannelMasked, HardLightChannel, false>(Result.get(), Result.get(), Blended.get(), bApplyColorBlendToAlpha, bOnlyOneMip); break;
+			case EBlendType::BT_BURN: BufferLayerEmbeddedMask<BurnChannelMasked, BurnChannel, false>(Result.get(), Result.get(), Blended.get(), bApplyColorBlendToAlpha, bOnlyOneMip); break;
+			case EBlendType::BT_DODGE: BufferLayerEmbeddedMask<DodgeChannelMasked, DodgeChannel, false>(Result.get(), Result.get(), Blended.get(), bApplyColorBlendToAlpha, bOnlyOneMip); break;
+			case EBlendType::BT_SCREEN: BufferLayerEmbeddedMask<ScreenChannelMasked, ScreenChannel, false>(Result.get(), Result.get(), Blended.get(), bApplyColorBlendToAlpha, bOnlyOneMip); break;
+			case EBlendType::BT_OVERLAY: BufferLayerEmbeddedMask<OverlayChannelMasked, OverlayChannel, false>(Result.get(), Result.get(), Blended.get(), bApplyColorBlendToAlpha, bOnlyOneMip); break;
+			case EBlendType::BT_LIGHTEN: BufferLayerEmbeddedMask<LightenChannelMasked, LightenChannel, false>(Result.get(), Result.get(), Blended.get(), bApplyColorBlendToAlpha, bOnlyOneMip); break;
+			case EBlendType::BT_MULTIPLY: BufferLayerEmbeddedMask<MultiplyChannelMasked, MultiplyChannel, false>(Result.get(), Result.get(), Blended.get(), bApplyColorBlendToAlpha, bOnlyOneMip); break;
+			case EBlendType::BT_BLEND: BufferLayerEmbeddedMask<BlendChannelMasked, BlendChannel, false>(Result.get(), Result.get(), Blended.get(), bApplyColorBlendToAlpha, bOnlyOneMip); break;
+			case EBlendType::BT_NONE: break;
+			default: check(false);
+			}
+		}
+		else if (!bDone)
+		{
+			// Apply blend without mask to RGB
+			switch (EBlendType(Args.blendType))
+			{
+			case EBlendType::BT_NORMAL_COMBINE: ImageNormalCombine(Result.get(), Result.get(), Blended.get(), bOnlyOneMip); check(!bUseBlendSourceFromBlendAlpha);  break;
+			case EBlendType::BT_SOFTLIGHT: BufferLayer<SoftLightChannel, true>(Result.get(), Result.get(), Blended.get(), bApplyColorBlendToAlpha, bOnlyOneMip, bUseBlendSourceFromBlendAlpha); break;
+			case EBlendType::BT_HARDLIGHT: BufferLayer<HardLightChannel, true>(Result.get(), Result.get(), Blended.get(), bApplyColorBlendToAlpha, bOnlyOneMip, bUseBlendSourceFromBlendAlpha); break;
+			case EBlendType::BT_BURN: BufferLayer<BurnChannel, true>(Result.get(), Result.get(), Blended.get(), bApplyColorBlendToAlpha, bOnlyOneMip, bUseBlendSourceFromBlendAlpha); break;
+			case EBlendType::BT_DODGE: BufferLayer<DodgeChannel, true>(Result.get(), Result.get(), Blended.get(), bApplyColorBlendToAlpha, bOnlyOneMip, bUseBlendSourceFromBlendAlpha); break;
+			case EBlendType::BT_SCREEN: BufferLayer<ScreenChannel, true>(Result.get(), Result.get(), Blended.get(), bApplyColorBlendToAlpha, bOnlyOneMip, bUseBlendSourceFromBlendAlpha); break;
+			case EBlendType::BT_OVERLAY: BufferLayer<OverlayChannel, true>(Result.get(), Result.get(), Blended.get(), bApplyColorBlendToAlpha, bOnlyOneMip, bUseBlendSourceFromBlendAlpha); break;
+			case EBlendType::BT_LIGHTEN: BufferLayer<LightenChannel, true>(Result.get(), Result.get(), Blended.get(), bApplyColorBlendToAlpha, bOnlyOneMip, bUseBlendSourceFromBlendAlpha); break;
+			case EBlendType::BT_MULTIPLY: BufferLayer<MultiplyChannel, true>(Result.get(), Result.get(), Blended.get(), bApplyColorBlendToAlpha, bOnlyOneMip, bUseBlendSourceFromBlendAlpha); break;
+			case EBlendType::BT_BLEND: BufferLayer<BlendChannel, true>(Result.get(), Result.get(), Blended.get(), bApplyColorBlendToAlpha, bOnlyOneMip, bUseBlendSourceFromBlendAlpha); break;
+			case EBlendType::BT_NONE: break;
+			default: check(false);
+			}
+		}
+
+		// Apply the separate blend operation for alpha
+		if (!bDone && !bApplyColorBlendToAlpha)
+		{
+			// Separate alpha operation ignores the mask.
+			switch (EBlendType(Args.blendTypeAlpha))
+			{
+			case EBlendType::BT_SOFTLIGHT: BufferLayerInPlace<SoftLightChannel, false, 1>(Result.get(), Blended.get(), bOnlyOneMip, 3, Args.BlendAlphaSourceChannel); break;
+			case EBlendType::BT_HARDLIGHT: BufferLayerInPlace<HardLightChannel, false, 1>(Result.get(), Blended.get(), bOnlyOneMip, 3, Args.BlendAlphaSourceChannel); break;
+			case EBlendType::BT_BURN: BufferLayerInPlace<BurnChannel, false, 1>(Result.get(), Blended.get(), bOnlyOneMip, 3, Args.BlendAlphaSourceChannel); break;
+			case EBlendType::BT_DODGE: BufferLayerInPlace<DodgeChannel, false, 1>(Result.get(), Blended.get(), bOnlyOneMip, 3, Args.BlendAlphaSourceChannel); break;
+			case EBlendType::BT_SCREEN: BufferLayerInPlace<ScreenChannel, false, 1>(Result.get(), Blended.get(), bOnlyOneMip, 3, Args.BlendAlphaSourceChannel); break;
+			case EBlendType::BT_OVERLAY: BufferLayerInPlace<OverlayChannel, false, 1>(Result.get(), Blended.get(), bOnlyOneMip, 3, Args.BlendAlphaSourceChannel); break;
+			case EBlendType::BT_LIGHTEN: BufferLayerInPlace<LightenChannel, false, 1>(Result.get(), Blended.get(), bOnlyOneMip, 3, Args.BlendAlphaSourceChannel); break;
+			case EBlendType::BT_MULTIPLY: BufferLayerInPlace<MultiplyChannel, false, 1>(Result.get(), Blended.get(), bOnlyOneMip, 3, Args.BlendAlphaSourceChannel); break;
+			case EBlendType::BT_BLEND: BufferLayerInPlace<BlendChannel, false, 1>(Result.get(), Blended.get(), bOnlyOneMip, 3, Args.BlendAlphaSourceChannel); break;
+			case EBlendType::BT_NONE: break;
+			default: check(false);
+			}
+		}
+
+		if (bOnlyOneMip)
+		{
+			MUTABLE_CPUPROFILER_SCOPE(ImageLayer_MipFix);
+			FMipmapGenerationSettings DummyMipSettings{};
+			ImageMipmapInPlace(ImageCompressionQuality, Result.get(), DummyMipSettings);
+		}
+
+		// Reset relevancy map.
+		Result->m_flags &= ~Image::EImageFlags::IF_HAS_RELEVANCY_MAP;
+	}
+
+
+	bool FImageLayerTask::Complete( CodeRunner* Runner )
+	{
+		// This runs in the Runner thread
+		Runner->Release(Blended);
+		Runner->Release(Mask);
+
+		// If no shortcut was taken
+		if (Result)
+		{
+			if (InitialFormat != Result->GetFormat())
+			{
+				Ptr<Image> Formatted = Runner->CreateImage(Result->GetSizeX(), Result->GetSizeY(), Result->GetLODCount(), InitialFormat, EInitializationType::NotInitialized);
+				bool bSuccess = false;
+
+				FImageOperator ImOp = MakeImageOperator(Runner);
+				ImOp.ImagePixelFormat(bSuccess, ImageCompressionQuality, Formatted.get(), Result.get());
+				check(bSuccess);
+
+				Runner->Release(Result);
+				Result = Formatted;
+			}
+
+			Runner->StoreImage(Op, Result);
+		}
+
+		bool bSuccess = true;
+		return bSuccess;
+	}
+
+
+	//---------------------------------------------------------------------------------------------
+	//---------------------------------------------------------------------------------------------
+	//---------------------------------------------------------------------------------------------
+	class FImageLayerColourTask : public CodeRunner::FIssuedTask
+	{
+	public:
+		FImageLayerColourTask(const FScheduledOp& InOp, const OP::ImageLayerColourArgs& InArgs)
+			: FIssuedTask(InOp), Args(InArgs)
+		{}
+
+		// FIssuedTask interface
+		virtual bool Prepare(CodeRunner*, bool& bOutFailed) override;
+		virtual void DoWork() override;
+		virtual bool Complete(CodeRunner*) override;
+
+	private:
+		int32 ImageCompressionQuality = 0;
+		OP::ImageLayerColourArgs Args;
+		FVector4f Color;
+		Ptr<const Image> Mask;
+		Ptr<Image> Result;
+		EImageFormat InitialFormat = EImageFormat::IF_NONE;
+	};
+
+
+	bool FImageLayerColourTask::Prepare(CodeRunner* Runner, bool& bOutFailed)
+	{
+		// This runs in the mutable Runner thread
+		MUTABLE_CPUPROFILER_SCOPE(FImageLayerColourTask_Prepare);
+		bOutFailed = false;
+
+		ImageCompressionQuality = Runner->m_pSettings->ImageCompressionQuality;
+
+		Ptr<const Image> Base = Runner->LoadImage({ Args.base, Op.ExecutionIndex, Op.ExecutionOptions });
+		check(!Base || Base->GetFormat() < EImageFormat::IF_COUNT);
+
+		Color = Runner->LoadColor({ Args.colour, Op.ExecutionIndex, 0 });
+		if (Args.mask)
+		{
+			Mask = Runner->LoadImage({ Args.mask, Op.ExecutionIndex, Op.ExecutionOptions });
+			check(!Mask || Mask->GetFormat() < EImageFormat::IF_COUNT);
+		}
+
+		// Shortcuts
+		if (!Base)
+		{
+			Runner->Release(Mask);
+			Runner->StoreImage(Op, Base);
+			return false;
+		}
+
+		bool bValid = Base->GetSizeX() > 0 && Base->GetSizeY() > 0;
+		if (!bValid)
+		{
+			Runner->Release(Mask);
+			Runner->StoreImage(Op, Base);
+			return false;
+		}
+
+		// Fix input data
+		InitialFormat = Base->GetFormat();
+		check(InitialFormat < EImageFormat::IF_COUNT);
+
+		if (Args.mask && Mask)
+		{
+			FImageOperator ImOp = MakeImageOperator(Runner);
+
+			if (Base->GetSize() != Mask->GetSize())
+			{
+				MUTABLE_CPUPROFILER_SCOPE(ImageResize_EmergencyFixSize);
+				Ptr<Image> Resized = Runner->CreateImage(Base->GetSizeX(), Base->GetSizeY(), 1, Mask->GetFormat(), EInitializationType::NotInitialized);
+				ImOp.ImageResizeLinear(Resized.get(), ImageCompressionQuality, Mask.get() );
+				Runner->Release(Mask);
+				Mask = Resized;
+			}
+
+			// emergy fix c36adf47-e40d-490f-b709-41142bafad78
+			if (Mask->GetLODCount() < Base->GetLODCount())
+			{
+				MUTABLE_CPUPROFILER_SCOPE(ImageResize_EmergencyFixMips);
+				int32 StartLevel = Mask->GetLODCount() - 1;
+				int32 LevelCount = Base->GetLODCount();
+
+				Ptr<Image> MaskFix = Runner->CloneOrTakeOver(Mask);
+				MaskFix->DataStorage.SetNumLODs(LevelCount);
+
+				FMipmapGenerationSettings Settings{};
+				ImOp.ImageMipmap(ImageCompressionQuality, MaskFix.get(), MaskFix.get(), StartLevel, LevelCount, Settings);
+
+				Mask = MaskFix;
+			}
+		}
+
+		// Create destination data
+		Result = Runner->CloneOrTakeOver(Base);
+		return true;
+	}
+
+
+	void FImageLayerColourTask::DoWork()
+	{
+		// This runs in a worker thread
+		MUTABLE_CPUPROFILER_SCOPE(FImageLayerColourTask);
+
+		bool bOnlyOneMip = false;
+
+		// Does it apply to colour?
+		if (EBlendType(Args.blendType) != EBlendType::BT_NONE)
+		{
+			// TODO: It could be done "in-place"
+			if (Args.mask && Mask)
+			{			
+				// Not implemented yet
+				check(Args.flags==0);
+
+				// \todo: precalculated tables for softlight
+				switch (EBlendType(Args.blendType))
+				{
+				case EBlendType::BT_NORMAL_COMBINE: ImageNormalCombine(Result.get(), Result.get(), Mask.get(), Color); break;
+				case EBlendType::BT_SOFTLIGHT: BufferLayerColour<SoftLightChannelMasked, SoftLightChannel>(Result.get(), Result.get(), Mask.get(), Color); break;
+				case EBlendType::BT_HARDLIGHT: BufferLayerColour<HardLightChannelMasked, HardLightChannel>(Result.get(), Result.get(), Mask.get(), Color); break;
+				case EBlendType::BT_BURN: BufferLayerColour<BurnChannelMasked, BurnChannel>(Result.get(), Result.get(), Mask.get(), Color); break;
+				case EBlendType::BT_DODGE: BufferLayerColour<DodgeChannelMasked, DodgeChannel>(Result.get(), Result.get(), Mask.get(), Color); break;
+				case EBlendType::BT_SCREEN: BufferLayerColour<ScreenChannelMasked, ScreenChannel>(Result.get(), Result.get(), Mask.get(), Color); break;
+				case EBlendType::BT_OVERLAY: BufferLayerColour<OverlayChannelMasked, OverlayChannel>(Result.get(), Result.get(), Mask.get(), Color); break;
+				case EBlendType::BT_LIGHTEN: BufferLayerColour<LightenChannelMasked, LightenChannel>(Result.get(), Result.get(), Mask.get(), Color); break;
+				case EBlendType::BT_MULTIPLY: BufferLayerColour<MultiplyChannelMasked, MultiplyChannel>(Result.get(), Result.get(), Mask.get(), Color); break;
+				case EBlendType::BT_BLEND: BufferLayerColour<BlendChannelMasked, BlendChannel>(Result.get(), Result.get(), Mask.get(), Color); break;
+				default: check(false);
+				}
+
+			}
+			else
+			{
+				if (Args.flags & OP::ImageLayerArgs::FLAGS::F_BASE_RGB_FROM_ALPHA)
+				{
+					switch (EBlendType(Args.blendType))
+					{
+					case EBlendType::BT_NORMAL_COMBINE: check(false); break;
+					case EBlendType::BT_SOFTLIGHT: BufferLayerColourFromAlpha<SoftLightChannel>(Result.get(), Result.get(), Color); break;
+					case EBlendType::BT_HARDLIGHT: BufferLayerColourFromAlpha<HardLightChannel>(Result.get(), Result.get(), Color); break;
+					case EBlendType::BT_BURN: BufferLayerColourFromAlpha<BurnChannel>(Result.get(), Result.get(), Color); break;
+					case EBlendType::BT_DODGE: BufferLayerColourFromAlpha<DodgeChannel>(Result.get(), Result.get(), Color); break;
+					case EBlendType::BT_SCREEN: BufferLayerColourFromAlpha<ScreenChannel>(Result.get(), Result.get(), Color); break;
+					case EBlendType::BT_OVERLAY: BufferLayerColourFromAlpha<OverlayChannel>(Result.get(), Result.get(), Color); break;
+					case EBlendType::BT_LIGHTEN: BufferLayerColourFromAlpha<LightenChannel>(Result.get(), Result.get(), Color); break;
+					case EBlendType::BT_MULTIPLY: BufferLayerColourFromAlpha<MultiplyChannel>(Result.get(), Result.get(), Color); break;
+					case EBlendType::BT_BLEND: check(false); break;
+					default: check(false);
+					}
+				}
+				else
+				{
+					switch (EBlendType(Args.blendType))
+					{
+					case EBlendType::BT_NORMAL_COMBINE: ImageNormalCombine(Result.get(), Result.get(), Color); break;
+					case EBlendType::BT_SOFTLIGHT: BufferLayerColour<SoftLightChannel>(Result.get(), Result.get(), Color); break;
+					case EBlendType::BT_HARDLIGHT: BufferLayerColour<HardLightChannel>(Result.get(), Result.get(), Color); break;
+					case EBlendType::BT_BURN: BufferLayerColour<BurnChannel>(Result.get(), Result.get(), Color); break;
+					case EBlendType::BT_DODGE: BufferLayerColour<DodgeChannel>(Result.get(), Result.get(), Color); break;
+					case EBlendType::BT_SCREEN: BufferLayerColour<ScreenChannel>(Result.get(), Result.get(), Color); break;
+					case EBlendType::BT_OVERLAY: BufferLayerColour<OverlayChannel>(Result.get(), Result.get(), Color); break;
+					case EBlendType::BT_LIGHTEN: BufferLayerColour<LightenChannel>(Result.get(), Result.get(), Color); break;
+					case EBlendType::BT_MULTIPLY: BufferLayerColour<MultiplyChannel>(Result.get(), Result.get(), Color); break;
+					case EBlendType::BT_BLEND: 
+					{
+						// In this case we know it is already an uncompressed image, and we won't need additional allocations;
+						FImageOperator ImOp = FImageOperator::GetDefault(FImageOperator::FImagePixelFormatFunc());
+						ImOp.FillColor( Result.get(), Color); 
+						break;
+					}
+					default: check(false);
+					}
+				}
+			}
+		}
+
+		// Does it apply to alpha?
+		if (EBlendType(Args.blendTypeAlpha) != EBlendType::BT_NONE)
+		{
+			if (Args.mask && Mask)
+			{
+				// \todo: precalculated tables for softlight
+				switch (EBlendType(Args.blendTypeAlpha))
+				{
+				case EBlendType::BT_NORMAL_COMBINE: check(false); break;
+				case EBlendType::BT_SOFTLIGHT: BufferLayerColourInPlace<SoftLightChannelMasked, SoftLightChannel, 1>(Result.get(), Mask.get(), Color, bOnlyOneMip, 3, Args.BlendAlphaSourceChannel); break;
+				case EBlendType::BT_HARDLIGHT: BufferLayerColourInPlace<HardLightChannelMasked, HardLightChannel, 1>(Result.get(), Mask.get(), Color, bOnlyOneMip, 3, Args.BlendAlphaSourceChannel); break;
+				case EBlendType::BT_BURN: BufferLayerColourInPlace<BurnChannelMasked, BurnChannel, 1>(Result.get(), Mask.get(), Color, bOnlyOneMip, 3, Args.BlendAlphaSourceChannel); break;
+				case EBlendType::BT_DODGE: BufferLayerColourInPlace<DodgeChannelMasked, DodgeChannel, 1>(Result.get(), Mask.get(), Color, bOnlyOneMip, 3, Args.BlendAlphaSourceChannel); break;
+				case EBlendType::BT_SCREEN: BufferLayerColourInPlace<ScreenChannelMasked, ScreenChannel, 1>(Result.get(), Mask.get(), Color, bOnlyOneMip, 3, Args.BlendAlphaSourceChannel); break;
+				case EBlendType::BT_OVERLAY: BufferLayerColourInPlace<OverlayChannelMasked, OverlayChannel, 1>(Result.get(), Mask.get(), Color, bOnlyOneMip, 3, Args.BlendAlphaSourceChannel); break;
+				case EBlendType::BT_LIGHTEN: BufferLayerColourInPlace<LightenChannelMasked, LightenChannel, 1>(Result.get(), Mask.get(), Color, bOnlyOneMip, 3, Args.BlendAlphaSourceChannel); break;
+				case EBlendType::BT_MULTIPLY: BufferLayerColourInPlace<MultiplyChannelMasked, MultiplyChannel, 1>(Result.get(), Mask.get(), Color, bOnlyOneMip, 3, Args.BlendAlphaSourceChannel); break;
+				case EBlendType::BT_BLEND: BufferLayerColourInPlace<BlendChannelMasked, BlendChannel, 1>(Result.get(), Mask.get(), Color, bOnlyOneMip, 3, Args.BlendAlphaSourceChannel); break;
+				default: check(false);
+				}
+
+			}
+			else
+			{
+				switch (EBlendType(Args.blendTypeAlpha))
+				{
+				case EBlendType::BT_NORMAL_COMBINE: check(false); break;
+				case EBlendType::BT_SOFTLIGHT: BufferLayerColourInPlace<SoftLightChannel, 1>(Result.get(), Color, bOnlyOneMip, 3, Args.BlendAlphaSourceChannel); break;
+				case EBlendType::BT_HARDLIGHT: BufferLayerColourInPlace<HardLightChannel, 1>(Result.get(), Color, bOnlyOneMip, 3, Args.BlendAlphaSourceChannel); break;
+				case EBlendType::BT_BURN: BufferLayerColourInPlace<BurnChannel, 1>(Result.get(), Color, bOnlyOneMip, 3, Args.BlendAlphaSourceChannel); break;
+				case EBlendType::BT_DODGE: BufferLayerColourInPlace<DodgeChannel, 1>(Result.get(), Color, bOnlyOneMip, 3, Args.BlendAlphaSourceChannel); break;
+				case EBlendType::BT_SCREEN: BufferLayerColourInPlace<ScreenChannel, 1>(Result.get(), Color, bOnlyOneMip, 3, Args.BlendAlphaSourceChannel); break;
+				case EBlendType::BT_OVERLAY: BufferLayerColourInPlace<OverlayChannel, 1>(Result.get(), Color, bOnlyOneMip, 3, Args.BlendAlphaSourceChannel); break;
+				case EBlendType::BT_LIGHTEN: BufferLayerColourInPlace<LightenChannel, 1>(Result.get(), Color, bOnlyOneMip, 3, Args.BlendAlphaSourceChannel); break;
+				case EBlendType::BT_MULTIPLY: BufferLayerColourInPlace<MultiplyChannel, 1>(Result.get(), Color, bOnlyOneMip, 3, Args.BlendAlphaSourceChannel); break;
+				case EBlendType::BT_BLEND: BufferLayerColourInPlace<BlendChannel, 1>(Result.get(), Color, bOnlyOneMip, 3, Args.BlendAlphaSourceChannel); break;
+				default: check(false);
+				}
+			}
+		}
+
+		// Reset relevancy map.
+		Result->m_flags &= ~Image::EImageFlags::IF_HAS_RELEVANCY_MAP;
+	}
+
+
+	bool FImageLayerColourTask::Complete(CodeRunner* Runner)
+	{
+		// This runs in the Runner thread
+		Runner->Release(Mask);
+
+		// If no shortcut was taken
+		if (Result)
+		{
+			Runner->StoreImage(Op, Result);
+		}
+
+		bool bSuccess = true;
+		return bSuccess;
+	}
+
+
+	//---------------------------------------------------------------------------------------------
+	//---------------------------------------------------------------------------------------------
+	//---------------------------------------------------------------------------------------------
+	class FImagePixelFormatTask : public CodeRunner::FIssuedTask
+	{
+	public:
+		FImagePixelFormatTask(const FScheduledOp& InOp, const OP::ImagePixelFormatArgs& InArgs)
+			: FIssuedTask(InOp), Args(InArgs)
+		{}
+
+		// FIssuedTask interface
+		virtual bool Prepare(CodeRunner*, bool& bOutFailed) override;
+		virtual void DoWork() override;
+		virtual bool Complete(CodeRunner*) override;
+
+	private:
+		int ImageCompressionQuality = 0;
+		OP::ImagePixelFormatArgs Args;
+		EImageFormat TargetFormat = EImageFormat::IF_NONE;
+		Ptr<const Image> Base;
+		Ptr<Image> Result;
+		FImageOperator::FImagePixelFormatFunc ImagePixelFormatFunc;
+	};
+
+
+	bool FImagePixelFormatTask::Prepare(CodeRunner* Runner, bool& bOutFailed)
+	{
+		// This runs in the mutable Runner thread
+		MUTABLE_CPUPROFILER_SCOPE(FImageLayerPixelFormatTask_Prepare);
+		bOutFailed = false;
+
+		ImageCompressionQuality = Runner->m_pSettings->ImageCompressionQuality;
+
+		Base = Runner->LoadImage({ Args.source, Op.ExecutionIndex, Op.ExecutionOptions });
+		check(!Base || Base->GetFormat() < EImageFormat::IF_COUNT);		
+
+		// Shortcuts
+		if (!Base)
+		{
+			Runner->StoreImage(Op, Base);
+			return false;
+		}
+
+		bool bValid = Base->GetSizeX() > 0 && Base->GetSizeY() > 0;
+		if (!bValid)
+		{
+			Runner->StoreImage(Op, Base);
+			return false;
+		}
+
+		TargetFormat = Args.format;
+		if (Args.formatIfAlpha != EImageFormat::IF_NONE
+			&&
+			GetImageFormatData(Base->GetFormat()).Channels > 3)
+		{
+			TargetFormat = Args.formatIfAlpha;
+		}
+
+		if (TargetFormat == EImageFormat::IF_NONE || TargetFormat == Base->GetFormat())
+		{
+			Runner->StoreImage(Op, Base);
+			return false;
+		}
+
+		ImagePixelFormatFunc = Runner->m_pSystem->ImagePixelFormatOverride;
+
+		// Create destination data
+		Result = Runner->CreateImage(Base->GetSizeX(), Base->GetSizeY(), Base->GetLODCount(), TargetFormat, EInitializationType::NotInitialized);
+		return true;
+	}
+
+
+	void FImagePixelFormatTask::DoWork()
+	{
+		// This runs in a worker thread
+		MUTABLE_CPUPROFILER_SCOPE(FImagePixelFormatTask);
+		
+		bool bSuccess = false;
+		FImageOperator ImOp = FImageOperator::GetDefault(ImagePixelFormatFunc);
+		ImOp.ImagePixelFormat(bSuccess, ImageCompressionQuality, Result.get(), Base.get(), -1);
+		check(bSuccess);
+	}
+
+
+	bool FImagePixelFormatTask::Complete(CodeRunner* Runner)
+	{
+		// This runs in the Runner thread
+		Runner->Release(Base);
+
+		// If no shortcut was taken
+		if (Result)
+		{
+			Runner->StoreImage(Op, Result);
+		}
+
+		bool bSuccess = true;
+		return bSuccess;
+	}
+
+	class FImageMipmapTask : public CodeRunner::FIssuedTask
+	{
+	public:
+		FImageMipmapTask(const FScheduledOp& InOp, const OP::ImageMipmapArgs& InArgs)
+			:FIssuedTask(InOp), Args(InArgs)
+		{}
+
+		// FIssuedTask interface
+		virtual bool Prepare(CodeRunner*, bool& bOutFailed) override;
+		virtual void DoWork() override;
+		virtual bool Complete(CodeRunner*) override;
+
+	private:
+		int32 ImageCompressionQuality = 0;
+		int32 StartLevel = -1;
+		OP::ImageMipmapArgs Args;
+		Ptr<const Image> Base;
+		Ptr<Image> Result;
+		FImageOperator::FScratchImageMipmap Scratch;
+		FImageOperator::FImagePixelFormatFunc ImagePixelFormatFunc;
+	};
+
+
+	bool FImageMipmapTask::Prepare(CodeRunner* Runner, bool& bOutFailed)
+	{
+		// This runs in the mutable Runner thread
+		MUTABLE_CPUPROFILER_SCOPE(FImageMipmapTask_Prepare);
+		bOutFailed = false;
+
+		ImageCompressionQuality = Runner->m_pSettings->ImageCompressionQuality;
+
+		Base = Runner->LoadImage({ Args.source, Op.ExecutionIndex, Op.ExecutionOptions });
+		check(!Base || Base->GetFormat() < EImageFormat::IF_COUNT);
+
+		// Shortcuts
+		if (!Base)
+		{
+			Runner->StoreImage(Op, Base);
+			return false;
+		}
+
+		bool bValid = Base->GetSizeX() > 0 && Base->GetSizeY() > 0;
+		if (!bValid)
+		{
+			Runner->StoreImage(Op, Base);
+			return false;
+		}
+
+		int32 LevelCount = Args.levels;
+		int32 MaxLevelCount = Image::GetMipmapCount(Base->GetSizeX(), Base->GetSizeY());
+		if (LevelCount == 0)
+		{
+			LevelCount = MaxLevelCount;
+		}
+		else if (LevelCount > MaxLevelCount)
+		{
+			// If code generation is smart enough, this should never happen.
+			// \todo But apparently it does, sometimes.
+			LevelCount = MaxLevelCount;
+		}
+
+		// At least keep the levels we already have.
+		LevelCount = FMath::Max(Base->GetLODCount(), LevelCount);
+		
+		if (LevelCount == Base->GetLODCount())
+		{
+			Runner->StoreImage(Op, Base);
+			return false;
+		}
+
+		StartLevel = Base->GetLODCount() - 1;
+		
+		Result = Runner->CloneOrTakeOver(Base);
+		Base = nullptr;
+
+		Result->DataStorage.SetNumLODs(LevelCount);
+
+		FImageOperator ImOp = MakeImageOperator(Runner);
+		ImOp.ImageMipmap_PrepareScratch(Result.get(), StartLevel, LevelCount, Scratch);
+
+		ImagePixelFormatFunc = Runner->m_pSystem->ImagePixelFormatOverride;
+
+		return true;
+	}
+
+
+	void FImageMipmapTask::DoWork()
+	{
+		// This runs in a worker thread
+		MUTABLE_CPUPROFILER_SCOPE(FImageMipmapTask);
+
+		check(StartLevel >= 0);
+
+		FMipmapGenerationSettings Settings{Args.FilterType, Args.AddressMode};
+		FImageOperator ImOp = FImageOperator::GetDefault(ImagePixelFormatFunc);
+		ImOp.ImageMipmap(Scratch, ImageCompressionQuality, Result.get(), Result.get(), StartLevel, Result->GetLODCount(), Settings);
+	}
+
+
+	bool FImageMipmapTask::Complete(CodeRunner* Runner)
+	{
+		FImageOperator ImOp = MakeImageOperator(Runner);
+		ImOp.ImageMipmap_ReleaseScratch(Scratch);
+
+		// This runs in the Runner thread
+		if (Base)
+		{
+			Runner->Release(Base);
+		}
+
+		// If no shortcut was taken
+		if (Result)
+		{
+			Runner->StoreImage(Op, Result);
+		}
+
+		bool bSuccess = true;
+		return bSuccess;
+	}
+
+
+	//---------------------------------------------------------------------------------------------
+	//---------------------------------------------------------------------------------------------
+	//---------------------------------------------------------------------------------------------
+	class FImageSwizzleTask : public CodeRunner::FIssuedTask
+	{
+	public:
+		FImageSwizzleTask(const FScheduledOp& InOp, const OP::ImageSwizzleArgs& InArgs)
+			: FIssuedTask(InOp), Args(InArgs)
+		{}
+
+		// FIssuedTask interface
+		virtual bool Prepare(CodeRunner*, bool& bOutFailed);
+		virtual void DoWork() override;
+		virtual bool Complete(CodeRunner*) override;
+
+	private:
+		OP::ImageSwizzleArgs Args;
+		Ptr<const Image> Sources[MUTABLE_OP_MAX_SWIZZLE_CHANNELS];
+		Ptr<Image> Result;
+	};
+
+
+	bool FImageSwizzleTask::Prepare(CodeRunner* Runner, bool& bOutFailed)
+	{
+		// This runs in the mutable Runner thread
+		MUTABLE_CPUPROFILER_SCOPE(FImageSwizzleTask_Prepare);
+		bOutFailed = false;
+
+		int32 FirstValidSourceIndex = -1;
+		for (int32 SourceIndex = 0; SourceIndex < MUTABLE_OP_MAX_SWIZZLE_CHANNELS; ++SourceIndex)
+		{
+			if (Args.sources[SourceIndex])
+			{
+				FirstValidSourceIndex = SourceIndex;
+				break;
+			}
+		}
+
+		for (int32 SourceIndex = 0; SourceIndex < MUTABLE_OP_MAX_SWIZZLE_CHANNELS; ++SourceIndex)
+		{
+			if (Args.sources[SourceIndex])
+			{
+				Sources[SourceIndex] = Runner->LoadImage({ Args.sources[SourceIndex], Op.ExecutionIndex, Op.ExecutionOptions });
+			}
+		}
+
+		// Shortcuts
+		if (FirstValidSourceIndex < 0 || !Sources[FirstValidSourceIndex])
+		{
+			for (int32 SourceIndex = 0; SourceIndex < MUTABLE_OP_MAX_SWIZZLE_CHANNELS; ++SourceIndex)
+			{
+				Runner->Release(Sources[SourceIndex]);
+			}
+			Runner->StoreImage(Op, nullptr);
+
+			return false;
+		}
+
+		// Create destination data
+		EImageFormat format = (EImageFormat)Args.format;
+
+		FImageOperator ImOp = MakeImageOperator(Runner);
+
+		int32 ResultLODs = Sources[FirstValidSourceIndex]->GetLODCount();
+
+		// Be defensive: ensure formats are uncompressed
+		for (int32 SourceIndex = 0; SourceIndex < MUTABLE_OP_MAX_SWIZZLE_CHANNELS; ++SourceIndex)
+		{
+			if (Sources[SourceIndex] && Sources[SourceIndex]->GetFormat() != GetUncompressedFormat(Sources[SourceIndex]->GetFormat()))
+			{
+				MUTABLE_CPUPROFILER_SCOPE(ImageFormat_ForSwizzle);
+
+				EImageFormat UncompressedFormat = GetUncompressedFormat(Sources[SourceIndex]->GetFormat());
+				Ptr<Image> Formatted = Runner->CreateImage(Sources[SourceIndex]->GetSizeX(), Sources[SourceIndex]->GetSizeY(), 1, UncompressedFormat, EInitializationType::NotInitialized);
+				bool bSuccess = false;
+				int32 ImageCompressionQuality = 4; // TODO
+				ImOp.ImagePixelFormat(bSuccess, ImageCompressionQuality, Formatted.get(), Sources[SourceIndex].get());
+				check(bSuccess); // Decompression cannot fail
+				Runner->Release(Sources[SourceIndex]);
+				Sources[SourceIndex] = Formatted;
+				ResultLODs = 1;
+			}
+		}
+
+		const FImageSize ResultSize = Sources[FirstValidSourceIndex]->GetSize();
+
+		// Be defensive: ensure image sizes match.
+		for (int32 SourceIndex = FirstValidSourceIndex + 1; SourceIndex < MUTABLE_OP_MAX_SWIZZLE_CHANNELS; ++SourceIndex)
+		{
+			if (Sources[SourceIndex] && ResultSize != Sources[SourceIndex]->GetSize())
+			{
+				MUTABLE_CPUPROFILER_SCOPE(ImageResize_ForSwizzle);
+				Ptr<Image> Resized = Runner->CreateImage(ResultSize.X, ResultSize.Y, 1, Sources[SourceIndex]->GetFormat(), EInitializationType::NotInitialized);
+				ImOp.ImageResizeLinear(Resized.get(), 0, Sources[SourceIndex].get());
+				Runner->Release(Sources[SourceIndex]);
+				Sources[SourceIndex] = Resized;
+				ResultLODs = 1;
+			}
+		}
+
+		// If any source has only 1 LOD, then the result has to have 1 LOD and the rest be regenerated later on
+		for (int32 SourceIndex = 0; SourceIndex < MUTABLE_OP_MAX_SWIZZLE_CHANNELS; ++SourceIndex)
+		{
+			if (Sources[SourceIndex] && Sources[SourceIndex]->GetLODCount() == 1)
+			{
+				ResultLODs = 1;
+			}
+		}
+
+		Result = Runner->CreateImage(ResultSize.X, ResultSize.Y, ResultLODs, Args.format, EInitializationType::Black);
+		return true;
+	}
+
+
+	void FImageSwizzleTask::DoWork()
+	{
+		// This runs in a worker thread
+		MUTABLE_CPUPROFILER_SCOPE(FImageSwizzleTask);
+
+		ImageSwizzle(Result.get(), Sources, Args.sourceChannels);
+	}
+
+
+	bool FImageSwizzleTask::Complete(CodeRunner* Runner)
+	{
+		// This runs in the Runner thread
+		for (int32 SourceIndex = 0; SourceIndex < MUTABLE_OP_MAX_SWIZZLE_CHANNELS; ++SourceIndex)
+		{
+			Runner->Release(Sources[SourceIndex]);
+		}
+
+		// \TODO: If Result LODs differ from Source[0]'s, rebuild mips?
+
+		// If no shortcut was taken
+		if (Result)
+		{
+			Runner->StoreImage(Op, Result);
+		}
+
+		bool bSuccess = true;
+		return bSuccess;
+	}
+
+
+	//---------------------------------------------------------------------------------------------
+	//---------------------------------------------------------------------------------------------
+	//---------------------------------------------------------------------------------------------
+	class FImageSaturateTask : public CodeRunner::FIssuedTask
+	{
+	public:
+		FImageSaturateTask(const FScheduledOp&, const OP::ImageSaturateArgs&);
+
+		// FIssuedTask interface
+		virtual bool Prepare(CodeRunner*, bool& bOutFailed) override;
+		virtual void DoWork() override;
+		virtual bool Complete(CodeRunner*) override;
+
+	private:
+		const OP::ImageSaturateArgs Args;
+		Ptr<Image> Result;
+		float Factor;
+	};
+
+
+	//---------------------------------------------------------------------------------------------
+	FImageSaturateTask::FImageSaturateTask(const FScheduledOp& InOp, const OP::ImageSaturateArgs& InArgs)
+		: FIssuedTask(InOp)
+		, Args(InArgs)
+	{
+	}
+
+	//---------------------------------------------------------------------------------------------
+	bool FImageSaturateTask::Prepare(CodeRunner* Runner, bool& bOutFailed)
+	{
+		// This runs in the mutable Runner thread
+		MUTABLE_CPUPROFILER_SCOPE(FImageSaturateTask_Prepare);
+		bOutFailed = false;
+
+		Ptr<const Image> Source = Runner->LoadImage(FCacheAddress(Args.base, Op));
+		Factor = Runner->LoadScalar(FScheduledOp::FromOpAndOptions(Args.factor, Op, 0));
+
+		if (!Source)
+		{
+			Runner->StoreImage(Op, Source);
+			return false;
+		}
+
+		const bool bOptimizeUnchanged = FMath::IsNearlyEqual(Factor, 1.0f);
+		if (bOptimizeUnchanged)
+		{
+			Runner->StoreImage(Op, Source);
+			return false;
+		}
+
+		Result = Runner->CloneOrTakeOver(Source);
+		return true;
+	}
+
+	//---------------------------------------------------------------------------------------------
+	void FImageSaturateTask::DoWork()
+	{
+		// This runs on a random worker thread
+		MUTABLE_CPUPROFILER_SCOPE(FImageSaturateTask);
+
+		constexpr bool bUseVectorIntrinsics = true;
+		ImageSaturate<bUseVectorIntrinsics>(Result.get(), Factor);
+	}
+
+	//---------------------------------------------------------------------------------------------
+	bool FImageSaturateTask::Complete(CodeRunner* Runner)
+	{
+		// This runs in the mutable Runner thread
+
+		// If we didn't take a shortcut
+		if (Result)
+		{
+			Runner->StoreImage(Op, Result);
+		}
+
+		bool bSuccess = true;
+		return bSuccess;
+	}
+
+
+	//---------------------------------------------------------------------------------------------
+	//---------------------------------------------------------------------------------------------
+	//---------------------------------------------------------------------------------------------
+	class FImageResizeTask : public CodeRunner::FIssuedTask
+	{
+	public:
+		FImageResizeTask(const FScheduledOp& InOp, const OP::ImageResizeArgs& InArgs)
+			: FIssuedTask(InOp), Args(InArgs)
+		{}
+
+		// FIssuedTask interface
+		virtual bool Prepare(CodeRunner* Runner, bool& bFailed) override;
+		virtual void DoWork() override;
+		virtual bool Complete(CodeRunner*) override;
+
+	private:
+		int32 ImageCompressionQuality = 0;
+		OP::ImageResizeArgs Args;
+		Ptr<const Image> Base;
+		Ptr<Image> Result;
+		FImageOperator::FImagePixelFormatFunc ImagePixelFormatFunc;
+	};
+
+
+	bool FImageResizeTask::Prepare(CodeRunner* Runner, bool& bOutFailed)
+	{
+		// This runs in the mutable Runner thread
+		MUTABLE_CPUPROFILER_SCOPE(FImageResizeTask_Prepare);
+		bOutFailed = false;
+
+		ImageCompressionQuality = Runner->m_pSettings->ImageCompressionQuality;
+
+		FImageSize destSize = FImageSize(Args.size[0], Args.size[1]);
+
+		// Apply the mips-to-skip to the dest size
+		int32 MipsToSkip = Op.ExecutionOptions;
+		destSize[0] = FMath::Max(destSize[0] >> MipsToSkip, 1);
+		destSize[1] = FMath::Max(destSize[1] >> MipsToSkip, 1);
+
+		Base = Runner->LoadImage(FCacheAddress(Args.source, Op));
+		if (!Base 
+			|| 
+			( Base->GetSizeX()==destSize[0] && Base->GetSizeY()==destSize[1] )
+			)
+		{
+			Runner->StoreImage(Op, Base);
+			return false;
+		}
+
+		int32 Lods = 1;
+
+		// If the source image had mips, generate them as well for the resized image.
+		// This shouldn't happen often since it should be usually optimised  during model compilation. 
+		// The mipmap generation below is not very precise with the number of mips that are needed and 
+		// will probably generate too many
+		bool bSourceHasMips = Base->GetLODCount() > 1;
+		if (bSourceHasMips)
+		{
+			Lods = Image::GetMipmapCount(destSize[0], destSize[1]);
+		}
+
+		if (Base->IsReference())
+		{
+			// We are trying to resize an external reference. This shouldn't happen, but be deffensive.
+			Runner->StoreImage(Op, Base);
+			return false;
+		}
+
+		Result = Runner->CreateImage( destSize[0], destSize[1], Lods, Base->GetFormat(), EInitializationType::NotInitialized );
+
+		ImagePixelFormatFunc = Runner->m_pSystem->ImagePixelFormatOverride;
+
+		return true;
+	}
+
+
+	//---------------------------------------------------------------------------------------------
+	void FImageResizeTask::DoWork()
+	{
+		// This runs on a random worker thread
+		MUTABLE_CPUPROFILER_SCOPE(FImageResizeTask);
+
+		FImageSize destSize = FImageSize( Args.size[0], Args.size[1]);
+
+		// Apply the mips-to-skip to the dest size
+		int32 MipsToSkip = Op.ExecutionOptions;
+		destSize[0] = FMath::Max(destSize[0] >> MipsToSkip, 1);
+		destSize[1] = FMath::Max(destSize[1] >> MipsToSkip, 1);
+
+		// Warning: This will actually allocate temp memory that may exceed the budget.
+		// \TODO: Fix it.
+		FImageOperator ImOp = FImageOperator::GetDefault(ImagePixelFormatFunc);
+		ImOp.ImageResizeLinear( Result.get(), ImageCompressionQuality, Base.get());
+
+		int32 LodCount = Result->GetLODCount();
+		if (LodCount>1)
+		{
+			FMipmapGenerationSettings mipSettings{};
+			ImageMipmapInPlace( ImageCompressionQuality, Result.get(), mipSettings );
+		}
+	}
+
+
+	//---------------------------------------------------------------------------------------------
+	bool FImageResizeTask::Complete(CodeRunner* Runner)
+	{
+		// This runs in the Runner thread
+		Runner->Release(Base);
+
+		// If didn't take a shortcut and set it already
+		if (Result)
+		{
+			Runner->StoreImage(Op, Result);
+		}
+
+		bool bSuccess = true;
+		return bSuccess;
+	}
+
+
+	//---------------------------------------------------------------------------------------------
+	//---------------------------------------------------------------------------------------------
+	//---------------------------------------------------------------------------------------------
+	class FImageResizeRelTask : public CodeRunner::FIssuedTask
+	{
+	public:
+		FImageResizeRelTask(const FScheduledOp& InOp, const OP::ImageResizeRelArgs& InArgs)
+			: FIssuedTask(InOp), Args(InArgs)
+		{}
+
+		// FIssuedTask interface
+		virtual bool Prepare(CodeRunner* Runner, bool& bFailed) override;
+		virtual void DoWork() override;
+		virtual bool Complete(CodeRunner*) override;
+
+	private:
+		OP::ImageResizeRelArgs Args;
+		int32 ImageCompressionQuality=0;
+		Ptr<const Image> Base;
+		Ptr<Image> Result;
+		FImageSize DestSize;
+		FImageOperator::FImagePixelFormatFunc ImagePixelFormatFunc;
+	};
+
+
+	bool FImageResizeRelTask::Prepare(CodeRunner* Runner, bool& bOutFailed)
+	{
+		// This runs in the mutable Runner thread
+		MUTABLE_CPUPROFILER_SCOPE(FImageResizeRelTask_Prepare);
+		bOutFailed = false;
+
+		ImageCompressionQuality = Runner->m_pSettings->ImageCompressionQuality;
+
+		Base = Runner->LoadImage(FCacheAddress(Args.source, Op));
+		if (!Base)
+		{
+			Runner->StoreImage(Op, Base);
+			return false;
+		}
+
+		DestSize = FImageSize(
+			uint16(FMath::Max(1.0, Base->GetSizeX() * Args.factor[0] + 0.5f)),
+			uint16(FMath::Max(1.0, Base->GetSizeY() * Args.factor[1] + 0.5f)));
+
+		if (Base->GetSizeX() == DestSize[0] && Base->GetSizeY() == DestSize[1])
+		{
+			Runner->StoreImage(Op, Base);
+			return false;
+		}
+
+		int32 Lods = 1;
+
+		// If the source image had mips, generate them as well for the resized image.
+		// This shouldn't happen often since it should be usually optimised  during model compilation. 
+		// The mipmap generation below is not very precise with the number of mips that are needed and 
+		// will probably generate too many
+		bool bSourceHasMips = Base->GetLODCount() > 1;
+		if (bSourceHasMips)
+		{
+			Lods = Image::GetMipmapCount(DestSize[0], DestSize[1]);
+		}
+
+		Result = Runner->CreateImage(DestSize[0], DestSize[1], Lods, Base->GetFormat(), EInitializationType::NotInitialized);
+
+		ImagePixelFormatFunc = Runner->m_pSystem->ImagePixelFormatOverride;
+
+		return true;
+	}
+
+
+	void FImageResizeRelTask::DoWork()
+	{
+		// This runs on a random worker thread
+		MUTABLE_CPUPROFILER_SCOPE(FImageResizeRelTask);
+
+		// \TODO: Track allocs
+		FImageOperator ImOp = FImageOperator::GetDefault(ImagePixelFormatFunc);
+		ImOp.ImageResizeLinear(Result.get(), ImageCompressionQuality, Base.get());
+
+		int32 LodCount = Result->GetLODCount();
+		if (LodCount > 1)
+		{
+			FMipmapGenerationSettings mipSettings{};
+			ImageMipmapInPlace(ImageCompressionQuality, Result.get(), mipSettings);
+		}
+	}
+
+
+	bool FImageResizeRelTask::Complete(CodeRunner* Runner)
+	{
+		// This runs in the Runner thread
+		Runner->Release(Base);
+
+		// If didn't take a shortcut and set it already
+		if (Result)
+		{
+			Runner->StoreImage(Op, Result);
+		}
+
+		bool bSuccess = true;
+		return bSuccess;
+	}
+
+
+	//---------------------------------------------------------------------------------------------
+	//---------------------------------------------------------------------------------------------
+	//---------------------------------------------------------------------------------------------
+	class FImageInvertTask : public CodeRunner::FIssuedTask
+	{
+	public:
+		FImageInvertTask(const FScheduledOp& InOp, const OP::ImageInvertArgs& InArgs)
+			: FIssuedTask(InOp), Args(InArgs)
+		{}
+
+		// FIssuedTask interface
+		virtual bool Prepare(CodeRunner*, bool& bOutFailed) override;
+		virtual void DoWork() override;
+		virtual bool Complete(CodeRunner* Runner) override;
+
+	private:
+		Ptr<Image> Result;
+		OP::ImageInvertArgs Args;
+	};
+
+
+	bool FImageInvertTask::Prepare(CodeRunner* Runner, bool& bOutFailed)
+	{
+		// This runs in the mutable Runner thread
+		MUTABLE_CPUPROFILER_SCOPE(FImageInvertTask_Prepare);
+		bOutFailed = false;
+
+		Ptr<const Image> Source = Runner->LoadImage({ Args.base, Op.ExecutionIndex, Op.ExecutionOptions });
+
+		// Create destination data
+		Result = Runner->CloneOrTakeOver(Source);
+		return true;
+	}
+
+
+	void FImageInvertTask::DoWork()
+	{
+		// This runs on a random worker thread
+		MUTABLE_CPUPROFILER_SCOPE(FImageInvertTask);
+
+		ImageInvert(Result.get());
+	}
+
+
+	bool FImageInvertTask::Complete(CodeRunner* Runner)
+	{
+		// If we didn't take a shortcut
+		if (Result)
+		{
+			Runner->StoreImage(Op, Result);
+		}
+
+		bool bSuccess = true;
+		return bSuccess;
+	}
+
+
+	//---------------------------------------------------------------------------------------------
+	//---------------------------------------------------------------------------------------------
+	//---------------------------------------------------------------------------------------------
+	class FImageComposeTask : public CodeRunner::FIssuedTask
+	{
+	public:
+		FImageComposeTask(const FScheduledOp& InOp, const OP::ImageComposeArgs& InArgs, const Ptr<const Layout>& InLayout)
+			: FIssuedTask(InOp), Args(InArgs), Layout(InLayout)
+		{}
+
+		// FIssuedTask interface
+		virtual bool Prepare(CodeRunner*, bool& bOutFailed) override;
+		virtual void DoWork() override;
+		virtual bool Complete(CodeRunner*) override;
+
+	private:
+		int32 ImageCompressionQuality = 0;
+		OP::ImageComposeArgs Args;
+		Ptr<const Layout> Layout;
+		Ptr<const Image> Block;
+		Ptr<const Image> Mask;
+		Ptr<Image> Result;
+		box<FIntVector2> Rect;
+		FImageOperator::FImagePixelFormatFunc ImagePixelFormatFunc;
+	};
+
+
+	bool FImageComposeTask::Prepare(CodeRunner* Runner, bool& bOutFailed)
+	{
+		// This runs in the mutable Runner thread
+		MUTABLE_CPUPROFILER_SCOPE(FImageComposeTask_Prepare);
+		bOutFailed = false;
+
+		ImageCompressionQuality = Runner->m_pSettings->ImageCompressionQuality;
+
+		Ptr<const Image> Base = Runner->LoadImage({ Args.base, Op.ExecutionIndex, Op.ExecutionOptions });
+
+		int32 RelBlockIndex = Layout->FindBlock(Args.BlockId);
+
+		// Shortcuts
+		if (RelBlockIndex < 0)
+		{
+			Runner->StoreImage(Op, Base);
+			return false;
+		}
+
+		// Only load the image is RelBlockIndex is valid, otherwise, we won't have requested it.
+		Block = Runner->LoadImage({ Args.blockImage, Op.ExecutionIndex, Op.ExecutionOptions });
+		if (Args.mask)
+		{
+			Mask = Runner->LoadImage({ Args.mask, Op.ExecutionIndex, Op.ExecutionOptions });
+		}
+
+		box<FIntVector2> RectInblocks;
+		RectInblocks.min = Layout->Blocks[RelBlockIndex].Min;
+		RectInblocks.size = Layout->Blocks[RelBlockIndex].Size;
+
+		// Convert the rect from blocks to pixels
+		FIntPoint Grid = Layout->GetGridSize();
+		int32 BlockSizeX = Base->GetSizeX() / Grid[0];
+		int32 BlockSizeY = Base->GetSizeY() / Grid[1];
+		Rect = RectInblocks;
+		Rect.min[0] *= BlockSizeX;
+		Rect.min[1] *= BlockSizeY;
+		Rect.size[0] *= BlockSizeX;
+		Rect.size[1] *= BlockSizeY;
+
+		if (!(Block && Rect.size[0] && Rect.size[1] && Block->GetSizeX() && Block->GetSizeY()))
+		{
+			Runner->Release(Block);
+			Runner->Release(Mask);
+			Runner->StoreImage(Op, Base);
+			return false;
+		}
+
+		// Create destination data
+		Result = Runner->CloneOrTakeOver(Base);
+		Result->m_flags = 0;
+
+		bool useMask = Args.mask != 0;
+		if (!useMask)
+		{
+			MUTABLE_CPUPROFILER_SCOPE(ImageComposeWithoutMask);
+
+			FImageOperator ImOp = MakeImageOperator(Runner);
+
+			EImageFormat Format = GetMostGenericFormat(Result->GetFormat(), Block->GetFormat());
+
+			// Resize image if it doesn't fit in the new block size
+			if (FIntVector2(Block->GetSize()) != Rect.size)
+			{
+				MUTABLE_CPUPROFILER_SCOPE(ImageComposeWithoutMask_BlockResize);
+
+				// This now happens more often since the generation of specific mips on request. For this reason
+				// this warning is usually acceptable.
+				Ptr<Image> Resized = Runner->CreateImage(Rect.size[0], Rect.size[1], 1, Block->GetFormat(), EInitializationType::NotInitialized );
+				ImOp.ImageResizeLinear(Resized.get(), ImageCompressionQuality, Block.get());
+				Runner->Release(Block);
+				Block = Resized;
+			}
+
+			// Change the block image format if it doesn't match the composed image
+			// This is usually enforced at object compilation time.
+			if (Result->GetFormat() != Block->GetFormat())
+			{
+				MUTABLE_CPUPROFILER_SCOPE(ImageComposeReformat);
+
+				if (Result->GetFormat() != Format)
+				{
+					Ptr<Image> Formatted = Runner->CreateImage(Result->GetSizeX(), Result->GetSizeY(), Result->GetLODCount(), Format, EInitializationType::NotInitialized);
+					bool bSuccess = false;
+					ImOp.ImagePixelFormat(bSuccess, ImageCompressionQuality, Formatted.get(), Result.get());
+					check(bSuccess); // Decompression cannot fail
+					Runner->Release(Result);
+					Result = Formatted;
+				}
+				if (Block->GetFormat() != Format)
+				{
+					Ptr<Image> Formatted = Runner->CreateImage(Block->GetSizeX(), Block->GetSizeY(), Block->GetLODCount(), Format, EInitializationType::NotInitialized);
+					bool bSuccess = false;
+					ImOp.ImagePixelFormat(bSuccess, ImageCompressionQuality, Formatted.get(), Block.get());
+					check(bSuccess); // Decompression cannot fail
+					Runner->Release(Block);
+					Block = Formatted;
+				}
+			}
+		}
+
+		ImagePixelFormatFunc = Runner->m_pSystem->ImagePixelFormatOverride;
+
+		return true;
+	}
+
+
+	void FImageComposeTask::DoWork()
+	{
+		// This runs on a random worker thread
+		MUTABLE_CPUPROFILER_SCOPE(FImageComposeTask);
+
+		bool useMask = Args.mask != 0;
+		if (!useMask)
+		{
+			MUTABLE_CPUPROFILER_SCOPE(ImageComposeWithoutMask);
+
+			// Compose without a mask
+			// \TODO: track allocs
+			FImageOperator ImOp = FImageOperator::GetDefault(ImagePixelFormatFunc);
+			ImOp.ImageCompose(Result.get(), Block.get(), Rect);
+		}
+		else
+		{
+			MUTABLE_CPUPROFILER_SCOPE(ImageComposeWithMask);
+
+			// Compose with a mask
+			ImageBlendOnBaseNoAlpha(Result.get(), Mask.get(), Block.get(), Rect);
+		}
+
+		Layout = nullptr;
+	}
+
+
+	bool FImageComposeTask::Complete(CodeRunner* Runner)
+	{
+		// This runs in the mutable Runner thread
+		Runner->Release(Block);
+		Runner->Release(Mask);
+
+		// If we didn't take a shortcut
+		if (Result)
+		{
+			Runner->StoreImage(Op, Result);
+		}
+
+		bool bSuccess = true;
+		return bSuccess;
+	}
+
+
+	//---------------------------------------------------------------------------------------------
+	//---------------------------------------------------------------------------------------------
+	//---------------------------------------------------------------------------------------------
+	bool CodeRunner::FLoadMeshRomTask::Prepare(CodeRunner* Runner, bool& bOutFailed)
+	{
+		// This runs in the mutable Runner thread
+		MUTABLE_CPUPROFILER_SCOPE(FLoadMeshRomTask_Prepare);
+
+		if (!Runner || !Runner->m_pSystem)
+		{
+			return false;
+		}
+
+		bOutFailed = false;
+
+		const FProgram& Program = Runner->m_pModel->GetPrivate()->m_program;
+
+		check(RomIndex < Program.m_roms.Num());
+		
+		FWorkingMemoryManager::FModelCacheEntry* ModelCache = Runner->m_pSystem->WorkingMemoryManager.FindModelCache(Runner->m_pModel.Get());
+		++ModelCache->PendingOpsPerRom[RomIndex];
+
+		if (Program.IsRomLoaded(RomIndex))
+		{
+			return false;
+		}
+
+		if (const FRomLoadOp* Result = Runner->RomLoadOps.Find(RomIndex))
+		{
+			Event = Result->Event;  // Wait for the read operation started by other task
+			return false;
+		}
+
+		FRomLoadOp& RomLoadOp = Runner->RomLoadOps.Create(RomIndex);
+		
+		check(Runner->m_pSystem->StreamInterface);
+
+		const uint32 RomSize = Program.m_roms[RomIndex].Size;
+		check(RomSize > 0);
+
+		// Free roms if necessary
+		{
+			MUTABLE_CPUPROFILER_SCOPE(FreeingRoms);
+
+			Runner->m_pSystem->WorkingMemoryManager.MarkRomUsed(RomIndex, Runner->m_pModel);
+			Runner->m_pSystem->WorkingMemoryManager.EnsureBudgetBelow(RomSize);
+		}
+
+		const int32 SizeBefore = RomLoadOp.m_streamBuffer.GetAllocatedSize();
+		RomLoadOp.m_streamBuffer.SetNumUninitialized(RomSize);
+		const int32 SizeAfter = RomLoadOp.m_streamBuffer.GetAllocatedSize();
+
+	 	UE::Tasks::FTaskEvent ReadCompletionEvent = UE::Tasks::FTaskEvent(TEXT("FLoadMeshRomTask"));	
+		RomLoadOp.Event = ReadCompletionEvent;
+		
+		TFunction<void(bool)> Callback = [ReadCompletionEvent](bool bSuccess) mutable
+		{
+			ReadCompletionEvent.Trigger();
+		};
+		
+		const uint32 RomId = Program.m_roms[RomIndex].Id;
+		RomLoadOp.m_streamID = Runner->m_pSystem->StreamInterface->BeginReadBlock(Runner->m_pModel.Get(), RomId, RomLoadOp.m_streamBuffer.GetData(), RomSize, &Callback);
+		if (RomLoadOp.m_streamID < 0)
+		{
+			bOutFailed = true;
+			return false;
+		}
+
+		Event = ReadCompletionEvent; // Wait for read operation to end
+		return false; // No worker thread work
+	}
+	
+
+	//---------------------------------------------------------------------------------------------
+	bool CodeRunner::FLoadMeshRomTask::Complete(CodeRunner* Runner)
+	{
+		// This runs in the Runner thread
+		MUTABLE_CPUPROFILER_SCOPE(FLoadMeshRomTask_Complete);
+
+		if (!Runner || !Runner->m_pSystem)
+		{
+			return false;
+		}
+
+		FProgram& Program = Runner->m_pModel->GetPrivate()->m_program;
+
+		// Since task could be reordered, we need to make sure we end the rom read before continuing
+		if (FRomLoadOp* RomLoadOp = Runner->RomLoadOps.Find(RomIndex))
+		{
+			Runner->m_pSystem->StreamInterface->EndRead(RomLoadOp->m_streamID);									
+
+			const int32 ResIndex = Program.m_roms[RomIndex].ResourceIndex;
+			check(!Program.ConstantMeshes[ResIndex].Value)
+			MUTABLE_CPUPROFILER_SCOPE(Unserialise);
+
+			InputMemoryStream Stream(RomLoadOp->m_streamBuffer.GetData(), RomLoadOp->m_streamBuffer.Num());
+			InputArchive Arch(&Stream);
+
+			check(!Program.ConstantMeshes[ResIndex].Value);
+			Ptr<Mesh> Value = Mesh::StaticUnserialise(Arch);
+			Program.SetMeshRomValue(RomIndex, Value);
+
+			check(Program.ConstantMeshes[ResIndex].Value);
+
+			Runner->RomLoadOps.Remove(*RomLoadOp);
+		}
+
+		// Process the constant op normally, now that the rom is loaded.
+		Runner->RunCode(Op, Runner->m_pParams, Runner->m_pModel, Runner->m_lodMask);
+
+		FWorkingMemoryManager::FModelCacheEntry* ModelCache = Runner->m_pSystem->WorkingMemoryManager.FindModelCache(Runner->m_pModel.Get());
+
+		Runner->m_pSystem->WorkingMemoryManager.MarkRomUsed(RomIndex, Runner->m_pModel);
+		--ModelCache->PendingOpsPerRom[RomIndex];
+
+		bool bSuccess = true;
+		return bSuccess;
+	}
+
+
+	//---------------------------------------------------------------------------------------------
+	//---------------------------------------------------------------------------------------------
+	//---------------------------------------------------------------------------------------------
+	bool CodeRunner::FLoadExtensionDataTask::Prepare(CodeRunner* Runner, bool& bOutFailed)
+	{
+		// This runs in the mutable Runner thread
+		MUTABLE_CPUPROFILER_SCOPE(FLoadExtensionDataTask_Prepare);
+		bOutFailed = false;
+
+		FProgram& Program = Runner->m_pModel->GetPrivate()->m_program;
+		check(Program.m_constantExtensionData.IsValidIndex(ModelConstantIndex));
+
+		const FExtensionDataConstant& Constant = Program.m_constantExtensionData[ModelConstantIndex];
+		check(Constant.Data == Data);
+
+		if (Constant.LoadState == FExtensionDataConstant::ELoadState::FailedToLoad)
+		{
+			// This already failed before, so don't waste time trying again.
+			//
+			// This behavior could be made an option in future if needed.
+			bOutFailed = true;
+			return false;
+		}
+
+		TArray<ExtensionDataPtrConst> UnloadedConstants;
+		LoadHandle = Runner->m_pSystem->GetExtensionDataStreamer()->StartLoad(Data, UnloadedConstants);
+		check(LoadHandle.IsValid());
+
+		// Update the load state of any constants that were unloaded
+		for (const ExtensionDataPtrConst& Unloaded : UnloadedConstants)
+		{
+			check(Unloaded->Origin == ExtensionData::EOrigin::ConstantStreamed);
+
+			FExtensionDataConstant* FoundConstant = Algo::FindBy(Program.m_constantExtensionData, Unloaded, &FExtensionDataConstant::Data);
+			check(FoundConstant);
+
+			FoundConstant->LoadState = FExtensionDataConstant::ELoadState::Unloaded;
+		}
+
+		// No worker thread work
+		return false;
+	}
+
+
+	//---------------------------------------------------------------------------------------------
+	bool CodeRunner::FLoadExtensionDataTask::Complete(CodeRunner* Runner)
+	{
+		MUTABLE_CPUPROFILER_SCOPE(FLoadExtensionDataTask_Complete);
+
+		bool bSuccess = true;
+
+		// Complete should only be called if the load is finished
+		check(LoadHandle->LoadState != FExtensionDataLoadHandle::ELoadState::Pending);
+
+		FProgram& Program = Runner->m_pModel->GetPrivate()->m_program;
+		check(Program.m_constantExtensionData.IsValidIndex(ModelConstantIndex));
+
+		FExtensionDataConstant& Constant = Program.m_constantExtensionData[ModelConstantIndex];
+		check(Constant.Data == Data);
+
+		// Update the load state in the Program
+		if (LoadHandle->LoadState == FExtensionDataLoadHandle::ELoadState::Loaded)
+		{
+			Constant.LoadState = FExtensionDataConstant::ELoadState::CurrentlyLoaded;
+		}
+		else
+		{
+			check(LoadHandle->LoadState == FExtensionDataLoadHandle::ELoadState::FailedToLoad);
+			Constant.LoadState = FExtensionDataConstant::ELoadState::FailedToLoad;
+			bSuccess = false;
+		}
+
+		// Process the constant op normally, now that the data is loaded.
+		Runner->RunCode(Op, Runner->m_pParams, Runner->m_pModel, Runner->m_lodMask);
+
+		return bSuccess;
+	}
+
+
+	//---------------------------------------------------------------------------------------------
+	bool CodeRunner::FLoadExtensionDataTask::IsComplete(CodeRunner* Runner)
+	{
+		check(LoadHandle.IsValid());
+		return LoadHandle->LoadState != FExtensionDataLoadHandle::ELoadState::Pending;
+	}
+
+
+	//---------------------------------------------------------------------------------------------
+	//---------------------------------------------------------------------------------------------
+	//---------------------------------------------------------------------------------------------
+	bool CodeRunner::FLoadImageRomsTask::Prepare(CodeRunner* Runner, bool& bOutFailed )
+	{
+		if (!Runner || !Runner->m_pSystem)
+		{
+			return false;
+		}
+
+		// This runs in the mutable Runner thread
+		MUTABLE_CPUPROFILER_SCOPE(FLoadImageRomsTask_Prepare);
+		bOutFailed = false;
+
+		FProgram& program = Runner->m_pModel->GetPrivate()->m_program;
+
+		FWorkingMemoryManager::FModelCacheEntry* ModelCache = Runner->m_pSystem->WorkingMemoryManager.FindModelCache(Runner->m_pModel.Get());
+
+		TArray<UE::Tasks::FTask> ReadCompleteEvents;
+
+		ReadCompleteEvents.Reserve(LODIndexCount); 
+
+		for (int32 LODIndex = 0; LODIndex < LODIndexCount; ++LODIndex)
+		{
+			const int32 CurrentIndexIndex = LODIndexIndex + LODIndex;
+			const int32 CurrentIndex = program.ConstantImageLODIndices[CurrentIndexIndex];
+
+			if (program.ConstantImageLODs[CurrentIndex].Key < 0)
+			{
+				// This data is always resident.
+				continue;
+			}
+
+			int32 RomIndex = program.ConstantImageLODs[CurrentIndex].Key;
+			check(RomIndex < program.m_roms.Num());
+
+			++ModelCache->PendingOpsPerRom[RomIndex];
+
+			if (DebugRom && (DebugRomAll || RomIndex == DebugRomIndex))
+				UE_LOG(LogMutableCore, Log, TEXT("Preparing rom %d, now peding ops is %d."), RomIndex, ModelCache->PendingOpsPerRom[RomIndex]);
+
+			if (program.IsRomLoaded(RomIndex))
+			{
+				continue;
+			}
+
+			RomIndices.Add(RomIndex);
+
+			if (const FRomLoadOp* Result = Runner->RomLoadOps.Find(RomIndex))
+			{
+				ReadCompleteEvents.Add(Result->Event); // Wait for the read operation started by other task
+				continue;
+			}
+
+			FRomLoadOp& RomLoadOp = Runner->RomLoadOps.Create(RomIndex);
+			
+			check(Runner->m_pSystem->StreamInterface);
+
+			const uint32 RomSize = program.m_roms[RomIndex].Size;
+			check(RomSize > 0);
+
+			// Free roms if necessary
+			{
+				Runner->m_pSystem->WorkingMemoryManager.MarkRomUsed(RomIndex, Runner->m_pModel);
+				Runner->m_pSystem->WorkingMemoryManager.EnsureBudgetBelow(RomSize);
+			}
+
+			const int32 SizeBefore = RomLoadOp.m_streamBuffer.GetAllocatedSize();
+			RomLoadOp.m_streamBuffer.SetNumUninitialized(RomSize);
+			const int32 SizeAfter = RomLoadOp.m_streamBuffer.GetAllocatedSize();
+
+			UE::Tasks::FTaskEvent ReadCompletionEvent(TEXT("FLoadImageRomsTaskRom"));
+			ReadCompleteEvents.Add(ReadCompletionEvent);
+			RomLoadOp.Event = ReadCompletionEvent;
+
+			TFunction<void(bool)> Callback = [ReadCompletionEvent](bool bSuccess) mutable // Mutable due Trigger not being const
+			{
+				ReadCompletionEvent.Trigger();
+			};
+			
+			const uint32 RomId = program.m_roms[RomIndex].Id;
+			RomLoadOp.m_streamID = Runner->m_pSystem->StreamInterface->BeginReadBlock(Runner->m_pModel.Get(), RomId, RomLoadOp.m_streamBuffer.GetData(), RomSize, &Callback);
+			if (RomLoadOp.m_streamID < 0)
+			{
+				bOutFailed = true;
+				return false;
+			}
+		}
+
+		// Wait for all read operations to end
+		UE::Tasks::FTaskEvent GatherReadsCompletionEvent(TEXT("FLoadImageRomsTask"));
+		GatherReadsCompletionEvent.AddPrerequisites(ReadCompleteEvents);
+		GatherReadsCompletionEvent.Trigger();
+
+		Event = GatherReadsCompletionEvent;
+			
+		return false; // No worker thread work
+	}
+	
+	
+	//---------------------------------------------------------------------------------------------
+	bool CodeRunner::FLoadImageRomsTask::Complete(CodeRunner* Runner)
+	{
+		// This runs in the Runner thread
+		MUTABLE_CPUPROFILER_SCOPE(FLoadImageRomsTask_Complete);
+
+		if (!Runner || !Runner->m_pSystem)
+		{
+			return false;
+		}
+
+		FProgram& Program = Runner->m_pModel->GetPrivate()->m_program;
+		
+		FWorkingMemoryManager::FModelCacheEntry* ModelCache = Runner->m_pSystem->WorkingMemoryManager.FindModelCache(Runner->m_pModel.Get());
+
+		bool bSomeMissingData = false;
+
+		for (const int32 RomIndex : RomIndices)
+		{
+			// Since task could be reordered, we need to make sure we end the rom read before continuing
+			if (FRomLoadOp* RomLoadOp = Runner->RomLoadOps.Find(RomIndex))
+			{				
+				bool bSuccess = Runner->m_pSystem->StreamInterface->EndRead(RomLoadOp->m_streamID);
+
+				MUTABLE_CPUPROFILER_SCOPE(Unserialise);
+
+				if (bSuccess)
+				{
+					InputMemoryStream Stream(RomLoadOp->m_streamBuffer.GetData(), RomLoadOp->m_streamBuffer.Num());
+					InputArchive Arch(&Stream);
+
+					const int32 ResIndex = Program.m_roms[RomIndex].ResourceIndex;
+
+					// TODO: Try to reuse buffer from PooledImages.
+					check(!Program.ConstantImageLODs[ResIndex].Value);
+					Ptr<Image> Value = Image::StaticUnserialise(Arch);
+
+					Program.SetImageRomValue(RomIndex, Value);
+					check(Program.ConstantImageLODs[ResIndex].Value);
+				}
+				else
+				{
+					bSomeMissingData = true;
+				}
+				
+				Runner->RomLoadOps.Remove(*RomLoadOp);
+			}
+		}
+
+		if (bSomeMissingData)
+		{
+			// Some data may be missing. We can try to go on if some mips are there. 
+			UE_LOG(LogMutableCore, Verbose, TEXT("FLoadImageRomsTask::Complete failed: missing data?"));
+		}
+
+		// Process the constant op normally, now that the rom is loaded.	
+		bool bSuccess = Runner->RunCode_ConstantResource(Op, Runner->m_pModel.Get());
+
+		for (int32 LODIndex = 0; bSuccess && (LODIndex < LODIndexCount); ++LODIndex)
+		{
+			int32 CurrentIndexIndex = LODIndexIndex + LODIndex;
+			int32 CurrentIndex = Program.ConstantImageLODIndices[CurrentIndexIndex];
+
+			if (Program.ConstantImageLODs[CurrentIndex].Key < 0)
+			{
+				// This data is always resident.
+				continue;
+			}
+
+			int32 RomIndex = Program.ConstantImageLODs[CurrentIndex].Key;
+			check(RomIndex < Program.m_roms.Num());
+
+			Runner->m_pSystem->WorkingMemoryManager.MarkRomUsed(RomIndex, Runner->m_pModel);
+			--ModelCache->PendingOpsPerRom[RomIndex];
+
+			if (DebugRom && (DebugRomAll || RomIndex == DebugRomIndex))
+			{
+				UE_LOG(LogMutableCore, Log, TEXT("FLoadImageRomsTask::Complete rom %d, now peding ops is %d."), RomIndex, ModelCache->PendingOpsPerRom[RomIndex]);
+			}
+		}
+
+		return bSuccess;
+	}
+
+
+	/** This task is used to load an image parameter (by its FName) or an image reference (from its ID).
+	*/
+	class FImageExternalLoadTask : public CodeRunner::FIssuedTask
+	{
+	public:
+
+		FImageExternalLoadTask(const FScheduledOp& InItem, uint8 InMipmapsToSkip, CodeRunner::FExternalResourceId InId);
+
+		// FIssuedTask interface
+		virtual bool Prepare(CodeRunner*, bool& bOutFailed) override;
+		virtual bool Complete(CodeRunner* Runner) override;
+		
+	private:
+		uint8 MipmapsToSkip;
+		CodeRunner::FExternalResourceId Id;
+
+		Ptr<Image> Result;
+		
+		TFunction<void()> ExternalCleanUpFunc;
+	};
+
+
+	FImageExternalLoadTask::FImageExternalLoadTask(const FScheduledOp& InOp, uint8 InMipmapsToSkip, CodeRunner::FExternalResourceId InId)
+		: FIssuedTask(InOp)
+	{
+		MipmapsToSkip = InMipmapsToSkip;
+		Id = InId;
+	}
+
+
+	bool FImageExternalLoadTask::Prepare(CodeRunner* Runner, bool& bOutFailed)
+	{
+		// This runs in the mutable Runner thread
+		MUTABLE_CPUPROFILER_SCOPE(FImageExternalLoadTask_Prepare);
+
+		// LoadExternalImageAsync will always generate some image even if it is a dummy one.
+		bOutFailed = false;
+
+		// Capturing this here should not be a problem. The lifetime of the callback lambda is tied to 
+		// the task and the later will always outlive the former. 
+		
+		// Probably we could simply pass a reference to the result image. 
+		TFunction<void (Ptr<Image>)> ResultCallback = [this](Ptr<Image> InResult)
+		{
+			Result = InResult;
+		};
+
+		Tie(Event, ExternalCleanUpFunc) = Runner->LoadExternalImageAsync(Id, MipmapsToSkip, ResultCallback);
+
+		// return false indicating there is no work to do so Event is not overriden by a DoWork task.
+		return false;
+	}
+
+
+	bool FImageExternalLoadTask::Complete(CodeRunner* Runner)
+	{
+		if (ExternalCleanUpFunc)
+		{
+			Invoke(ExternalCleanUpFunc);
+		}
+
+		Runner->StoreImage(Op, Result);
+
+		bool bSuccess = true;
+		return bSuccess;
+	}	
+
+
+	/** This task is used to load a mesh parameter (by its FName) or a mesh reference (from its ID).
+	*/
+	class FMeshExternalLoadTask : public CodeRunner::FIssuedTask
+	{
+	public:
+
+		FMeshExternalLoadTask(const FScheduledOp&, CodeRunner::FExternalResourceId);
+
+		// FIssuedTask interface
+		virtual bool Prepare(CodeRunner*, bool& bOutFailed) override;
+		virtual bool Complete(CodeRunner* Runner) override;
+
+	private:
+		uint8 MipmapsToSkip;
+		CodeRunner::FExternalResourceId Id;
+
+		Ptr<Mesh> Result;
+
+		TFunction<void()> ExternalCleanUpFunc;
+	};
+
+
+	FMeshExternalLoadTask::FMeshExternalLoadTask(const FScheduledOp& InOp, CodeRunner::FExternalResourceId InId)
+		: FIssuedTask(InOp)
+		, Id(InId)
+	{
+	}
+
+
+	bool FMeshExternalLoadTask::Prepare(CodeRunner* Runner, bool& bOutFailed)
+	{
+		// This runs in the mutable Runner thread
+		MUTABLE_CPUPROFILER_SCOPE(FMeshExternalLoadTask_Prepare);
+
+		// LoadExternalMeshAsync will always generate some mesh even if it is a dummy one.
+		bOutFailed = false;
+
+		// Capturing this here should not be a problem. The lifetime of the callback lambda is tied to 
+		// the task and the later will always outlive the former. 
+
+		// Probably we could simply pass a reference to the result mesh. 
+		TFunction<void(Ptr<Mesh>)> ResultCallback = [this](Ptr<Mesh> InResult)
+			{
+				Result = InResult;
+			};
+
+		Tie(Event, ExternalCleanUpFunc) = Runner->LoadExternalMeshAsync(Id, ResultCallback);
+
+		// return false indicating there is no work to do so Event is not overriden by a DoWork task.
+		return false;
+	}
+
+
+	bool FMeshExternalLoadTask::Complete(CodeRunner* Runner)
+	{
+		if (ExternalCleanUpFunc)
+		{
+			Invoke(ExternalCleanUpFunc);
+		}
+
+		Runner->StoreMesh(Op, Result);
+
+		bool bSuccess = true;
+		return bSuccess;
+	}
+
+
+	//---------------------------------------------------------------------------------------------
+	//---------------------------------------------------------------------------------------------
+	//---------------------------------------------------------------------------------------------
+	TSharedPtr<CodeRunner::FIssuedTask> CodeRunner::IssueOp(FScheduledOp item)
+	{
+		TSharedPtr<FIssuedTask> Issued;
+
+		FProgram& program = m_pModel->GetPrivate()->m_program;
+
+		OP_TYPE type = program.GetOpType(item.At);
+
+		switch (type)
+		{
+		case OP_TYPE::ME_CONSTANT:
+		{
+			OP::MeshConstantArgs args = program.GetOpArgs<OP::MeshConstantArgs>(item.At);
+			int32 RomIndex = program.ConstantMeshes[args.value].Key;
+			if (RomIndex >= 0 && !program.ConstantMeshes[args.value].Value)
+			{
+				Issued = MakeShared<FLoadMeshRomTask>(item, RomIndex);
+			}
+			else
+			{
+				// If already available, the rest of the constant code will run right away.
+			}
+			break;
+		}
+		
+		case OP_TYPE::ED_CONSTANT:
+		{
+			OP::ResourceConstantArgs Args = program.GetOpArgs<OP::ResourceConstantArgs>(item.At);
+
+			const FExtensionDataConstant::ELoadState LoadState = program.m_constantExtensionData[Args.value].LoadState;
+
+			check(LoadState != FExtensionDataConstant::ELoadState::Invalid);
+			
+			if (LoadState == FExtensionDataConstant::ELoadState::AlwaysLoaded
+				|| LoadState == FExtensionDataConstant::ELoadState::CurrentlyLoaded)
+			{
+				// Already loaded. The rest of the constant code will run right away.
+			}
+			else
+			{
+				Issued = MakeShared<FLoadExtensionDataTask>(item, program.m_constantExtensionData[Args.value].Data, Args.value);
+			}
+			break;
+		}
+
+		case OP_TYPE::IM_CONSTANT:
+		{
+			OP::ResourceConstantArgs args = program.GetOpArgs<OP::ResourceConstantArgs>(item.At);
+			int32 MipsToSkip = item.ExecutionOptions;
+			int32 ImageIndex = args.value;
+			int32 ReallySkip = FMath::Min(MipsToSkip, program.ConstantImages[ImageIndex].LODCount - 1);
+			int32 LODIndexIndex = program.ConstantImages[ImageIndex].FirstIndex + ReallySkip;
+			int32 LODIndexCount = program.ConstantImages[ImageIndex].LODCount - ReallySkip;
+			check(LODIndexCount > 0);
+
+			// We always need to follow this path, or roms may not be protected for long enough and might be unloaded 
+			// because of memory budget contraints.
+			bool bAnyMissing = true;
+			//bool bAnyMissing = false;
+			//for (int32 i=0; i<LODIndexCount; ++i)
+			//{
+			//	uint32 LODIndex = program.ConstantImageLODIndices[LODIndexIndex+i];
+			//	if ( !program.ConstantImageLODs[LODIndex].Value )
+			//	{
+			//		bAnyMissing = true;
+			//		break;
+			//	}
+			//}
+
+			if (bAnyMissing)
+			{
+				Issued = MakeShared<FLoadImageRomsTask>(item, LODIndexIndex, LODIndexCount);
+
+				if (DebugRom && (DebugRomAll || ImageIndex == DebugImageIndex))
+					UE_LOG(LogMutableCore, Log, TEXT("Issuing image %d skipping %d ."), ImageIndex, ReallySkip);
+			}
+			else
+			{
+				// If already available, the rest of the constant code will run right away.
+				if (DebugRom && (DebugRomAll || ImageIndex == DebugImageIndex))
+					UE_LOG(LogMutableCore, Log, TEXT("Image %d skipping %d is already loaded."), ImageIndex, ReallySkip);
+			}
+			break;
+		}		
+
+		case OP_TYPE::IM_PARAMETER:
+		{
+			OP::ParameterArgs args = program.GetOpArgs<OP::ParameterArgs>(item.At);
+			Ptr<RangeIndex> Index = BuildCurrentOpRangeIndex(item, m_pParams, m_pModel.Get(), args.variable);
+
+			const FName Id = m_pParams->GetImageValue(args.variable, Index);
+
+			check(ImageLOD < TNumericLimits<uint8>::Max() && ImageLOD >= 0);
+			check(ImageLOD + static_cast<int32>(item.ExecutionOptions) < TNumericLimits<uint8>::Max());
+
+			const uint8 MipmapsToSkip = item.ExecutionOptions + static_cast<uint8>(ImageLOD);
+
+			CodeRunner::FExternalResourceId FullId;
+			FullId.ParameterId = Id;
+			Issued = MakeShared<FImageExternalLoadTask>(item, MipmapsToSkip, FullId);
+
+			break;
+		}
+
+		case OP_TYPE::IM_REFERENCE:
+		{
+			OP::ResourceReferenceArgs Args = m_pModel->GetPrivate()->m_program.GetOpArgs<OP::ResourceReferenceArgs>(item.At);
+
+			// We only convert references to images if indicated in the operation.
+			if (Args.ForceLoad)
+			{
+				check(item.Stage==0);
+
+				const uint8 MipmapsToSkip = item.ExecutionOptions + static_cast<uint8>(ImageLOD);
+
+				FExternalResourceId FullId;
+				FullId.ReferenceResourceId = Args.ID;
+				Issued = MakeShared<FImageExternalLoadTask>(item, MipmapsToSkip, FullId);
+			}
+
+			break;
+		}
+
+		case OP_TYPE::ME_REFERENCE:
+		{
+			OP::ResourceReferenceArgs Args = m_pModel->GetPrivate()->m_program.GetOpArgs<OP::ResourceReferenceArgs>(item.At);
+
+			// We only convert references to meshes if indicated in the operation.
+			if (Args.ForceLoad)
+			{
+				check(item.Stage == 0);
+
+				FExternalResourceId FullId;
+				FullId.ReferenceResourceId = Args.ID;
+				Issued = MakeShared<FMeshExternalLoadTask>(item, FullId);
+			}
+
+			break;
+		}
+
+		case OP_TYPE::IM_PIXELFORMAT:
+		{
+			if (item.Stage == 1)
+			{
+				OP::ImagePixelFormatArgs Args = program.GetOpArgs<OP::ImagePixelFormatArgs>(item.At);
+				Issued = MakeShared<FImagePixelFormatTask>(item, Args);
+			}
+			break;
+		}
+
+		case OP_TYPE::IM_LAYERCOLOUR:
+		{
+			if (item.Stage == 1)
+			{
+				OP::ImageLayerColourArgs Args = program.GetOpArgs<OP::ImageLayerColourArgs>(item.At);
+				Issued = MakeShared<FImageLayerColourTask>(item, Args);
+			}
+			break;
+		}
+
+		case OP_TYPE::IM_LAYER:
+		{
+			if ((ExecutionStrategy == EExecutionStrategy::MinimizeMemory && item.Stage == 2)
+				||
+				(ExecutionStrategy != EExecutionStrategy::MinimizeMemory && item.Stage == 1)
+				)
+			{
+				OP::ImageLayerArgs Args = program.GetOpArgs<OP::ImageLayerArgs>(item.At);
+				Issued = MakeShared<FImageLayerTask>(item, Args);
+			}
+			break;
+		}
+
+		case OP_TYPE::IM_MIPMAP:
+		{
+			if (item.Stage == 1)
+			{
+				OP::ImageMipmapArgs Args = program.GetOpArgs<OP::ImageMipmapArgs>(item.At);
+				Issued = MakeShared<FImageMipmapTask>(item, Args);
+			}
+			break;
+		}
+
+		case OP_TYPE::IM_SWIZZLE:
+		{
+			if (item.Stage == 1)
+			{
+				OP::ImageSwizzleArgs Args = program.GetOpArgs<OP::ImageSwizzleArgs>(item.At);
+				Issued = MakeShared<FImageSwizzleTask>(item, Args);
+			}
+			break;
+		}
+
+		case OP_TYPE::IM_SATURATE:
+		{
+			if (item.Stage == 1)
+			{
+				OP::ImageSaturateArgs Args = program.GetOpArgs<OP::ImageSaturateArgs>(item.At);
+				Issued = MakeShared<FImageSaturateTask>(item, Args);
+			}
+			break;
+		}
+
+		case OP_TYPE::IM_INVERT:
+		{
+			if (item.Stage == 1)
+			{
+				OP::ImageInvertArgs Args = program.GetOpArgs<OP::ImageInvertArgs>(item.At);
+				Issued = MakeShared<FImageInvertTask>(item, Args);
+			}
+			break;
+		}
+
+		case OP_TYPE::IM_RESIZE:
+		{
+			if (item.Stage == 1)
+			{
+				OP::ImageResizeArgs Args = program.GetOpArgs<OP::ImageResizeArgs>(item.At);
+				Issued = MakeShared<FImageResizeTask>(item, Args);
+			}
+			break;
+		}
+
+		case OP_TYPE::IM_RESIZEREL:
+		{
+			if (item.Stage == 1)
+			{
+				OP::ImageResizeRelArgs Args = program.GetOpArgs<OP::ImageResizeRelArgs>(item.At);
+				Issued = MakeShared<FImageResizeRelTask>(item, Args);
+			}
+			break;
+		}
+
+		case OP_TYPE::IM_COMPOSE:
+		{
+			if ((ExecutionStrategy == EExecutionStrategy::MinimizeMemory && item.Stage == 3) ||
+				(ExecutionStrategy != EExecutionStrategy::MinimizeMemory && item.Stage == 2))
+			{
+				OP::ImageComposeArgs Args = program.GetOpArgs<OP::ImageComposeArgs>(item.At);
+				Ptr<const Layout> ComposeLayout = static_cast<const Layout*>( m_heapData[item.CustomState].Resource.get());
+				Issued = MakeShared<FImageComposeTask>(item, Args, ComposeLayout);
+			}
+			break;
+		}
+
+		default:
+			break;
+		}
+
+		return Issued;
+	}
+
+} // namespace mu
