[mutable] Moved the Mutable plugin out of Experimental status into Beta.

diff --git a/Engine/Plugins/Mutable/Source/MutableRuntime/Private/MuR/CodeVisitor.cpp b/Engine/Plugins/Mutable/Source/MutableRuntime/Private/MuR/CodeVisitor.cpp
new file mode 100644
--- /dev/null
+++ b/Engine/Plugins/Mutable/Source/MutableRuntime/Private/MuR/CodeVisitor.cpp
@@ -0,0 +1,90 @@
+// Copyright Epic Games, Inc. All Rights Reserved.
+
+#include "MuR/CodeVisitor.h"
+
+#include "HAL/PlatformCrt.h"
+#include "MuR/MutableTrace.h"
+
+namespace mu
+{
+
+	//-------------------------------------------------------------------------------------------------
+	//-------------------------------------------------------------------------------------------------
+	//-------------------------------------------------------------------------------------------------
+	void SubtreeParametersVisitor::Run(OP::ADDRESS Root, const FProgram& Program)
+	{
+		// Cached?
+		TArray<int32>* it = ResultCache.Find(Root);
+		if (it)
+		{
+			RelevantParams = *it;
+			return;
+		}
+
+		// Not cached
+		MUTABLE_CPUPROFILER_SCOPE(SubtreeParametersVisitor);
+
+		Visited.SetNum(Program.m_opAddress.Num());
+		if (Program.m_opAddress.Num())
+		{
+			FMemory::Memzero(&Visited[0], Visited.Num());
+		}
+
+		CurrentParams.SetNum(Program.m_parameters.Num());
+		if (CurrentParams.Num())
+		{
+			FMemory::Memzero(CurrentParams.GetData(), CurrentParams.GetAllocatedSize());
+		}
+
+		Pending.Reserve(Program.m_opAddress.Num() / 4);
+		Pending.Add(Root);
+
+		while (Pending.Num())
+		{
+			OP::ADDRESS at = Pending.Pop();
+
+			if (!Visited[at])
+			{
+				Visited[at] = true;
+
+				switch (Program.GetOpType(at))
+				{
+				case OP_TYPE::NU_PARAMETER:
+				case OP_TYPE::SC_PARAMETER:
+				case OP_TYPE::BO_PARAMETER:
+				case OP_TYPE::CO_PARAMETER:
+				case OP_TYPE::PR_PARAMETER:
+				case OP_TYPE::IM_PARAMETER:
+				case OP_TYPE::ST_PARAMETER:
+					CurrentParams[Program.GetOpArgs<OP::ParameterArgs>(at).variable]++;
+					break;
+
+				default:
+					break;
+				}
+
+				ForEachReference(Program, at, [&](OP::ADDRESS ref)
+					{
+						if (ref)
+						{
+							Pending.Add(ref);
+						}
+					});
+			}
+		}
+
+
+		// Build result
+		RelevantParams.Empty();
+		for (int32 i = 0; i < CurrentParams.Num(); ++i)
+		{
+			if (CurrentParams[i])
+			{
+				RelevantParams.Add(i);
+			}
+		}
+
+		ResultCache.Add(Root, RelevantParams);
+	}
+
+}
\ No newline at end of file
diff --git a/Engine/Plugins/Mutable/Source/MutableRuntime/Private/MuR/Instance.cpp b/Engine/Plugins/Mutable/Source/MutableRuntime/Private/MuR/Instance.cpp
new file mode 100644
--- /dev/null
+++ b/Engine/Plugins/Mutable/Source/MutableRuntime/Private/MuR/Instance.cpp
@@ -0,0 +1,577 @@
+// Copyright Epic Games, Inc. All Rights Reserved.
+
+
+#include "MuR/Instance.h"
+
+#include "HAL/LowLevelMemTracker.h"
+#include "Misc/AssertionMacros.h"
+#include "MuR/InstancePrivate.h"
+#include "MuR/MutableMath.h"
+
+namespace mu
+{
+
+	//---------------------------------------------------------------------------------------------
+	//---------------------------------------------------------------------------------------------
+	//---------------------------------------------------------------------------------------------
+	Instance::Instance()
+	{
+		LLM_SCOPE_BYNAME(TEXT("MutableRuntime"));
+		m_pD = new Private();
+	}
+
+
+	//---------------------------------------------------------------------------------------------
+	Instance::~Instance()
+	{
+		LLM_SCOPE_BYNAME(TEXT("MutableRuntime"));
+        check( m_pD );
+		delete m_pD;
+		m_pD = 0;
+	}
+
+
+	//---------------------------------------------------------------------------------------------
+	Instance::Private* Instance::GetPrivate() const
+	{
+		return m_pD;
+	}
+
+
+    //---------------------------------------------------------------------------------------------
+    InstancePtr Instance::Clone() const
+    {
+		LLM_SCOPE_BYNAME(TEXT("MutableRuntime"));
+        InstancePtr pResult = new Instance();
+
+        *pResult->GetPrivate() = *m_pD;
+
+        return pResult;
+    }
+
+
+	//---------------------------------------------------------------------------------------------
+	int32 Instance::GetDataSize() const
+	{
+		return 16 + sizeof(Private) + m_pD->Components.GetAllocatedSize() + m_pD->ExtensionData.GetAllocatedSize();
+	}
+
+
+    //---------------------------------------------------------------------------------------------
+    Instance::ID Instance::GetId() const
+    {
+        return m_pD->Id;
+    }
+
+
+    //---------------------------------------------------------------------------------------------
+    int32 Instance::GetComponentCount() const
+    {
+		return m_pD->Components.Num();
+    }
+
+
+	//---------------------------------------------------------------------------------------------
+	int32 Instance::GetLODCount( int32 ComponentIndex ) const
+	{
+		if (m_pD->Components.IsValidIndex(ComponentIndex))
+		{
+			return m_pD->Components[ComponentIndex].LODs.Num();
+		}
+		check(false);
+		return 0;
+	}
+
+	
+	//---------------------------------------------------------------------------------------------
+	uint16 Instance::GetComponentId( int32 ComponentIndex ) const
+	{
+		if (m_pD->Components.IsValidIndex(ComponentIndex))
+		{
+			return m_pD->Components[ComponentIndex].Id;
+		}
+		else
+		{
+			check(false);
+		}
+
+		return 0;
+	}
+
+
+    //---------------------------------------------------------------------------------------------
+    int32 Instance::GetSurfaceCount( int32 ComponentIndex, int32 LODIndex ) const
+    {
+		if (m_pD->Components.IsValidIndex(ComponentIndex) &&
+			m_pD->Components[ComponentIndex].LODs.IsValidIndex(LODIndex))
+		{
+			return m_pD->Components[ComponentIndex].LODs[LODIndex].Surfaces.Num();
+		}
+		else
+		{
+			check(false);
+		}
+
+		return 0;
+	}
+
+
+    //---------------------------------------------------------------------------------------------
+    uint32 Instance::GetSurfaceId(int32 ComponentIndex, int32 LODIndex, int32 SurfaceIndex ) const
+    {
+        if (m_pD->Components.IsValidIndex(ComponentIndex) &&
+			m_pD->Components[ComponentIndex].LODs.IsValidIndex(LODIndex) &&
+			m_pD->Components[ComponentIndex].LODs[LODIndex].Surfaces.IsValidIndex(SurfaceIndex) )
+        {
+            return m_pD->Components[ComponentIndex].LODs[LODIndex].Surfaces[SurfaceIndex].InternalId;
+        }
+		else
+		{
+			check(false);
+		}
+
+        return 0;
+    }
+
+
+    //---------------------------------------------------------------------------------------------
+    int32 Instance::FindSurfaceById(int32 ComponentIndex, int32 LODIndex, uint32 id ) const
+    {
+		if (m_pD->Components.IsValidIndex(ComponentIndex) &&
+			m_pD->Components[ComponentIndex].LODs.IsValidIndex(LODIndex))
+		{
+			for (int32 i = 0; i < m_pD->Components[ComponentIndex].LODs[LODIndex].Surfaces.Num(); ++i)
+			{
+				if (m_pD->Components[ComponentIndex].LODs[LODIndex].Surfaces[i].InternalId == id)
+				{
+					return i;
+				}
+			}
+		}
+		else
+		{
+			check(false);
+		}
+
+        return -1;
+    }
+
+	
+	//---------------------------------------------------------------------------------------------
+	void Instance::FindBaseSurfaceBySharedId(int32 CompIndex, int32 SharedId, int32& OutSurfaceIndex, int32& OutLODIndex) const
+	{
+		if (m_pD->Components.IsValidIndex(CompIndex))
+		{
+			for (int32 LodIndex = 0; LodIndex < m_pD->Components[CompIndex].LODs.Num(); LodIndex++)
+			{
+				FInstanceLOD& LOD = m_pD->Components[CompIndex].LODs[LodIndex];
+				for (int32 SurfaceIndex = 0; SurfaceIndex < LOD.Surfaces.Num(); ++SurfaceIndex)
+				{
+					if (LOD.Surfaces[SurfaceIndex].SharedId == SharedId)
+					{
+						OutSurfaceIndex = SurfaceIndex;
+						OutLODIndex = LodIndex;
+						return;
+					}
+				}
+			}
+
+		}
+
+		OutSurfaceIndex = INDEX_NONE;
+		OutLODIndex = INDEX_NONE;
+	}
+
+
+	//---------------------------------------------------------------------------------------------
+	int32 Instance::GetSharedSurfaceId(int32 ComponentIndex, int32 LODIndex, int32 SurfaceIndex) const
+	{
+		if (m_pD->Components.IsValidIndex(ComponentIndex) &&
+			m_pD->Components[ComponentIndex].LODs.IsValidIndex(LODIndex) &&
+			m_pD->Components[ComponentIndex].LODs[LODIndex].Surfaces.IsValidIndex(SurfaceIndex))
+		{
+			return m_pD->Components[ComponentIndex].LODs[LODIndex].Surfaces[SurfaceIndex].SharedId;
+		}
+		else
+		{
+			check(false);
+		}
+
+		return 0;
+	}
+	
+
+    //---------------------------------------------------------------------------------------------
+    uint32 Instance::GetSurfaceCustomId(int32 ComponentIndex, int32 LODIndex, int32 SurfaceIndex ) const
+    {
+        if (m_pD->Components.IsValidIndex(ComponentIndex) &&
+			m_pD->Components[ComponentIndex].LODs.IsValidIndex(LODIndex) &&
+			m_pD->Components[ComponentIndex].LODs[LODIndex].Surfaces.IsValidIndex(SurfaceIndex))
+        {
+            return m_pD->Components[ComponentIndex].LODs[LODIndex].Surfaces[SurfaceIndex].ExternalId;
+        }
+		else
+		{
+			check(false);
+		}
+
+        return 0;
+    }
+
+
+	//---------------------------------------------------------------------------------------------
+    int32 Instance::GetImageCount(int32 ComponentIndex, int32 LODIndex, int32 SurfaceIndex ) const
+	{
+		check(m_pD->Components.IsValidIndex(ComponentIndex));
+		check(m_pD->Components[ComponentIndex].LODs.IsValidIndex(LODIndex));
+		check(m_pD->Components[ComponentIndex].LODs[LODIndex].Surfaces.IsValidIndex(SurfaceIndex));
+
+		return m_pD->Components[ComponentIndex].LODs[LODIndex].Surfaces[SurfaceIndex].Images.Num();
+	}
+
+
+	//---------------------------------------------------------------------------------------------
+    int32 Instance::GetVectorCount(int32 ComponentIndex, int32 LODIndex, int32 SurfaceIndex ) const
+	{
+		check(m_pD->Components.IsValidIndex(ComponentIndex));
+		check(m_pD->Components[ComponentIndex].LODs.IsValidIndex(LODIndex));
+		check(m_pD->Components[ComponentIndex].LODs[LODIndex].Surfaces.IsValidIndex(SurfaceIndex));
+
+		return m_pD->Components[ComponentIndex].LODs[LODIndex].Surfaces[SurfaceIndex].Vectors.Num();
+	}
+
+
+	//---------------------------------------------------------------------------------------------
+    int32 Instance::GetScalarCount(int32 ComponentIndex, int32 LODIndex, int32 SurfaceIndex ) const
+	{
+		check(m_pD->Components.IsValidIndex(ComponentIndex));
+		check(m_pD->Components[ComponentIndex].LODs.IsValidIndex(LODIndex));
+		check(m_pD->Components[ComponentIndex].LODs[LODIndex].Surfaces.IsValidIndex(SurfaceIndex));
+
+		return m_pD->Components[ComponentIndex].LODs[LODIndex].Surfaces[SurfaceIndex].Scalars.Num();
+	}
+
+
+    //---------------------------------------------------------------------------------------------
+    int32 Instance::GetStringCount(int32 ComponentIndex, int32 LODIndex, int32 SurfaceIndex ) const
+    {
+		check(m_pD->Components.IsValidIndex(ComponentIndex));
+		check(m_pD->Components[ComponentIndex].LODs.IsValidIndex(LODIndex));
+		check(m_pD->Components[ComponentIndex].LODs[LODIndex].Surfaces.IsValidIndex(SurfaceIndex));
+
+		return m_pD->Components[ComponentIndex].LODs[LODIndex].Surfaces[SurfaceIndex].Strings.Num();
+	}
+
+
+    //---------------------------------------------------------------------------------------------
+	FResourceID Instance::GetMeshId(int32 ComponentIndex, int32 LODIndex ) const
+    {
+        check(m_pD->Components.IsValidIndex(ComponentIndex));
+        check(m_pD->Components[ComponentIndex].LODs.IsValidIndex(LODIndex));
+
+		return m_pD->Components[ComponentIndex].LODs[LODIndex].MeshId;
+    }
+
+
+	//---------------------------------------------------------------------------------------------
+	FResourceID Instance::GetImageId(int32 ComponentIndex, int32 LODIndex, int32 SurfaceIndex, int32 ImageIndex) const
+	{
+		check(m_pD->Components.IsValidIndex(ComponentIndex));
+		check(m_pD->Components[ComponentIndex].LODs.IsValidIndex(LODIndex));
+        check(m_pD->Components[ComponentIndex].LODs[LODIndex].Surfaces.IsValidIndex(SurfaceIndex));
+        check(m_pD->Components[ComponentIndex].LODs[LODIndex].Surfaces[SurfaceIndex].Images.IsValidIndex(ImageIndex));
+
+		return m_pD->Components[ComponentIndex].LODs[LODIndex].Surfaces[SurfaceIndex].Images[ImageIndex].Id;
+	}
+
+
+	//---------------------------------------------------------------------------------------------
+    FName Instance::GetImageName(int32 ComponentIndex, int32 LODIndex, int32 SurfaceIndex, int32 ImageIndex ) const
+	{
+		check(m_pD->Components.IsValidIndex(ComponentIndex));
+		check(m_pD->Components[ComponentIndex].LODs.IsValidIndex(LODIndex));
+        check(m_pD->Components[ComponentIndex].LODs[LODIndex].Surfaces.IsValidIndex(SurfaceIndex));
+		check(m_pD->Components[ComponentIndex].LODs[LODIndex].Surfaces[SurfaceIndex].Images.IsValidIndex(ImageIndex));
+
+        return m_pD->Components[ComponentIndex].LODs[LODIndex].Surfaces[SurfaceIndex].Images[ImageIndex].Name;
+	}
+
+
+	//---------------------------------------------------------------------------------------------
+	FVector4f Instance::GetVector(int32 ComponentIndex, int32 LODIndex, int32 SurfaceIndex, int32 VectorIndex) const
+	{
+		check(m_pD->Components.IsValidIndex(ComponentIndex));
+		check(m_pD->Components[ComponentIndex].LODs.IsValidIndex(LODIndex));
+        check(m_pD->Components[ComponentIndex].LODs[LODIndex].Surfaces.IsValidIndex(SurfaceIndex));
+		check(m_pD->Components[ComponentIndex].LODs[LODIndex].Surfaces[SurfaceIndex].Vectors.IsValidIndex(VectorIndex));
+
+        return m_pD->Components[ComponentIndex].LODs[LODIndex].Surfaces[SurfaceIndex].Vectors[VectorIndex].Value;
+	}
+
+
+	//---------------------------------------------------------------------------------------------
+	FName Instance::GetVectorName(int32 ComponentIndex, int32 LODIndex, int32 SurfaceIndex, int32 VectorIndex) const
+	{
+		check(m_pD->Components.IsValidIndex(ComponentIndex));
+		check(m_pD->Components[ComponentIndex].LODs.IsValidIndex(LODIndex));
+        check(m_pD->Components[ComponentIndex].LODs[LODIndex].Surfaces.IsValidIndex(SurfaceIndex));
+		check(m_pD->Components[ComponentIndex].LODs[LODIndex].Surfaces[SurfaceIndex].Vectors.IsValidIndex(VectorIndex));
+
+        return m_pD->Components[ComponentIndex].LODs[LODIndex].Surfaces[SurfaceIndex].Vectors[VectorIndex].Name;
+	}
+
+
+	//---------------------------------------------------------------------------------------------
+    float Instance::GetScalar(int32 ComponentIndex, int32 LODIndex, int32 SurfaceIndex, int32 ScalarIndex ) const
+	{
+		check(m_pD->Components.IsValidIndex(ComponentIndex));
+		check(m_pD->Components[ComponentIndex].LODs.IsValidIndex(LODIndex));
+        check(m_pD->Components[ComponentIndex].LODs[LODIndex].Surfaces.IsValidIndex(SurfaceIndex));
+		check(m_pD->Components[ComponentIndex].LODs[LODIndex].Surfaces[SurfaceIndex].Scalars.IsValidIndex(ScalarIndex));
+
+        return m_pD->Components[ComponentIndex].LODs[LODIndex].Surfaces[SurfaceIndex].Scalars[ScalarIndex].Value;
+	}
+
+
+	//---------------------------------------------------------------------------------------------
+	FName Instance::GetScalarName(int32 ComponentIndex, int32 LODIndex, int32 SurfaceIndex, int32 ScalarIndex) const
+	{
+		check(m_pD->Components.IsValidIndex(ComponentIndex));
+		check(m_pD->Components[ComponentIndex].LODs.IsValidIndex(LODIndex));
+        check(m_pD->Components[ComponentIndex].LODs[LODIndex].Surfaces.IsValidIndex(SurfaceIndex));
+		check(m_pD->Components[ComponentIndex].LODs[LODIndex].Surfaces[SurfaceIndex].Scalars.IsValidIndex(ScalarIndex));
+
+        return m_pD->Components[ComponentIndex].LODs[LODIndex].Surfaces[SurfaceIndex].Scalars[ScalarIndex].Name;
+	}
+
+
+    //---------------------------------------------------------------------------------------------
+    FString Instance::GetString(int32 ComponentIndex, int32 LODIndex, int32 SurfaceIndex, int32 StringIndex ) const
+    {
+        check(m_pD->Components.IsValidIndex(ComponentIndex));
+        check(m_pD->Components[ComponentIndex].LODs.IsValidIndex(LODIndex));
+        check(m_pD->Components[ComponentIndex].LODs[LODIndex].Surfaces.IsValidIndex(SurfaceIndex));
+		check(m_pD->Components[ComponentIndex].LODs[LODIndex].Surfaces[SurfaceIndex].Strings.IsValidIndex(StringIndex));
+
+		bool bValid = m_pD->Components[ComponentIndex].LODs[LODIndex].Surfaces[SurfaceIndex].Strings.IsValidIndex(StringIndex);
+		if (bValid)
+        {
+            return m_pD->Components[ComponentIndex].LODs[LODIndex].Surfaces[SurfaceIndex].Strings[StringIndex].Value;
+        }
+
+        return "";
+    }
+
+
+    //---------------------------------------------------------------------------------------------
+	FName Instance::GetStringName(int32 ComponentIndex, int32 LODIndex, int32 SurfaceIndex, int32 StringIndex) const
+    {
+        check(m_pD->Components.IsValidIndex(ComponentIndex));
+        check(m_pD->Components[ComponentIndex].LODs.IsValidIndex(LODIndex));
+        check(m_pD->Components[ComponentIndex].LODs[LODIndex].Surfaces.IsValidIndex(SurfaceIndex));
+
+		bool bValid = m_pD->Components[ComponentIndex].LODs[LODIndex].Surfaces[SurfaceIndex].Strings.IsValidIndex(StringIndex);
+		if (bValid)
+		{
+            return m_pD->Components[ComponentIndex].LODs[LODIndex].Surfaces[SurfaceIndex].Strings[StringIndex].Name;
+        }
+
+        return NAME_None;
+    }
+
+
+    //---------------------------------------------------------------------------------------------
+	int32 Instance::GetExtensionDataCount() const
+	{
+		return m_pD->ExtensionData.Num();
+	}
+
+
+    //---------------------------------------------------------------------------------------------
+	void Instance::GetExtensionData(int32 Index, Ptr<const class ExtensionData>& OutExtensionData, FName& OutName) const
+	{
+		check(m_pD->ExtensionData.IsValidIndex(Index));
+
+		OutExtensionData = m_pD->ExtensionData[Index].Data;
+		OutName = m_pD->ExtensionData[Index].Name;
+	}
+
+
+    //---------------------------------------------------------------------------------------------
+	int32 Instance::Private::AddComponent()
+	{
+		int32 result = Components.Emplace();
+		return result;
+	}
+
+
+    //---------------------------------------------------------------------------------------------
+    int32 Instance::Private::AddLOD( int32 ComponentIndex )
+    {
+        // Automatically create the necessary lods and components
+        while (ComponentIndex >= Components.Num())
+        {
+            AddComponent();
+        }
+
+        return Components[ComponentIndex].LODs.Emplace();
+    }
+
+
+    //---------------------------------------------------------------------------------------------
+    int32 Instance::Private::AddSurface( int32 ComponentIndex, int32 LODIndex )
+    {
+        // Automatically create the necessary lods and components
+        while (ComponentIndex >= Components.Num())
+        {
+            AddComponent();
+        }
+        while (LODIndex >= Components[ComponentIndex].LODs.Num())
+        {
+            AddLOD(ComponentIndex);
+        }
+
+        return Components[ComponentIndex].LODs[LODIndex].Surfaces.Emplace();
+    }
+
+
+    //---------------------------------------------------------------------------------------------
+    void Instance::Private::SetSurfaceName( int32 ComponentIndex, int32 LODIndex, int32 SurfaceIndex, FName Name)
+    {
+        // Automatically create the necessary lods and components
+		while (ComponentIndex >= Components.Num())
+		{
+			AddComponent();
+		}
+		while (LODIndex >= Components[ComponentIndex].LODs.Num())
+		{
+			AddLOD(ComponentIndex);
+		}
+		while ( SurfaceIndex>=Components[ComponentIndex].LODs[LODIndex].Surfaces.Num() )
+        {
+            AddSurface( LODIndex, ComponentIndex );
+        }
+
+        FInstanceSurface& surface = Components[ComponentIndex].LODs[LODIndex].Surfaces[SurfaceIndex];
+        surface.Name = Name;
+    }
+
+
+	//---------------------------------------------------------------------------------------------
+	void Instance::Private::SetMesh(int32 ComponentIndex, int32 LODIndex, FResourceID meshId, FName Name)
+	{
+		// Automatically create the necessary lods and components
+		while (ComponentIndex >= Components.Num())
+		{
+			AddComponent();
+		}
+		while (LODIndex >= Components[ComponentIndex].LODs.Num())
+		{
+			AddLOD(ComponentIndex);
+		}
+
+		FInstanceLOD& LOD = Components[ComponentIndex].LODs[LODIndex];
+		LOD.MeshId = meshId;
+		LOD.MeshName = Name;
+	}
+
+
+	//---------------------------------------------------------------------------------------------
+    int32 Instance::Private::AddImage( int32 ComponentIndex, int32 LODIndex, int32 SurfaceIndex, FResourceID imageId, FName Name)
+	{
+		// Automatically create the necessary lods and components
+		while (ComponentIndex >= Components.Num())
+		{
+			AddComponent();
+		}
+		while (LODIndex >= Components[ComponentIndex].LODs.Num())
+		{
+			AddLOD(ComponentIndex);
+		}
+		while (SurfaceIndex >= Components[ComponentIndex].LODs[LODIndex].Surfaces.Num())
+		{
+			AddSurface(LODIndex, ComponentIndex);
+		}
+
+		FInstanceSurface& Surface = Components[ComponentIndex].LODs[LODIndex].Surfaces[SurfaceIndex];
+		return Surface.Images.Add({ imageId, Name });
+	}
+
+
+	//---------------------------------------------------------------------------------------------
+    int32 Instance::Private::AddVector( int32 ComponentIndex, int32 LODIndex, int32 SurfaceIndex, const FVector4f& vec, FName Name)
+	{
+		// Automatically create the necessary lods and components
+		while (ComponentIndex >= Components.Num())
+		{
+			AddComponent();
+		}
+		while (LODIndex >= Components[ComponentIndex].LODs.Num())
+		{
+			AddLOD(ComponentIndex);
+		}
+		while (SurfaceIndex >= Components[ComponentIndex].LODs[LODIndex].Surfaces.Num())
+		{
+			AddSurface(LODIndex, ComponentIndex);
+		}
+
+		FInstanceSurface& Surface = Components[ComponentIndex].LODs[LODIndex].Surfaces[SurfaceIndex];
+		return Surface.Vectors.Add({ vec, Name } );
+	}
+
+
+    //---------------------------------------------------------------------------------------------
+    int32 Instance::Private::AddScalar( int32 ComponentIndex, int32 LODIndex, int32 SurfaceIndex, float sca, FName Name)
+    {
+        // Automatically create the necessary lods and components
+		while (ComponentIndex >= Components.Num())
+		{
+			AddComponent();
+		}
+		while (LODIndex >= Components[ComponentIndex].LODs.Num())
+		{
+			AddLOD(ComponentIndex);
+		}
+		while (SurfaceIndex >= Components[ComponentIndex].LODs[LODIndex].Surfaces.Num())
+		{
+			AddSurface(LODIndex, ComponentIndex);
+		}
+
+		FInstanceSurface& Surface = Components[ComponentIndex].LODs[LODIndex].Surfaces[SurfaceIndex];
+		return Surface.Scalars.Add({ sca, Name });
+    }
+
+
+    //---------------------------------------------------------------------------------------------
+    int32 Instance::Private::AddString( int32 ComponentIndex, int32 LODIndex, int32 SurfaceIndex, const FString& Value, FName Name)
+    {
+        // Automatically create the necessary lods and components
+		while (ComponentIndex >= Components.Num())
+		{
+			AddComponent();
+		}
+		while (LODIndex >= Components[ComponentIndex].LODs.Num())
+		{
+			AddLOD(ComponentIndex);
+		}
+		while (SurfaceIndex >= Components[ComponentIndex].LODs[LODIndex].Surfaces.Num())
+		{
+			AddSurface(LODIndex, ComponentIndex);
+		}
+
+		FInstanceSurface& Surface = Components[ComponentIndex].LODs[LODIndex].Surfaces[SurfaceIndex];
+		return Surface.Strings.Add({ Value, Name });
+    }
+
+
+    //---------------------------------------------------------------------------------------------
+	void Instance::Private::AddExtensionData(const Ptr<const class ExtensionData>& Data, FName Name)
+	{
+		check(Data);
+
+		NamedExtensionData& Entry = ExtensionData.AddDefaulted_GetRef();
+		Entry.Data = Data;
+		Entry.Name = Name;
+	}
+}
+
diff --git a/Engine/Plugins/Mutable/Source/MutableRuntime/Private/MuR/InstancePrivate.h b/Engine/Plugins/Mutable/Source/MutableRuntime/Private/MuR/InstancePrivate.h
new file mode 100644
--- /dev/null
+++ b/Engine/Plugins/Mutable/Source/MutableRuntime/Private/MuR/InstancePrivate.h
@@ -0,0 +1,111 @@
+// Copyright Epic Games, Inc. All Rights Reserved.
+
+#pragma once
+
+#include "MuR/Instance.h"
+
+#include "MuR/ImagePrivate.h"
+#include "MuR/MeshPrivate.h"
+
+
+namespace mu
+{
+
+	/** Helper functions to make and read FResourceIDs */
+	inline FResourceID MakeResourceID(uint32 RootAddress, uint32 ParameterBlobIndex)
+	{
+		return (uint64(RootAddress) << 32) | uint64(ParameterBlobIndex);
+	}
+
+	/** */
+    struct FInstanceSurface
+	{
+		FName Name;
+        uint32 InternalId=0;
+        uint32 ExternalId =0;
+        uint32 SharedId =0;
+
+		struct FInstanceImage
+		{
+			FResourceID Id;
+			FName Name;
+		};
+
+		TArray<FInstanceImage, TInlineAllocator<4>> Images;
+
+		struct FInstanceVector
+		{
+			FVector4f Value;
+			FName Name;
+		};
+
+		TArray<FInstanceVector> Vectors;
+
+        struct FInstanceScalar
+        {
+            float Value;
+			FName Name;
+		};
+
+		TArray<FInstanceScalar> Scalars;
+
+        struct FInstanceString
+        {
+			FString Value;
+			FName Name;
+		};
+
+		TArray<FInstanceString> Strings;
+    };
+
+
+    struct FInstanceLOD
+    {
+		FResourceID MeshId;
+		FName MeshName;
+
+		// The order must match the meshes surfaces
+		TArray<FInstanceSurface, TInlineAllocator<4>> Surfaces;
+	};
+
+
+    struct FInstanceComponent
+    {
+		uint16 Id;
+
+		TArray<FInstanceLOD, TInlineAllocator<4>> LODs;
+    };
+
+	struct NamedExtensionData
+	{
+		Ptr<const ExtensionData> Data;
+		FName Name;
+	};
+
+	class Instance::Private
+	{
+	public:
+
+        //!
+        Instance::ID Id = 0;
+
+		//!
+		TArray<FInstanceComponent,TInlineAllocator<4>> Components;
+
+		// Every entry must have a valid ExtensionData and name
+		TArray<NamedExtensionData> ExtensionData;
+
+		int32 AddComponent();
+		int32 AddLOD(int32 ComponentIndex);
+		void SetMesh(int32 ComponentIndex, int32 LODIndex, FResourceID, FName Name);
+		int32 AddSurface(int32 ComponentIndex, int32 LODIndex);
+        void SetSurfaceName( int32 ComponentIndex, int32 LODIndex, int32 SurfaceIndex, FName Name);
+		int32 AddImage( int32 ComponentIndex, int32 LODIndex, int32 SurfaceIndex, FResourceID, FName Name);
+        int32 AddVector( int32 ComponentIndex, int32 LODIndex, int32 SurfaceIndex, const FVector4f&, FName Name);
+        int32 AddScalar( int32 ComponentIndex, int32 LODIndex, int32 SurfaceIndex, float, FName Name);
+        int32 AddString( int32 ComponentIndex, int32 LODIndex, int32 SurfaceIndex, const FString& Value, FName Name);
+		
+		// Data must be non-null
+		void AddExtensionData(const Ptr<const class ExtensionData>& Data, FName Name);
+    };
+}
diff --git a/Engine/Plugins/Mutable/Source/MutableRuntime/Private/MuR/Memory.cpp b/Engine/Plugins/Mutable/Source/MutableRuntime/Private/MuR/Memory.cpp
new file mode 100644
--- /dev/null
+++ b/Engine/Plugins/Mutable/Source/MutableRuntime/Private/MuR/Memory.cpp
@@ -0,0 +1,33 @@
+// Copyright Epic Games, Inc. All Rights Reserved.
+
+#include "MuR/BlockCompression/Miro/Miro.h"
+#include "MuR/MutableMemory.h"
+
+namespace mu
+{
+	static int s_initialized = 0;
+	static int s_finalized = 0;
+
+	void Initialize()
+	{
+		if (!s_initialized)
+		{
+			s_initialized = 1;
+			s_finalized = 0;
+
+			miro::initialize();
+		}
+	}
+
+	void Finalize()
+	{
+		if (s_initialized && !s_finalized)
+		{
+			miro::finalize();
+
+			s_finalized = 1;
+			s_initialized = 0;
+		}
+	}
+
+}
diff --git a/Engine/Plugins/Mutable/Source/MutableRuntime/Private/MuR/Model.cpp b/Engine/Plugins/Mutable/Source/MutableRuntime/Private/MuR/Model.cpp
new file mode 100644
--- /dev/null
+++ b/Engine/Plugins/Mutable/Source/MutableRuntime/Private/MuR/Model.cpp
@@ -0,0 +1,992 @@
+// Copyright Epic Games, Inc. All Rights Reserved.
+
+
+#include "MuR/Model.h"
+
+#include "Containers/Array.h"
+#include "Containers/Map.h"
+#include "HAL/LowLevelMemTracker.h"
+#include "HAL/PlatformMath.h"
+#include "Misc/AssertionMacros.h"
+#include "MuR/Image.h"
+#include "MuR/Mesh.h"
+#include "MuR/ModelPrivate.h"
+#include "MuR/Operations.h"
+#include "MuR/ParametersPrivate.h"
+#include "MuR/Serialisation.h"
+#include "MuR/SerialisationPrivate.h"
+#include "MuR/System.h"
+#include "MuR/Types.h"
+#include "Templates/Tuple.h"
+
+
+namespace mu
+{
+    MUTABLE_IMPLEMENT_POD_SERIALISABLE(FRomData);
+	MUTABLE_IMPLEMENT_POD_SERIALISABLE(FImageLODRange);
+
+	
+    //---------------------------------------------------------------------------------------------
+    void FProgram::Check()
+    {
+    #ifdef MUTABLE_DEBUG
+		// Insert debug checks here.
+    #endif
+    }
+
+
+    //---------------------------------------------------------------------------------------------
+    void FProgram::LogHistogram() const
+    {
+#if 0
+        uint64 countPerType[(int32)OP_TYPE::COUNT];
+        mutable_memset(countPerType,0,sizeof(countPerType));
+
+        for ( const uint32& o: m_opAddress )
+        {
+            OP_TYPE type = GetOpType(o);
+            countPerType[(int32)type]++;
+        }
+
+		TArray< TPair<uint64,OP_TYPE> > sorted((int32)OP_TYPE::COUNT);
+        for (int32 i=0; i<(int32)OP_TYPE::COUNT; ++i)
+        {
+            sorted[i].second = (OP_TYPE)i;
+            sorted[i].first = countPerType[i];
+        }
+
+        std::sort(sorted.begin(),sorted.end(), []( const pair<uint64,OP_TYPE>& a, const pair<uint64,OP_TYPE>& b )
+        {
+            return a.first>b.first;
+        });
+
+        UE_LOG(LogMutableCore,Log, TEXT("Op histogram (%llu ops):"), m_opAddress.Num());
+        for(int32 i=0; i<8; ++i)
+        {
+            float p = sorted[i].first/float(m_opAddress.Num())*100.0f;
+            UE_LOG(LogMutableCore,Log, TEXT("  %3.2f%% : %d"), p, (int32)sorted[i].second );
+        }
+#endif
+    }
+
+	
+	//---------------------------------------------------------------------------------------------
+    void Model::Private::UnloadRoms()
+    {
+	    for (int32 RomIndex = 0; RomIndex < m_program.m_roms.Num(); ++RomIndex)
+	    {
+		    m_program.UnloadRom(RomIndex);
+	    }
+    }
+	
+
+    //---------------------------------------------------------------------------------------------
+    //---------------------------------------------------------------------------------------------
+    //---------------------------------------------------------------------------------------------
+    Model::Model()
+    {
+        m_pD = new Private();
+    }
+
+
+    //---------------------------------------------------------------------------------------------
+    Model::~Model()
+    {
+		MUTABLE_CPUPROFILER_SCOPE(ModelDestructor);
+
+        check( m_pD );
+        delete m_pD;
+        m_pD = 0;
+    }
+
+
+    //---------------------------------------------------------------------------------------------
+    void Model::Serialise( const Model* p, OutputArchive& arch )
+    {
+		LLM_SCOPE_BYNAME(TEXT("MutableRuntime"));
+
+        arch << *p->m_pD;
+    }
+
+    //---------------------------------------------------------------------------------------------
+    //---------------------------------------------------------------------------------------------
+    //---------------------------------------------------------------------------------------------
+    class OutputModelStream : public OutputStream
+    {
+    public:
+
+        //-----------------------------------------------------------------------------------------
+        // Life cycle
+        //-----------------------------------------------------------------------------------------
+        OutputModelStream(ModelWriter* pStreamer )
+            : m_pStreamer( pStreamer )
+        {
+        }
+
+        //-----------------------------------------------------------------------------------------
+        // InputStream interface
+        //-----------------------------------------------------------------------------------------
+        void Write( const void* pData, uint64 size ) override
+        {
+            m_pStreamer->Write( pData, size );
+        }
+
+
+    private:
+
+        ModelWriter* m_pStreamer;
+    };
+
+
+    //---------------------------------------------------------------------------------------------
+    void Model::Serialise( Model* p, ModelWriter& streamer, bool bDropData )
+    {
+		LLM_SCOPE_BYNAME(TEXT("MutableRuntime"));
+
+		mu::FProgram& program = p->m_pD->m_program;
+
+		TArray<TPair<int32, mu::ImagePtrConst>> InitialImages;
+		TArray<TPair<int32, mu::MeshPtrConst>> InitialMeshes;
+		if (!bDropData)
+		{
+			InitialImages = program.ConstantImageLODs;
+			InitialMeshes = program.ConstantMeshes;
+		}
+
+		OutputMemoryStream MemStream(16 * 1024 * 1024);
+
+		// Save images and unload from memory
+		for (int32 ResourceIndex = 0; ResourceIndex < program.ConstantImageLODs.Num(); ++ResourceIndex)
+		{
+			TPair<int32, mu::ImagePtrConst>& ResData = program.ConstantImageLODs[ResourceIndex];
+
+			// This shouldn't have been serialised with rom support before.
+			if (ResData.Key < 0)
+			{
+				continue;
+			}
+
+			int32 RomIndex = ResData.Key;
+			const FRomData& RomData = program.m_roms[RomIndex];
+			check(RomData.ResourceType == DT_IMAGE);
+			check(RomData.ResourceIndex == ResourceIndex);
+
+			// Serialize to memory, to find out final size of this rom
+			MemStream.Reset();
+			OutputArchive MemoryArch(&MemStream);
+			Image::Serialise(ResData.Value.get(), MemoryArch);
+			check(RomData.Size == MemStream.GetBufferSize());
+
+			streamer.OpenWriteFile(RomData.Id);
+			streamer.Write(MemStream.GetBuffer(), MemStream.GetBufferSize());
+			streamer.CloseWriteFile();
+
+			// We clear it to avoid including it also with the main model data
+			ResData.Value = nullptr;
+		}
+
+		// Save meshes and unload from memory
+		for (int32 ResourceIndex = 0; ResourceIndex < program.ConstantMeshes.Num(); ++ResourceIndex)
+		{
+			TPair<int32, mu::MeshPtrConst>& ResData = program.ConstantMeshes[ResourceIndex];
+
+			// This shouldn't have been serialised with rom support before.
+			if (ResData.Key < 0)
+			{
+				continue;
+			}
+
+			int32 RomIndex = ResData.Key;
+			const FRomData& RomData = program.m_roms[RomIndex];
+			check(RomData.ResourceType == DT_MESH);
+			check(RomData.ResourceIndex == ResourceIndex);
+
+			// Serialize to memory, to find out final size of this rom
+			MemStream.Reset();
+			OutputArchive MemoryArch(&MemStream);
+			Mesh::Serialise(ResData.Value.get(), MemoryArch);
+			check(RomData.Size == MemStream.GetBufferSize());
+
+			streamer.OpenWriteFile(RomData.Id);
+			streamer.Write(MemStream.GetBuffer(), MemStream.GetBufferSize());
+			streamer.CloseWriteFile();
+
+			// We clear it to avoid including it also with the main model data
+			ResData.Value = nullptr;
+		}
+
+		// Store the main data of the model
+		{
+			streamer.OpenWriteFile(0);
+			OutputModelStream stream(&streamer);
+			OutputArchive arch(&stream);
+
+			arch << *p->m_pD;
+
+			streamer.CloseWriteFile();
+		}
+
+		// Restore full data
+		if (!bDropData)
+		{
+			program.ConstantImageLODs = InitialImages;
+			program.ConstantMeshes = InitialMeshes;
+		}
+	}
+
+
+    //---------------------------------------------------------------------------------------------
+    bool Model::HasExternalData() const
+    {
+		return m_pD->m_program.m_roms.Num() > 0;
+    }
+
+#if WITH_EDITOR
+	//---------------------------------------------------------------------------------------------
+	bool Model::IsValid() const
+	{
+		return m_pD->m_program.bIsValid;
+	}
+
+
+	//---------------------------------------------------------------------------------------------
+	void Model::Invalidate()
+    {
+		m_pD->m_program.bIsValid = false;
+    }
+#endif
+
+    //---------------------------------------------------------------------------------------------
+    void Model::UnloadExternalData()
+    {
+		LLM_SCOPE_BYNAME(TEXT("MutableRuntime"));
+
+		for (int32 ResIndex = 0; ResIndex < m_pD->m_program.ConstantImageLODs.Num(); ++ResIndex)
+		{
+			if (m_pD->m_program.ConstantImageLODs[ResIndex].Key >= 0)
+			{
+				m_pD->m_program.ConstantImageLODs[ResIndex].Value = nullptr;
+			}
+		}
+
+		for (int32 ResIndex = 0; ResIndex < m_pD->m_program.ConstantMeshes.Num(); ++ResIndex)
+		{
+			if (m_pD->m_program.ConstantMeshes[ResIndex].Key >= 0)
+			{
+				m_pD->m_program.ConstantMeshes[ResIndex].Value = nullptr;
+			}
+		}
+	}
+
+
+    //---------------------------------------------------------------------------------------------
+	TSharedPtr<Model> Model::StaticUnserialise( InputArchive& arch )
+    {
+		LLM_SCOPE_BYNAME(TEXT("MutableRuntime"));
+
+		TSharedPtr<Model> pResult = MakeShared<Model>();
+        arch >> *pResult->m_pD;
+        return pResult;
+    }
+
+    //---------------------------------------------------------------------------------------------
+    Model::Private* Model::GetPrivate() const
+    {
+        return m_pD;
+    }
+
+
+    //---------------------------------------------------------------------------------------------
+    bool Model::GetBoolDefaultValue(int32 Index) const
+    {
+    	check(m_pD->m_program.m_parameters.IsValidIndex(Index));
+		check(m_pD->m_program.m_parameters[Index].m_type == PARAMETER_TYPE::T_BOOL);
+
+        // Early out in case of invalid parameters
+        if (!m_pD->m_program.m_parameters.IsValidIndex(Index) ||
+            m_pD->m_program.m_parameters[Index].m_type != PARAMETER_TYPE::T_BOOL)
+        {
+            return false;
+        }
+		
+        return m_pD->m_program.m_parameters[Index].m_defaultValue.Get<ParamBoolType>();
+    }
+
+
+    //---------------------------------------------------------------------------------------------
+	int32 Model::GetIntDefaultValue(int32 Index) const
+	{
+		check(m_pD->m_program.m_parameters.IsValidIndex(Index));
+		check(m_pD->m_program.m_parameters[Index].m_type == PARAMETER_TYPE::T_INT);
+
+        // Early out in case of invalid parameters
+        if (!m_pD->m_program.m_parameters.IsValidIndex(Index) ||
+            m_pD->m_program.m_parameters[Index].m_type != PARAMETER_TYPE::T_INT)
+        {
+            return 0;
+        }
+		
+        return m_pD->m_program.m_parameters[Index].m_defaultValue.Get<ParamIntType>();
+	}
+
+
+    //---------------------------------------------------------------------------------------------
+	float Model::GetFloatDefaultValue(int32 Index) const
+	{
+    	check(m_pD->m_program.m_parameters.IsValidIndex(Index));
+		check(m_pD->m_program.m_parameters[Index].m_type == PARAMETER_TYPE::T_FLOAT);
+
+        // Early out in case of invalid parameters
+        if (!m_pD->m_program.m_parameters.IsValidIndex(Index) ||
+            m_pD->m_program.m_parameters[Index].m_type != PARAMETER_TYPE::T_FLOAT)
+        {
+            return 0.0f;
+        }
+		
+        return m_pD->m_program.m_parameters[Index].m_defaultValue.Get<ParamFloatType>();
+	}
+
+
+	void Model::GetColourDefaultValue(int32 Index, float* R, float* G, float* B, float* A) const
+    {
+    	check(m_pD->m_program.m_parameters.IsValidIndex(Index));
+		check(m_pD->m_program.m_parameters[Index].m_type == PARAMETER_TYPE::T_COLOUR);
+
+        // Early out in case of invalid parameters
+        if (!m_pD->m_program.m_parameters.IsValidIndex(Index) ||
+            m_pD->m_program.m_parameters[Index].m_type != PARAMETER_TYPE::T_COLOUR)
+        {
+            return;
+        }
+
+        ParamColorType& Color = m_pD->m_program.m_parameters[Index].m_defaultValue.Get<ParamColorType>();
+		
+		if (R)
+		{
+			*R = Color[0];
+		}
+
+		if (G)
+		{
+			*G = Color[1];
+		}
+
+		if (B)
+		{
+			*B = Color[2];
+		}
+
+		if (A)
+		{
+			*A = Color[3];
+		}
+    }
+
+
+	void Model::GetProjectorDefaultValue(int32 Index, PROJECTOR_TYPE* OutProjectionType, FVector3f* OutPos,
+		FVector3f* OutDir, FVector3f* OutUp, FVector3f* OutScale, float* OutProjectionAngle) const
+	{
+    	check(m_pD->m_program.m_parameters.IsValidIndex(Index));
+		check(m_pD->m_program.m_parameters[Index].m_type == PARAMETER_TYPE::T_PROJECTOR);
+
+        // Early out in case of invalid parameters
+        if (!m_pD->m_program.m_parameters.IsValidIndex(Index) ||
+            m_pD->m_program.m_parameters[Index].m_type != PARAMETER_TYPE::T_PROJECTOR)
+        {
+            return;
+        }
+
+        const FProjector& Projector = m_pD->m_program.m_parameters[Index].m_defaultValue.Get<ParamProjectorType>();
+        if (OutProjectionType) *OutProjectionType = Projector.type;
+    	if (OutPos) *OutPos = Projector.position;
+		if (OutDir) *OutDir = Projector.direction;
+    	if (OutUp) *OutUp = Projector.up;
+    	if (OutScale) *OutScale = Projector.scale;
+    	if (OutProjectionAngle) *OutProjectionAngle = Projector.projectionAngle;
+	}
+
+
+	//---------------------------------------------------------------------------------------------
+	FName Model::GetImageDefaultValue(int32 Index) const
+    {
+	    check(m_pD->m_program.m_parameters.IsValidIndex(Index));
+		check(m_pD->m_program.m_parameters[Index].m_type == PARAMETER_TYPE::T_IMAGE);
+
+        // Early out in case of invalid parameters
+        if (!m_pD->m_program.m_parameters.IsValidIndex(Index) ||
+            m_pD->m_program.m_parameters[Index].m_type != PARAMETER_TYPE::T_IMAGE)
+        {
+            return {};
+        }
+		
+        return m_pD->m_program.m_parameters[Index].m_defaultValue.Get<ParamImageType>();
+    }
+
+
+    //---------------------------------------------------------------------------------------------
+    int32 Model::GetRomCount() const
+    {
+    	return m_pD->m_program.m_roms.Num();
+    }
+
+
+    uint32 Model::GetRomId(int32 Index) const
+    {
+    	return m_pD->m_program.m_roms[Index].Id;
+    }
+
+
+    uint32 Model::GetRomSize(int32 Index) const
+    {
+    	return m_pD->m_program.m_roms[Index].Size;
+    }
+
+	ERomFlags Model::GetRomFlags(int32 Index) const
+	{
+		return m_pD->m_program.m_roms[Index].Flags;
+	}
+
+
+    //---------------------------------------------------------------------------------------------
+    ParametersPtr Model::NewParameters(TSharedPtr<const Model> Model, const Parameters* pOld )
+    {
+		LLM_SCOPE_BYNAME(TEXT("MutableRuntime"));
+
+        ParametersPtr pRes = new Parameters();
+
+        pRes->GetPrivate()->m_pModel = Model;
+
+		const FProgram& Program = Model->GetPrivate()->m_program;
+        pRes->GetPrivate()->m_values.SetNum(Program.m_parameters.Num());
+        for ( int32 p=0; p< Program.m_parameters.Num(); ++p )
+        {
+            pRes->GetPrivate()->m_values[p] = Program.m_parameters[p].m_defaultValue;
+        }
+
+        // Copy old values
+        if ( pOld )
+        {
+            for ( int p=0; p<pOld->GetCount(); ++p )
+            {
+                int thisP = pRes->GetPrivate()->Find( pOld->GetName(p) );
+
+                if ( thisP>=0 )
+                {
+                    if ( pOld->GetType(p)==pRes->GetType(thisP) )
+                    {
+                        switch ( pRes->GetType(thisP) )
+                        {
+						case PARAMETER_TYPE::T_BOOL:
+                            pRes->SetBoolValue( thisP, pOld->GetBoolValue(p) );
+                            break;
+
+                        case PARAMETER_TYPE::T_INT:
+                            pRes->SetIntValue( thisP, pOld->GetIntValue(p) );
+                            break;
+
+                        case PARAMETER_TYPE::T_FLOAT:
+                            pRes->SetFloatValue( thisP, pOld->GetFloatValue(p) );
+                            break;
+
+                        case PARAMETER_TYPE::T_COLOUR:
+                        {
+                            float R, G, B, A;
+                            pOld->GetColourValue( p, &R, &G, &B, &A );
+                            pRes->SetColourValue( thisP, R, G, B, A );
+                            break;
+                        }
+
+                        case PARAMETER_TYPE::T_PROJECTOR:
+                        {
+//							float m[16];
+//							pOld->GetProjectorValue( p, m );
+                            pRes->GetPrivate()->m_values[thisP].Set<ParamProjectorType>(pOld->GetPrivate()->m_values[p].Get<ParamProjectorType>());
+                            break;
+                        }
+
+                        case PARAMETER_TYPE::T_IMAGE:
+                            pRes->SetImageValue( thisP, pOld->GetImageValue(p) );
+                            break;
+
+                        default:
+                            check(false);
+                            break;
+                        }
+                    }
+                }
+            }
+        }
+
+        return pRes;
+    }
+
+
+    //---------------------------------------------------------------------------------------------
+    bool Model::IsParameterMultidimensional(const int32 ParamIndex) const
+    {
+		if (m_pD->m_program.m_parameters.IsValidIndex(ParamIndex))
+		{
+			return m_pD->m_program.m_parameters[ParamIndex].m_ranges.Num() > 0;
+		}
+
+		return false;
+    }
+	
+
+    //---------------------------------------------------------------------------------------------
+    int Model::GetStateCount() const
+    {
+        return (int)m_pD->m_program.m_states.Num();
+    }
+
+
+    //---------------------------------------------------------------------------------------------
+    const FString& Model::GetStateName( int32 index ) const
+    {
+        const char* strRes = 0;
+
+        if ( index>=0 && index<(int)m_pD->m_program.m_states.Num() )
+        {
+            return m_pD->m_program.m_states[index].Name;
+        }
+
+		static FString None;
+        return None;
+    }
+
+
+    //---------------------------------------------------------------------------------------------
+    int32 Model::FindState( const FString& Name ) const
+    {
+        int res = -1;
+
+        for ( int i=0; res<0 && i<(int)m_pD->m_program.m_states.Num(); ++i )
+        {
+            if ( m_pD->m_program.m_states[i].Name == Name )
+            {
+                res = i;
+            }
+        }
+
+        return res;
+    }
+
+
+    //---------------------------------------------------------------------------------------------
+    int Model::GetStateParameterCount( int stateIndex ) const
+    {
+        int res = -1;
+
+        if ( stateIndex>=0 && stateIndex<(int)m_pD->m_program.m_states.Num() )
+        {
+            res = (int)m_pD->m_program.m_states[stateIndex].m_runtimeParameters.Num();
+        }
+
+        return res;
+    }
+
+
+    //---------------------------------------------------------------------------------------------
+    int Model::GetStateParameterIndex( int stateIndex, int paramIndex ) const
+    {
+        int res = -1;
+
+        if ( stateIndex>=0 && stateIndex<(int)m_pD->m_program.m_states.Num() )
+        {
+            const FProgram::FState& state = m_pD->m_program.m_states[stateIndex];
+            if ( paramIndex>=0 && paramIndex<(int)state.m_runtimeParameters.Num() )
+            {
+                res = (int)state.m_runtimeParameters[paramIndex];
+            }
+        }
+
+        return res;
+    }
+
+
+    //---------------------------------------------------------------------------------------------
+    //---------------------------------------------------------------------------------------------
+    //---------------------------------------------------------------------------------------------
+    ModelParametersGenerator::ModelParametersGenerator
+        (
+			TSharedPtr<const Model> pModel,
+            System* pSystem
+        )
+    {
+		LLM_SCOPE_BYNAME(TEXT("MutableRuntime"));
+
+        m_pD = new Private();
+        m_pD->m_pModel = pModel;
+        m_pD->m_pSystem = pSystem;
+
+        size_t paramCount = pModel->GetPrivate()->m_program.m_parameters.Num();
+
+		m_pD->m_instanceCount = 1;
+        for (size_t i=0;i<paramCount; ++i)
+        {
+            const FParameterDesc& param =
+                    m_pD->m_pModel->GetPrivate()->m_program.m_parameters[ i ];
+            switch ( param.m_type )
+            {
+
+            case PARAMETER_TYPE::T_INT:
+            {
+                m_pD->m_instanceCount *= param.m_possibleValues.Num();
+                break;
+            }
+
+            case PARAMETER_TYPE::T_BOOL:
+            {
+                m_pD->m_instanceCount *= 2;
+                break;
+            }
+
+            default:
+                // Parameter not discrete, so ignored for combinations.
+                break;
+            }
+        }
+    }
+
+
+    //---------------------------------------------------------------------------------------------
+    ModelParametersGenerator::~ModelParametersGenerator()
+    {
+		LLM_SCOPE_BYNAME(TEXT("MutableRuntime"));
+
+        check( m_pD );
+        delete m_pD;
+        m_pD = 0;
+    }
+
+
+    //---------------------------------------------------------------------------------------------
+    ModelParametersGenerator::Private* ModelParametersGenerator::GetPrivate() const
+    {
+        return m_pD;
+    }
+
+
+    //---------------------------------------------------------------------------------------------
+    int64 ModelParametersGenerator::GetInstanceCount()
+    {
+        return m_pD->m_instanceCount;
+    }
+
+
+    //---------------------------------------------------------------------------------------------
+    Ptr<Parameters> ModelParametersGenerator::GetInstance( int64 index, TFunction<float()> RandomGenerator )
+    {
+		LLM_SCOPE_BYNAME(TEXT("MutableRuntime"));
+
+		Ptr<Parameters> res = Model::NewParameters(m_pD->m_pModel);
+
+		const FProgram& Program = m_pD->m_pModel->GetPrivate()->m_program;
+        int paramCount = Program.m_parameters.Num();
+        int64 currentInstance = index;
+        for (int i=0;i<paramCount;++i)
+        {
+            const FParameterDesc& param = Program.m_parameters[ i ];
+			Ptr<RangeIndex> RangeIndex = res->NewRangeIndex(i);
+
+            switch ( param.m_type )
+            {
+
+            case PARAMETER_TYPE::T_INT:
+            {
+				bool bIsRangeSize = Program.m_ranges.ContainsByPredicate([i](const FRangeDesc& r) { return r.m_dimensionParameter == i; });
+				if (bIsRangeSize)
+				{
+					res->SetIntValue(i, m_pD->DefaultRangeDimension);
+				}
+				else
+				{
+					int numOptions = res->GetIntPossibleValueCount(i);
+					int value = res->GetIntPossibleValue(i, (int)(currentInstance % numOptions));
+					res->SetIntValue(i, value);
+					currentInstance /= numOptions;
+				}
+                break;
+            }
+
+            case PARAMETER_TYPE::T_BOOL:
+            {
+                res->SetBoolValue( i, currentInstance%2!=0 );
+                currentInstance /= 2;
+                break;
+            }
+
+            case PARAMETER_TYPE::T_FLOAT:
+			{
+				bool bIsRangeSize = Program.m_ranges.ContainsByPredicate([i](const FRangeDesc& r) { return r.m_dimensionParameter == i; });
+				if (bIsRangeSize)
+				{
+					if (RangeIndex)
+					{
+						for (int32 Dimensions = 0; Dimensions < RangeIndex->GetRangeCount(); ++Dimensions)
+						{
+							for (int32 RangePosition = 0; RangePosition < m_pD->DefaultRangeDimension; ++RangePosition)
+							{
+								RangeIndex->SetPosition(Dimensions, RangePosition);
+								res->SetFloatValue(i, float(m_pD->DefaultRangeDimension), RangeIndex);
+							}
+						}
+					}
+					else
+					{
+						res->SetFloatValue(i, float(m_pD->DefaultRangeDimension));
+					}
+				}
+				else if (RandomGenerator)
+				{
+					if (RangeIndex)
+					{
+						for (int32 Dimensions = 0; Dimensions < RangeIndex->GetRangeCount(); ++Dimensions)
+						{
+							for (int32 RangePosition = 0; RangePosition < m_pD->DefaultRangeDimension; ++RangePosition)
+							{
+								RangeIndex->SetPosition(Dimensions, RangePosition);
+								res->SetFloatValue(i, RandomGenerator(), RangeIndex);
+							}
+						}
+					}
+					else
+					{
+						res->SetFloatValue(i, RandomGenerator());
+					}
+				}
+				break;
+			}
+
+            case PARAMETER_TYPE::T_COLOUR:
+                if (RandomGenerator)
+                {
+					if (RangeIndex)
+					{
+						for (int32 Dimensions = 0; Dimensions < RangeIndex->GetRangeCount(); ++Dimensions)
+						{
+							for (int32 RangePosition = 0; RangePosition < m_pD->DefaultRangeDimension; ++RangePosition)
+							{
+								RangeIndex->SetPosition(Dimensions, RangePosition);
+								res->SetColourValue(i, RandomGenerator(), RandomGenerator(), RandomGenerator(), RandomGenerator(), RangeIndex);
+							}
+						}
+					}
+					else
+					{
+						res->SetColourValue(i, RandomGenerator(), RandomGenerator(), RandomGenerator(), RandomGenerator());
+					}
+                }
+                break;
+
+			case PARAMETER_TYPE::T_PROJECTOR:
+				if (RandomGenerator)
+				{
+					// For projectors we just warp the position a little bit just to get something different
+					FVector3f Position, Direction, Up, Scale;
+					float Angle = 1.0f;
+					res->GetProjectorValue(i, nullptr, &Position, &Direction, &Up, &Scale, &Angle);
+
+					if (RangeIndex)
+					{
+						for (int32 Dimensions = 0; Dimensions < RangeIndex->GetRangeCount(); ++Dimensions)
+						{
+							for (int32 RangePosition = 0; RangePosition < m_pD->DefaultRangeDimension; ++RangePosition)
+							{
+								RangeIndex->SetPosition(Dimensions, RangePosition);
+								Position.X *= 0.9 + RandomGenerator() * 0.2f;
+								Position.Y *= 0.9 + RandomGenerator() * 0.2f;
+								Position.Z *= 0.9 + RandomGenerator() * 0.2f;
+								res->SetProjectorValue(i, Position, Direction, Up, Scale, Angle, RangeIndex);
+
+							}
+						}
+					}
+					else
+					{
+						Position.X *= 0.9 + RandomGenerator() * 0.2f;
+						Position.Y *= 0.9 + RandomGenerator() * 0.2f;
+						Position.Z *= 0.9 + RandomGenerator() * 0.2f;
+						res->SetProjectorValue(i, Position, Direction, Up, Scale, Angle);
+					}
+				}
+				break;
+
+            default:
+                break;
+            }
+        }
+
+        return res;
+    }
+
+
+    //---------------------------------------------------------------------------------------------
+    ParametersPtr ModelParametersGenerator::GetRandomInstance( TFunctionRef<float()> randomGenerator)
+    {
+		LLM_SCOPE_BYNAME(TEXT("MutableRuntime"));
+
+        ParametersPtr res = Model::NewParameters(m_pD->m_pModel);
+
+		const FProgram& Program = m_pD->m_pModel->GetPrivate()->m_program;
+        int32 ParamCount = Program.m_parameters.Num();
+
+        for (int i=0;i<ParamCount;++i)
+        {
+            const FParameterDesc& param = Program.m_parameters[ i ];
+			Ptr<RangeIndex> RangeIndex = res->NewRangeIndex(i);
+
+            switch ( param.m_type )
+            {
+
+            case PARAMETER_TYPE::T_INT:
+            {
+				bool bIsRangeSize = Program.m_ranges.ContainsByPredicate([i](const FRangeDesc& r) { return r.m_dimensionParameter == i; });
+				if (bIsRangeSize)
+				{
+					if (RangeIndex)
+					{
+						for (int32 Dimensions = 0; Dimensions < RangeIndex->GetRangeCount(); ++Dimensions)
+						{
+							for (int32 RangePosition = 0; RangePosition < m_pD->DefaultRangeDimension; ++RangePosition)
+							{
+								RangeIndex->SetPosition(Dimensions, RangePosition);
+								res->SetIntValue(i, m_pD->DefaultRangeDimension, RangeIndex);
+							}
+						}
+					}
+					else
+					{
+						res->SetIntValue(i, m_pD->DefaultRangeDimension);
+					}
+				}
+				else
+				{
+					if (RangeIndex)
+					{
+						for (int32 Dimensions = 0; Dimensions < RangeIndex->GetRangeCount(); ++Dimensions)
+						{
+							for (int32 RangePosition = 0; RangePosition < m_pD->DefaultRangeDimension; ++RangePosition)
+							{
+								int numOptions = res->GetIntPossibleValueCount(i);
+								int valueIndex = (int)(FMath::Min(numOptions - 1, int(randomGenerator() * numOptions)));
+								int value = res->GetIntPossibleValue(i, valueIndex);
+								res->SetIntValue(i, value, RangeIndex);
+							}
+						}
+					}
+					else
+					{
+						int numOptions = res->GetIntPossibleValueCount(i);
+						int valueIndex = (int)(FMath::Min(numOptions - 1, int(randomGenerator() * numOptions)));
+						int value = res->GetIntPossibleValue(i, valueIndex);
+						res->SetIntValue(i, value);
+					}
+				}
+                break;
+            }
+
+            case PARAMETER_TYPE::T_BOOL:
+            {
+                res->SetBoolValue( i, randomGenerator()>0.5f );
+                break;
+            }
+
+            case PARAMETER_TYPE::T_FLOAT:
+			{
+				bool bIsRangeSize = Program.m_ranges.ContainsByPredicate([i](const FRangeDesc& r) { return r.m_dimensionParameter == i; });
+				if (bIsRangeSize)
+				{
+					if (RangeIndex)
+					{
+						for (int32 Dimensions = 0; Dimensions < RangeIndex->GetRangeCount(); ++Dimensions)
+						{
+							for (int32 RangePosition = 0; RangePosition < m_pD->DefaultRangeDimension; ++RangePosition)
+							{
+								res->SetFloatValue(i, float(m_pD->DefaultRangeDimension), RangeIndex);
+							}
+						}
+					}
+					else
+					{
+						res->SetFloatValue(i, float(m_pD->DefaultRangeDimension));
+					}
+				}
+				else
+				{
+					if (RangeIndex)
+					{
+						for (int32 Dimensions = 0; Dimensions < RangeIndex->GetRangeCount(); ++Dimensions)
+						{
+							for (int32 RangePosition = 0; RangePosition < m_pD->DefaultRangeDimension; ++RangePosition)
+							{
+								res->SetFloatValue(i, randomGenerator(), RangeIndex);
+							}
+						}
+					}
+					else
+					{
+						res->SetFloatValue(i, randomGenerator());
+					}
+				}
+				break;
+			}
+
+            case PARAMETER_TYPE::T_COLOUR:
+				if (RangeIndex)
+				{
+					for (int32 Dimensions = 0; Dimensions < RangeIndex->GetRangeCount(); ++Dimensions)
+					{
+						for (int32 RangePosition = 0; RangePosition < m_pD->DefaultRangeDimension; ++RangePosition)
+						{
+							res->SetColourValue(i, randomGenerator(), randomGenerator(), randomGenerator(), randomGenerator(), 
+												RangeIndex);
+						}
+					}
+				}
+				else
+				{
+					res->SetColourValue(i, randomGenerator(), randomGenerator(), randomGenerator(), randomGenerator());
+				}
+                break;
+
+			case PARAMETER_TYPE::T_PROJECTOR:
+			{
+				// For projectors we just warp the position a little bit just to get something different
+				FVector3f Position, Direction, Up, Scale;
+				float Angle = 1.0f;
+				res->GetProjectorValue(i, nullptr, &Position, &Direction, &Up, &Scale, &Angle);
+				if (RangeIndex)
+				{
+					for (int32 Dimensions = 0; Dimensions < RangeIndex->GetRangeCount(); ++Dimensions)
+					{
+						for (int32 RangePosition = 0; RangePosition < m_pD->DefaultRangeDimension; ++RangePosition)
+						{
+							Position.X *= 0.9 + randomGenerator() * 0.2f;
+							Position.Y *= 0.9 + randomGenerator() * 0.2f;
+							Position.Z *= 0.9 + randomGenerator() * 0.2f;
+							res->SetProjectorValue(i, Position, Direction, Up, Scale, Angle, RangeIndex);
+						}
+					}
+				}
+				else
+				{
+					Position.X *= 0.9 + randomGenerator() * 0.2f;
+					Position.Y *= 0.9 + randomGenerator() * 0.2f;
+					Position.Z *= 0.9 + randomGenerator() * 0.2f;
+					res->SetProjectorValue(i, Position, Direction, Up, Scale, Angle);
+				}
+				break;
+			}
+
+            default:
+                break;
+            }
+        }
+
+        return res;
+    }
+
+
+}
+
diff --git a/Engine/Plugins/Mutable/Source/MutableRuntime/Private/MuR/MutableRuntimeModule.cpp b/Engine/Plugins/Mutable/Source/MutableRuntime/Private/MuR/MutableRuntimeModule.cpp
new file mode 100644
--- /dev/null
+++ b/Engine/Plugins/Mutable/Source/MutableRuntime/Private/MuR/MutableRuntimeModule.cpp
@@ -0,0 +1,28 @@
+// Copyright Epic Games, Inc. All Rights Reserved.
+
+#include "MuR/MutableRuntimeModule.h"
+
+#include "HAL/UnrealMemory.h"
+#include "Logging/LogMacros.h"
+#include "Misc/AssertionMacros.h"
+#include "Modules/ModuleManager.h"
+#include "MuR/MutableMemory.h"
+#include "MuR/Platform.h"
+
+IMPLEMENT_MODULE(FMutableRuntimeModule, MutableRuntime);
+
+DEFINE_LOG_CATEGORY(LogMutableCore);
+
+
+void FMutableRuntimeModule::StartupModule()
+{
+	mu::Initialize();
+}
+
+
+void FMutableRuntimeModule::ShutdownModule()
+{
+	// Finalize the mutable runtime in this module
+	mu::Finalize();
+}
+
diff --git a/Engine/Plugins/Mutable/Source/MutableRuntime/Private/MuR/MutableString.cpp b/Engine/Plugins/Mutable/Source/MutableRuntime/Private/MuR/MutableString.cpp
new file mode 100644
--- /dev/null
+++ b/Engine/Plugins/Mutable/Source/MutableRuntime/Private/MuR/MutableString.cpp
@@ -0,0 +1,54 @@
+// Copyright Epic Games, Inc. All Rights Reserved.
+
+
+#include "MuR/MutableString.h"
+
+
+namespace mu {
+
+	//---------------------------------------------------------------------------------------------
+	String::String( const FString& InValue)
+	{
+        Value = InValue;
+	}
+
+
+	////---------------------------------------------------------------------------------------------
+	//void String::Serialise( const String* p, OutputArchive& arch )
+	//{
+	//	arch << *p;
+	//}
+
+
+	////---------------------------------------------------------------------------------------------
+	//Ptr<String> String::StaticUnserialise( InputArchive& arch )
+	//{
+	//	Ptr<String> pResult = new String();
+	//	arch >> *pResult;
+	//	return pResult;
+	//}
+
+
+	//---------------------------------------------------------------------------------------------
+	Ptr<String> String::Clone() const
+	{
+		Ptr<String> pResult = new String(Value);
+		return pResult;
+	}
+
+
+	//---------------------------------------------------------------------------------------------
+	int32 String::GetDataSize() const
+	{
+		return sizeof(String) + Value.GetAllocatedSize();
+	}
+
+
+	//---------------------------------------------------------------------------------------------
+	const FString& String::GetValue() const
+	{
+        return Value;
+	}
+
+}
+
diff --git a/Engine/Plugins/Mutable/Source/MutableRuntime/Private/MuR/OpImageApplyComposite.h b/Engine/Plugins/Mutable/Source/MutableRuntime/Private/MuR/OpImageApplyComposite.h
new file mode 100644
--- /dev/null
+++ b/Engine/Plugins/Mutable/Source/MutableRuntime/Private/MuR/OpImageApplyComposite.h
@@ -0,0 +1,82 @@
+// Copyright Epic Games, Inc. All Rights Reserved.
+
+#pragma once
+
+#include "Math/UnrealMathUtility.h"
+#include "Math/Vector.h"
+
+#include "MuR/ImagePrivate.h"
+#include "MuR/OpImageLayer.h"
+
+
+namespace mu
+{
+	// Combine layer oprations get packed color.
+	inline unsigned NormalComposite(
+			const uint32 Base, 
+			const uint32 Normal, 
+			const uint8 Channel, 
+			const float Power )
+	{
+		const int32 ChannelShift = Channel << 3;
+
+		float ReferencValue = static_cast<float>(Base >> ChannelShift & 0xFF) / 255.0f;
+
+		// Rougness is stored as Unorm8
+		float Roughness = static_cast<float>( (Base >> ChannelShift) & 0xFF ) / 255.0f;
+
+		FVector3f N = ( FVector3f( static_cast<float>((Normal >> 0 ) & 0xFF) / 255.0f,
+							       static_cast<float>((Normal >> 8 ) & 0xFF) / 255.0f,
+							       static_cast<float>((Normal >> 16) & 0xFF) / 255.0f ) * 2.0f - 1.0f );	
+	
+		// See TexturecCompressorModule.cpp::ApplyCompositeTexture
+		// Toksvig estimation of variance
+		float LengthN = FMath::Min( N.Size(), 1.0f );
+		float Variance = ( 1.0f - LengthN ) / LengthN;
+		Variance = FMath::Max( 0.0f, Variance - 0.00004f );
+
+		Variance *= Power;
+		
+		float a = Roughness * Roughness;
+		float a2 = a * a;
+		float B = 2.0f * Variance * (a2 - 1.0f);
+		a2 = ( B - a2 ) / ( B - 1.0f );
+		Roughness = FMath::Pow( a2, 0.25f );
+
+		const uint32 ResMask = ~(0xFF << ChannelShift);
+		const uint32 Value = FMath::Clamp<uint32>( static_cast<uint32>(Roughness * 255.0f), 0u, 255u );
+
+		return  (Value << ChannelShift) | (Base & ResMask); 
+	}
+
+	inline void ImageNormalComposite(Image* pResult, const Image* pBase, const Image* pNormal, ECompositeImageMode mode, float power)
+	{
+		if (mode == ECompositeImageMode::CIM_Disabled)
+		{
+			// Copy base to result. Using functor layer combine with identity to do the copy for simplicity.
+			// This could be optimized but will never happen since it is checked at compile time. 	
+			ImageLayerCombineFunctor( pResult, pBase, pNormal, 
+				[](uint32 base, uint32) -> uint32 { return base; });
+
+			return;
+		}	
+
+		const int8 channel = [mode]() 
+			{
+				switch (mode)
+				{
+					case ECompositeImageMode::CIM_NormalRoughnessToRed   : return 0;
+					case ECompositeImageMode::CIM_NormalRoughnessToGreen : return 1;
+					case ECompositeImageMode::CIM_NormalRoughnessToBlue  : return 2;
+					case ECompositeImageMode::CIM_NormalRoughnessToAlpha : return 3;
+				}
+			
+				check(false);
+				return 0;		
+			}();
+
+		ImageLayerCombineFunctor( pResult, pBase, pNormal,
+				[channel, power](uint32 base, uint32 normal) -> uint32 
+					{ return NormalComposite(base, normal, channel, power); });
+	}
+}
diff --git a/Engine/Plugins/Mutable/Source/MutableRuntime/Private/MuR/OpImageColourMap.h b/Engine/Plugins/Mutable/Source/MutableRuntime/Private/MuR/OpImageColourMap.h
new file mode 100644
--- /dev/null
+++ b/Engine/Plugins/Mutable/Source/MutableRuntime/Private/MuR/OpImageColourMap.h
@@ -0,0 +1,155 @@
+// Copyright Epic Games, Inc. All Rights Reserved.
+
+#pragma once
+
+#include "MuR/ImagePrivate.h"
+#include "MuR/Platform.h"
+
+#include "Async/ParallelFor.h"
+
+namespace mu
+{
+
+	inline void ImageColourMap(Image* DestImage, const Image* SourceImage, const Image* MaskImage, const Image* MapImage, bool bOnlyOneMip)
+	{
+		check(SourceImage->GetSizeX() == MaskImage->GetSizeX());
+		check(SourceImage->GetSizeY() == MaskImage->GetSizeY());
+		check(SourceImage->GetLODCount() == MaskImage->GetLODCount() || bOnlyOneMip);
+		check(SourceImage->GetFormat() == DestImage->GetFormat());
+		check(SourceImage->GetSize() == DestImage->GetSize());
+		check(MaskImage->GetFormat() == EImageFormat::IF_L_UBYTE);
+
+		// Generic implementation
+
+		// Make a palette for faster conversion
+        uint8 Palette[256][4];
+		for (int32 I = 0; I < 256; ++I)
+		{
+			FVector4f Color = MapImage->Sample(FVector2f(float(I)/255.0f, 0.0f));
+            Palette[I][0] = (uint8)FMath::Max(0, FMath::Min(255, int32(Color[0]*255.0f)));
+            Palette[I][1] = (uint8)FMath::Max(0, FMath::Min(255, int32(Color[1]*255.0f)));
+            Palette[I][2] = (uint8)FMath::Max(0, FMath::Min(255, int32(Color[2]*255.0f)));
+            Palette[I][3] = (uint8)FMath::Max(0, FMath::Min(255, int32(Color[3]*255.0f)));
+		}
+
+		const int32 LODBegin = 0;
+		const int32 LODEnd = bOnlyOneMip ? 1 : DestImage->GetLODCount();
+	
+		constexpr int32 NumBatchElems = 1 << 14;
+		const int32 BytesPerElem = GetImageFormatData(DestImage->GetFormat()).BytesPerBlock;
+		check(GetImageFormatData(SourceImage->GetFormat()).BytesPerBlock == BytesPerElem);
+
+		const int32 NumBatches = DestImage->DataStorage.GetNumBatchesLODRange(NumBatchElems, BytesPerElem, LODBegin, LODEnd); 
+		check(NumBatches == SourceImage->DataStorage.GetNumBatchesLODRange(NumBatchElems, BytesPerElem, LODBegin, LODEnd));
+
+		auto ProcessBatch = [DestImage, SourceImage, MaskImage, Palette, NumBatchElems, BytesPerElem, LODBegin, LODEnd](int32 BatchId)
+		{
+			TArrayView<uint8> DestView = DestImage->DataStorage.GetBatchLODRange(BatchId, NumBatchElems, BytesPerElem, LODBegin, LODEnd); 
+			TArrayView<const uint8> SourceView = SourceImage->DataStorage.GetBatchLODRange(BatchId, NumBatchElems, BytesPerElem, LODBegin, LODEnd); 
+			TArrayView<const uint8> MaskView = MaskImage->DataStorage.GetBatchLODRange(BatchId, NumBatchElems, 1, LODBegin, LODEnd); 
+
+			const int32 NumElems = DestView.Num() / BytesPerElem;
+			check(SourceView.Num() / BytesPerElem == NumElems);
+			check(MaskView.Num() / 1 == NumElems);
+
+			uint8* DestBuf = DestView.GetData();
+			const uint8* SourceBuf = SourceView.GetData();
+			const uint8* MaskBuf = MaskView.GetData();
+
+			switch (SourceImage->GetFormat())
+			{
+			case EImageFormat::IF_L_UBYTE:
+			{
+				for (int32 I = 0; I < NumElems; ++I)
+				{
+					if (MaskBuf[I] > 127)
+					{
+						DestBuf[I] = Palette[SourceBuf[I]][0];
+					}
+					else
+					{
+						DestBuf[I] = SourceBuf[I];
+					}
+				}
+				break;
+			}
+			case EImageFormat::IF_RGB_UBYTE:
+			{
+				for (int32 I = 0; I < NumElems; ++I)
+				{
+					if (MaskBuf[I])
+					{
+						DestBuf[3*I + 0] = Palette[SourceBuf[3*I + 0]][0];
+						DestBuf[3*I + 1] = Palette[SourceBuf[3*I + 1]][1];
+						DestBuf[3*I + 2] = Palette[SourceBuf[3*I + 2]][2];
+					}
+					else
+					{
+						DestBuf[3*I + 0] = SourceBuf[3*I + 0];
+						DestBuf[3*I + 1] = SourceBuf[3*I + 1];
+						DestBuf[3*I + 2] = SourceBuf[3*I + 2];
+					}
+				}
+				break;
+			}
+
+			case EImageFormat::IF_RGBA_UBYTE:
+			{
+				for (int32 I = 0; I < NumElems; ++I)
+				{
+					if (MaskBuf[I] > 127)
+					{
+						DestBuf[4*I + 0] = Palette[SourceBuf[4*I + 0]][0];
+						DestBuf[4*I + 1] = Palette[SourceBuf[4*I + 1]][1];
+						DestBuf[4*I + 2] = Palette[SourceBuf[4*I + 2]][2];
+						DestBuf[4*I + 3] = Palette[SourceBuf[4*I + 3]][3];
+					}
+					else
+					{
+						DestBuf[4*I + 0] = SourceBuf[4*I + 0];
+						DestBuf[4*I + 1] = SourceBuf[4*I + 1];
+						DestBuf[4*I + 2] = SourceBuf[4*I + 2];
+						DestBuf[4*I + 3] = SourceBuf[4*I + 3];
+					}
+				}
+				break;
+			}
+
+			case EImageFormat::IF_BGRA_UBYTE:
+			{
+				for (int32 I = 0; I < NumElems; ++I)
+				{
+					if (MaskBuf[I] > 127)
+					{
+						DestBuf[4*I + 0] = Palette[SourceBuf[4*I + 0]][2];
+						DestBuf[4*I + 1] = Palette[SourceBuf[4*I + 1]][1];
+						DestBuf[4*I + 2] = Palette[SourceBuf[4*I + 2]][0];
+						DestBuf[4*I + 3] = Palette[SourceBuf[4*I + 3]][3];
+					}
+					else
+					{
+						DestBuf[4*I + 0] = SourceBuf[4*I + 0];
+						DestBuf[4*I + 1] = SourceBuf[4*I + 1];
+						DestBuf[4*I + 2] = SourceBuf[4*I + 2];
+						DestBuf[4*I + 3] = SourceBuf[4*I + 3];
+					}
+				}
+				break;
+			}
+
+			default:
+				check(false);
+			}
+		};
+
+		if (NumBatches == 1)
+		{
+			ProcessBatch(0);
+		}
+		else if (NumBatches > 1)
+		{
+			ParallelFor(NumBatches, ProcessBatch);
+		}
+	}
+
+}
diff --git a/Engine/Plugins/Mutable/Source/MutableRuntime/Private/MuR/OpImageGradient.h b/Engine/Plugins/Mutable/Source/MutableRuntime/Private/MuR/OpImageGradient.h
new file mode 100644
--- /dev/null
+++ b/Engine/Plugins/Mutable/Source/MutableRuntime/Private/MuR/OpImageGradient.h
@@ -0,0 +1,41 @@
+// Copyright Epic Games, Inc. All Rights Reserved.
+
+#pragma once
+
+#include "MuR/ImagePrivate.h"
+#include "MuR/Platform.h"
+
+
+namespace mu
+{
+
+	inline void ImageGradient( Image* pDest, const FVector4f& c0, const FVector4f& c1 )
+	{
+		check(pDest && pDest->GetFormat()==EImageFormat::IF_RGBA_UBYTE);
+
+        uint8* pDestBuf = pDest->GetLODData(0);
+
+		int32 sizeX = pDest->GetSizeX();
+		int32 sizeY = pDest->GetSizeX();
+
+		for ( int i=0; i< sizeX; ++i )
+		{
+			float delta = float(i)/float(sizeX -1);
+			FVector3f c = c0 * (1.0f-delta) + c1 * delta;
+
+            uint8 colour[4];
+            colour[0] = (uint8)FMath::Max( 0, FMath::Min( 255, int(c[0]*255.0f) ) );
+            colour[1] = (uint8)FMath::Max( 0, FMath::Min( 255, int(c[1]*255.0f) ) );
+            colour[2] = (uint8)FMath::Max( 0, FMath::Min( 255, int(c[2]*255.0f) ) );
+			colour[3] = 255;
+
+			for ( int j=0; j< sizeY; ++j )
+			{
+				pDestBuf[ (sizeX*j+i)*3 + 0 ] = colour[0];
+				pDestBuf[ (sizeX*j+i)*3 + 1 ] = colour[1];
+				pDestBuf[ (sizeX*j+i)*3 + 2 ] = colour[2];
+			}
+		}
+	}
+
+}
