[mutable] Moved the Mutable plugin out of Experimental status into Beta.

diff --git a/Engine/Plugins/Mutable/Source/MutableRuntime/Private/MuR/OpImageLayer.h b/Engine/Plugins/Mutable/Source/MutableRuntime/Private/MuR/OpImageLayer.h
new file mode 100644
--- /dev/null
+++ b/Engine/Plugins/Mutable/Source/MutableRuntime/Private/MuR/OpImageLayer.h
@@ -0,0 +1,2561 @@
+// Copyright Epic Games, Inc. All Rights Reserved.
+
+#pragma once
+
+#include "MuR/ImagePrivate.h"
+#include "MuR/ImageRLE.h"
+#include "Async/ParallelFor.h"
+#include "Templates/UnrealTemplate.h"
+
+#include "MuR/OpImageBlend.h"
+#include "MuR/MutableRuntimeModule.h"
+
+namespace mu
+{
+	inline bool IsAnyComponentLargerThan1(FVector4f Value)
+	{
+		return (Value[0] > 1) | (Value[1] > 1) | (Value[2] > 1) | (Value[3] > 1);
+	}
+	
+	/** Apply a blending function to an image with a colour source.
+	* It only affects the RGB or L channels, leaving alpha untouched.
+	*/
+	namespace Private
+	{
+		template<uint32 NumChannels, uint32 (*BLEND_FUNC)(uint32, uint32), bool bClamp>
+		FORCENOINLINE void BufferLayerColourGenericChannel(uint8* DestBuf, const uint8* BaseBuf, int32 NumElems, const FIntVector& Color)
+		{
+			static_assert(NumChannels > 0 && NumChannels <= 4);
+
+			for (int32 I = 0; I < NumElems; ++I)
+			{
+				for (uint32 C = 0; C < NumChannels; ++C)
+				{
+					uint32 Base = BaseBuf[NumChannels * I + C];
+					uint32 Result = BLEND_FUNC(Base, Color[C]);
+					if constexpr (bClamp)
+					{
+						DestBuf[NumChannels * I + C] = (uint8)FMath::Min(255u, Result);
+					}
+					else
+					{
+						DestBuf[NumChannels * I + C] = (uint8)Result;
+					}
+				}
+			}
+		}
+	}
+
+	template< unsigned (*BLEND_FUNC)(unsigned,unsigned), bool CLAMP >
+    inline void BufferLayerColour(Image* ResultImage, const Image* BaseImage, FVector4f Color)
+	{
+        check(ResultImage->GetFormat() == BaseImage->GetFormat());
+        check(ResultImage->GetSizeX() == BaseImage->GetSizeX());
+        check(ResultImage->GetSizeY() == BaseImage->GetSizeY());
+        check(ResultImage->GetLODCount() == BaseImage->GetLODCount());
+
+		const EImageFormat BaseFormat = BaseImage->GetFormat();
+		// Generic implementation
+		constexpr int32 NumBatchElems = 4096*4;
+		const int32 BytesPerElem = GetImageFormatData(BaseFormat).BytesPerBlock;	
+
+		const int32 NumBatches = BaseImage->DataStorage.GetNumBatches(NumBatchElems, BytesPerElem);
+
+		const FIntVector ColorValue = FIntVector(Color.X * 255.0f, Color.Y * 255.0f, Color.Z * 255.0f);
+		auto ProcessBatch = [NumBatchElems, BytesPerElem, ColorValue, BaseImage, ResultImage, BaseFormat](int32 BatchId)
+		{
+			TArrayView<const uint8> BaseView = BaseImage->DataStorage.GetBatch(BatchId, NumBatchElems, BytesPerElem);
+			TArrayView<uint8> ResultView = ResultImage->DataStorage.GetBatch(BatchId, NumBatchElems, BytesPerElem);
+			
+			check(BaseView.Num() == ResultView.Num());
+
+			const int32 NumElems = BaseView.Num() / BytesPerElem;
+
+			switch (BaseFormat)
+			{
+			case EImageFormat::IF_L_UBYTE:
+			{
+				check(BytesPerElem == 1);
+				Private::BufferLayerColourGenericChannel<1, BLEND_FUNC, CLAMP>(ResultView.GetData(), BaseView.GetData(), NumElems, ColorValue);
+				break;
+			}
+			case EImageFormat::IF_RGB_UBYTE:
+			{
+				check(BytesPerElem == 3);
+				Private::BufferLayerColourGenericChannel<3, BLEND_FUNC, CLAMP>(ResultView.GetData(), BaseView.GetData(), NumElems, ColorValue);
+				break;
+			}
+			case EImageFormat::IF_RGBA_UBYTE:
+			{
+				check(BytesPerElem == 4);
+				Private::BufferLayerColourGenericChannel<4, BLEND_FUNC, CLAMP>(ResultView.GetData(), BaseView.GetData(), NumElems, ColorValue);
+				break;
+			}
+			case EImageFormat::IF_BGRA_UBYTE:
+			{
+				check(BytesPerElem == 4);
+				FIntVector BGRAColorValue = FIntVector(ColorValue.Z, ColorValue.Y, ColorValue.X);
+				Private::BufferLayerColourGenericChannel<4, BLEND_FUNC, CLAMP>(ResultView.GetData(), BaseView.GetData(), NumElems, BGRAColorValue);
+				break;
+			}
+			default:
+			{
+				checkf(false, TEXT("Unsupported format."));
+				break;
+			}
+			}
+		};
+
+		if (NumBatches == 1)
+		{
+			ProcessBatch(0);
+		}
+		else if (NumBatches > 1)
+		{
+			ParallelFor(NumBatches, ProcessBatch);
+		}
+    }
+
+
+	template< uint32 (*BLEND_FUNC)(uint32, uint32) >
+	inline void BufferLayerColour(Image* ResultImage, const Image* BaseImage, FVector4f Color)
+	{
+		bool bIsClampNeeded = IsAnyComponentLargerThan1(Color);
+		if (bIsClampNeeded)
+		{
+			BufferLayerColour<BLEND_FUNC, true>(ResultImage, BaseImage, Color);
+		}
+		else
+		{
+			BufferLayerColour<BLEND_FUNC, false>(ResultImage, BaseImage, Color);
+		}
+	}
+
+	namespace Private
+	{
+		template<uint32 NumChannels, uint32 (*BLEND_FUNC)(uint32,uint32), bool Clamp>
+		FORCENOINLINE void BufferLayerColourFromAlphaGenericChannel(uint8* DestBuf, const uint8* BaseBuf, int32 NumElems, const FIntVector& Color)
+		{
+			static_assert(NumChannels > 0 && NumChannels <= 4);
+
+			for (int32 I = 0; I < NumElems; ++I)
+			{
+				const uint32 Alpha = Invoke([&]() -> uint32
+				{
+					if constexpr (NumChannels <= 3)
+					{
+						return 255;
+					}
+					else
+					{
+						return BaseBuf[NumChannels * I + 3];
+					}
+				});
+
+				for (uint32 C = 0; C < NumChannels; ++C)
+				{
+					uint32 Result = BLEND_FUNC(Alpha, Color[C]);
+					if constexpr (Clamp)
+					{
+						DestBuf[NumChannels * I + C] = (uint8)FMath::Min(255u, Result);
+					}
+					else
+					{
+						DestBuf[NumChannels * I + C] = (uint8)Result;
+					}
+				}
+			}
+		}
+	}
+
+	template<
+		uint32 (*BLEND_FUNC)(uint32, uint32), 
+		bool CLAMP>
+    inline void BufferLayerColourFromAlpha(Image* ResultImage, const Image* BaseImage, FVector4f Color)
+	{
+        check(ResultImage->GetFormat() == BaseImage->GetFormat());
+        check(ResultImage->GetSizeX() == BaseImage->GetSizeX());
+        check(ResultImage->GetSizeY() == BaseImage->GetSizeY());
+        check(ResultImage->GetLODCount() == BaseImage->GetLODCount());
+
+		const EImageFormat BaseFormat = BaseImage->GetFormat();
+		// Generic implementation
+		constexpr int32 NumBatchElems = 4096*2;
+		const int32 BytesPerElem = GetImageFormatData(BaseFormat).BytesPerBlock;	
+
+		const int32 NumBatches = BaseImage->DataStorage.GetNumBatches(NumBatchElems, BytesPerElem);
+		check(NumBatches == ResultImage->DataStorage.GetNumBatches(NumBatchElems, BytesPerElem));
+
+		const FIntVector ColorValue = FIntVector(Color.X * 255.0f, Color.Y * 255.0f, Color.Z * 255.0f);
+		auto ProcessBatch = [NumBatchElems, BytesPerElem, ColorValue, BaseImage, ResultImage, BaseFormat](int32 BatchId)
+		{
+			TArrayView<const uint8> BaseView = BaseImage->DataStorage.GetBatch(BatchId, NumBatchElems, BytesPerElem);
+			TArrayView<uint8> ResultView = ResultImage->DataStorage.GetBatch(BatchId, NumBatchElems, BytesPerElem);
+			
+			check(BaseView.Num() == ResultView.Num());
+
+			const int32 NumElems = ResultView.Num() / BytesPerElem;
+			switch (BaseFormat)
+			{
+			case EImageFormat::IF_L_UBYTE:
+			{
+				check(BytesPerElem == 1);
+				Private::BufferLayerColourFromAlphaGenericChannel<1, BLEND_FUNC, CLAMP>(ResultView.GetData(), BaseView.GetData(), NumElems, ColorValue);
+				break;
+			}
+			case EImageFormat::IF_RGB_UBYTE:
+			{
+				check(BytesPerElem == 3);
+				Private::BufferLayerColourFromAlphaGenericChannel<3, BLEND_FUNC, CLAMP>(ResultView.GetData(), BaseView.GetData(), NumElems, ColorValue);
+				break;
+			}
+			case EImageFormat::IF_RGBA_UBYTE:
+			{
+				check(BytesPerElem == 4);
+				Private::BufferLayerColourFromAlphaGenericChannel<4, BLEND_FUNC, CLAMP>(ResultView.GetData(), BaseView.GetData(), NumElems, ColorValue);
+				break;
+			}
+			case EImageFormat::IF_BGRA_UBYTE:
+			{
+				check(BytesPerElem == 4);
+				FIntVector BGRAColorValues = FIntVector(ColorValue.Z, ColorValue.Y, ColorValue.X);
+				Private::BufferLayerColourFromAlphaGenericChannel<4, BLEND_FUNC, CLAMP>(ResultView.GetData(), BaseView.GetData(), NumElems , ColorValue);
+				break;
+			}
+			default:
+			{
+				checkf(false, TEXT("Unsupported format."));
+				break;
+			}
+			}
+		};
+
+		if (NumBatches == 1)
+		{
+			ProcessBatch(0);
+		}
+		else if (NumBatches > 1)
+		{
+			ParallelFor(NumBatches, ProcessBatch);
+		}
+    }
+
+
+	template<uint32 (*BLEND_FUNC)(uint32, uint32)>
+	inline void BufferLayerColourFromAlpha(Image* ResultImage, const Image* BaseImage, FVector4f Color)
+	{
+		bool bIsClampNeeded = IsAnyComponentLargerThan1(Color);
+		if (bIsClampNeeded)
+		{
+			BufferLayerColourFromAlpha<BLEND_FUNC, true>(ResultImage, BaseImage, Color);
+		}
+		else
+		{
+			BufferLayerColourFromAlpha<BLEND_FUNC, false>(ResultImage, BaseImage, Color);
+		}
+	}
+
+
+	template< 
+		uint32 (*BLEND_FUNC_MASKED)(uint32, uint32, uint32),
+		uint32 (*BLEND_FUNC)(uint32, uint32),
+		bool CLAMP,
+		// Number of total channels to actually process
+		uint32 CHANNELS_TO_BLEND,
+		// Number of total channels in the base image
+		int32 BASE_CHANNEL_STRIDE>
+    inline void BufferLayerColourFormat(Image* DestImage, const Image* BaseImage, const Image* MaskImage, FVector4f Col, uint32 BaseChannelOffset, uint8 ColorChannelOffset, bool bOnlyFirstLOD)
+	{
+		check(CHANNELS_TO_BLEND + BaseChannelOffset <= BASE_CHANNEL_STRIDE);
+		check(DestImage->GetLODCount() <= BaseImage->GetLODCount());
+
+		FUintVector4 TopColor = FUintVector4(Col.X * 255.0f, Col.Y * 255.0f, Col.Z * 255.0f, Col.W * 255);
+
+		const EImageFormat BaseFormat = BaseImage->GetFormat();
+		const EImageFormat MaskFormat = MaskImage->GetFormat();
+
+		int32 NumLODs = bOnlyFirstLOD ? 1 : BaseImage->GetLODCount();
+
+        const bool bIsMaskUncompressed = (MaskFormat == EImageFormat::IF_L_UBYTE);
+
+		constexpr uint32 NumColorChannels = FMath::Min(CHANNELS_TO_BLEND, 3u);
+		check(NumColorChannels + ColorChannelOffset < 4);
+
+        if (bIsMaskUncompressed)
+        {
+			constexpr int32 NumBatchElems = 4096 * 2;
+			const int32 BytesPerElem = GetImageFormatData(BaseFormat).BytesPerBlock;
+			check(GetImageFormatData(MaskFormat).BytesPerBlock == 1);
+			check(GetImageFormatData(DestImage->GetFormat()).BytesPerBlock == BytesPerElem);
+
+			const int32 NumBatches = 
+					BaseImage->DataStorage.GetNumBatchesLODRange(NumBatchElems, BytesPerElem, 0, NumLODs);
+
+			check(NumBatches == DestImage->DataStorage.GetNumBatchesLODRange(NumBatchElems, BytesPerElem, 0, NumLODs));
+			check(NumBatches == MaskImage->DataStorage.GetNumBatchesLODRange(NumBatchElems, 1, 0, NumLODs));
+
+			auto ProcessBatch = [
+				BaseImage, MaskImage, DestImage, 
+				BytesPerElem, NumBatchElems, 
+				BaseChannelOffset, TopColor, ColorChannelOffset, NumColorChannels, 
+				NumLODs
+			](int32 BatchId)
+			{
+				TArrayView<const uint8> BaseBatchView = 
+						BaseImage->DataStorage.GetBatchLODRange(BatchId, NumBatchElems, BytesPerElem, 0, NumLODs);
+
+				TArrayView<const uint8> MaskBatchView =
+						MaskImage->DataStorage.GetBatchLODRange(BatchId, NumBatchElems, 1, 0, NumLODs);
+
+				TArrayView<uint8> DestBatchView =
+						DestImage->DataStorage.GetBatchLODRange(BatchId, NumBatchElems, BytesPerElem, 0, NumLODs);
+
+				const uint8* BaseBuf = BaseBatchView.GetData() + BaseChannelOffset;
+				const uint8* MaskBuf = MaskBatchView.GetData();
+				uint8* DestBuf = DestBatchView.GetData() + BaseChannelOffset;
+
+				const int32 NumElems = BaseBatchView.Num() / BytesPerElem;
+				check(NumElems == MaskBatchView.Num());
+				check(NumElems == DestBatchView.Num() / BytesPerElem);
+
+				for (int32 I = 0; I < NumElems; ++I)
+				{
+					uint32 MaskData = MaskBuf[I];
+					for (int32 C = 0; C < NumColorChannels; ++C)
+					{
+						const uint32 Base = BaseBuf[BASE_CHANNEL_STRIDE*I + C];
+						const uint32 Result = BLEND_FUNC_MASKED(Base, TopColor[C + ColorChannelOffset], MaskData);
+						if constexpr (CLAMP)
+						{
+							DestBuf[BASE_CHANNEL_STRIDE*I + C] = (uint8)FMath::Min(255u, Result);
+						}
+						else
+						{
+							DestBuf[BASE_CHANNEL_STRIDE*I + C] = (uint8)Result;
+						}
+					}
+
+					constexpr bool bIsNC4 = (BASE_CHANNEL_STRIDE == 4);
+					if constexpr (bIsNC4)
+					{
+						DestBuf[BASE_CHANNEL_STRIDE*I + 3] = BaseBuf[BASE_CHANNEL_STRIDE*I + 3];
+					}
+				}
+			};
+
+			if (NumBatches == 1)
+			{
+				ProcessBatch(0);
+			}
+			else if (NumBatches > 1)
+			{
+				ParallelFor(NumBatches, ProcessBatch);
+			}
+        }
+        else if (MaskFormat == EImageFormat::IF_L_UBYTE_RLE)
+        {
+            int32 Rows = BaseImage->GetSizeY();
+            int32 Width = BaseImage->GetSizeX();
+
+            for (int32 Lod = 0; Lod < NumLODs; ++Lod)
+            {
+                const uint8* BaseBuf = BaseImage->GetLODData(Lod);
+                const uint8* MaskBuf = MaskImage->GetLODData(Lod);
+				uint8* DestBuf = DestImage->GetLODData(Lod);
+
+				// Remove RLE header, mip size and row sizes.
+				MaskBuf += sizeof(uint32);
+				MaskBuf += Rows * sizeof(uint32);
+
+                for (int32 RowIndex = 0; RowIndex < Rows; ++RowIndex)
+                {
+                    const uint8* DestRowEnd = DestBuf + Width * BASE_CHANNEL_STRIDE;
+                    while (DestBuf != DestRowEnd)
+                    {
+                        // Decode header
+						uint16 Equal = 0;
+						FMemory::Memmove(&Equal, MaskBuf, sizeof(uint16));
+                        MaskBuf += 2;
+
+                        uint8 Different = *MaskBuf;
+                        ++MaskBuf;
+
+                        uint8 EqualPixel = *MaskBuf;
+                        ++MaskBuf;
+
+                        // Equal pixels
+						//check(DestBuf + BASE_CHANNEL_STRIDE * Equal <= BaseImage->GetDataSize(Lod));
+                        if (EqualPixel == 255)
+                        {
+                            for (int32 I = 0; I < Equal; ++I)
+                            {
+                                for (int32 C = 0; C < NumColorChannels; ++C)
+                                {
+                                    uint32 Base = BaseBuf[BASE_CHANNEL_STRIDE * I + C];
+                                    uint32 Result = BLEND_FUNC(Base, TopColor[C + ColorChannelOffset]);
+                                    if constexpr (CLAMP)
+                                    {
+                                        DestBuf[BASE_CHANNEL_STRIDE * I + C] = (uint8)FMath::Min(255u, Result);
+                                    }
+                                    else
+                                    {
+                                        DestBuf[BASE_CHANNEL_STRIDE * I + C] = (uint8)Result;
+                                    }
+                                }
+
+                                constexpr bool bIsNC4 = (BASE_CHANNEL_STRIDE == 4);
+                                if (bIsNC4)
+                                {
+                                    DestBuf[BASE_CHANNEL_STRIDE * I + 3] = BaseBuf[BASE_CHANNEL_STRIDE * I + 3];
+                                }
+                            }
+                        }
+                        else if (EqualPixel > 0)
+                        {
+                            for (int32 I = 0; I < Equal; ++I)
+                            {
+                                for (int32 C = 0; C < NumColorChannels; ++C)
+                                {
+                                    uint32 Base = BaseBuf[BASE_CHANNEL_STRIDE * I + C];
+                                    uint32 Result = BLEND_FUNC_MASKED(Base, TopColor[C + ColorChannelOffset], EqualPixel);
+                                    if constexpr (CLAMP)
+                                    {
+                                        DestBuf[BASE_CHANNEL_STRIDE * I + C] = (uint8)FMath::Min(255u, Result);
+                                    }
+                                    else
+                                    {
+                                        DestBuf[BASE_CHANNEL_STRIDE * I + C] = (uint8)Result;
+                                    }
+                                }
+
+                                constexpr bool bIsNC4 = (BASE_CHANNEL_STRIDE == 4);
+                                if (bIsNC4)
+                                {
+                                    DestBuf[BASE_CHANNEL_STRIDE * I + 3] = BaseBuf[BASE_CHANNEL_STRIDE * I + 3];
+                                }
+                            }
+                        }
+                        else
+                        {
+                            // It could happen if xxxxxOnBase
+                            if (DestBuf != BaseBuf)
+                            {
+                                FMemory::Memmove(DestBuf, BaseBuf, BASE_CHANNEL_STRIDE*Equal);
+                            }
+                        }
+
+                        DestBuf += BASE_CHANNEL_STRIDE * Equal;
+                        BaseBuf += BASE_CHANNEL_STRIDE * Equal;
+
+                        // Different pixels
+						//check(DestBuf + BASE_CHANNEL_STRIDE * Different <= StartDestBuf + BaseImage->GetDataSize(Lod));
+                        for (int32 I = 0; I < Different; ++I)
+                        {
+                            for (int32 C = 0; C < NumColorChannels; ++C)
+                            {
+                                uint32 Mask = MaskBuf[I];
+                                uint32 Base = BaseBuf[BASE_CHANNEL_STRIDE * I + C];
+                                uint32 Result = BLEND_FUNC_MASKED(Base, TopColor[C + ColorChannelOffset], Mask);
+                                if constexpr (CLAMP)
+                                {
+                                    DestBuf[BASE_CHANNEL_STRIDE * I + C] = (uint8)FMath::Min(255u, Result);
+                                }
+                                else
+                                {
+                                    DestBuf[BASE_CHANNEL_STRIDE * I + C] = (uint8)Result;
+                                }
+                            }
+
+                            constexpr bool bIsNC4 = (BASE_CHANNEL_STRIDE == 4);
+                            if (bIsNC4)
+                            {
+                                DestBuf[BASE_CHANNEL_STRIDE * I + 3] = BaseBuf[BASE_CHANNEL_STRIDE * I + 3];
+                            }
+                        }
+
+                        DestBuf += BASE_CHANNEL_STRIDE * Different;
+                        BaseBuf += BASE_CHANNEL_STRIDE * Different;
+                        MaskBuf += Different;
+                    }
+                }
+
+                Rows = FMath::DivideAndRoundUp(Rows, 2);
+                Width = FMath::DivideAndRoundUp(Width, 2);
+            }
+        }
+        else
+        {
+            checkf( false, TEXT("Unsupported mask format.") );
+        }
+	}
+
+
+	/**
+	* Apply a blending function to an image with a colour source and a mask
+	*/
+	template<
+		uint32 (*BLEND_FUNC_MASKED)(uint32,uint32,uint32),
+		uint32 (*BLEND_FUNC)(uint32,uint32),
+		bool CLAMP>
+    inline void BufferLayerColour(Image* ResultImage, const Image* BaseImage, const Image* MaskImage, FVector4f Col)
+	{
+		check(BaseImage->GetSizeX() == MaskImage->GetSizeX());
+		check(BaseImage->GetSizeY() == MaskImage->GetSizeY());
+		check(MaskImage->GetFormat() == EImageFormat::IF_L_UBYTE ||
+			  MaskImage->GetFormat() == EImageFormat::IF_L_UBYTE_RLE);
+
+		const bool bValid = (BaseImage->GetSizeX() == MaskImage->GetSizeX()) &&
+						    (BaseImage->GetSizeY() == MaskImage->GetSizeY()) &&
+							(MaskImage->GetFormat() == EImageFormat::IF_L_UBYTE || MaskImage->GetFormat() == EImageFormat::IF_L_UBYTE_RLE);
+		if (!bValid)
+		{
+			return;
+		}
+
+		EImageFormat BaseFormat = BaseImage->GetFormat();
+		if (BaseFormat == EImageFormat::IF_RGB_UBYTE)
+		{
+			BufferLayerColourFormat<BLEND_FUNC_MASKED, BLEND_FUNC, CLAMP, 3, 3>(ResultImage, BaseImage, MaskImage, Col, 0, 0, false);
+		}
+        else if (BaseFormat == EImageFormat::IF_RGBA_UBYTE)
+        {
+            BufferLayerColourFormat<BLEND_FUNC_MASKED, BLEND_FUNC, CLAMP, 3, 4>(ResultImage, BaseImage, MaskImage, Col, 0, 0, false);
+        }
+        else if (BaseFormat == EImageFormat::IF_BGRA_UBYTE)
+        {
+            float Temp = Col[0];
+            Col[0] = Col[2];
+            Col[2] = Temp;
+            BufferLayerColourFormat<BLEND_FUNC_MASKED, BLEND_FUNC, CLAMP, 3, 4>(ResultImage, BaseImage, MaskImage, Col, 0, 0, false);
+        }
+        else if (BaseFormat == EImageFormat::IF_L_UBYTE)
+		{
+			BufferLayerColourFormat<BLEND_FUNC_MASKED, BLEND_FUNC, CLAMP, 1, 1>(ResultImage, BaseImage, MaskImage, Col, 0, 0, false);
+		}
+		else
+		{
+			checkf(false, TEXT("Unsupported format."));
+		}
+	}
+
+
+	template< 
+		uint32 (*BLEND_FUNC_MASKED)(uint32, uint32, uint32),
+		uint32 (*BLEND_FUNC)(uint32, uint32)>
+	inline void BufferLayerColour(Image* DestImage, const Image* BaseImage, const Image* MaskImage, FVector4f Col)
+	{
+		bool bIsClampNeeded = IsAnyComponentLargerThan1(Col);
+		if (bIsClampNeeded)
+		{
+			BufferLayerColour<BLEND_FUNC_MASKED, BLEND_FUNC, true>(DestImage, BaseImage, MaskImage, Col);
+		}
+		else
+		{
+			BufferLayerColour<BLEND_FUNC_MASKED, BLEND_FUNC, false>(DestImage, BaseImage, MaskImage, Col);
+		}
+	}
+
+
+	/**
+	* Apply a blending function to an image with another image as blending layer
+	*/
+	template< 
+		uint32 (*BLEND_FUNC)(uint32, uint32), 
+		bool CLAMP,
+		// Number of total channels to actually process
+		uint32 CHANNELS_TO_BLEND,
+		// Number of total channels in the base image
+		int32 BASE_CHANNEL_STRIDE>
+	inline void BufferLayerColourFormatInPlace(Image* BaseImage, FVector4f Color, uint32 BaseChannelOffset, uint8 ColorChannelOffset, bool bOnlyFirstLOD)
+	{
+		FUintVector4 TopColor = FUintVector4(Color.X * 255.0f, Color.Y * 255.0f, Color.Z * 255.0f, Color.W * 255.0f);
+
+		int32 NumLODs = bOnlyFirstLOD ? 1 : BaseImage->GetLODCount();
+
+		constexpr int32 NumBatchElems = 4096*2;
+
+		const int32 BytesPerElem = GetImageFormatData(BaseImage->GetFormat()).BytesPerBlock;
+		const int32 NumBatches = 
+				BaseImage->DataStorage.GetNumBatchesLODRange(NumBatchElems, BytesPerElem, 0, NumLODs);
+
+		auto ProcessBatch =
+		[
+			BaseImage, TopColor, BaseChannelOffset, ColorChannelOffset, BytesPerElem, NumBatchElems, NumLODs
+		] (uint32 BatchId)
+		{	
+			const TArrayView<uint8> BaseBatchView = 
+					BaseImage->DataStorage.GetBatchLODRange(BatchId, NumBatchElems, BytesPerElem, 0, NumLODs);
+			
+			uint8* BaseBuf = BaseBatchView.GetData() + BaseChannelOffset;
+			const int32 NumElems = BaseBatchView.Num() / BytesPerElem;
+			
+			for (int32 I = 0; I < NumElems; ++I)
+			{
+				for (int32 C = 0; C < CHANNELS_TO_BLEND; ++C)
+				{
+					uint32 Base = BaseBuf[BASE_CHANNEL_STRIDE*I + C];
+					uint32 Blended = TopColor[C + ColorChannelOffset];
+					uint32 Result = BLEND_FUNC(Base, Blended);
+					if constexpr (CLAMP)
+					{
+						BaseBuf[BASE_CHANNEL_STRIDE*I + C] = (uint8)FMath::Min(255u, Result);
+					}
+					else
+					{
+						BaseBuf[BASE_CHANNEL_STRIDE*I + C] = (uint8)Result;
+					}
+				}
+			}
+		};
+
+		if (NumBatches == 1)
+		{
+			ProcessBatch(0);
+		}
+		else if (NumBatches > 1)
+		{
+			ParallelFor(NumBatches, ProcessBatch);
+		}
+	}
+
+
+	template<
+		uint32 (*BLEND_FUNC)(uint32, uint32), 
+		bool CLAMP, 
+		uint32 CHANNEL_COUNT>
+	inline void BufferLayerColourInPlace(Image* BaseImage, FVector4f Col, bool bOnlyOneMip, uint32 BaseOffset, uint8 ColOffset)
+	{
+		EImageFormat BaseFormat = BaseImage->GetFormat();
+
+		if (BaseFormat == EImageFormat::IF_RGB_UBYTE)
+		{
+			check(BaseOffset + CHANNEL_COUNT <= 3);
+			BufferLayerColourFormatInPlace<BLEND_FUNC, CLAMP, CHANNEL_COUNT, 3>
+					(BaseImage, Col, BaseOffset, ColOffset, bOnlyOneMip);
+		}
+		else if (BaseFormat == EImageFormat::IF_RGBA_UBYTE)
+		{
+			check(BaseOffset + CHANNEL_COUNT <= 4);
+			BufferLayerColourFormatInPlace<BLEND_FUNC, CLAMP, CHANNEL_COUNT, 4>
+					(BaseImage, Col, BaseOffset, ColOffset, bOnlyOneMip);
+		}
+		else if (BaseFormat == EImageFormat::IF_BGRA_UBYTE)
+		{
+			float Temp = Col[0];
+			Col[0] = Col[2];
+			Col[2] = Temp;
+			BufferLayerColourFormatInPlace<BLEND_FUNC, CLAMP, CHANNEL_COUNT, 4>
+					(BaseImage, Col, BaseOffset, ColOffset, bOnlyOneMip);
+		}
+
+		else if (BaseFormat == EImageFormat::IF_L_UBYTE)
+		{
+			check(BaseOffset + CHANNEL_COUNT <= 1);
+			BufferLayerColourFormatInPlace< BLEND_FUNC, CLAMP, CHANNEL_COUNT, 1>
+					(BaseImage, Col, BaseOffset, ColOffset, bOnlyOneMip);
+		}
+		else
+		{
+			checkf(false, TEXT("Unsupported format."));
+		}
+	}
+
+
+	template< 
+		uint32 (*BLEND_FUNC)(uint32, uint32), 
+		uint32 CHANNEL_COUNT>
+	inline void BufferLayerColourInPlace(Image* BaseImage, FVector4f Color, bool bOnlyOneMip, uint32 BaseOffset, uint8 ColorOffset)
+	{
+		bool bIsClampNeeded = IsAnyComponentLargerThan1(Color);
+		if (bIsClampNeeded)
+		{
+			BufferLayerColourInPlace<BLEND_FUNC, true, CHANNEL_COUNT>(BaseImage, Color, bOnlyOneMip, BaseOffset, ColorOffset);
+		}
+		else
+		{
+			BufferLayerColourInPlace<BLEND_FUNC, false, CHANNEL_COUNT>(BaseImage, Color, bOnlyOneMip, BaseOffset, ColorOffset);
+		}
+	}
+
+
+	template< 
+		uint32 (*BLEND_FUNC_MASKED)(uint32, uint32, uint32),
+		uint32 (*BLEND_FUNC)(uint32, uint32), 
+		bool CLAMP, 
+		int32 CHANNEL_COUNT>
+	inline void BufferLayerColourInPlace(Image* BaseImage, const Image* MaskImage, FVector4f Color, bool bOnlyOneMip, uint32 BaseOffset, uint8 ColorOffset)
+	{
+		check(BaseImage->GetSizeX() == MaskImage->GetSizeX());
+		check(BaseImage->GetSizeY() == MaskImage->GetSizeY());
+
+		EImageFormat BaseFormat = BaseImage->GetFormat();
+
+		if (BaseFormat == EImageFormat::IF_RGB_UBYTE)
+		{
+			check(BaseOffset + CHANNEL_COUNT <= 3);
+			BufferLayerColourFormat<BLEND_FUNC_MASKED, BLEND_FUNC, CLAMP, CHANNEL_COUNT, 3>
+					(BaseImage, BaseImage, MaskImage, Color, BaseOffset, ColorOffset, bOnlyOneMip);
+		}
+		else if (BaseFormat == EImageFormat::IF_RGBA_UBYTE)
+		{
+			check(BaseOffset + CHANNEL_COUNT <= 4);
+			BufferLayerColourFormat<BLEND_FUNC_MASKED, BLEND_FUNC, CLAMP, CHANNEL_COUNT, 4>
+					(BaseImage, BaseImage, MaskImage, Color, BaseOffset, ColorOffset, bOnlyOneMip);
+		}
+		else if (BaseFormat == EImageFormat::IF_BGRA_UBYTE)
+		{
+			check(BaseOffset + CHANNEL_COUNT <= 4);
+			float Temp = Color[0];
+			Color[0] = Color[2];
+			Color[2] = Temp;
+
+			BufferLayerColourFormat<BLEND_FUNC_MASKED, BLEND_FUNC, CLAMP, CHANNEL_COUNT, 4>
+					(BaseImage, BaseImage, MaskImage, Color, BaseOffset, ColorOffset, bOnlyOneMip);
+		}
+		else if (BaseFormat == EImageFormat::IF_L_UBYTE)
+		{
+			check(BaseOffset + CHANNEL_COUNT <= 1);
+			BufferLayerColourFormat<BLEND_FUNC_MASKED, BLEND_FUNC, CLAMP, CHANNEL_COUNT, 1>
+					(BaseImage, BaseImage, MaskImage, Color, BaseOffset, ColorOffset, bOnlyOneMip);
+		}
+		else
+		{
+			checkf(false, TEXT("Unsupported format."));
+		}
+	}
+
+	template< 
+		uint32 (*BLEND_FUNC_MASKED)(uint32, uint32, uint32), 
+		uint32 (*BLEND_FUNC)(uint32, uint32), uint32 CHANNEL_COUNT>
+	inline void BufferLayerColourInPlace(Image* BaseImage, const Image* MaskImage, FVector4f Color, bool bOnlyOneMip, uint32 BaseOffset, uint8 ColorOffset)
+	{
+		bool bIsClampNeeded = IsAnyComponentLargerThan1(Color);
+		if (bIsClampNeeded)
+		{
+			BufferLayerColourInPlace<BLEND_FUNC_MASKED, BLEND_FUNC, true, CHANNEL_COUNT>
+					(BaseImage, MaskImage, Color, bOnlyOneMip, BaseOffset, ColorOffset);
+		}
+		else
+		{
+			BufferLayerColourInPlace<BLEND_FUNC_MASKED, BLEND_FUNC, false, CHANNEL_COUNT>
+					(BaseImage, MaskImage, Color, bOnlyOneMip, BaseOffset, ColorOffset);
+		}
+	}
+
+
+	/**	
+	* Apply a blending function to an image with another image as blending layer
+	*/
+	template< uint32 (*BLEND_FUNC)(uint32, uint32), bool CLAMP,
+		// Number of total channels to actually process
+		int32 CHANNELS_TO_BLEND,
+		// Number of total channels in the base image
+		int32 BASE_CHANNEL_STRIDE,
+		// Number of total channels in the blend image
+		int32 BLENDED_CHANNEL_STRIDE>
+    inline void BufferLayerFormatInPlace(
+		Image* BaseImage, 
+		const Image* BlendedImage,
+		uint32 BaseChannelOffset,
+		uint32 BlendedChannelOffset,
+		bool bOnlyFirstLOD)
+	{
+		check(BaseImage->GetSizeX() == BlendedImage->GetSizeX());
+		check(BaseImage->GetSizeY() == BlendedImage->GetSizeY());
+
+		// No longer required.
+		//check(BaseImage->GetFormat() == BlendedImage->GetFormat());
+		check(BaseChannelOffset + CHANNELS_TO_BLEND <= GetImageFormatData(BaseImage->GetFormat()).Channels);
+		check(BlendedChannelOffset + CHANNELS_TO_BLEND <= GetImageFormatData(BlendedImage->GetFormat()).Channels);
+
+		check(bOnlyFirstLOD || BaseImage->GetLODCount() <= BlendedImage->GetLODCount());
+
+		const int32 BaseBytesPerElem = GetImageFormatData(BaseImage->GetFormat()).BytesPerBlock;
+		const int32 BlendBytesPerElem = GetImageFormatData(BlendedImage->GetFormat()).BytesPerBlock;
+	
+		const int32 NumLODs = bOnlyFirstLOD ? 1 : BaseImage->GetLODCount();
+
+		constexpr int32 NumBatchElems = 4096*2;
+		const int32 NumBatches =
+				BaseImage->DataStorage.GetNumBatchesLODRange(NumBatchElems, BaseBytesPerElem, 0, NumLODs);
+		
+		check(NumBatches <= BlendedImage->DataStorage.GetNumBatches(NumBatchElems, BlendBytesPerElem));
+
+		auto ProcessBatch =
+		[
+			BaseImage, BlendedImage, BaseBytesPerElem, BlendBytesPerElem, 
+			NumBatchElems, BaseChannelOffset, BlendedChannelOffset, NumLODs
+		] (int32 BatchId)
+		{
+			TArrayView<uint8> BaseBatchView =
+					BaseImage->DataStorage.GetBatchLODRange(BatchId, NumBatchElems, BaseBytesPerElem, 0, NumLODs);
+
+			TArrayView<const uint8> BlendedBatchView =
+					BlendedImage->DataStorage.GetBatchLODRange(BatchId, NumBatchElems, BlendBytesPerElem, 0, NumLODs);
+
+        	uint8* BaseBuf = BaseBatchView.GetData() + BaseChannelOffset;
+        	const uint8* BlendedBuf = BlendedBatchView.GetData() + BlendedChannelOffset;
+
+			const int32 NumElems = BaseBatchView.Num() / BaseBytesPerElem;
+			check(BlendedBatchView.Num() / BlendBytesPerElem == NumElems);
+
+			for (int32 I = 0; I < NumElems; ++I)
+			{
+				for (int32 C = 0; C < CHANNELS_TO_BLEND; ++C)
+				{
+					uint32 Base = BaseBuf[BASE_CHANNEL_STRIDE * I + C];
+					uint32 Blended = BlendedBuf[BLENDED_CHANNEL_STRIDE * I + C];
+					uint32 Result = BLEND_FUNC(Base, Blended);
+					
+					if constexpr (CLAMP)
+					{
+						BaseBuf[BASE_CHANNEL_STRIDE * I + C] = (uint8)FMath::Min(255u, Result);
+					}
+					else
+					{
+						BaseBuf[BASE_CHANNEL_STRIDE * I + C] = (uint8)Result;
+					}
+				}
+			}
+		};
+
+		if (NumBatches == 1)
+		{
+			ProcessBatch(0);
+		}
+		else if (NumBatches > 1)
+		{
+			ParallelFor(NumBatches, ProcessBatch);
+		}
+	}
+	
+	/**
+	* Apply a blending function to an image with another image as blending layer
+	*/
+	template<
+		uint32 (*BLEND_FUNC)(uint32, uint32), 
+		bool CLAMP,
+		// Number of total channels to actually process
+		int32 CHANNELS_TO_BLEND,
+		// Number of total channels in the base image
+		int32 BASE_CHANNEL_STRIDE,
+		// Number of total channels in the blend image
+		int32 BLENDED_CHANNEL_STRIDE,
+		int32 BLENDED_CHANNEL_OFFSET>
+	inline void BufferLayerFormat(Image* DestImage, const Image* BaseImage, const Image* BlendedImage, bool bOnlyFirstLOD)
+	{
+		check(BaseImage->GetSizeX() == BlendedImage->GetSizeX());
+		check(BaseImage->GetSizeY() == BlendedImage->GetSizeY());
+		// Not true anymore, since the BLENDED_CHANNEL_OFFSET has been added.
+		// check(BaseImage->GetFormat() == BlendedImage->GetFormat());
+		check(bOnlyFirstLOD || BaseImage->GetLODCount() <= BlendedImage->GetLODCount());
+
+		constexpr int32 UnblendedChannels = BASE_CHANNEL_STRIDE - CHANNELS_TO_BLEND;
+		constexpr int32 NumBatchElems = 4096*2;
+		
+		const int32 BaseBytesPerElem = GetImageFormatData(BaseImage->GetFormat()).BytesPerBlock;
+		const int32 BlendBytesPerElem = GetImageFormatData(BlendedImage->GetFormat()).BytesPerBlock;
+		const int32 DestBytesPerElem = GetImageFormatData(DestImage->GetFormat()).BytesPerBlock;
+
+		const int32 NumLODs = bOnlyFirstLOD ? 1 : BaseImage->GetLODCount();
+
+		const int32 NumBatches = 
+				BaseImage->DataStorage.GetNumBatchesLODRange(NumBatchElems, BaseBytesPerElem, 0, NumLODs);
+
+		check(NumBatches == BlendedImage->DataStorage.GetNumBatchesLODRange(NumBatchElems, BlendBytesPerElem, 0, NumLODs));
+		check(NumBatches == DestImage->DataStorage.GetNumBatchesLODRange(NumBatchElems, DestBytesPerElem, 0, NumLODs));
+
+		auto ProcessBatch =
+		[
+			BaseImage, BlendedImage, DestImage, 
+			BaseBytesPerElem, BlendBytesPerElem, DestBytesPerElem, 
+			UnblendedChannels, NumBatchElems, NumLODs
+		] (int32 BatchId)
+		{
+			TArrayView<const uint8> BaseBatchView =
+					BaseImage->DataStorage.GetBatchLODRange(BatchId, NumBatchElems, BaseBytesPerElem, 0, NumLODs);
+
+			TArrayView<const uint8> BlendedBatchView =
+					BlendedImage->DataStorage.GetBatchLODRange(BatchId, NumBatchElems, BlendBytesPerElem, 0, NumLODs);
+
+			TArrayView<uint8> DestBatchView =
+					DestImage->DataStorage.GetBatchLODRange(BatchId, NumBatchElems, DestBytesPerElem, 0, NumLODs);
+	
+			const int32 NumElems = BaseBatchView.Num() / BaseBytesPerElem;
+			check(NumElems == BlendedBatchView.Num() / BlendBytesPerElem);
+			check(NumElems == DestBatchView.Num() / DestBytesPerElem);
+
+			const uint8* BaseBuf = BaseBatchView.GetData();
+			const uint8* BlendedBuf = BlendedBatchView.GetData() + BLENDED_CHANNEL_OFFSET;
+			uint8* DestBuf = DestBatchView.GetData();
+
+			for (int32 I = 0; I < NumElems; ++I)
+			{
+				for (int32 C = 0; C < CHANNELS_TO_BLEND; ++C)
+				{
+					uint32 Base = BaseBuf[BASE_CHANNEL_STRIDE * I + C];
+					uint32 Blended = BlendedBuf[BLENDED_CHANNEL_STRIDE * I + C];
+					uint32 Result = BLEND_FUNC(Base, Blended);
+
+					if constexpr (CLAMP)
+					{
+						DestBuf[BASE_CHANNEL_STRIDE * I + C] = (uint8)FMath::Min(255u, Result);
+					}
+					else
+					{
+						DestBuf[BASE_CHANNEL_STRIDE * I + C] = (uint8)Result;
+					}
+				}
+
+				// Copy the unblended channels
+				// \TODO: unnecessary when doing it in-place?
+				if constexpr (UnblendedChannels > 0)
+				{
+					for (int32 C = 0; C < UnblendedChannels; ++C)
+					{
+						DestBuf[BASE_CHANNEL_STRIDE * I + CHANNELS_TO_BLEND + C] 
+							= BaseBuf[BASE_CHANNEL_STRIDE * I + CHANNELS_TO_BLEND + C];
+					}
+				}
+			}
+		};
+
+		if (NumBatches == 1)
+		{
+			ProcessBatch(0);
+		}
+		else if (NumBatches > 1)
+		{
+			ParallelFor(NumBatches, ProcessBatch);
+		}
+	}
+
+	/**
+	* Apply a blending function to an image with another image as blending layer
+	*/
+	template< uint32 (*BLEND_FUNC)(uint32, uint32), bool CLAMP >
+	inline void BufferLayer(Image* ResultImage, const Image* BaseImage, const Image* BlendedImage, bool bApplyToAlpha, bool bOnlyOneMip, bool bUseBlendSourceFromBlendAlpha)
+	{
+		check(ResultImage->GetFormat() == BaseImage->GetFormat());
+		check(ResultImage->GetSizeX() == BaseImage->GetSizeX());
+		check(ResultImage->GetSizeY() == BaseImage->GetSizeY());
+		check(bOnlyOneMip || ResultImage->GetLODCount() == BaseImage->GetLODCount());
+		check(BaseImage->GetSizeX() == BlendedImage->GetSizeX());
+		check(BaseImage->GetSizeY() == BlendedImage->GetSizeY());
+		check(bOnlyOneMip || ResultImage->GetLODCount() <= BlendedImage->GetLODCount());
+
+		EImageFormat BaseFormat = BaseImage->GetFormat();
+		EImageFormat BlendedFormat = BlendedImage->GetFormat();
+
+		if (bUseBlendSourceFromBlendAlpha)
+		{
+			if (BlendedFormat == EImageFormat::IF_RGBA_UBYTE || BlendedFormat == EImageFormat::IF_BGRA_UBYTE)
+			{
+				if (BaseFormat == EImageFormat::IF_L_UBYTE)
+				{
+					BufferLayerFormat<BLEND_FUNC, CLAMP, 1, 1, 4, 3>(ResultImage, BaseImage, BlendedImage, bOnlyOneMip);
+				}
+				else
+				{
+					checkf(false, TEXT("Unsupported format."));
+				}
+			}
+			else if (BlendedFormat == EImageFormat::IF_L_UBYTE)
+			{
+				BufferLayerFormat<BLEND_FUNC, CLAMP, 1, 1, 1, 0>(ResultImage, BaseImage, BlendedImage, bOnlyOneMip);
+			}
+		}
+		else
+		{
+			check(BaseFormat == BlendedFormat);
+			if (BaseFormat == EImageFormat::IF_RGB_UBYTE)
+			{
+				check(!bUseBlendSourceFromBlendAlpha);
+				BufferLayerFormat<BLEND_FUNC, CLAMP, 3, 3, 3, 0>(ResultImage, BaseImage, BlendedImage, bOnlyOneMip);
+			}
+			else if (BaseFormat == EImageFormat::IF_RGBA_UBYTE || BaseFormat == EImageFormat::IF_BGRA_UBYTE)
+			{
+				check(!bUseBlendSourceFromBlendAlpha);
+				if (bApplyToAlpha)
+				{
+					BufferLayerFormat<BLEND_FUNC, CLAMP, 4, 4, 4, 0>(ResultImage, BaseImage, BlendedImage, bOnlyOneMip);
+				}
+				else
+				{
+					BufferLayerFormat<BLEND_FUNC, CLAMP, 3, 4, 4, 0>(ResultImage, BaseImage, BlendedImage, bOnlyOneMip);
+				}
+			}
+			else if (BaseFormat == EImageFormat::IF_L_UBYTE)
+			{
+				BufferLayerFormat<BLEND_FUNC, CLAMP, 1, 1, 1, 0>(ResultImage, BaseImage, BlendedImage, bOnlyOneMip);
+			}
+			else
+			{
+				checkf(false, TEXT("Unsupported format."));
+			}
+		}
+	}
+
+
+	template< uint32 (*BLEND_FUNC)(uint32, uint32), bool CLAMP, int32 CHANNEL_COUNT >
+	inline void BufferLayerInPlace(Image* BaseImage, const Image* BlendedImage, bool bOnlyOneMip, uint32 BaseOffset, uint32 BlendedOffset)
+	{
+		check(BaseImage->GetSizeX() == BlendedImage->GetSizeX());
+		check(BaseImage->GetSizeY() == BlendedImage->GetSizeY());
+		// Not required since we have the CHANNEL_COUNT and offsets. 
+		// check(BaseImage->GetFormat() == BlendedImage->GetFormat());
+
+		EImageFormat BaseFormat = BaseImage->GetFormat();
+		EImageFormat BlendFormat = BlendedImage->GetFormat();
+
+		if (BaseFormat == EImageFormat::IF_RGB_UBYTE && BlendFormat==EImageFormat::IF_RGB_UBYTE)
+		{
+			check(BaseOffset + CHANNEL_COUNT <= 3);
+			check(BlendedOffset + CHANNEL_COUNT <= 3);
+			BufferLayerFormatInPlace<BLEND_FUNC, CLAMP, CHANNEL_COUNT, 3, 3>(BaseImage, BlendedImage, BaseOffset, BlendedOffset, bOnlyOneMip);
+		}
+		else if ((BaseFormat == EImageFormat::IF_RGBA_UBYTE || BaseFormat == EImageFormat::IF_BGRA_UBYTE) &&
+				 (BlendFormat == EImageFormat::IF_RGBA_UBYTE || BlendFormat == EImageFormat::IF_BGRA_UBYTE) )
+		{
+			check(BaseOffset + CHANNEL_COUNT <= 4);
+			check(BlendedOffset + CHANNEL_COUNT <= 4);
+			BufferLayerFormatInPlace< BLEND_FUNC, CLAMP, CHANNEL_COUNT, 4, 4>
+					(BaseImage, BlendedImage, BaseOffset, BlendedOffset, bOnlyOneMip);
+		}
+		else if ((BaseFormat == EImageFormat::IF_RGBA_UBYTE || BaseFormat == EImageFormat::IF_BGRA_UBYTE) &&
+				 BlendFormat == EImageFormat::IF_L_UBYTE )
+		{
+			check(BaseOffset + CHANNEL_COUNT <= 4);
+			check(BlendedOffset + CHANNEL_COUNT <= 1);
+			BufferLayerFormatInPlace< BLEND_FUNC, CLAMP, CHANNEL_COUNT, 4, 1>
+					(BaseImage, BlendedImage, BaseOffset, BlendedOffset, bOnlyOneMip);
+		}
+		else if (BaseFormat == EImageFormat::IF_L_UBYTE && BlendFormat==EImageFormat::IF_L_UBYTE)
+		{
+			check(BaseOffset + CHANNEL_COUNT <= 1);
+			BufferLayerFormatInPlace< BLEND_FUNC, CLAMP, CHANNEL_COUNT, 1, 1>
+					(BaseImage, BlendedImage, BaseOffset, BlendedOffset, bOnlyOneMip);
+		}
+		else
+		{
+			checkf(false, TEXT("Unsupported format."));
+		}
+	}
+
+
+	template< uint32 (*BLEND_FUNC_MASKED)(uint32, uint32, uint32),
+		uint32 (*BLEND_FUNC)(uint32, uint32),
+		bool CLAMP,
+		// Number of total channels to actually process
+		int32 CHANNELS_TO_BLEND,
+		// Number of total channels in the base image
+		int32 BASE_CHANNEL_STRIDE,
+		// Number of total channels in the blend image
+		int32 BLENDED_CHANNEL_STRIDE>
+	inline void BufferLayerFormat(
+		Image* DestImage,
+		const Image* BaseImage,
+		const Image* MaskImage,
+		const Image* BlendImage,
+		uint32 DestOffset,
+		uint32 BaseChannelOffset,
+		uint32 BlendedChannelOffset,
+		bool bOnlyFirstLOD)
+	{
+		check(BaseImage->GetSizeX() == MaskImage->GetSizeX() && BaseImage->GetSizeY() == MaskImage->GetSizeY());
+		check(BaseImage->GetSizeX() == BlendImage->GetSizeX() && BaseImage->GetSizeY() == BlendImage->GetSizeY());
+		check(bOnlyFirstLOD || BaseImage->GetLODCount() <= MaskImage->GetLODCount());
+		check(bOnlyFirstLOD || BaseImage->GetLODCount() <= BlendImage->GetLODCount());
+		check(BaseImage->GetFormat() == BlendImage->GetFormat());
+		check(MaskImage->GetFormat() == EImageFormat::IF_L_UBYTE ||
+			  MaskImage->GetFormat() == EImageFormat::IF_L_UBYTE_RLE);
+
+		const EImageFormat MaskFormat = MaskImage->GetFormat();
+
+		const int32 BaseBytesPerElem = GetImageFormatData(BaseImage->GetFormat()).BytesPerBlock; 
+		const int32 BlendBytesPerElem = GetImageFormatData(BlendImage->GetFormat()).BytesPerBlock; 
+
+		const int32 NumLODs = bOnlyFirstLOD ? 1 : BaseImage->GetLODCount();
+
+		const bool bIsMaskUncompressed = MaskFormat == EImageFormat::IF_L_UBYTE;
+		constexpr int32 UnblendedChannels = BASE_CHANNEL_STRIDE - CHANNELS_TO_BLEND;
+
+		if (bIsMaskUncompressed)
+		{	
+			constexpr int32 NumBatchElems = 4096*2;
+
+			const int32 NumBatches = 
+				BaseImage->DataStorage.GetNumBatchesLODRange(NumBatchElems, BaseBytesPerElem, 0, NumLODs);
+
+			check(NumBatches == BlendImage->DataStorage.GetNumBatchesLODRange(NumBatchElems, BlendBytesPerElem, 0, NumLODs));
+			check(NumBatches == MaskImage->DataStorage.GetNumBatchesLODRange(NumBatchElems, 1, 0, NumLODs));
+			check(NumBatches == DestImage->DataStorage.GetNumBatchesLODRange(NumBatchElems, BaseBytesPerElem, 0, NumLODs));
+
+			auto ProcessBatch = 
+			[
+				BaseImage, BlendImage, MaskImage, DestImage, UnblendedChannels, BaseBytesPerElem, BlendBytesPerElem, NumLODs
+			] (uint32 BatchId)
+			{
+				TArrayView<const uint8> BaseBatchView =
+						BaseImage->DataStorage.GetBatchLODRange(BatchId, NumBatchElems, BaseBytesPerElem, 0, NumLODs);
+
+				TArrayView<const uint8> BlendBatchView =
+						BlendImage->DataStorage.GetBatchLODRange(BatchId, NumBatchElems, BlendBytesPerElem, 0, NumLODs);
+
+				TArrayView<const uint8> MaskBatchView =
+						MaskImage->DataStorage.GetBatchLODRange(BatchId, NumBatchElems, 1, 0, NumLODs);
+
+				TArrayView<uint8> DestBatchView = 
+						DestImage->DataStorage.GetBatchLODRange(BatchId, NumBatchElems, BaseBytesPerElem, 0, NumLODs);
+
+				const uint8* BaseBuf = BaseBatchView.GetData();
+				const uint8* BlendedBuf = BlendBatchView.GetData();
+				const uint8* MaskBuf = MaskBatchView.GetData();
+				uint8* DestBuf = DestBatchView.GetData();
+
+				const int32 NumElems = BaseBatchView.Num() / BaseBytesPerElem;
+				check(NumElems == BlendBatchView.Num() / BlendBytesPerElem);
+				check(NumElems == MaskBatchView.Num() / 1);
+				check(NumElems == DestBatchView.Num() / BaseBytesPerElem);
+
+				for (int32 I = 0; I < NumElems; ++I)
+				{
+					uint32 Mask = MaskBuf[I];
+					for (int32 C = 0; C < CHANNELS_TO_BLEND; ++C)
+					{
+						const uint32 Base = BaseBuf[BASE_CHANNEL_STRIDE * I + C];
+						const uint32 Blended = BlendedBuf[BLENDED_CHANNEL_STRIDE * I + C];
+						const uint32 Result = BLEND_FUNC_MASKED(Base, Blended, Mask);
+						if constexpr (CLAMP)
+						{
+							DestBuf[BASE_CHANNEL_STRIDE * I + C] = (uint8)FMath::Min(255u, Result);
+						}
+						else
+						{
+							DestBuf[BASE_CHANNEL_STRIDE * I + C] = (uint8)Result;
+						}
+					}
+					// Copy the unblended channels
+					// \TODO: unnecessary when doing it in-place?
+					if constexpr (UnblendedChannels > 0)
+					{
+						for (int32 C = 0; C < UnblendedChannels; ++C)
+						{
+							DestBuf[BASE_CHANNEL_STRIDE * I + CHANNELS_TO_BLEND + C] = BaseBuf[BASE_CHANNEL_STRIDE * I + CHANNELS_TO_BLEND + C];
+						}
+					}
+				}
+			};
+
+			if (NumBatches == 1)
+			{
+				ProcessBatch(0);
+			}
+			else if (NumBatches > 1)
+			{
+				ParallelFor(NumBatches, ProcessBatch);
+			}
+		}
+		else if (MaskFormat == EImageFormat::IF_L_UBYTE_RLE)
+		{
+			int32 Rows = BaseImage->GetSizeY();
+			int32 Width = BaseImage->GetSizeX();
+
+			for (int32 LOD = 0; LOD < NumLODs; ++LOD)
+			{
+				const uint8* BaseBuf = BaseImage->GetLODData(LOD);
+				uint8* DestBuf = DestImage->GetLODData(LOD);
+				const uint8* BlendedBuf = BlendImage->GetLODData(LOD);
+				const uint8* MaskBuf = MaskImage->GetLODData(LOD);
+
+				// Remove RLE header, mip size and row sizes.
+				MaskBuf += sizeof(uint32) + Rows * sizeof(uint32);
+
+				for (int32 RowIndex = 0; RowIndex < Rows; ++RowIndex)
+				{
+					const uint8* DestRowEnd = DestBuf + Width * BASE_CHANNEL_STRIDE;
+					while (DestBuf != DestRowEnd)
+					{
+						// Decode header
+						uint16 Equal;
+						FMemory::Memcpy(&Equal, MaskBuf, 2);
+						MaskBuf += 2;
+
+						uint8 Different = *MaskBuf;
+						++MaskBuf;
+
+						uint8 EqualPixel = *MaskBuf;
+						++MaskBuf;
+
+						// Equal pixels
+						//check(DestBuf + BASE_CHANNEL_STRIDE * Equal <= BaseImage->GetLODDataSize(0));
+						if (EqualPixel == 255)
+						{
+							for (int32 I = 0; I < Equal; ++I)
+							{
+								for (int32 C = 0; C < CHANNELS_TO_BLEND; ++C)
+								{
+									uint32 Base = BaseBuf[BASE_CHANNEL_STRIDE * I + C];
+									uint32 Blended = BlendedBuf[BLENDED_CHANNEL_STRIDE * I + C];
+									uint32 Result = BLEND_FUNC(Base, Blended);
+									if constexpr (CLAMP)
+									{
+										DestBuf[BASE_CHANNEL_STRIDE * I + C] = (uint8)FMath::Min(255u, Result);
+									}
+									else
+									{
+										DestBuf[BASE_CHANNEL_STRIDE * I + C] = (uint8)Result;
+									}
+								}
+
+								// Copy the unblended channels
+								// \TODO: unnecessary when doing it in-place?
+								if constexpr (UnblendedChannels > 0)
+								{
+									for (int32 C = 0; C < UnblendedChannels; ++C)
+									{
+										DestBuf[BASE_CHANNEL_STRIDE * I + CHANNELS_TO_BLEND + C] = BaseBuf[BASE_CHANNEL_STRIDE * I + CHANNELS_TO_BLEND + C];
+									}
+								}
+							}
+						}
+						else if (EqualPixel > 0)
+						{
+							for (int32 I = 0; I < Equal; ++I)
+							{
+								for (int32 C = 0; C < CHANNELS_TO_BLEND; ++C)
+								{
+									uint32 Base = BaseBuf[BASE_CHANNEL_STRIDE * I + C];
+									uint32 Blended = BlendedBuf[BLENDED_CHANNEL_STRIDE * I + C];
+									uint32 Result = BLEND_FUNC_MASKED(Base, Blended, EqualPixel);
+
+									if constexpr (CLAMP)
+									{
+										DestBuf[BASE_CHANNEL_STRIDE * I + C] = (uint8)FMath::Min(255u, Result);
+									}
+									else
+									{
+										DestBuf[BASE_CHANNEL_STRIDE * I + C] = (uint8)Result;
+									}
+								}
+
+								// Copy the unblended channels
+								// \TODO: unnecessary when doing it in-place?
+								if constexpr (UnblendedChannels > 0)
+								{
+									for (int32 C = 0; C < UnblendedChannels; ++C)
+									{
+										DestBuf[BASE_CHANNEL_STRIDE * I + CHANNELS_TO_BLEND + C] = BaseBuf[BASE_CHANNEL_STRIDE * I + CHANNELS_TO_BLEND + C];
+									}
+								}
+							}
+						}
+						else
+						{
+							// It could happen if xxxxxOnBase
+							if (DestBuf != BaseBuf)
+							{
+								FMemory::Memmove(DestBuf, BaseBuf, BASE_CHANNEL_STRIDE * Equal);
+							}
+						}
+
+						DestBuf += BASE_CHANNEL_STRIDE * Equal;
+						BaseBuf += BASE_CHANNEL_STRIDE * Equal;
+						BlendedBuf += BLENDED_CHANNEL_STRIDE * Equal;
+
+						// Different pixels
+						//check(DestBuf + BASE_CHANNEL_STRIDE * Different <= BaseImage->GetDataSize(0));
+						for (int32 I = 0; I < Different; ++I)
+						{
+							for (int32 C = 0; C < CHANNELS_TO_BLEND; ++C)
+							{
+								uint32 Mask = MaskBuf[I];
+								uint32 Base = BaseBuf[BASE_CHANNEL_STRIDE * I + C];
+								uint32 Blended = BlendedBuf[BLENDED_CHANNEL_STRIDE * I + C];
+								uint32 Result = BLEND_FUNC_MASKED(Base, Blended, Mask);
+								if (CLAMP)
+								{
+									DestBuf[BASE_CHANNEL_STRIDE * I + C] = (uint8)FMath::Min(255u, Result);
+								}
+								else
+								{
+									DestBuf[BASE_CHANNEL_STRIDE * I + C] = (uint8)Result;
+								}
+							}
+
+							// Copy the unblended channels
+							// \TODO: unnecessary when doing it in-place?
+							if constexpr (UnblendedChannels > 0)
+							{
+								for (int32 C = 0; C < UnblendedChannels; ++C)
+								{
+									DestBuf[BASE_CHANNEL_STRIDE * I + CHANNELS_TO_BLEND + C] = BaseBuf[BASE_CHANNEL_STRIDE * I + CHANNELS_TO_BLEND + C];
+								}
+							}
+						}
+
+						DestBuf += BASE_CHANNEL_STRIDE * Different;
+						BaseBuf += BASE_CHANNEL_STRIDE * Different;
+						BlendedBuf += BLENDED_CHANNEL_STRIDE * Different;
+						MaskBuf += Different;
+					}
+				}
+
+				Rows = FMath::DivideAndRoundUp(Rows, 2);
+				Width = FMath::DivideAndRoundUp(Width, 2);
+			}
+		}
+		else
+		{
+			checkf(false, TEXT("Unsupported mask format."));
+		}
+	}
+
+	template< uint32 (*BLEND_FUNC_MASKED)(uint32, uint32, uint32),
+		uint32 (*BLEND_FUNC)(uint32, uint32),
+		bool CLAMP,
+		// Number of total channels to actually process
+		int32 CHANNELS_TO_BLEND,
+		// Number of total channels in the base image
+		int32 BASE_CHANNEL_STRIDE,
+		// Number of total channels in the blend image
+		int32 BLENDED_CHANNEL_STRIDE>
+	inline void BufferLayerFormatEmbeddedMask(
+		Image* DestImage,
+		const Image* BaseImage,
+		const Image* BlendImage,
+		uint32 DestOffset,
+		uint32 BaseChannelOffset,
+		bool bOnlyFirstLOD)
+	{
+		check(BaseImage->GetSizeX() == BlendImage->GetSizeX() && BaseImage->GetSizeY() == BlendImage->GetSizeY());
+		check(bOnlyFirstLOD || BaseImage->GetLODCount() <= BlendImage->GetLODCount());
+
+		constexpr int32 UnblendedChannels = BASE_CHANNEL_STRIDE - CHANNELS_TO_BLEND;
+
+		const int32 BytesPerElem = GetImageFormatData(BaseImage->GetFormat()).BytesPerBlock; 
+	
+		const int32 NumLODs = bOnlyFirstLOD ? 1 : BaseImage->GetLODCount();
+
+		constexpr int32 NumBatchElems = 4096*2;
+		const int32 NumBatches = BaseImage->DataStorage.GetNumBatchesLODRange(NumBatchElems, BytesPerElem, 0, NumLODs);
+
+		check(BytesPerElem == GetImageFormatData(DestImage->GetFormat()).BytesPerBlock);
+		check(BytesPerElem == GetImageFormatData(BlendImage->GetFormat()).BytesPerBlock);
+		
+		check(NumBatches == BlendImage->DataStorage.GetNumBatchesLODRange(NumBatchElems, BytesPerElem, 0, NumLODs));
+		check(NumBatches == DestImage->DataStorage.GetNumBatchesLODRange(NumBatchElems, BytesPerElem, 0, NumLODs));
+
+		auto ProcessBatch = 
+		[
+			BaseImage, BlendImage, DestImage, UnblendedChannels, NumBatchElems, BytesPerElem, NumLODs
+		] (int32 BatchId)
+		{
+			TArrayView<const uint8> BlendedView = 
+					BlendImage->DataStorage.GetBatchLODRange(BatchId, NumBatchElems, BytesPerElem, 0, NumLODs);
+
+			TArrayView<const uint8> BaseView =
+					BaseImage->DataStorage.GetBatchLODRange(BatchId, NumBatchElems, BytesPerElem, 0, NumLODs);
+
+			TArrayView<uint8> DestView =
+					DestImage->DataStorage.GetBatchLODRange(BatchId, NumBatchElems, BytesPerElem, 0, NumLODs);
+
+			uint8* const DestBuf = DestView.GetData();
+			uint8 const * const BlendedBuf = BlendedView.GetData();
+			uint8 const * const BaseBuf = BaseView.GetData(); 
+
+			const int32 NumElems = BaseView.Num() / BytesPerElem;
+			check(NumElems == DestView.Num() / BytesPerElem);
+			check(NumElems == BlendedView.Num() / BytesPerElem);
+
+			if (BlendImage->GetFormat() == EImageFormat::IF_RGBA_UBYTE)
+			{
+				const uint8* MaskBuf = BlendedBuf + 3;
+
+				for (int32 I = 0; I < NumElems; ++I)
+				{
+					uint32 Mask = MaskBuf[BLENDED_CHANNEL_STRIDE * I];
+					for (int32 C = 0; C < CHANNELS_TO_BLEND; ++C)
+					{
+						uint32 Base = BaseBuf[BASE_CHANNEL_STRIDE * I + C];
+						uint32 Blended = BlendedBuf[BLENDED_CHANNEL_STRIDE * I + C];
+						uint32 Result = BLEND_FUNC_MASKED(Base, Blended, Mask);
+						if constexpr (CLAMP)
+						{
+							DestBuf[BASE_CHANNEL_STRIDE * I + C] = (uint8)FMath::Min(255u, Result);
+						}
+						else
+						{
+							DestBuf[BASE_CHANNEL_STRIDE * I + C] = (uint8)Result;
+						}
+					}
+					// Copy the unblended channels
+					// \TODO: unnecessary when doing it in-place?
+					if constexpr (UnblendedChannels > 0)
+					{
+						for (int32 C = 0; C < UnblendedChannels; ++C)
+						{
+							DestBuf[BASE_CHANNEL_STRIDE * I + CHANNELS_TO_BLEND + C] = BaseBuf[BASE_CHANNEL_STRIDE * I + CHANNELS_TO_BLEND + C];
+						}
+					}
+				}
+			}
+			else
+			{
+				for (int32 I = 0; I < NumElems; ++I)
+				{
+					for (int32 C = 0; C < CHANNELS_TO_BLEND; ++C)
+					{
+						uint32 Base = BaseBuf[BASE_CHANNEL_STRIDE * I + C];
+						uint32 Blended = BlendedBuf[BLENDED_CHANNEL_STRIDE * I + C];
+						uint32 Result = BLEND_FUNC(Base, Blended);
+						if constexpr (CLAMP)
+						{
+							DestBuf[BASE_CHANNEL_STRIDE * I + C] = (uint8)FMath::Min(255u, Result);
+						}
+						else
+						{
+							DestBuf[BASE_CHANNEL_STRIDE * I + C] = (uint8)Result;
+						}
+					}
+					// Copy the unblended channels
+					// \TODO: unnecessary when doing it in-place?
+					if constexpr (UnblendedChannels > 0)
+					{
+						for (int32 C = 0; C < UnblendedChannels; ++C)
+						{
+							DestBuf[BASE_CHANNEL_STRIDE * I + CHANNELS_TO_BLEND + C] = BaseBuf[BASE_CHANNEL_STRIDE * I + CHANNELS_TO_BLEND + C];
+						}
+					}
+				}
+			}
+		};
+
+		if (NumBatches == 1)
+		{
+			ProcessBatch(0);
+		}
+		else if (NumBatches > 1)
+		{
+			ParallelFor(NumBatches, ProcessBatch);
+		}
+	}
+
+	template< 
+		uint32 (*BLEND_FUNC)(uint32, uint32),
+		bool CLAMP,
+		uint32 NC>
+    inline void BufferLayerFormatStrideNoAlpha(
+		Image* DestImage,
+		int32 DestOffset,
+		int32 Stride,
+		const Image* MaskImage,
+		const Image* BlendImage/*, int32 LODCount*/)
+	{
+        const uint8* MaskBuf = MaskImage->GetLODData(0);
+        const uint8* BlendedBuf = BlendImage->GetLODData(0);
+		uint8* DestBuf = DestImage->GetLODData(0) + DestOffset;
+
+		EImageFormat MaskFormat = MaskImage->GetFormat();
+        bool bIsUncompressed = (MaskFormat == EImageFormat::IF_L_UBYTE);
+
+        if (bIsUncompressed)
+		{
+			int32 RowCount = BlendImage->GetSizeY();
+			int32 PixelCount = BlendImage->GetSizeX();
+			for (int32 RowIndex = 0; RowIndex < RowCount; ++RowIndex)
+			{
+				for (int32 PixelIndex = 0; PixelIndex < PixelCount; ++PixelIndex)
+				{
+					uint32 Mask = *MaskBuf;
+					if (Mask)
+					{
+						for (int32 C = 0; C < NC; ++C)
+						{
+							uint32 Base = *DestBuf;
+							uint32 Blended = *BlendedBuf;
+							uint32 Result = BLEND_FUNC(Base, Blended);
+							if constexpr (CLAMP)
+							{
+								*DestBuf = (uint8)FMath::Min(255u, Result);
+							}
+							else
+							{
+								*DestBuf = (uint8)Result;
+							}
+							++DestBuf;
+							++BlendedBuf;
+						}
+					}
+					else
+					{
+						DestBuf += NC;
+						BlendedBuf += NC;
+					}
+					++MaskBuf;
+				}
+
+				DestBuf += Stride;
+			}
+		}
+        else if (MaskFormat == EImageFormat::IF_L_UBIT_RLE)
+		{
+			int32 Rows = MaskImage->GetSizeY();
+			int32 Width = MaskImage->GetSizeX();
+
+            //for (int32 lod = 0; lod < LODCount; ++lod)
+            //{
+				// Remove RLE header.
+                MaskBuf += 4 + Rows*sizeof(uint32);
+
+                for (int32 RowIndex = 0; RowIndex < Rows; ++RowIndex)
+                {
+                    const uint8* DestRowEnd = DestBuf + Width*NC;
+                    while (DestBuf != DestRowEnd)
+                    {
+                        // Decode header
+                        uint16 Zeros = *(const uint16*)MaskBuf;
+                        MaskBuf += 2;
+
+                        uint16 Ones = *(const uint16*)MaskBuf;
+                        MaskBuf += 2;
+
+                        // Skip
+                        DestBuf += Zeros*NC;
+                        BlendedBuf += Zeros*NC;
+
+                        // Copy
+                        FMemory::Memmove(DestBuf, BlendedBuf, Ones*NC);
+
+                        DestBuf += NC*Ones;
+                        BlendedBuf += NC*Ones;
+                    }
+
+                    DestBuf += Stride;
+                }
+
+                //Rows = FMath::DivideAndRoundUp(Rows, 2);
+                //Width = FMath::DivideAndRoundUp(Width, 2);
+            //}
+		}
+		else
+		{
+			checkf( false, TEXT("Unsupported mask format.") );
+		}
+	}
+
+	template<uint32(*BLEND_FUNC_MASKED)(uint32, uint32, uint32),
+			 uint32(*BLEND_FUNC)(uint32, uint32),
+			 bool CLAMP>
+    inline void BufferLayer(Image* DestImage,
+							const Image* BaseImage,
+							const Image* MaskImage,
+                            const Image* BlendImage,
+                            bool bApplyToAlpha,
+							bool bOnlyFirstLOD)
+	{
+		if (BaseImage->GetFormat() == EImageFormat::IF_RGB_UBYTE)
+		{
+            BufferLayerFormat<BLEND_FUNC_MASKED, BLEND_FUNC, CLAMP, 3, 3, 3>
+                    (DestImage, BaseImage, MaskImage, BlendImage, 0, 0, 0, bOnlyFirstLOD);
+		}
+        else if (BaseImage->GetFormat() == EImageFormat::IF_RGBA_UBYTE || 
+				 BaseImage->GetFormat() == EImageFormat::IF_BGRA_UBYTE)
+		{
+			if (bApplyToAlpha)
+			{
+				BufferLayerFormat<BLEND_FUNC_MASKED, BLEND_FUNC, CLAMP, 4, 4, 4>
+					(DestImage, BaseImage, MaskImage, BlendImage, 0, 0, 0, bOnlyFirstLOD);
+			}
+			else
+			{
+				BufferLayerFormat<BLEND_FUNC_MASKED, BLEND_FUNC, CLAMP, 3, 4, 4>
+					(DestImage, BaseImage, MaskImage, BlendImage, 0, 0, 0, bOnlyFirstLOD);
+			}
+		}
+		else if (BaseImage->GetFormat() == EImageFormat::IF_L_UBYTE)
+		{
+            BufferLayerFormat<BLEND_FUNC_MASKED, BLEND_FUNC, CLAMP, 1, 1, 1>
+                    (DestImage, BaseImage, MaskImage, BlendImage, 0, 0, 0, bOnlyFirstLOD);
+		}
+        else
+		{
+			checkf(false, TEXT("Unsupported format."));
+		}
+	}
+
+
+	template< uint32 (*BLEND_FUNC_MASKED)(uint32, uint32, uint32),
+		uint32 (*BLEND_FUNC)(uint32, uint32),
+		bool CLAMP >
+	inline void BufferLayerEmbeddedMask(
+		Image* DestImage,
+		const Image* BaseImage,
+		const Image* BlendImage,
+		bool bApplyToAlpha,
+		bool bOnlyFirstLOD)
+	{
+		if (BaseImage->GetFormat() == EImageFormat::IF_RGB_UBYTE)
+		{
+			BufferLayerFormatEmbeddedMask<BLEND_FUNC_MASKED, BLEND_FUNC, CLAMP, 3, 3, 3>
+				(DestImage, BaseImage, BlendImage, 0, 0, bOnlyFirstLOD);
+		}
+		else if (BaseImage->GetFormat() == EImageFormat::IF_RGBA_UBYTE || 
+				 BaseImage->GetFormat() == EImageFormat::IF_BGRA_UBYTE)
+		{
+			if (bApplyToAlpha)
+			{
+				BufferLayerFormatEmbeddedMask<BLEND_FUNC_MASKED, BLEND_FUNC, CLAMP, 4, 4, 4>
+					(DestImage, BaseImage, BlendImage, 0, 0, bOnlyFirstLOD);
+			}
+			else
+			{
+				BufferLayerFormatEmbeddedMask<BLEND_FUNC_MASKED, BLEND_FUNC, CLAMP, 3, 4, 4>
+					(DestImage, BaseImage, BlendImage, 0, 0, bOnlyFirstLOD);
+			}
+		}
+		else
+		{
+			checkf(false, TEXT("Unsupported format."));
+		}
+	}
+
+
+	template< 
+		uint32 (*RGB_FUNC_MASKED)(uint32, uint32, uint32),
+		uint32 (*A_FUNC)(uint32, uint32),
+		bool CLAMP >
+	inline void BufferLayerComposite(
+		Image* BaseImage,
+		const Image* BlendImage,
+		bool bOnlyFirstLOD,
+		uint8 BlendAlphaSourceChannel)
+	{
+		check(BaseImage->GetFormat() == EImageFormat::IF_RGBA_UBYTE);
+		check(BlendImage->GetFormat() == EImageFormat::IF_RGBA_UBYTE);
+		check(BaseImage->GetSizeX() == BlendImage->GetSizeX() && BaseImage->GetSizeY() == BlendImage->GetSizeY());
+		check(bOnlyFirstLOD || BaseImage->GetLODCount() <= BlendImage->GetLODCount());
+
+		bOnlyFirstLOD = bOnlyFirstLOD || BaseImage->GetLODCount() == 1;
+		
+		int32 FirstLODDataOffset = 0;	
+		int32 NumRelevantElems = -1;
+
+		if (BlendImage->m_flags & Image::IF_HAS_RELEVANCY_MAP && bOnlyFirstLOD)
+		{
+			check(BlendImage->RelevancyMaxY < BaseImage->GetSizeY());
+			check(BlendImage->RelevancyMaxY >= BlendImage->RelevancyMinY);
+
+			uint16 SizeX = BaseImage->GetSizeX();
+			NumRelevantElems = (BlendImage->RelevancyMaxY - BlendImage->RelevancyMinY + 1) * SizeX * 4;
+			
+			FirstLODDataOffset = BlendImage->RelevancyMinY * SizeX * 4;
+		}
+
+		const int32 BytesPerElem = GetImageFormatData(BaseImage->GetFormat()).BytesPerBlock;
+		check(BytesPerElem == GetImageFormatData(BlendImage->GetFormat()).BytesPerBlock);
+
+		const int32 LODBegin = 0;
+		const int32 LODEnd = BaseImage->GetLODCount();
+
+		constexpr int32 NumBatchElems = 4096*2;
+
+		const int32 NumBatches = bOnlyFirstLOD
+				? BaseImage->DataStorage.GetNumBatchesFirstLODOffset(NumBatchElems, BytesPerElem, FirstLODDataOffset)
+				: BaseImage->DataStorage.GetNumBatchesLODRange(NumBatchElems, BytesPerElem, LODBegin, LODEnd);
+
+		// This will always be an upper-bound for bOnlyFirtsLODs, check if it performs as expected or it needs more fine tune.
+		const int32 NumRelevantBatches = NumRelevantElems != -1
+				? FMath::DivideAndRoundUp(NumRelevantElems, NumBatchElems)
+				: NumBatches;
+
+		// This check only applies if bOnlyFirstLOD is false. The other case could give false negatives
+		// but that is already checked above.
+		check(NumBatches <= BlendImage->DataStorage.GetNumBatches(NumBatchElems, BytesPerElem));
+
+		auto ProcessBatch = 
+		[
+			BaseImage, BlendImage, BlendAlphaSourceChannel, BytesPerElem, NumBatchElems, bOnlyFirstLOD, LODBegin, LODEnd, FirstLODDataOffset
+		] (int32 BatchId)
+		{
+			TArrayView<uint8> BaseBatchView = bOnlyFirstLOD 
+					? BaseImage->DataStorage.GetBatchFirstLODOffset(BatchId, NumBatchElems, BytesPerElem, FirstLODDataOffset)
+					: BaseImage->DataStorage.GetBatchLODRange(BatchId, NumBatchElems, BytesPerElem, LODBegin, LODEnd);
+
+			TArrayView<const uint8> BlendBatchView = bOnlyFirstLOD 
+					? BlendImage->DataStorage.GetBatchFirstLODOffset(BatchId, NumBatchElems, BytesPerElem, FirstLODDataOffset)
+					: BlendImage->DataStorage.GetBatchLODRange(BatchId, NumBatchElems, BytesPerElem, LODBegin, LODEnd);
+
+			const int32 NumElems = BaseBatchView.Num() / BytesPerElem; 
+			check(BlendBatchView.Num() / BytesPerElem == NumElems);
+
+			uint8* BaseBuf = BaseBatchView.GetData();
+			const uint8* BlendBuf = BlendBatchView.GetData();
+
+			for (int32 I = 0; I < NumElems; ++I)
+			{
+				// TODO: Optimize this (SIMD?)
+				uint32 Mask = BlendBuf[4 * I + 3];
+
+				// RGB
+				for (int32 C = 0; C < 3; ++C)
+				{
+					uint32 Base = BaseBuf[4 * I + C];
+					uint32 Blended = BlendBuf[4 * I + C];
+					uint32 Result = RGB_FUNC_MASKED(Base, Blended, Mask);
+					if constexpr (CLAMP)
+					{
+						BaseBuf[4 * I + C] = (uint8)FMath::Min(255u, Result);
+					}
+					else
+					{
+						BaseBuf[4 * I + C] = (uint8)Result;
+					}
+				}
+
+				// A
+				{
+					uint32 Base = BaseBuf[4 * I + 3];
+					uint32 Blended = BlendBuf[4 * I + BlendAlphaSourceChannel];
+					uint32 Result = A_FUNC(Base, Blended);
+					if constexpr (CLAMP)
+					{
+						BaseBuf[4 * I + 3] = (uint8)FMath::Min(255u, Result);
+					}
+					else
+					{
+						BaseBuf[4 * I + 3] = (uint8)Result;
+					}
+				}
+			}
+		};
+
+		if (NumRelevantBatches == 1)
+		{
+			ProcessBatch(0);
+		}
+		else if (NumRelevantBatches > 1)
+		{
+			ParallelFor(NumRelevantBatches, ProcessBatch);
+		}
+	}
+
+	template<>
+	inline void BufferLayerComposite<BlendChannelMasked, LightenChannel, false>
+	(
+		Image* BaseImage,
+		const Image* BlendImage,
+		bool bOnlyFirstLOD,
+		uint8 BlendAlphaSourceChannel)
+	{
+		check(BaseImage->GetFormat() == EImageFormat::IF_RGBA_UBYTE);
+		check(BlendImage->GetFormat() == EImageFormat::IF_RGBA_UBYTE);
+		check(BaseImage->GetSizeX() == BlendImage->GetSizeX() && BaseImage->GetSizeY() == BlendImage->GetSizeY());
+		check(bOnlyFirstLOD || BaseImage->GetLODCount() <= BlendImage->GetLODCount());
+		check(BlendAlphaSourceChannel == 3);
+
+		bOnlyFirstLOD = bOnlyFirstLOD || BaseImage->GetLODCount() == 1;
+
+		int32 FirstLODDataOffset = 0;
+		int32 NumRelevantElems = -1;
+		if (BlendImage->m_flags & Image::IF_HAS_RELEVANCY_MAP && bOnlyFirstLOD)
+		{
+			check(BlendImage->RelevancyMaxY < BaseImage->GetSizeY());
+			check(BlendImage->RelevancyMaxY >= BlendImage->RelevancyMinY);
+
+			uint16 SizeX = BaseImage->GetSizeX();
+			NumRelevantElems = (BlendImage->RelevancyMaxY - BlendImage->RelevancyMinY + 1) * SizeX;
+			
+			FirstLODDataOffset = BlendImage->RelevancyMinY * SizeX * 4;
+		}
+
+		// \TODO: Ensure batches operate on 32Kb aligned buffers?
+		constexpr int32 NumBatchElems = 4096 * 2;
+
+		constexpr int32 BytesPerElem = 4;
+
+		const int32 NumLODs = BaseImage->GetLODCount();
+		const int32 NumBatches = bOnlyFirstLOD
+				? BaseImage->DataStorage.GetNumBatchesFirstLODOffset(NumBatchElems, BytesPerElem, FirstLODDataOffset)
+				: BaseImage->DataStorage.GetNumBatchesLODRange(NumBatchElems, BytesPerElem, 0, NumLODs);
+
+		// This will always be an upper-bound for bOnlyFirtsLODs, check if it performs as expected or it needs more fine tune.
+		const int32 NumRelevantBatches = NumRelevantElems != -1
+				? FMath::DivideAndRoundUp(NumRelevantElems, NumBatchElems)
+				: NumBatches;
+
+		// This check only applies if bOnlyFirstLOD is false. The other case could give false negatives
+		// but that is already checked above.
+		check(NumBatches <= BlendImage->DataStorage.GetNumBatchesLODRange(NumBatchElems, BytesPerElem, 0, NumLODs));
+		
+		auto ProcessBatch =
+			[
+				BaseImage, BlendImage, NumBatchElems, bOnlyFirstLOD, FirstLODDataOffset, NumLODs
+			] (uint32 BatchId)
+		{
+			TArrayView<uint8> BaseBatchView = bOnlyFirstLOD 
+					? BaseImage->DataStorage.GetBatchFirstLODOffset(BatchId, NumBatchElems, BytesPerElem, FirstLODDataOffset)
+					: BaseImage->DataStorage.GetBatchLODRange(BatchId, NumBatchElems, BytesPerElem, 0, NumLODs);
+
+			TArrayView<const uint8> BlendBatchView = bOnlyFirstLOD 
+					? BlendImage->DataStorage.GetBatchFirstLODOffset(BatchId, NumBatchElems, BytesPerElem, FirstLODDataOffset)
+					: BlendImage->DataStorage.GetBatchLODRange(BatchId, NumBatchElems, BytesPerElem, 0, NumLODs);
+
+			const int32 NumElems = BaseBatchView.Num() / BytesPerElem; 
+			check(NumElems == BlendBatchView.Num() / BytesPerElem);
+			
+			uint8* BaseBuf = BaseBatchView.GetData();
+			const uint8* BlendBuf = BlendBatchView.GetData();
+			
+			for (int32 I = 0; I < NumElems; ++I)
+			{
+				// TODO: Optimize this (SIMD?)
+				uint32 FullBase;
+				FMemory::Memcpy(&FullBase, BaseBuf + I*sizeof(uint32), sizeof(uint32)); 
+
+				uint32 FullBlended;
+				FMemory::Memcpy(&FullBlended, BlendBuf + I*sizeof(uint32), sizeof(uint32)); 
+				uint32 Mask = (FullBlended & 0xff000000) >> 24;
+
+				uint32 FullResult = 0;
+				FullResult |= BlendChannelMasked((FullBase >>  0) & 0xff, (FullBlended >>  0) & 0xff, Mask) << 0;
+				FullResult |= BlendChannelMasked((FullBase >>  8) & 0xff, (FullBlended >>  8) & 0xff, Mask) << 8;
+				FullResult |= BlendChannelMasked((FullBase >> 16) & 0xff, (FullBlended >> 16) & 0xff, Mask) << 16;
+				FullResult |= LightenChannel	((FullBase >> 24) & 0xff, (FullBlended >> 24) & 0xff) << 24;
+
+				FMemory::Memcpy(BaseBuf + I*sizeof(uint32), &FullResult, sizeof(uint32));
+			}
+		};
+
+		if (NumRelevantBatches == 1)
+		{
+			ProcessBatch(0);
+		}
+		else if (NumRelevantBatches > 1)
+		{
+			ParallelFor(NumBatches, ProcessBatch);
+		}
+	}
+
+	template< 
+		VectorRegister4Int (*RGB_FUNC_MASKED)(const VectorRegister4Int&, const VectorRegister4Int&, const VectorRegister4Int&),
+		int32 (*A_FUNC)(int32, int32),
+		bool CLAMP >
+	inline void BufferLayerCompositeVector(
+		Image* BaseImage,
+		const Image* BlendImage,
+		bool bOnlyFirstLOD,
+		uint8 BlendAlphaSourceChannel)
+	{
+		check(BaseImage->GetFormat() == EImageFormat::IF_RGBA_UBYTE);
+		check(BlendImage->GetFormat() == EImageFormat::IF_RGBA_UBYTE);
+		check(BaseImage->GetSizeX() == BlendImage->GetSizeX() && BaseImage->GetSizeY() == BlendImage->GetSizeY());
+		check(bOnlyFirstLOD || BaseImage->GetLODCount() <= BlendImage->GetLODCount());
+
+		constexpr int32 NumBatchElems = 4096*2;
+		constexpr int32 BytesPerElem = 4;
+
+		const int32 NumLODs = bOnlyFirstLOD ? 1 : BaseImage->GetLODCount();
+
+		const int32 NumBatches = BaseImage->DataStorage.GetNumBatchesLODRange(NumBatchElems, BytesPerElem, 0, NumLODs);
+	
+		check(NumBatches == BlendImage->DataStorage.GetNumBatchesLODRange(NumBatchElems, BytesPerElem, 0, NumLODs));
+
+		auto ProcessBatch = 
+		[
+			BaseImage, BlendImage, BytesPerElem, NumBatchElems, BlendAlphaSourceChannel, NumLODs
+		] (uint32 BatchId)
+		{
+			const TArrayView<uint8> BaseBatchView =
+					BaseImage->DataStorage.GetBatchLODRange(BatchId, NumBatchElems, BytesPerElem, 0, NumLODs);
+
+			const TArrayView<const uint8> BlendBatchView =
+					BlendImage->DataStorage.GetBatchLODRange(BatchId, NumBatchElems, BytesPerElem, 0, NumLODs);
+
+			uint8* BaseBuf = BaseBatchView.GetData();
+			const uint8* BlendBuf = BlendBatchView.GetData();
+
+			const int32 NumElems = BaseBatchView.Num() / BytesPerElem;
+			check(NumElems == BlendBatchView.Num()/BytesPerElem);
+
+			for (int32 I = 0; I < NumElems; ++I)
+			{
+				// TODO: Optimize this (SIMD?)
+				const int32 BaseAlpha = BaseBuf[4 * I + BlendAlphaSourceChannel];
+				const int32 BlendedAlpha = BlendBuf[4 * I + BlendAlphaSourceChannel];
+
+				const VectorRegister4Int Mask = VectorIntSet1(BlendedAlpha);
+
+				uint32 BlendPixel;
+				FMemory::Memcpy(&BlendPixel, BlendBuf, sizeof(uint32));
+
+				uint32 BasePixel;
+				FMemory::Memcpy(&BasePixel, BaseBuf, sizeof(uint32));
+
+				const VectorRegister4Int Blended = MakeVectorRegisterInt(
+						(BlendPixel >> 0) & 0xFF, 
+						(BlendPixel >> 8) & 0xFF, 
+						(BlendPixel >> 16) & 0xFF, 
+						(BlendPixel >> 24) & 0xFF);
+
+				const VectorRegister4Int Base = MakeVectorRegisterInt(
+						(BasePixel >> 0) & 0xFF, 
+						(BasePixel >> 8) & 0xFF, 
+						(BasePixel >> 16) & 0xFF, 
+						(BasePixel >> 24) & 0xFF);
+
+				VectorRegister4Int Result = RGB_FUNC_MASKED(Base, Blended, Mask);
+				if constexpr (CLAMP)
+				{
+					Result = VectorIntMin(MakeVectorRegisterIntConstant(255, 255, 255, 255), Result);
+				}
+
+				int32 AlphaResult = A_FUNC(BaseAlpha, BlendedAlpha);
+				if constexpr (CLAMP)
+				{
+					AlphaResult = FMath::Min(255, AlphaResult);
+				}
+
+				alignas(VectorRegister4Int) int32 IndexableRegister[4];
+				VectorIntStoreAligned(Result, &IndexableRegister);
+
+				BaseBuf[4 * I + 0] = static_cast<uint8>(IndexableRegister[0]);
+				BaseBuf[4 * I + 1] = static_cast<uint8>(IndexableRegister[1]);
+				BaseBuf[4 * I + 2] = static_cast<uint8>(IndexableRegister[2]);
+				BaseBuf[4 * I + 3] = static_cast<uint8>(IndexableRegister[3]);
+
+				BaseBuf[4 * I + BlendAlphaSourceChannel] = static_cast<uint8>(AlphaResult);
+			}
+		};
+
+		if (NumBatches == 1)
+		{
+			ProcessBatch(0);
+		}
+		else if (NumBatches > 1)
+		{
+			ParallelFor(NumBatches, ProcessBatch);
+		}
+	}
+
+	template<uint32 (*BLEND_FUNC)(uint32,uint32),
+			 bool CLAMP>
+    inline void BufferLayerStrideNoAlpha(Image* DestImage, int32 DestOffset, int32 Stride, const Image* MaskImage, const Image* BlendImage/*, int32 LODCount*/)
+	{
+		if (BlendImage->GetFormat() == EImageFormat::IF_RGB_UBYTE)
+		{
+			BufferLayerFormatStrideNoAlpha<BLEND_FUNC, CLAMP, 3>
+					(DestImage, DestOffset, Stride, MaskImage, BlendImage/*, LODCount*/);
+		}
+        else if (BlendImage->GetFormat() == EImageFormat::IF_RGBA_UBYTE || 
+				 BlendImage->GetFormat() == EImageFormat::IF_BGRA_UBYTE)
+		{
+			BufferLayerFormatStrideNoAlpha<BLEND_FUNC, CLAMP, 4>
+					(DestImage, DestOffset, Stride, MaskImage, BlendImage/*, LODCount*/);
+		}
+		else if (BlendImage->GetFormat() == EImageFormat::IF_L_UBYTE)
+		{
+			BufferLayerFormatStrideNoAlpha<BLEND_FUNC, CLAMP, 1>
+					(DestImage, DestOffset, Stride, MaskImage, BlendImage/*, LODCount*/);
+		}
+		else
+		{
+			checkf( false, TEXT("Unsupported format.") );
+		}
+	}
+
+
+	/**
+	* Apply a blending function to an image with another image as blending layer, on a subrect of
+	* the base image.
+	* \warning this method applies the blending function to the alpha channel too
+	* \warning this method uses the mask as a binary mask (>0)
+	*/
+	template<uint32 (*BLEND_FUNC)(uint32, uint32), bool CLAMP>
+	inline void ImageLayerOnBaseNoAlpha(
+			Image* BaseImage,
+			const Image* MaskImage,
+			const Image* BlendedImage,
+			const box<FIntVector2>& Rect)
+	{
+		check(BaseImage->GetSizeX() >= Rect.min[0] + Rect.size[0]);
+		check(BaseImage->GetSizeY() >= Rect.min[1] + Rect.size[1]);
+		check(MaskImage->GetSizeX() == BlendedImage->GetSizeX());
+		check(MaskImage->GetSizeY() == BlendedImage->GetSizeY());
+		check(BaseImage->GetFormat() == BlendedImage->GetFormat());
+		check(MaskImage->GetFormat() == EImageFormat::IF_L_UBYTE ||
+			  //UBYTE_RLE does not look to be supported.
+			  //MaskImage->GetFormat() == EImageFormat::IF_L_UBYTE_RLE || 
+			  MaskImage->GetFormat() == EImageFormat::IF_L_UBIT_RLE);
+        check(BaseImage->GetLODCount() <= MaskImage->GetLODCount());
+        check(BaseImage->GetLODCount() <= BlendedImage->GetLODCount());
+
+		int32 PixelSize = GetImageFormatData(BaseImage->GetFormat()).BytesPerBlock;
+
+		int32 Start = (BaseImage->GetSizeX() * Rect.min[1] + Rect.min[0]) * PixelSize;
+		int32 Stride = (BaseImage->GetSizeX() - Rect.size[0]) * PixelSize;
+
+		// Stride is only valid for LOD 0, BufferLayerStride variants cannot operate on multiple lods.
+		// TODO: review if this needs to be supported, and implement using a rect lod reducction at this level.
+		BufferLayerStrideNoAlpha<BLEND_FUNC, CLAMP>(BaseImage, Start, Stride, MaskImage, BlendedImage/*, BaseImage->GetLODCount()*/);
+	}
+
+	template<uint32 NC>
+	FORCEINLINE uint32 PackPixel(const uint8* PixelPtr)
+	{
+		static_assert(NC > 0 && NC <= 4);
+
+		uint32 PixelPack = 0;
+
+		// The compiler should be able to optimize this given that NC is a constant expression.
+		FMemory::Memcpy(&PixelPack, PixelPtr, NC);
+
+		return PixelPack;
+	}
+
+	template<uint32 NC>
+	FORCEINLINE void UnpackPixel(uint8* PixelPtr, uint32 PixelData)
+	{
+		static_assert(NC > 0 && NC <= 4);
+
+		// The compiler should be able to optimize this given that NC is a constant expression
+		FMemory::Memcpy(PixelPtr, &PixelData, NC);
+	}
+
+	template< 
+		uint32 (*BLEND_FUNC)(uint32, uint32), 
+		uint32 NC>
+	inline void BufferLayerCombineColour(Image* ResultImage, const Image* BaseImage, FVector4f Color, bool bOnlyFirstLOD = false)
+	{
+		static_assert(NC > 0 && NC <= 4);
+
+		check(BaseImage->GetSizeX() == ResultImage->GetSizeX());
+		check(BaseImage->GetSizeY() == ResultImage->GetSizeY());
+
+		const uint32 TopColor = 
+			static_cast<uint32>(255.0f * Color[0]) << 0 |  
+			static_cast<uint32>(255.0f * Color[1]) << 8 |  
+			static_cast<uint32>(255.0f * Color[2]) << 16;  
+
+		const int32 NumLODs = bOnlyFirstLOD ? 1 : ResultImage->GetLODCount();
+		
+		constexpr int32 BatchNumElems = 4098*2;
+		const int32 NumBatches = BaseImage->DataStorage.GetNumBatchesLODRange(BatchNumElems, NC, 0, NumLODs);
+
+		check(NumBatches == ResultImage->DataStorage.GetNumBatchesLODRange(BatchNumElems, NC, 0, NumLODs));
+
+		auto ProcessBatch = [ResultImage, BaseImage, BatchNumElems, TopColor, NumLODs](int32 BatchId)
+		{
+			const TArrayView<const uint8> BaseBatchView = 
+					BaseImage->DataStorage.GetBatchLODRange(BatchId, BatchNumElems, NC, 0, NumLODs);
+			
+			const TArrayView<uint8> ResultBatchView =
+					ResultImage->DataStorage.GetBatchLODRange(BatchId, BatchNumElems, NC, 0, NumLODs);
+
+			const int32 NumElems = BaseBatchView.Num() / NC;
+
+			check(NumElems == ResultBatchView.Num() / NC);
+			
+			const uint8* BaseBuf = BaseBatchView.GetData();
+			uint8* ResultBuf = ResultBatchView.GetData();
+			
+			for (int32 I = 0; I < NumElems; ++I)
+			{
+				const uint32 Base = PackPixel<NC>(&BaseBuf[NC * I]);
+
+				const uint32 Result = BLEND_FUNC(Base, TopColor);
+				UnpackPixel<NC>(&ResultBuf[NC * I], Result);
+			}
+		};
+
+		if (NumBatches == 1)
+		{
+			ProcessBatch(0);
+		}
+		else if (NumBatches > 1)
+		{
+			ParallelFor(NumBatches, ProcessBatch);
+		}
+	}
+
+	template< 
+		uint32 (*BLEND_FUNC)(uint32, uint32), 
+		uint32 NC>
+	inline void BufferLayerCombine(Image* ResultImage, const Image* BaseImage, const Image* BlendImage, bool bOnlyFirstLOD)
+	{
+		static_assert(NC > 0 && NC <= 4);
+
+		check(BaseImage->GetSizeX() == BlendImage->GetSizeX());
+		check(BaseImage->GetSizeY() == BlendImage->GetSizeY());
+		check(BaseImage->GetSizeX() == ResultImage->GetSizeX());
+		check(BaseImage->GetSizeY() == ResultImage->GetSizeY());
+		check(BaseImage->GetFormat() == BlendImage->GetFormat());
+		check(BaseImage->GetFormat() == ResultImage->GetFormat());
+
+		const int32 NumLODs = bOnlyFirstLOD ? 1 : ResultImage->GetLODCount();
+
+		constexpr int32 BatchNumElems = 4098*2;
+		const int32 NumBatches = BaseImage->DataStorage.GetNumBatchesLODRange(BatchNumElems, NC, 0, NumLODs);
+
+		check(NumBatches == BlendImage->DataStorage.GetNumBatchesLODRange(BatchNumElems, NC, 0, NumLODs));
+		check(NumBatches == ResultImage->DataStorage.GetNumBatchesLODRange(BatchNumElems, NC, 0, NumLODs));
+
+		auto ProcessBatch = [ResultImage, BaseImage, BlendImage, BatchNumElems, NumLODs](int32 BatchId)
+		{
+			const TArrayView<const uint8> BaseBatchView = 
+					BaseImage->DataStorage.GetBatchLODRange(BatchId, BatchNumElems, NC, 0, NumLODs);
+
+			const TArrayView<const uint8> BlendBatchView = 
+					BlendImage->DataStorage.GetBatchLODRange(BatchId, BatchNumElems, NC, 0, NumLODs);
+
+			const TArrayView<uint8> ResultBatchView =
+					ResultImage->DataStorage.GetBatchLODRange(BatchId, BatchNumElems, NC, 0, NumLODs);
+
+			const int32 NumElems = BaseBatchView.Num() / NC;
+			check(NumElems == BlendBatchView.Num() / NC);
+			check(NumElems == ResultBatchView.Num() / NC);
+			
+			const uint8* BaseBuf = BaseBatchView.GetData();
+			const uint8* BlendBuf = BlendBatchView.GetData();
+			uint8* ResultBuf = ResultBatchView.GetData();
+			
+			for (int32 I = 0; I < NumElems; ++I)
+			{
+				const uint32 Base = PackPixel<NC>(&BaseBuf[NC * I]);
+				const uint32 Blend = PackPixel<NC>(&BlendBuf[NC * I]);
+
+				const uint32 Result = BLEND_FUNC(Base, Blend);
+				UnpackPixel<NC>(&ResultBuf[NC * I], Result);
+			}
+		};
+
+		if (NumBatches == 1)
+		{
+			ProcessBatch(0);
+		}
+		else if (NumBatches > 1)
+		{
+			ParallelFor(NumBatches, ProcessBatch);
+		}
+	}
+
+	template< 
+		uint32 (*BLEND_FUNC_MASKED)(uint32, uint32, uint32), 
+		uint32 NC>
+	inline void BufferLayerCombine(Image* ResultImage, const Image* BaseImage, const Image* MaskImage, const Image* BlendImage, bool bOnlyFirstLOD)
+	{
+		static_assert(NC > 0 && NC <= 4);
+
+		check(BaseImage->GetSizeX() == BlendImage->GetSizeX());
+		check(BaseImage->GetSizeY() == BlendImage->GetSizeY());
+		check(BaseImage->GetSizeX() == MaskImage->GetSizeX());
+		check(BaseImage->GetSizeY() == MaskImage->GetSizeY());
+		check(BaseImage->GetFormat() == BlendImage->GetFormat());
+
+		const int32 NumLODs = bOnlyFirstLOD ? 1 : ResultImage->GetLODCount();
+
+		constexpr int32 BatchNumElems = 4098*2;
+		const int32 NumBatches = BaseImage->DataStorage.GetNumBatchesLODRange(BatchNumElems, NC, 0, NumLODs);
+
+		check(NumBatches <= MaskImage->DataStorage.GetNumBatchesLODRange(BatchNumElems, 1, 0, NumLODs));
+		check(NumBatches <= BlendImage->DataStorage.GetNumBatchesLODRange(BatchNumElems, NC, 0, NumLODs));
+		check(NumBatches <= ResultImage->DataStorage.GetNumBatchesLODRange(BatchNumElems, NC, 0, NumLODs));
+
+		auto ProcessBatch = [ResultImage, BaseImage, MaskImage, BlendImage, BatchNumElems, NumLODs](int32 BatchId)
+		{
+			const TArrayView<const uint8> BaseBatchView = 
+					BaseImage->DataStorage.GetBatchLODRange(BatchId, BatchNumElems, NC, 0, NumLODs);
+
+			const TArrayView<const uint8> BlendBatchView = 
+					BlendImage->DataStorage.GetBatchLODRange(BatchId, BatchNumElems, NC, 0, NumLODs);
+
+			const TArrayView<const uint8> MaskBatchView =
+					MaskImage->DataStorage.GetBatchLODRange(BatchId, BatchNumElems, 1, 0, NumLODs);
+
+			const TArrayView<uint8> ResultBatchView =
+					ResultImage->DataStorage.GetBatchLODRange(BatchId, BatchNumElems, NC, 0, NumLODs);
+
+			const int32 NumElems = BaseBatchView.Num() / NC;
+
+			check(NumElems == BlendBatchView.Num() / NC);
+			check(NumElems == MaskBatchView.Num() / 1);
+			check(NumElems == ResultBatchView.Num() / NC);
+			
+			const uint8* BaseBuf = BaseBatchView.GetData();
+			const uint8* BlendBuf = BlendBatchView.GetData();
+			const uint8* MaskBuf = MaskBatchView.GetData();
+			uint8* ResultBuf = ResultBatchView.GetData();
+			
+			for (int32 I = 0; I < NumElems; ++I)
+			{
+				const uint32 Base = PackPixel<NC>(&BaseBuf[NC * I]);
+				const uint32 Blend = PackPixel<NC>(&BlendBuf[NC * I]);
+				const uint32 Mask = PackPixel<1>(&MaskBuf[1 * I]);
+
+				const uint32 Result = BLEND_FUNC_MASKED(Base, Blend, Mask);
+				UnpackPixel<NC>(&ResultBuf[NC * I], Result);
+			}
+		};
+
+		if (NumBatches == 1)
+		{
+			ProcessBatch(0);
+		}
+		else if (NumBatches > 1)
+		{
+			ParallelFor(NumBatches, ProcessBatch);
+		}
+	}
+
+	template< 
+		uint32 (*BLEND_FUNC_MASKED)(uint32, uint32, uint32), 
+		uint32 NC>
+	inline void BufferLayerCombineColour(Image* ResultImage, const Image* BaseImage, const Image* MaskImage, FVector4f Color)
+	{
+		static_assert(NC > 0 && NC <= 4);
+
+		check(BaseImage->GetSizeX() == MaskImage->GetSizeX());
+		check(BaseImage->GetSizeY() == MaskImage->GetSizeY());
+
+		const uint32 TopColor = 
+			static_cast<uint32>(255.0f * Color[0]) << 0 |  
+			static_cast<uint32>(255.0f * Color[1]) << 8 |  
+			static_cast<uint32>(255.0f * Color[2]) << 16;  
+
+		constexpr int32 BatchNumElems = 4098*2;
+		int32 NumBatches = BaseImage->DataStorage.GetNumBatches(BatchNumElems, NC);
+
+		check(NumBatches == MaskImage->DataStorage.GetNumBatches(BatchNumElems, 1));
+		check(NumBatches == ResultImage->DataStorage.GetNumBatches(BatchNumElems, NC));
+
+		auto ProcessBatch = [ResultImage, BaseImage, MaskImage, BatchNumElems, TopColor](int32 BatchId)
+		{
+			const TArrayView<const uint8> BaseBatchView = BaseImage->DataStorage.GetBatch(BatchId, BatchNumElems, NC);
+			const TArrayView<const uint8> MaskBatchView = MaskImage->DataStorage.GetBatch(BatchId, BatchNumElems, 1);
+			const TArrayView<uint8> ResultBatchView = ResultImage->DataStorage.GetBatch(BatchId, BatchNumElems, NC);
+
+			const int32 NumElems = BaseBatchView.Num() / NC;
+
+			check(NumElems == MaskBatchView.Num() / 1);
+			check(NumElems == ResultBatchView.Num() / NC);
+			
+			const uint8* BaseBuf = BaseBatchView.GetData();
+			const uint8* MaskBuf = MaskBatchView.GetData();
+			uint8* ResultBuf = ResultBatchView.GetData();
+			
+			for (int32 I = 0; I < NumElems; ++I)
+			{
+				const uint32 Base = PackPixel<NC>(&BaseBuf[NC * I]);
+				const uint32 Mask = PackPixel<1>(&MaskBuf[1 * I]);
+
+				const uint32 Result = BLEND_FUNC_MASKED(Base, TopColor, Mask);
+				UnpackPixel<NC>(&ResultBuf[NC * I], Result);
+			}
+		};
+
+		if (NumBatches == 1)
+		{
+			ProcessBatch(0);
+		}
+		else if (NumBatches > 1)
+		{
+			ParallelFor(NumBatches, ProcessBatch);
+		}
+	}
+
+	template<uint32 (*BLEND_FUNC)(uint32, uint32)>
+	inline void ImageLayerCombine(Image* ResultImage, const Image* BaseImage, const Image* BlendedImage, bool bOnlyFirstLOD)
+	{
+		check(ResultImage->GetFormat() == BaseImage->GetFormat());
+		check(ResultImage->GetSizeX() == BaseImage->GetSizeX());
+		check(ResultImage->GetSizeY() == BaseImage->GetSizeY());
+		check(bOnlyFirstLOD || ResultImage->GetLODCount() == BaseImage->GetLODCount());
+		check(BaseImage->GetSizeX() == BlendedImage->GetSizeX());
+		check(BaseImage->GetSizeY() == BlendedImage->GetSizeY());
+		check(BaseImage->GetFormat() == BlendedImage->GetFormat());
+		check(bOnlyFirstLOD || ResultImage->GetLODCount() <= BlendedImage->GetLODCount());
+
+		const EImageFormat BaseFormat = BaseImage->GetFormat();
+
+		if (BaseFormat == EImageFormat::IF_L_UBYTE)
+		{
+			BufferLayerCombine<BLEND_FUNC, 1>(ResultImage, BaseImage, BlendedImage, bOnlyFirstLOD);
+		}
+		else if (BaseFormat == EImageFormat::IF_RGB_UBYTE)
+		{
+			BufferLayerCombine<BLEND_FUNC, 3>(ResultImage, BaseImage, BlendedImage, bOnlyFirstLOD);
+		}
+		else if (BaseFormat == EImageFormat::IF_RGBA_UBYTE || BaseFormat == EImageFormat::IF_BGRA_UBYTE)
+		{
+			// \todo: pass swizzle template argument if BGRA_UBYTE, not yet supported.
+			BufferLayerCombine<BLEND_FUNC, 4>(ResultImage, BaseImage, BlendedImage, bOnlyFirstLOD);
+		}
+		else
+		{
+			checkf(false, TEXT("Unsupported format."));
+		}
+	}
+
+	template< 
+		uint32 (*BLEND_FUNC)(uint32, uint32), 
+		uint32 (*BLEND_FUNC_MASKED)(uint32, uint32, uint32)>
+	inline void ImageLayerCombine(Image* ResultImage, const Image* BaseImage, const Image* MaskImage, const Image* BlendedImage, bool bOnlyFirstLOD)
+	{
+		check(ResultImage->GetFormat() == BaseImage->GetFormat());
+		check(ResultImage->GetSizeX() == BaseImage->GetSizeX());
+		check(ResultImage->GetSizeY() == BaseImage->GetSizeY());
+		check(bOnlyFirstLOD || ResultImage->GetLODCount() == BaseImage->GetLODCount());
+		check(BaseImage->GetSizeX() == BlendedImage->GetSizeX());
+		check(BaseImage->GetSizeY() == BlendedImage->GetSizeY());
+		check(BaseImage->GetFormat() == BlendedImage->GetFormat());
+		check(bOnlyFirstLOD || ResultImage->GetLODCount() <= BlendedImage->GetLODCount());
+
+		const EImageFormat BaseFormat = BaseImage->GetFormat();
+
+		Ptr<Image> TempMaskImage;
+		if (MaskImage->GetFormat() != EImageFormat::IF_L_UBYTE)
+		{
+			UE_LOG(LogMutableCore, Log, TEXT("Image layer format not supported. A generic one will be used. "));
+
+			FImageOperator ImOp = FImageOperator::GetDefault(nullptr);
+			constexpr int32 Quality = 4;
+			TempMaskImage = ImOp.ImagePixelFormat( Quality, MaskImage, EImageFormat::IF_L_UBYTE );
+			MaskImage = TempMaskImage.get();
+		}
+
+		if (BaseFormat == EImageFormat::IF_L_UBYTE)
+		{
+			BufferLayerCombine<BLEND_FUNC_MASKED, 1>(ResultImage, BaseImage, MaskImage, BlendedImage, bOnlyFirstLOD);
+		}
+		else if (BaseFormat == EImageFormat::IF_RGB_UBYTE)
+		{
+			BufferLayerCombine<BLEND_FUNC_MASKED, 3>(ResultImage, BaseImage, MaskImage, BlendedImage, bOnlyFirstLOD);
+		}
+		else if (BaseFormat == EImageFormat::IF_RGBA_UBYTE || BaseFormat == EImageFormat::IF_BGRA_UBYTE)
+		{
+			// \todo: pass swizzle template argument if BGRA_UBYTE, not yet supported.
+			BufferLayerCombine<BLEND_FUNC_MASKED, 4>(ResultImage, BaseImage, MaskImage, BlendedImage, bOnlyFirstLOD);
+		}
+		else
+		{
+			UE_LOG(LogMutableCore, Log, TEXT("Image layer format not supported. A generic one will be used. "));
+
+			FImageOperator ImOp = FImageOperator::GetDefault(nullptr);
+			constexpr int32 Quality = 4;
+			Ptr<Image> TempBaseImage = ImOp.ImagePixelFormat(Quality, BaseImage, EImageFormat::IF_RGBA_UBYTE);
+			Ptr<Image> TempBlededImage = ImOp.ImagePixelFormat(Quality, BlendedImage, EImageFormat::IF_RGBA_UBYTE);
+			BufferLayerCombine<BLEND_FUNC_MASKED, 4>(ResultImage, TempBaseImage.get(), MaskImage, TempBlededImage.get(), bOnlyFirstLOD);
+		}
+	}
+
+	template<uint32 (*BLEND_FUNC)(uint32, uint32)>
+	inline void ImageLayerCombineColour(Image* ResultImage, const Image* BaseImage, FVector4f Color)
+	{
+		check(ResultImage->GetFormat() == BaseImage->GetFormat());
+		check(ResultImage->GetSizeX() == BaseImage->GetSizeX());
+		check(ResultImage->GetSizeY() == BaseImage->GetSizeY());
+		check(ResultImage->GetLODCount() == BaseImage->GetLODCount());
+
+		const EImageFormat BaseFormat = BaseImage->GetFormat();
+
+		if (BaseFormat == EImageFormat::IF_L_UBYTE)
+		{
+			BufferLayerCombineColour<BLEND_FUNC, 1>(ResultImage, BaseImage, Color);
+		}
+		else if (BaseFormat == EImageFormat::IF_RGB_UBYTE)
+		{
+			BufferLayerCombineColour<BLEND_FUNC, 3>(ResultImage, BaseImage, Color);
+		}
+		else if (BaseFormat == EImageFormat::IF_RGBA_UBYTE || BaseFormat == EImageFormat::IF_BGRA_UBYTE)
+		{
+			// \todo: pass swizzle template argument if BGRA_UBYTE, not yet supported.
+ 			BufferLayerCombineColour<BLEND_FUNC, 4>(ResultImage, BaseImage, Color);
+		}
+		else
+		{
+			checkf(false, TEXT("Unsupported format."));
+		}
+	}
+
+	template< 
+		uint32 (*BLEND_FUNC)(uint32, uint32),
+		uint32 (*BLEND_FUNC_MASKED)(uint32, uint32, uint32)>
+	inline void ImageLayerCombineColour(Image* ResultImage, const Image* BaseImage, const Image* MaskImage, FVector4f Color)
+	{
+		check(ResultImage->GetFormat() == BaseImage->GetFormat());
+		check(ResultImage->GetSizeX() == BaseImage->GetSizeX());
+		check(ResultImage->GetSizeY() == BaseImage->GetSizeY());
+		check(ResultImage->GetLODCount() == BaseImage->GetLODCount());
+
+		const EImageFormat BaseFormat = BaseImage->GetFormat();
+
+		if (MaskImage->GetFormat() != EImageFormat::IF_L_UBYTE)
+		{
+			checkf(false, TEXT("Unsupported mask format."));
+
+			BufferLayerCombineColour<BLEND_FUNC, 1>(ResultImage, BaseImage, Color);
+		}
+
+		if (BaseFormat == EImageFormat::IF_L_UBYTE)
+		{
+			BufferLayerCombineColour<BLEND_FUNC_MASKED, 1>(ResultImage, BaseImage, MaskImage, Color);
+		}
+		else if (BaseFormat == EImageFormat::IF_RGB_UBYTE)
+		{
+			BufferLayerCombineColour<BLEND_FUNC_MASKED, 3>(ResultImage, BaseImage, MaskImage, Color);
+		}
+		else if (BaseFormat == EImageFormat::IF_RGBA_UBYTE || BaseFormat == EImageFormat::IF_BGRA_UBYTE)
+		{
+			// \todo: pass swizzle template argument if BGRA_UBYTE, not yet supported.
+ 			BufferLayerCombineColour<BLEND_FUNC_MASKED, 4>(ResultImage, BaseImage, MaskImage, Color);
+		}
+		else
+		{
+			checkf(false, TEXT("Unsupported format."));
+		}
+	}
+
+	template<uint32 NCBase, uint32 NCBlend, class ImageCombineFn>
+	inline void BufferLayerCombineFunctor(Image* DestImage, const Image* BaseImage, const Image* BlendImage, ImageCombineFn&& ImageCombine)
+	{
+		static_assert(NCBase > 0 && NCBase <= 4);
+		static_assert(NCBlend > 0 && NCBlend <= 4);
+
+		check(BaseImage->GetFormat() == DestImage->GetFormat());
+		check(BaseImage->GetSizeX() == BlendImage->GetSizeX());
+		check(BaseImage->GetSizeY() == BlendImage->GetSizeY());
+		check(BaseImage->GetSizeX() == DestImage->GetSizeX());
+		check(BaseImage->GetSizeY() == DestImage->GetSizeY());
+		check(BaseImage->GetLODCount() <= BlendImage->GetLODCount());
+		check(BaseImage->GetLODCount() <= DestImage->GetLODCount());
+	
+		constexpr int32 BatchNumElems = 4098*2;
+		int32 NumBatches = BaseImage->DataStorage.GetNumBatches(BatchNumElems, NCBase);
+
+		check(NumBatches == BlendImage->DataStorage.GetNumBatches(BatchNumElems, NCBlend));
+		check(NumBatches == DestImage->DataStorage.GetNumBatches(BatchNumElems, NCBase));
+
+		auto ProcessBatch = [DestImage, BaseImage, BlendImage, BatchNumElems, ImageCombine](int32 BatchId)
+		{
+			const TArrayView<const uint8> BaseBatchView = BaseImage->DataStorage.GetBatch(BatchId, BatchNumElems, NCBase);
+			const TArrayView<const uint8> BlendBatchView = BlendImage->DataStorage.GetBatch(BatchId, BatchNumElems, NCBlend);
+			const TArrayView<uint8> DestBatchView = DestImage->DataStorage.GetBatch(BatchId, BatchNumElems, NCBase);
+
+			const int32 NumElems = BaseBatchView.Num() / NCBase;
+
+			const uint8* BaseBuf = BaseBatchView.GetData();
+			const uint8* BlendBuf = BlendBatchView.GetData();
+			uint8* DestBuf = DestBatchView.GetData();
+
+			check(NumElems == BlendBatchView.Num() / NCBlend);
+			check(NumElems == DestBatchView.Num() / NCBase);
+			
+			for (int32 I = 0; I < NumElems; ++I)
+			{
+				const uint32 Base = PackPixel<NCBase>(&BaseBuf[NCBase * I]);
+				const uint32 Blend = PackPixel<NCBlend>(&BlendBuf[NCBlend * I]);
+
+				const uint32 Result = ImageCombine(Base, Blend);
+				UnpackPixel<NCBase>(&DestBuf[NCBase* I], Result);
+			}
+		};
+
+		if (NumBatches == 1)
+		{
+			ProcessBatch(0);
+		}
+		else if (NumBatches > 1)
+		{
+			ParallelFor(NumBatches, ProcessBatch);
+		}
+	}
+
+	// Same functionality as above, in this case we use a functor which allows to pass user data. 
+	template<class ImageCombineFn>
+	inline void ImageLayerCombineFunctor(Image* ResultImage, const Image* BaseImage, const Image* BlendedImage, ImageCombineFn&& ImageCombine)
+	{
+		check(ResultImage->GetFormat() == BaseImage->GetFormat());
+		check(ResultImage->GetSizeX() == BaseImage->GetSizeX());
+		check(ResultImage->GetSizeY() == BaseImage->GetSizeY());
+		check(ResultImage->GetLODCount() == BaseImage->GetLODCount());
+		check(BaseImage->GetSizeX() == BlendedImage->GetSizeX());
+		check(BaseImage->GetSizeY() == BlendedImage->GetSizeY());
+		check(ResultImage->GetLODCount() <= BlendedImage->GetLODCount());
+
+		const EImageFormat BaseFormat = BaseImage->GetFormat();
+		const EImageFormat BlendFormat = BlendedImage->GetFormat();
+
+		if (BaseFormat == EImageFormat::IF_L_UBYTE )
+		{
+			if (BlendFormat == EImageFormat::IF_L_UBYTE )
+			{
+				BufferLayerCombineFunctor<1, 1>(ResultImage, BaseImage, BlendedImage, Forward<ImageCombineFn>(ImageCombine));
+			}
+			else if (BlendFormat == EImageFormat::IF_RGB_UBYTE )
+			{
+				BufferLayerCombineFunctor<1, 3>(ResultImage, BaseImage, BlendedImage, Forward<ImageCombineFn>(ImageCombine));
+			}
+			else if (BlendFormat == EImageFormat::IF_RGBA_UBYTE || BlendFormat == EImageFormat::IF_BGRA_UBYTE)
+			{
+				BufferLayerCombineFunctor<1, 4>(ResultImage, BaseImage, BlendedImage, Forward<ImageCombineFn>(ImageCombine));
+			}
+			else
+			{
+				checkf(false, TEXT("Unsupported format."));
+			}
+		}
+		else if (BaseFormat == EImageFormat::IF_RGB_UBYTE)
+		{
+			if (BlendFormat == EImageFormat::IF_L_UBYTE )
+			{
+				BufferLayerCombineFunctor<3, 1>(ResultImage, BaseImage, BlendedImage, Forward<ImageCombineFn>(ImageCombine));
+			}
+			else if (BlendFormat == EImageFormat::IF_RGB_UBYTE )
+			{
+				BufferLayerCombineFunctor<3, 3>(ResultImage, BaseImage, BlendedImage, Forward<ImageCombineFn>(ImageCombine));
+			}
+			else if (BlendFormat == EImageFormat::IF_RGBA_UBYTE || BlendFormat == EImageFormat::IF_BGRA_UBYTE )
+			{
+				BufferLayerCombineFunctor<3, 4>(ResultImage, BaseImage, BlendedImage, Forward<ImageCombineFn>(ImageCombine));
+			}
+			else
+			{
+				checkf(false, TEXT("Unsupported format."));
+			}
+		}
+		else if (BaseFormat == EImageFormat::IF_RGBA_UBYTE || BaseFormat == EImageFormat::IF_BGRA_UBYTE)
+		{
+			if (BlendFormat == EImageFormat::IF_L_UBYTE )
+			{
+				BufferLayerCombineFunctor<4, 1>(ResultImage, BaseImage, BlendedImage, Forward<ImageCombineFn>(ImageCombine));
+			}
+			else if (BlendFormat == EImageFormat::IF_RGB_UBYTE )
+			{
+				BufferLayerCombineFunctor<4, 3>(ResultImage, BaseImage, BlendedImage, Forward<ImageCombineFn>(ImageCombine));
+			}
+			else if (BlendFormat == EImageFormat::IF_RGBA_UBYTE || BlendFormat == EImageFormat::IF_BGRA_UBYTE )
+			{
+				BufferLayerCombineFunctor<4, 4>(ResultImage, BaseImage, BlendedImage, Forward<ImageCombineFn>(ImageCombine));
+			}
+			else
+			{
+				checkf(false, TEXT("Unsupported format."));
+			}
+		}
+		else
+		{
+			checkf(false, TEXT("Unsupported format."));
+		}
+	}
+
+
+	//! Blend a subimage on the base using a mask.
+	inline void ImageBlendOnBaseNoAlpha(Image* BaseImage,
+		const Image* MaskImage,
+		const Image* BlendedImage,
+		const box<FIntVector2>& Rect)
+	{
+		ImageLayerOnBaseNoAlpha<BlendChannel, false>(BaseImage, MaskImage, BlendedImage, Rect);
+	}
+
+}
diff --git a/Engine/Plugins/Mutable/Source/MutableRuntime/Private/MuR/OpImageNormalCombine.h b/Engine/Plugins/Mutable/Source/MutableRuntime/Private/MuR/OpImageNormalCombine.h
new file mode 100644
--- /dev/null
+++ b/Engine/Plugins/Mutable/Source/MutableRuntime/Private/MuR/OpImageNormalCombine.h
@@ -0,0 +1,84 @@
+// Copyright Epic Games, Inc. All Rights Reserved.
+
+#pragma once
+
+#include "MuR/ImagePrivate.h"
+#include "MuR/OpImageLayer.h"
+
+namespace mu
+{
+	// Combine layer oprations get packed color.
+	inline unsigned CombineNormal(unsigned base, unsigned blended)
+	{
+		// See /Engine/Functions/Engine_MaterialFunctions02/Utility/BlendAngleCorrectedNormals.BlendAngleCorrectedNormals
+		// for the source of the effect.
+
+		const float baseRf = (static_cast<float>((base >> 0)  & 0xFF) / 255.0f) * 2.0f - 1.0f;
+		const float baseGf = (static_cast<float>((base >> 8)  & 0xFF) / 255.0f) * 2.0f - 1.0f;
+		const float baseBf = (static_cast<float>((base >> 16) & 0xFF) / 255.0f) * 2.0f; //One added to the b channel.
+
+		const float blendedRf = (static_cast<float>((blended >> 0)  & 0xFF) / 255.0f) * 2.0f - 1.0f;
+		const float blendedGf = (static_cast<float>((blended >> 8)  & 0xFF) / 255.0f) * 2.0f - 1.0f;
+		const float blendedBf = (static_cast<float>((blended >> 16) & 0xFF) / 255.0f) * 2.0f - 1.0f;
+
+		FVector3f a = FVector3f(baseRf, baseGf, baseBf);
+		FVector3f b = FVector3f(-blendedRf, -blendedGf, blendedBf);
+
+		// The original does not normalize but if not done results don't look good due to signal clipping.
+		FVector3f n = (a* FVector3f::DotProduct(a, b) - b*a.Z);
+		n.Normalize();
+
+		return 
+			(FMath::Clamp(static_cast<unsigned>((n.X + 1.0f) * 0.5f * 255.0f), 0u, 255u) << 0) |
+			(FMath::Clamp(static_cast<unsigned>((n.Y + 1.0f) * 0.5f * 255.0f), 0u, 255u) << 8) |
+			(FMath::Clamp(static_cast<unsigned>((n.Z + 1.0f) * 0.5f * 255.0f), 0u, 255u) << 16);
+	}
+
+	inline unsigned CombineNormalMasked(unsigned base, unsigned blended, unsigned mask)
+	{
+		// See /Engine/Functions/Engine_MaterialFunctions02/Utility/BlendAngleCorrectedNormals.BlendAngleCorrectedNormals
+		// for the source of the effect.
+
+		const float baseRf = (static_cast<float>((base >> 0)  & 0xFF) / 255.0f) * 2.0f - 1.0f;
+		const float baseGf = (static_cast<float>((base >> 8)  & 0xFF) / 255.0f) * 2.0f - 1.0f;
+		const float baseBf = (static_cast<float>((base >> 16) & 0xFF) / 255.0f) * 2.0f - 1.0f;
+
+		const float blendedRf = (static_cast<float>((blended >> 0)  & 0xFF) / 255.0f) * 2.0f - 1.0f;
+		const float blendedGf = (static_cast<float>((blended >> 8)  & 0xFF) / 255.0f) * 2.0f - 1.0f;
+		const float blendedBf = (static_cast<float>((blended >> 16) & 0xFF) / 255.0f) * 2.0f - 1.0f;
+
+		const float maskRf = (static_cast<float>(mask & 0xFF) / 255.0f);
+
+		FVector3f a = FVector3f(baseRf, baseGf, baseBf + 1.0f);
+		FVector3f b = FVector3f(-blendedRf, -blendedGf, blendedBf);
+
+		// The original does not normalize but if not done results don't look good due to signal clipping.
+		FVector3f n = ( FMath::Lerp( FVector3f(baseRf, baseGf, baseBf), a*FVector3f::DotProduct(a, b) - b*a.Z, maskRf ) );
+		n.Normalize();
+
+		return 
+			(FMath::Clamp(static_cast<unsigned>((n.X + 1.0f) * 0.5f * 255.0f), 0u, 255u) << 0) |
+			(FMath::Clamp(static_cast<unsigned>((n.Y + 1.0f) * 0.5f * 255.0f), 0u, 255u) << 8) |
+			(FMath::Clamp(static_cast<unsigned>((n.Z + 1.0f) * 0.5f * 255.0f), 0u, 255u) << 16);
+	}
+
+	inline void ImageNormalCombine(Image* pResult, const Image* pBase, FVector4f col)
+	{
+		ImageLayerCombineColour<CombineNormal>(pResult, pBase, col);
+	}
+
+	inline void ImageNormalCombine(Image* pResult, const Image* pBase, const Image* pBlended, bool bOnlyFirstLOD)
+	{
+		ImageLayerCombine<CombineNormal>(pResult, pBase, pBlended, bOnlyFirstLOD);
+	}
+
+	inline void ImageNormalCombine(Image* pResult, const Image* pBase, const Image* pMask, FVector4f col)
+	{
+		ImageLayerCombineColour<CombineNormal, CombineNormalMasked>(pResult, pBase, pMask, col);
+	}
+
+	inline void ImageNormalCombine(Image* pResult, const Image* pBase, const Image* pMask, const Image* pBlended, bool bOnlyFirstLOD)
+	{
+		ImageLayerCombine<CombineNormal, CombineNormalMasked>(pResult, pBase, pMask, pBlended, bOnlyFirstLOD);
+	}
+}
