[mutable] Moved the Mutable plugin out of Experimental status into Beta.

diff --git a/Engine/Plugins/Mutable/Source/MutableRuntime/Private/MuR/OpImagePixelFormat.cpp b/Engine/Plugins/Mutable/Source/MutableRuntime/Private/MuR/OpImagePixelFormat.cpp
new file mode 100644
--- /dev/null
+++ b/Engine/Plugins/Mutable/Source/MutableRuntime/Private/MuR/OpImagePixelFormat.cpp
@@ -0,0 +1,1599 @@
+// Copyright Epic Games, Inc. All Rights Reserved.
+
+
+#include "MuR/Image.h"
+
+#include "Containers/Array.h"
+#include "HAL/UnrealMemory.h"
+#include "Math/UnrealMathSSE.h"
+#include "Misc/AssertionMacros.h"
+#include "MuR/BlockCompression/Miro/Miro.h"
+#include "MuR/ImageRLE.h"
+#include "MuR/MutableMath.h"
+#include "MuR/MutableTrace.h"
+#include "MuR/MutableRuntimeModule.h"
+#include "MuR/Ptr.h"
+
+
+namespace mu
+{
+
+
+Ptr<Image> FImageOperator::ImagePixelFormat(int32 CompressionQuality, const Image* Base, EImageFormat TargetFormat, int32 OnlyLOD)
+{
+	MUTABLE_CPUPROFILER_SCOPE(ImagePixelFormat);
+	TRACE_CPUPROFILER_EVENT_SCOPE_TEXT(*FString::Printf(TEXT("From %d to %d"), int32(Base->GetFormat()), int32(TargetFormat)));
+
+	FIntVector2 ResultSize;
+	int32 ResultLODCount = 0;
+
+	if (OnlyLOD == -1)
+	{
+		ResultSize[0] = Base->GetSizeX();
+		ResultSize[1] = Base->GetSizeY();
+		ResultLODCount = Base->GetLODCount();
+	}
+	else
+	{
+		ResultSize = Base->CalculateMipSize(OnlyLOD);
+		ResultLODCount = 1;
+	}
+
+	Ptr<Image> Result = CreateImage(ResultSize[0], ResultSize[1], ResultLODCount, TargetFormat, EInitializationType::NotInitialized);
+	Result->m_flags = Base->m_flags;
+
+	if (Base->GetSizeX() <= 0 || Base->GetSizeY() <= 0)
+	{
+		return Result;
+	}
+
+	bool bSuccess = false;
+	ImagePixelFormat(bSuccess, CompressionQuality, Result.get(), Base, OnlyLOD);
+	check(bSuccess);
+
+	return Result;
+}
+
+namespace ImagePixelFormatInternal
+{
+	template<class PixelFormatFuncType>
+	void ProcessBatchedHelper(
+			Image* Dest, const Image* Src, 
+			int32 DestElemSizeInBytes, int32 SrcElemSizeInBytes, 
+			int32 DestLODBegin, int32 SrcLODBegin, int32 NumLODs, 
+			PixelFormatFuncType&& BatchFunc)
+	{
+		check(Dest->GetLODCount() >= DestLODBegin + NumLODs);
+		check(Src->GetLODCount() >= SrcLODBegin + NumLODs);
+		
+		const int32 DestLODEnd = DestLODBegin + NumLODs;
+		const int32 SrcLODEnd = SrcLODBegin + NumLODs;
+
+		constexpr int32 NumBatchElems = 1 << 15;
+
+		const int32 NumBatches = Dest->DataStorage.GetNumBatchesLODRange(NumBatchElems, DestElemSizeInBytes, DestLODBegin, DestLODEnd);
+		check(NumBatches == Src->DataStorage.GetNumBatchesLODRange(NumBatchElems, SrcElemSizeInBytes, SrcLODBegin, SrcLODEnd));
+
+		for (int32 BatchIndex = 0; BatchIndex < NumBatches; ++BatchIndex)
+		{
+			TArrayView<const uint8> SrcView = Src->DataStorage.GetBatchLODRange(BatchIndex, NumBatchElems, SrcElemSizeInBytes, SrcLODBegin, SrcLODEnd);
+			TArrayView<uint8> DestView = Dest->DataStorage.GetBatchLODRange(BatchIndex, NumBatchElems, DestElemSizeInBytes, DestLODBegin, DestLODEnd);
+
+			const int32 NumElems = DestView.Num() / DestElemSizeInBytes;
+			check(DestView.Num() % DestElemSizeInBytes == 0);
+			check(NumElems == SrcView.Num() / SrcElemSizeInBytes);
+
+			BatchFunc(DestView.GetData(), SrcView.GetData(), NumElems);
+		}
+	}
+
+	template<class PixelFormatFuncType>
+	void ProcessBatchedHelper(
+			Image* Dest, int32 DestElemSizeInBytes, int32 DestLODBegin, int32 NumLODs, 
+			PixelFormatFuncType&& BatchFunc)
+	{
+		check(Dest->GetLODCount() >= DestLODBegin + NumLODs);
+		
+		const int32 DestLODEnd = DestLODBegin + NumLODs;
+
+		constexpr int32 NumBatchElems = 1 << 15;
+
+		const int32 NumBatches = Dest->DataStorage.GetNumBatchesLODRange(NumBatchElems, DestElemSizeInBytes, DestLODBegin, DestLODEnd);
+
+		for (int32 BatchIndex = 0; BatchIndex < NumBatches; ++BatchIndex)
+		{
+			TArrayView<uint8> DestView = Dest->DataStorage.GetBatchLODRange(BatchIndex, NumBatchElems, DestElemSizeInBytes, DestLODBegin, DestLODEnd);
+
+			const int32 NumElems = DestView.Num() / DestElemSizeInBytes;
+			check(DestView.Num() % DestElemSizeInBytes == 0);
+
+			BatchFunc(DestView.GetData(), NumElems);
+		}
+	}
+
+	template<class BlockDecompressionFuncType> 
+	void BlockDecompressionLODRangeHelper(
+		Image* Result, const Image* Base, 
+		int32 ResultLODBegin, int32 BaseLODBegin, int32 NumLODs, 
+		BlockDecompressionFuncType&& DecFunc)
+	{
+		check(Base->GetLODCount() >= BaseLODBegin + NumLODs);
+		check(Result->GetLODCount() >= ResultLODBegin + NumLODs);
+		for (int32 L = 0; L < NumLODs; ++L)
+		{
+			const int32 BaseLOD = BaseLODBegin + L;
+			const int32 ResultLOD = ResultLODBegin + L;
+			FIntVector2 MipSize = Result->CalculateMipSize(ResultLOD);
+			
+			check(Base->CalculateMipSize(BaseLOD) == MipSize);
+			
+			DecFunc(MipSize.X, MipSize.Y, Base->GetLODData(BaseLOD), Result->GetLODData(ResultLOD));
+		}
+	}
+
+	template<class BlockCompressionFuncType> 
+	void BlockCompressionLODRangeHelper(
+		Image* Result, const Image* Base, 
+		int32 ResultLODBegin, int32 BaseLODBegin, int32 NumLODs, int32 Quality,
+		BlockCompressionFuncType&& CompFunc)
+	{
+		check(Base->GetLODCount() >= BaseLODBegin + NumLODs);
+		check(Result->GetLODCount() >= ResultLODBegin + NumLODs);
+		for (int32 L = 0; L < NumLODs; ++L)
+		{
+			const int32 BaseLOD = BaseLODBegin + L;
+			const int32 ResultLOD = ResultLODBegin + L;
+			FIntVector2 MipSize = Result->CalculateMipSize(ResultLOD);
+		
+			check(Base->CalculateMipSize(BaseLOD) == MipSize);
+
+			CompFunc(MipSize.X, MipSize.Y, Base->GetLODData(BaseLOD), Result->GetLODData(ResultLOD), Quality);
+		}
+	}
+
+}
+
+void FImageOperator::ImagePixelFormat(bool& bOutSuccess, int32 CompressionQuality, Image* Result, const Image* Base, int32 OnlyLOD)
+{
+	if (FormatImageOverride)
+	{
+		MUTABLE_CPUPROFILER_SCOPE(ImagePixelFormatOverride);
+		FormatImageOverride(bOutSuccess, CompressionQuality, Result, Base, OnlyLOD);
+		if (bOutSuccess)
+		{
+			return;
+		}
+	}
+	
+	bOutSuccess = true;
+	
+	const int32 BaseLODBegin = OnlyLOD != -1 ? OnlyLOD : 0;
+	const int32 ResultLODBegin = 0;
+	const int32 NumLODs = OnlyLOD != -1 ? 1 : Base->GetLODCount();
+
+	ImagePixelFormat(bOutSuccess, CompressionQuality, Result, Base, ResultLODBegin, BaseLODBegin, NumLODs);
+	check(bOutSuccess);
+}
+
+void FImageOperator::ImagePixelFormat(bool& bOutSuccess, int32 CompressionQuality, Image* Result, const Image* Base, 
+		int32 ResultLODBegin, int32 BaseLODBegin, int32 NumLODs)
+{
+	MUTABLE_CPUPROFILER_SCOPE(ImagePixelFormatInPlace);
+
+	check(Result->GetLODCount() >= ResultLODBegin + NumLODs);
+	check(Base->GetLODCount() >= BaseLODBegin + NumLODs);
+	check(Base->CalculateMipSize(BaseLODBegin) == Result->CalculateMipSize(ResultLODBegin));
+
+	const bool bFormatAllLODs = Result->GetLODCount() == NumLODs && 
+								BaseLODBegin == 0 && ResultLODBegin == 0;
+	
+	check(!bFormatAllLODs || Base->GetLODCount() >= NumLODs);
+
+	bOutSuccess = true;
+
+	const bool bIsPixelFormatNeeded = 
+			Result->GetFormat()   == Base->GetFormat()   &&
+			Result->GetLODCount() == Base->GetLODCount() &&
+			Result->GetLODCount() == NumLODs;
+
+	if (bIsPixelFormatNeeded)
+	{
+		// Shouldn't really happen
+		Result->DataStorage = Base->DataStorage;
+	}
+	else
+	{
+		switch (Result->GetFormat())
+		{
+		case EImageFormat::IF_L_UBYTE:
+		{
+			switch (Base->GetFormat())
+			{
+			case EImageFormat::IF_L_UBYTE:
+			{
+				ImagePixelFormatInternal::ProcessBatchedHelper(Result, Base, 1, 1, ResultLODBegin, BaseLODBegin, NumLODs,
+				[](uint8* DestBuf, const uint8* BaseBuf, int32 NumElems)
+				{
+					FMemory::Memcpy(DestBuf, BaseBuf, NumElems * 1);
+				});
+
+				break;
+			}
+			case EImageFormat::IF_L_UBYTE_RLE:
+			{
+				if (bFormatAllLODs)
+				{
+					UncompressRLE_L(Base, Result);
+				}
+				else
+				{
+					check(NumLODs == 1);
+					FImageSize ResultSize = Result->GetSize();
+					UncompressRLE_L(ResultSize.X, ResultSize.Y, Base->GetLODData(BaseLODBegin), Result->GetLODData(ResultLODBegin));
+				}
+				break;
+			}
+			case EImageFormat::IF_L_UBIT_RLE:
+			{
+				if(bFormatAllLODs)
+				{
+					UncompressRLE_L1(Base, Result);
+				}
+				else
+				{
+					check(NumLODs == 1);
+					FImageSize ResultSize = Result->GetSize();
+					UncompressRLE_L1(ResultSize.X, ResultSize.Y, Base->GetLODData(BaseLODBegin), Result->GetLODData(ResultLODBegin));
+				}
+				break;
+			}
+			case EImageFormat::IF_RGB_UBYTE:
+			{
+				ImagePixelFormatInternal::ProcessBatchedHelper(Result, Base, 1, 3, ResultLODBegin, BaseLODBegin, NumLODs,
+				[](uint8* DestBuf, const uint8* BaseBuf, int32 NumElems)
+				{
+					for (int32 I = 0; I < NumElems; ++I)
+					{
+						uint32 ResultValue = 77 * BaseBuf[3*I + 0] + 150 * BaseBuf[3*I + 1] + 29 * BaseBuf[3*I + 2];
+						DestBuf[I] = (uint8)FMath::Min(255u, ResultValue >> 8);
+					}
+				});
+
+				break;
+			}
+			case EImageFormat::IF_RGB_UBYTE_RLE:
+			{
+				Ptr<Image> TempBase =
+						ImagePixelFormat(CompressionQuality, Base, EImageFormat::IF_RGB_UBYTE);
+
+				ImagePixelFormatInternal::ProcessBatchedHelper(Result, TempBase.get(), 1, 3, ResultLODBegin, BaseLODBegin, NumLODs,
+				[](uint8* DestBuf, const uint8* BaseBuf, int32 NumElems)
+				{
+					for (int32 I = 0; I < NumElems; ++I)
+					{
+						uint32 ResultValue = 77 * BaseBuf[3*I + 0] + 150 * BaseBuf[3*I + 1] + 29 * BaseBuf[3*I + 2];
+						DestBuf[I] = (uint8)FMath::Min(255u, ResultValue >> 8);
+					}
+				});
+
+				ReleaseImage(TempBase);
+
+				break;
+			}
+			case EImageFormat::IF_RGBA_UBYTE:
+			{
+				ImagePixelFormatInternal::ProcessBatchedHelper(Result, Base, 1, 4, ResultLODBegin, BaseLODBegin, NumLODs,
+				[](uint8* DestBuf, const uint8* BaseBuf, int32 NumElems)
+				{
+					for (int32 I = 0; I < NumElems; ++I)
+					{
+						uint32 ResultValue = 77 * BaseBuf[4*I + 0] + 150 * BaseBuf[4*I + 1] + 29 * BaseBuf[4*I + 2];
+						DestBuf[I] = (uint8)FMath::Min(255u, ResultValue >> 8);
+					}
+				});
+
+				break;
+			}
+			case EImageFormat::IF_RGBA_UBYTE_RLE:
+			{
+				Ptr<Image> TempBase = ImagePixelFormat(CompressionQuality, Base, EImageFormat::IF_RGBA_UBYTE);
+				
+				ImagePixelFormatInternal::ProcessBatchedHelper(Result, TempBase.get(), 1, 4, ResultLODBegin, BaseLODBegin, NumLODs,
+				[](uint8* DestBuf, const uint8* BaseBuf, int32 NumElems)
+				{
+					for (int32 I = 0; I < NumElems; ++I)
+					{
+						uint32 ResultValue = 77 * BaseBuf[4*I + 0] + 150 * BaseBuf[4*I + 1] + 29 * BaseBuf[4*I + 2];
+						DestBuf[I] = uint8( ResultValue >> 8);
+					}
+				});
+
+				ReleaseImage(TempBase);
+
+				break;
+			}
+			case EImageFormat::IF_BGRA_UBYTE:
+			{
+				ImagePixelFormatInternal::ProcessBatchedHelper(Result, Base, 1, 4, ResultLODBegin, BaseLODBegin, NumLODs,
+				[](uint8* DestBuf, const uint8* BaseBuf, int32 NumElems)
+				{
+					for (int32 I = 0; I < NumElems; ++I)
+					{
+						uint32 ResultValue = 77 * BaseBuf[4*I + 2] + 150 * BaseBuf[4*I + 1] + 29 * BaseBuf[4*I + 0];
+						DestBuf[I] = uint8(ResultValue >> 8);
+					}
+				});
+
+				break;
+			}
+			case EImageFormat::IF_BC1:
+			case EImageFormat::IF_BC2:
+			case EImageFormat::IF_BC3:
+			{
+				Ptr<Image> TempBase = ImagePixelFormat(CompressionQuality, Base, EImageFormat::IF_RGB_UBYTE);
+				
+				ImagePixelFormatInternal::ProcessBatchedHelper(Result, TempBase.get(), 1, 3, ResultLODBegin, BaseLODBegin, NumLODs,
+				[](uint8* DestBuf, const uint8* BaseBuf, int32 NumElems)
+				{
+					for (int32 I = 0; I < NumElems; ++I)
+					{
+						uint32 ResultValue = 77 * BaseBuf[3*I + 0] + 150 * BaseBuf[3*I + 1] + 29 * BaseBuf[3*I + 2];
+						DestBuf[I] = uint8(ResultValue >> 8);
+					}
+				});
+
+				ReleaseImage(TempBase);
+
+				break;
+			}
+			case EImageFormat::IF_BC4:
+			{
+				ImagePixelFormatInternal::BlockDecompressionLODRangeHelper(Result, Base, ResultLODBegin, BaseLODBegin, NumLODs, miro::BC4_to_L);
+				break;
+			}
+			default:
+				// Case not implemented
+				check(false);
+			}
+
+			break;
+		}
+
+		case EImageFormat::IF_L_UBYTE_RLE:
+		{
+			switch (Base->GetFormat())
+			{
+
+			case EImageFormat::IF_L_UBYTE:
+			{
+				check(bFormatAllLODs);
+
+				// Try to compress
+				if (Result->DataStorage.IsEmpty())
+				{
+					// Allocate memory for the compressed data. TODO: Smaller?	
+					check(Base->DataStorage.Buffers.Num() == Result->DataStorage.Buffers.Num());
+
+					for (int32 I = 0; I < Base->DataStorage.Buffers.Num(); ++I)
+					{
+						int32 BufferSize = Base->DataStorage.Buffers[I].Num();
+						Result->DataStorage.Buffers[I].SetNumUninitialized(BufferSize);
+					}
+					Result->DataStorage.CompactedTailOffsets = Base->DataStorage.CompactedTailOffsets;
+				}
+				
+				CompressRLE_L(Base, Result);
+
+//#ifdef MUTABLE_DEBUG_RLE		
+//				if (bOutSuccess)
+//				{
+//					// Verify
+//					Ptr<Image> Test = ImagePixelFormat(CompressionQuality, Result, EImageFormat::IF_L_UBYTE);
+//					check(Test->DataStorage == Base->DataStorage);
+//				}
+//#endif
+				break;
+			}
+
+			case EImageFormat::IF_RGB_UBYTE:
+			case EImageFormat::IF_RGBA_UBYTE:
+			case EImageFormat::IF_BGRA_UBYTE:
+			{
+				check(bFormatAllLODs);
+				Ptr<Image> TempBase = ImagePixelFormat(CompressionQuality, Base, EImageFormat::IF_L_UBYTE, -1);
+ 
+				// Try to compress
+				if (Result->DataStorage.IsEmpty())
+				{
+					// Allocate memory for the compressed data. TODO: Smaller?	
+					check(Base->DataStorage.Buffers.Num() == Result->DataStorage.Buffers.Num());
+
+					for (int32 I = 0; I < Base->DataStorage.Buffers.Num(); ++I)
+					{
+						int32 BufferSize = Base->DataStorage.Buffers[I].Num();
+						Result->DataStorage.Buffers[I].SetNumUninitialized(BufferSize);
+					}
+					Result->DataStorage.CompactedTailOffsets = Base->DataStorage.CompactedTailOffsets;
+				}
+
+				CompressRLE_L(TempBase.get(), Result);
+
+//#ifdef MUTABLE_DEBUG_RLE
+//				if (bOutSuccess)
+//				{
+//					// Verify
+//					Ptr<Image> Test = ImagePixelFormat(CompressionQuality, Result, EImageFormat::IF_L_UBYTE);
+//					check(Test->DataStorage == TempBase->DataStorage);
+//				}
+//#endif
+				ReleaseImage(TempBase);
+
+				break;
+			}
+
+			default:
+				// Case not implemented
+				check(false);
+
+			}
+
+			break;
+		}
+
+		case EImageFormat::IF_L_UBIT_RLE:
+		{
+			switch (Base->GetFormat())
+			{
+
+			case EImageFormat::IF_L_UBYTE:
+			{
+				check(bFormatAllLODs);
+
+				// Try to compress
+				if (Result->DataStorage.IsEmpty())
+				{
+					// Allocate memory for the compressed data. TODO: Smaller?	
+					check(Base->DataStorage.Buffers.Num() == Result->DataStorage.Buffers.Num());
+
+					for (int32 I = 0; I < Base->DataStorage.Buffers.Num(); ++I)
+					{
+						int32 BufferSize = Base->DataStorage.Buffers[I].Num();
+						Result->DataStorage.Buffers[I].SetNumUninitialized(BufferSize);
+					}
+					Result->DataStorage.CompactedTailOffsets = Base->DataStorage.CompactedTailOffsets;
+				}
+				
+				CompressRLE_L1(Base, Result);
+				
+//#ifdef MUTABLE_DEBUG_RLE		
+//				if (bOutSuccess)
+//				{
+//					// Verify
+//					Ptr<Image> Test = ImagePixelFormat(CompressionQuality, Result, EImageFormat::IF_L_UBYTE);
+//					check(Test->DataStorage == Base->DataStorage);
+//				}
+//#endif
+
+				break;
+			}
+
+			case EImageFormat::IF_RGB_UBYTE:
+			case EImageFormat::IF_RGBA_UBYTE:
+			case EImageFormat::IF_BGRA_UBYTE:
+			{
+				check(bFormatAllLODs);
+				Ptr<Image> TempBase = ImagePixelFormat(CompressionQuality, Base, EImageFormat::IF_L_UBYTE, -1);
+				CompressRLE_L1(TempBase.get(), Result);
+
+//#ifdef MUTABLE_DEBUG_RLE
+//				if (bOutSuccess)
+//				{
+//					// Verify
+//					Ptr<Image> Test = ImagePixelFormat(CompressionQuality, Result, EImageFormat::IF_L_UBYTE);
+//					check(Test->DataStorage == TempBase->DataStorage);
+//				}
+//#endif
+				ReleaseImage(TempBase);
+				break;
+			}
+
+			default:
+				// Case not implemented
+				check(false);
+
+			}
+
+			break;
+		}
+
+		case EImageFormat::IF_RGB_UBYTE:
+		{
+			switch (Base->GetFormat())
+			{
+
+			case EImageFormat::IF_L_UBYTE:
+			{
+				ImagePixelFormatInternal::ProcessBatchedHelper(Result, Base, 3, 1, ResultLODBegin, BaseLODBegin, NumLODs,
+				[](uint8* DestBuf, const uint8* BaseBuf, int32 NumElems)
+				{
+					for (int32 I = 0; I < NumElems; ++I)
+					{
+						DestBuf[I*3 + 0] = BaseBuf[I];
+						DestBuf[I*3 + 1] = BaseBuf[I];
+						DestBuf[I*3 + 2] = BaseBuf[I];
+					}
+				});
+
+				break;
+			}
+			case EImageFormat::IF_RGB_UBYTE:
+			{
+				ImagePixelFormatInternal::ProcessBatchedHelper(Result, Base, 3, 3, ResultLODBegin, BaseLODBegin, NumLODs,
+				[](uint8* DestBuf, const uint8* BaseBuf, int32 NumElems)
+				{
+					FMemory::Memcpy(DestBuf, BaseBuf, 3*NumElems);
+				});
+
+				break;
+			}
+			case EImageFormat::IF_RGBA_UBYTE:
+			{
+				ImagePixelFormatInternal::ProcessBatchedHelper(Result, Base, 3, 4, ResultLODBegin, BaseLODBegin, NumLODs,
+				[](uint8* DestBuf, const uint8* BaseBuf, int32 NumElems)
+				{
+					for (int32 I = 0; I < NumElems; ++I)
+					{
+						DestBuf[I*3 + 0] = BaseBuf[I*4 + 0];
+						DestBuf[I*3 + 1] = BaseBuf[I*4 + 1];
+						DestBuf[I*3 + 2] = BaseBuf[I*4 + 2];
+					}
+				});
+
+				break;
+			}
+			case EImageFormat::IF_BGRA_UBYTE:
+			{
+				ImagePixelFormatInternal::ProcessBatchedHelper(Result, Base, 3, 4, ResultLODBegin, BaseLODBegin, NumLODs,
+				[](uint8* DestBuf, const uint8* BaseBuf, int32 NumElems)
+				{
+					for (int32 I = 0; I < NumElems; ++I)
+					{
+						DestBuf[I*3 + 0] = BaseBuf[I*4 + 2];
+						DestBuf[I*3 + 1] = BaseBuf[I*4 + 1];
+						DestBuf[I*3 + 2] = BaseBuf[I*4 + 0];
+					}
+				});
+
+				break;
+			}
+			case EImageFormat::IF_RGB_UBYTE_RLE:
+			{
+				UncompressRLE_RGB(Base, Result);
+				break;
+			}
+			case EImageFormat::IF_RGBA_UBYTE_RLE:
+			{
+				check(bFormatAllLODs);
+
+				Ptr<Image> TempBase = ImagePixelFormat(CompressionQuality, Base, EImageFormat::IF_RGBA_UBYTE, -1);
+
+				ImagePixelFormatInternal::ProcessBatchedHelper(Result, TempBase.get(), 3, 4, ResultLODBegin, BaseLODBegin, NumLODs,
+				[](uint8* DestBuf, const uint8* BaseBuf, int32 NumElems)
+				{
+					for (int32 I = 0; I < NumElems; ++I)
+					{
+						DestBuf[I*3 + 0] = BaseBuf[I*4 + 2];
+						DestBuf[I*3 + 1] = BaseBuf[I*4 + 1];
+						DestBuf[I*3 + 2] = BaseBuf[I*4 + 0];
+					}
+				});
+
+				ReleaseImage(TempBase);
+				break;
+			}
+			case EImageFormat::IF_BC1:
+			{
+				ImagePixelFormatInternal::BlockDecompressionLODRangeHelper(
+						Result, Base, ResultLODBegin, BaseLODBegin, NumLODs, miro::BC1_to_RGB);
+				break;
+			}
+			case EImageFormat::IF_BC2:
+			{
+				ImagePixelFormatInternal::BlockDecompressionLODRangeHelper(
+						Result, Base, ResultLODBegin, BaseLODBegin, NumLODs, miro::BC2_to_RGB);
+				break;
+			}
+			case EImageFormat::IF_BC3:
+			{
+				ImagePixelFormatInternal::BlockDecompressionLODRangeHelper(
+						Result, Base, ResultLODBegin, BaseLODBegin, NumLODs, miro::BC3_to_RGB);
+				break;
+			}
+			case EImageFormat::IF_BC4:
+			{
+				ImagePixelFormatInternal::BlockDecompressionLODRangeHelper(
+						Result, Base, ResultLODBegin, BaseLODBegin, NumLODs, miro::BC4_to_RGB);
+				break;
+			}
+			case EImageFormat::IF_BC5:
+			{
+				ImagePixelFormatInternal::BlockDecompressionLODRangeHelper(
+						Result, Base, ResultLODBegin, BaseLODBegin, NumLODs, miro::BC5_to_RGB);
+				break;
+			}
+			case EImageFormat::IF_ASTC_4x4_RGB_LDR:
+			{
+				ImagePixelFormatInternal::BlockDecompressionLODRangeHelper(
+						Result, Base, ResultLODBegin, BaseLODBegin, NumLODs, miro::ASTC4x4RGBL_to_RGB);
+				break;
+			}
+			case EImageFormat::IF_ASTC_4x4_RGBA_LDR:
+			{
+				ImagePixelFormatInternal::BlockDecompressionLODRangeHelper(
+						Result, Base, ResultLODBegin, BaseLODBegin, NumLODs, miro::ASTC4x4RGBAL_to_RGB);
+				break;
+			}
+			case EImageFormat::IF_ASTC_4x4_RG_LDR:
+			{
+				ImagePixelFormatInternal::BlockDecompressionLODRangeHelper(
+						Result, Base, ResultLODBegin, BaseLODBegin, NumLODs, miro::ASTC4x4RGL_to_RGB);
+				break;
+			}
+			case EImageFormat::IF_ASTC_6x6_RGB_LDR:
+			{
+				ImagePixelFormatInternal::BlockDecompressionLODRangeHelper(
+						Result, Base, ResultLODBegin, BaseLODBegin, NumLODs, miro::ASTC6x6RGBL_to_RGB);
+				break;
+			}
+			case EImageFormat::IF_ASTC_6x6_RGBA_LDR:
+			{
+				ImagePixelFormatInternal::BlockDecompressionLODRangeHelper(
+						Result, Base, ResultLODBegin, BaseLODBegin, NumLODs, miro::ASTC6x6RGBAL_to_RGB);
+				break;
+			}
+			case EImageFormat::IF_ASTC_6x6_RG_LDR:
+			{
+				ImagePixelFormatInternal::BlockDecompressionLODRangeHelper(
+						Result, Base, ResultLODBegin, BaseLODBegin, NumLODs, miro::ASTC6x6RGL_to_RGB);
+				break;
+			}
+			case EImageFormat::IF_ASTC_8x8_RGB_LDR:
+			{
+				ImagePixelFormatInternal::BlockDecompressionLODRangeHelper(
+						Result, Base, ResultLODBegin, BaseLODBegin, NumLODs, miro::ASTC8x8RGBL_to_RGB);
+				break;
+			}
+			case EImageFormat::IF_ASTC_8x8_RGBA_LDR:
+			{
+				ImagePixelFormatInternal::BlockDecompressionLODRangeHelper(
+						Result, Base, ResultLODBegin, BaseLODBegin, NumLODs, miro::ASTC8x8RGBAL_to_RGB);
+				break;
+			}
+			case EImageFormat::IF_ASTC_8x8_RG_LDR:
+			{
+				ImagePixelFormatInternal::BlockDecompressionLODRangeHelper(
+						Result, Base, ResultLODBegin, BaseLODBegin, NumLODs, miro::ASTC8x8RGL_to_RGB);
+				break;
+			}
+			case EImageFormat::IF_ASTC_10x10_RGB_LDR:
+			{
+				ImagePixelFormatInternal::BlockDecompressionLODRangeHelper(
+						Result, Base, ResultLODBegin, BaseLODBegin, NumLODs, miro::ASTC10x10RGBL_to_RGB);
+				break;
+			}
+			case EImageFormat::IF_ASTC_10x10_RGBA_LDR:
+			{
+				ImagePixelFormatInternal::BlockDecompressionLODRangeHelper(
+						Result, Base, ResultLODBegin, BaseLODBegin, NumLODs, miro::ASTC10x10RGBAL_to_RGB);
+				break;
+			}
+			case EImageFormat::IF_ASTC_10x10_RG_LDR:
+			{
+				ImagePixelFormatInternal::BlockDecompressionLODRangeHelper(
+						Result, Base, ResultLODBegin, BaseLODBegin, NumLODs, miro::ASTC10x10RGL_to_RGB);
+				break;
+			}
+			case EImageFormat::IF_ASTC_12x12_RGB_LDR:
+			{
+				ImagePixelFormatInternal::BlockDecompressionLODRangeHelper(
+						Result, Base, ResultLODBegin, BaseLODBegin, NumLODs, miro::ASTC12x12RGBL_to_RGB);
+				break;
+			}
+			case EImageFormat::IF_ASTC_12x12_RGBA_LDR:
+			{
+				ImagePixelFormatInternal::BlockDecompressionLODRangeHelper(
+						Result, Base, ResultLODBegin, BaseLODBegin, NumLODs, miro::ASTC12x12RGBAL_to_RGB);
+				break;
+			}
+			case EImageFormat::IF_ASTC_12x12_RG_LDR:
+			{
+				ImagePixelFormatInternal::BlockDecompressionLODRangeHelper(
+						Result, Base, ResultLODBegin, BaseLODBegin, NumLODs, miro::ASTC12x12RGL_to_RGB);
+				break;
+			}
+
+			default:
+				// Case not implemented
+				check(false);
+			}
+
+			break;
+		}
+
+		case EImageFormat::IF_BGRA_UBYTE:
+		{
+			// TODO: Optimise
+			Result->DataStorage.ImageFormat = EImageFormat::IF_RGBA_UBYTE;
+			ImagePixelFormat(bOutSuccess, CompressionQuality, Result, Base, ResultLODBegin, BaseLODBegin, NumLODs);
+			check(bOutSuccess);
+			Result->DataStorage.ImageFormat = EImageFormat::IF_BGRA_UBYTE;
+
+			ImagePixelFormatInternal::ProcessBatchedHelper(Result, 4, ResultLODBegin, NumLODs,
+			[](uint8* DestBuf, int32 NumElems)
+			{
+				for (int32 I = 0; I < NumElems; ++I)
+				{
+					uint8 Temp = DestBuf[I*4 + 0];
+					DestBuf[I*4 + 0] = DestBuf[I*4 + 2];
+					DestBuf[I*4 + 2] = Temp;
+				}
+			});
+
+			break;
+		}
+
+		case EImageFormat::IF_RGBA_UBYTE:
+		{
+			switch (Base->GetFormat())
+			{
+
+			case EImageFormat::IF_RGBA_UBYTE_RLE:
+			{
+				check(bFormatAllLODs);
+				UncompressRLE_RGBA(Base, Result);
+				break;
+			}
+			case EImageFormat::IF_L_UBYTE:
+			{
+				ImagePixelFormatInternal::ProcessBatchedHelper(Result, Base, 4, 1, ResultLODBegin, BaseLODBegin, NumLODs,
+				[](uint8* DestBuf, const uint8* BaseBuf, int32 NumElems)
+				{
+					for (int32 I = 0; I < NumElems; ++I)
+					{
+						DestBuf[I*4 + 0] = BaseBuf[I];
+						DestBuf[I*4 + 1] = BaseBuf[I];
+						DestBuf[I*4 + 2] = BaseBuf[I];
+						DestBuf[I*4 + 3] = 255;
+					}
+				});
+
+				break;
+			}
+			case EImageFormat::IF_RGB_UBYTE:
+			{
+				ImagePixelFormatInternal::ProcessBatchedHelper(Result, Base, 4, 3, ResultLODBegin, BaseLODBegin, NumLODs,
+				[](uint8* DestBuf, const uint8* BaseBuf, int32 NumElems)
+				{
+					for (int32 I = 0; I < NumElems; ++I)
+					{
+						DestBuf[I*4 + 0] = BaseBuf[I*3 + 0];
+						DestBuf[I*4 + 1] = BaseBuf[I*3 + 1];
+						DestBuf[I*4 + 2] = BaseBuf[I*3 + 2];
+						DestBuf[I*4 + 3] = 255;
+					}
+				});
+
+				break;
+			}
+			case EImageFormat::IF_RGBA_UBYTE:
+			{
+				ImagePixelFormatInternal::ProcessBatchedHelper(Result, Base, 4, 4, ResultLODBegin, BaseLODBegin, NumLODs,
+				[](uint8* DestBuf, const uint8* BaseBuf, int32 NumElems)
+				{
+					FMemory::Memcpy(DestBuf, BaseBuf, NumElems*4);
+				});
+
+				break;
+			}
+			case EImageFormat::IF_BGRA_UBYTE:
+			{
+				ImagePixelFormatInternal::ProcessBatchedHelper(Result, Base, 4, 4, ResultLODBegin, BaseLODBegin, NumLODs,
+				[](uint8* DestBuf, const uint8* BaseBuf, int32 NumElems)
+				{
+					for (int32 I = 0; I < NumElems; ++I)
+					{
+						DestBuf[I*4 + 0] = BaseBuf[I*3 + 2];
+						DestBuf[I*4 + 1] = BaseBuf[I*3 + 1];
+						DestBuf[I*4 + 2] = BaseBuf[I*3 + 0];
+						DestBuf[I*4 + 3] = BaseBuf[I*3 + 3];
+					}
+				});
+
+				break;
+			}
+			case EImageFormat::IF_BC1:
+			{
+				ImagePixelFormatInternal::BlockDecompressionLODRangeHelper(
+						Result, Base, ResultLODBegin, BaseLODBegin, NumLODs, miro::BC1_to_RGBA);
+				break;
+			}
+			case EImageFormat::IF_BC2:
+			{
+				ImagePixelFormatInternal::BlockDecompressionLODRangeHelper(
+						Result, Base, ResultLODBegin, BaseLODBegin, NumLODs, miro::BC2_to_RGBA);
+				break;
+			}
+			case EImageFormat::IF_BC3:
+			{
+				ImagePixelFormatInternal::BlockDecompressionLODRangeHelper(
+						Result, Base, ResultLODBegin, BaseLODBegin, NumLODs, miro::BC3_to_RGBA);
+				break;
+			}
+			case EImageFormat::IF_BC4:
+			{
+				ImagePixelFormatInternal::BlockDecompressionLODRangeHelper(
+						Result, Base, ResultLODBegin, BaseLODBegin, NumLODs, miro::BC4_to_RGBA);
+				break;
+			}
+			case EImageFormat::IF_BC5:
+			{
+				ImagePixelFormatInternal::BlockDecompressionLODRangeHelper(
+						Result, Base, ResultLODBegin, BaseLODBegin, NumLODs, miro::BC5_to_RGBA);
+				break;
+			}
+			case EImageFormat::IF_ASTC_4x4_RGB_LDR:
+			{
+				ImagePixelFormatInternal::BlockDecompressionLODRangeHelper(
+						Result, Base, ResultLODBegin, BaseLODBegin, NumLODs, miro::ASTC4x4RGBL_to_RGBA);
+				break;
+			}
+			case EImageFormat::IF_ASTC_4x4_RGBA_LDR:
+			{
+				ImagePixelFormatInternal::BlockDecompressionLODRangeHelper(
+						Result, Base, ResultLODBegin, BaseLODBegin, NumLODs, miro::ASTC4x4RGBAL_to_RGBA);
+				break;
+			}
+			case EImageFormat::IF_ASTC_4x4_RG_LDR:
+			{
+				ImagePixelFormatInternal::BlockDecompressionLODRangeHelper(
+						Result, Base, ResultLODBegin, BaseLODBegin, NumLODs, miro::ASTC4x4RGL_to_RGBA);
+				break;
+			}
+			case EImageFormat::IF_ASTC_6x6_RGB_LDR:
+			{
+				ImagePixelFormatInternal::BlockDecompressionLODRangeHelper(
+						Result, Base, ResultLODBegin, BaseLODBegin, NumLODs, miro::ASTC6x6RGBL_to_RGBA);
+				break;
+			}
+			case EImageFormat::IF_ASTC_6x6_RGBA_LDR:
+			{
+				ImagePixelFormatInternal::BlockDecompressionLODRangeHelper(
+						Result, Base, ResultLODBegin, BaseLODBegin, NumLODs, miro::ASTC6x6RGBAL_to_RGBA);
+				break;
+			}
+			case EImageFormat::IF_ASTC_6x6_RG_LDR:
+			{
+				ImagePixelFormatInternal::BlockDecompressionLODRangeHelper(
+						Result, Base, ResultLODBegin, BaseLODBegin, NumLODs, miro::ASTC6x6RGL_to_RGBA);
+				break;
+			}
+			case EImageFormat::IF_ASTC_8x8_RGB_LDR:
+			{
+				ImagePixelFormatInternal::BlockDecompressionLODRangeHelper(
+						Result, Base, ResultLODBegin, BaseLODBegin, NumLODs, miro::ASTC8x8RGBL_to_RGBA);
+				break;
+			}
+			case EImageFormat::IF_ASTC_8x8_RGBA_LDR:
+			{
+				ImagePixelFormatInternal::BlockDecompressionLODRangeHelper(
+						Result, Base, ResultLODBegin, BaseLODBegin, NumLODs, miro::ASTC8x8RGBAL_to_RGBA);
+				break;
+			}
+			case EImageFormat::IF_ASTC_8x8_RG_LDR:
+			{
+				ImagePixelFormatInternal::BlockDecompressionLODRangeHelper(
+						Result, Base, ResultLODBegin, BaseLODBegin, NumLODs, miro::ASTC8x8RGL_to_RGBA);
+				break;
+			}
+			case EImageFormat::IF_ASTC_10x10_RGB_LDR:
+			{
+				ImagePixelFormatInternal::BlockDecompressionLODRangeHelper(
+						Result, Base, ResultLODBegin, BaseLODBegin, NumLODs, miro::ASTC10x10RGBL_to_RGBA);
+				break;
+			}
+			case EImageFormat::IF_ASTC_10x10_RGBA_LDR:
+			{
+				ImagePixelFormatInternal::BlockDecompressionLODRangeHelper(
+						Result, Base, ResultLODBegin, BaseLODBegin, NumLODs, miro::ASTC10x10RGBAL_to_RGBA);
+				break;
+			}
+			case EImageFormat::IF_ASTC_10x10_RG_LDR:
+			{
+				ImagePixelFormatInternal::BlockDecompressionLODRangeHelper(
+						Result, Base, ResultLODBegin, BaseLODBegin, NumLODs, miro::ASTC10x10RGL_to_RGBA);
+				break;
+			}
+			case EImageFormat::IF_ASTC_12x12_RGB_LDR:
+			{
+				ImagePixelFormatInternal::BlockDecompressionLODRangeHelper(
+						Result, Base, ResultLODBegin, BaseLODBegin, NumLODs, miro::ASTC12x12RGBL_to_RGBA);
+				break;
+			}
+			case EImageFormat::IF_ASTC_12x12_RGBA_LDR:
+			{
+				ImagePixelFormatInternal::BlockDecompressionLODRangeHelper(
+						Result, Base, ResultLODBegin, BaseLODBegin, NumLODs, miro::ASTC12x12RGBAL_to_RGBA);
+				break;
+			}
+			case EImageFormat::IF_ASTC_12x12_RG_LDR:
+			{
+				ImagePixelFormatInternal::BlockDecompressionLODRangeHelper(
+						Result, Base, ResultLODBegin, BaseLODBegin, NumLODs, miro::ASTC12x12RGL_to_RGBA);
+				break;
+			}
+			case EImageFormat::IF_RGB_UBYTE_RLE:
+			{
+				check(bFormatAllLODs);
+				Ptr<Image> TempBase = ImagePixelFormat(CompressionQuality, Base, EImageFormat::IF_RGB_UBYTE);
+
+				ImagePixelFormatInternal::ProcessBatchedHelper(Result, TempBase.get(), 4, 3, ResultLODBegin, BaseLODBegin, NumLODs,
+				[](uint8* DestBuf, const uint8* BaseBuf, int32 NumElems)
+				{
+					for (int32 I = 0; I < NumElems; ++I)
+					{
+						DestBuf[I*4 + 0] = BaseBuf[I*3 + 0];
+						DestBuf[I*4 + 1] = BaseBuf[I*3 + 1];
+						DestBuf[I*4 + 2] = BaseBuf[I*3 + 2];
+						DestBuf[I*4 + 3] = 255;
+					}
+				});
+
+				ReleaseImage(TempBase);
+				break;
+			}
+			case EImageFormat::IF_L_UBYTE_RLE:
+			{
+				check(bFormatAllLODs);
+				Ptr<Image> TempBase = ImagePixelFormat(CompressionQuality, Base, EImageFormat::IF_L_UBYTE);
+
+				ImagePixelFormatInternal::ProcessBatchedHelper(Result, TempBase.get(), 4, 1, ResultLODBegin, BaseLODBegin, NumLODs,
+				[](uint8* DestBuf, const uint8* BaseBuf, int32 NumElems)
+				{
+					for (int32 I = 0; I < NumElems; ++I)
+					{
+						DestBuf[I*4 + 0] = BaseBuf[I];
+						DestBuf[I*4 + 1] = BaseBuf[I];
+						DestBuf[I*4 + 2] = BaseBuf[I];
+						DestBuf[I*4 + 3] = 255;
+					}
+				});
+
+				ReleaseImage(TempBase);
+				break;
+			}
+
+			default:
+				// Case not implemented
+				check(false);
+			}
+
+			break;
+		}
+
+		case EImageFormat::IF_RGBA_UBYTE_RLE:
+		{
+			switch (Base->GetFormat())
+			{
+
+			case EImageFormat::IF_RGBA_UBYTE:
+			{
+				check(bFormatAllLODs);
+				CompressRLE_RGBA(Base, Result);
+
+				break;
+			}
+			case EImageFormat::IF_RGB_UBYTE:
+			{
+				check(bFormatAllLODs);
+
+				// \todo: optimise
+				Ptr<Image> TempBase = ImagePixelFormat(CompressionQuality, Base, EImageFormat::IF_RGBA_UBYTE);
+				CompressRLE_RGBA(TempBase.get(), Result);
+
+				// Test
+				//Ptr<Image> Test = ImagePixelFormat(CompressionQuality, Result, EImageFormat::IF_RGBA_UBYTE); 
+				//check(Test->DataStorage == Base->DataStorage);
+
+				ReleaseImage(TempBase);
+
+				break;
+			}
+			case EImageFormat::IF_RGB_UBYTE_RLE:
+			{
+				check(bFormatAllLODs);
+
+				// \todo: optimise
+				Ptr<Image> Temp1 = ImagePixelFormat(CompressionQuality, Base, EImageFormat::IF_RGB_UBYTE);
+				Ptr<Image> Temp2 = ImagePixelFormat(CompressionQuality, Temp1.get(), EImageFormat::IF_RGBA_UBYTE);
+				ReleaseImage(Temp1);
+				CompressRLE_RGBA(Temp2.get(), Result);
+				ReleaseImage(Temp2);
+
+				break;
+			}
+			default:
+				// Case not implemented
+				check(false);
+
+			}
+
+			break;
+		}
+
+		case EImageFormat::IF_RGB_UBYTE_RLE:
+		{
+			switch (Base->GetFormat())
+			{
+
+			case EImageFormat::IF_RGB_UBYTE:
+			{
+				check(bFormatAllLODs);
+				CompressRLE_RGB(Base, Result);
+
+				// Test
+				// Ptr<Image> Test = ImagePixelFormat(CompressionQuality, Result, EImageFormat::IF_RGB_UBYTE); 
+				// check(Test->DataStorage == Base->DataStorage);
+				break;
+			}
+
+			case EImageFormat::IF_RGBA_UBYTE_RLE:
+			{
+				check(bFormatAllLODs);
+
+				// \todo: optimise
+				Ptr<Image> Temp1 = ImagePixelFormat(CompressionQuality, Base, EImageFormat::IF_RGBA_UBYTE);
+				Ptr<Image> Temp2 = ImagePixelFormat(CompressionQuality, Temp1.get(), EImageFormat::IF_RGB_UBYTE);
+				ReleaseImage(Temp1);
+				CompressRLE_RGB(Temp2.get(), Result);
+				ReleaseImage(Temp2);
+
+				break;
+			}
+
+			default:
+				// Case not implemented
+				check(false);
+			}
+
+			break;
+		}
+
+		case EImageFormat::IF_BC1:
+		{
+			switch (Base->GetFormat())
+			{
+			case EImageFormat::IF_RGB_UBYTE:
+			{
+				ImagePixelFormatInternal::BlockCompressionLODRangeHelper(
+						Result, Base, ResultLODBegin, BaseLODBegin, NumLODs, CompressionQuality, miro::RGB_to_BC1);
+				break;
+			}
+			case EImageFormat::IF_RGBA_UBYTE:
+			{
+				ImagePixelFormatInternal::BlockCompressionLODRangeHelper(
+						Result, Base, ResultLODBegin, BaseLODBegin, NumLODs, CompressionQuality, miro::RGBA_to_BC1);
+				break;
+			}
+			case EImageFormat::IF_L_UBYTE:
+			{
+				Ptr<Image> TempBase = ImagePixelFormat(CompressionQuality, Base, EImageFormat::IF_RGB_UBYTE);
+				ImagePixelFormatInternal::BlockCompressionLODRangeHelper(
+						Result, TempBase.get(), ResultLODBegin, BaseLODBegin, NumLODs, CompressionQuality, miro::RGB_to_BC1);
+				ReleaseImage(TempBase);
+
+				break;
+			}
+			case EImageFormat::IF_BC3:
+			{
+				Ptr<Image> TempBase = ImagePixelFormat(CompressionQuality, Base, EImageFormat::IF_RGBA_UBYTE);
+				ImagePixelFormatInternal::BlockCompressionLODRangeHelper(
+						Result, TempBase.get(), ResultLODBegin, BaseLODBegin, NumLODs, CompressionQuality, miro::RGBA_to_BC1);
+				ReleaseImage(TempBase);
+
+				break;
+			}
+			default:
+			{
+				UE_LOG(LogMutableCore, Log, TEXT("Image format conversion not implemented. Expensive generic one used. "));
+
+				// Case not implemented, try a generic more expensive operation.
+				Ptr<Image> TempBase = ImagePixelFormat(CompressionQuality, Base, EImageFormat::IF_RGBA_UBYTE);
+				ImagePixelFormatInternal::BlockCompressionLODRangeHelper(
+						Result, TempBase.get(), ResultLODBegin, BaseLODBegin, NumLODs, CompressionQuality, miro::RGBA_to_BC1);
+				ReleaseImage(TempBase);
+			}
+
+			}
+
+			break;
+		}
+
+		case EImageFormat::IF_BC2:
+		{
+			switch (Base->GetFormat())
+			{
+			case EImageFormat::IF_RGBA_UBYTE:
+			{
+				ImagePixelFormatInternal::BlockCompressionLODRangeHelper(
+						Result, Base, ResultLODBegin, BaseLODBegin, NumLODs, CompressionQuality, miro::RGBA_to_BC2);
+				break;
+			}
+			default:
+			{
+				UE_LOG(LogMutableCore, Log, TEXT("Image format conversion not implemented. Expensive generic one used."));
+
+				// Case not implemented, try a generic more expensive operatino.
+				Ptr<Image> TempBase = ImagePixelFormat(CompressionQuality, Base, EImageFormat::IF_RGBA_UBYTE);
+				ImagePixelFormatInternal::BlockCompressionLODRangeHelper(
+						Result, TempBase.get(), ResultLODBegin, BaseLODBegin, NumLODs, CompressionQuality, miro::RGBA_to_BC2);
+				ReleaseImage(TempBase);
+			}
+
+			}
+
+			break;
+		}
+
+		case EImageFormat::IF_BC3:
+		{
+			switch (Base->GetFormat())
+			{
+			case EImageFormat::IF_RGBA_UBYTE:
+			{
+				ImagePixelFormatInternal::BlockCompressionLODRangeHelper(
+						Result, Base, ResultLODBegin, BaseLODBegin, NumLODs, CompressionQuality, miro::RGBA_to_BC3);
+				break;
+			}
+			case EImageFormat::IF_RGB_UBYTE:
+			{
+				ImagePixelFormatInternal::BlockCompressionLODRangeHelper(
+						Result, Base, ResultLODBegin, BaseLODBegin, NumLODs, CompressionQuality, miro::RGB_to_BC3);
+				break;
+			}
+			case EImageFormat::IF_BC1:
+			{
+				MUTABLE_CPUPROFILER_SCOPE(BC1toBC3);
+				ImagePixelFormatInternal::BlockCompressionLODRangeHelper(
+						Result, Base, ResultLODBegin, BaseLODBegin, NumLODs, CompressionQuality, miro::BC1_to_BC3);
+				break;
+			}
+			default:
+			{
+				UE_LOG(LogMutableCore, Log, TEXT("Image format conversion not implemented. Expensive generic one used."));
+
+				// Case not implemented, try a generic more expensive operation.
+				Ptr<Image> TempBase = ImagePixelFormat(CompressionQuality, Base, EImageFormat::IF_RGBA_UBYTE);
+				ImagePixelFormatInternal::BlockCompressionLODRangeHelper(
+						Result, TempBase.get(), ResultLODBegin, BaseLODBegin, NumLODs, CompressionQuality, miro::RGBA_to_BC3);
+				ReleaseImage(TempBase);
+			}
+
+			}
+
+			break;
+		}
+
+		case EImageFormat::IF_BC4:
+		{
+			switch (Base->GetFormat())
+			{
+			case EImageFormat::IF_L_UBYTE:
+			{
+				ImagePixelFormatInternal::BlockCompressionLODRangeHelper(
+						Result, Base, ResultLODBegin, BaseLODBegin, NumLODs, CompressionQuality, miro::L_to_BC4);
+				break;
+			}
+			case EImageFormat::IF_RGB_UBYTE:
+			case EImageFormat::IF_RGBA_UBYTE:
+			case EImageFormat::IF_BGRA_UBYTE:
+			default:
+			{	
+				Ptr<Image> TempBase = ImagePixelFormat(CompressionQuality, Base, EImageFormat::IF_L_UBYTE, -1);
+				ImagePixelFormatInternal::BlockCompressionLODRangeHelper(
+						Result, TempBase.get(), ResultLODBegin, BaseLODBegin, NumLODs, CompressionQuality, miro::L_to_BC4);
+				ReleaseImage(TempBase);
+				break;
+			}
+			}
+
+			break;
+		}
+
+		case EImageFormat::IF_BC5:
+		{
+			switch (Base->GetFormat())
+			{
+			case EImageFormat::IF_RGBA_UBYTE:
+			{
+				ImagePixelFormatInternal::BlockCompressionLODRangeHelper(
+						Result, Base, ResultLODBegin, BaseLODBegin, NumLODs, CompressionQuality, miro::RGBA_to_BC5);
+				break;
+			}
+			case EImageFormat::IF_RGB_UBYTE:
+			{
+				ImagePixelFormatInternal::BlockCompressionLODRangeHelper(
+						Result, Base, ResultLODBegin, BaseLODBegin, NumLODs, CompressionQuality, miro::RGB_to_BC5);
+				break;
+			}
+			default:
+			{
+				UE_LOG(LogMutableCore, Log, TEXT("Image format conversion not implemented. Expensive generic one used. "));
+
+				Ptr<Image> TempBase = ImagePixelFormat(CompressionQuality, Base, EImageFormat::IF_RGBA_UBYTE, -1);
+				ImagePixelFormatInternal::BlockCompressionLODRangeHelper(
+						Result, TempBase.get(), ResultLODBegin, BaseLODBegin, NumLODs, CompressionQuality, miro::RGBA_to_BC5);
+				ReleaseImage(TempBase);
+				break;
+			}
+			}
+
+			break;
+		}
+
+		case EImageFormat::IF_ASTC_4x4_RGB_LDR:
+		{
+			switch (Base->GetFormat())
+			{
+			case EImageFormat::IF_RGBA_UBYTE:
+			{
+				ImagePixelFormatInternal::BlockCompressionLODRangeHelper(
+						Result, Base, ResultLODBegin, BaseLODBegin, NumLODs, CompressionQuality, miro::RGBA_to_ASTC4x4RGBL);
+				break;
+			}
+			case EImageFormat::IF_RGB_UBYTE:
+			{
+				ImagePixelFormatInternal::BlockCompressionLODRangeHelper(
+						Result, Base, ResultLODBegin, BaseLODBegin, NumLODs, CompressionQuality, miro::RGB_to_ASTC4x4RGBL);
+				break;
+			}
+			case EImageFormat::IF_ASTC_4x4_RGBA_LDR:
+			{
+				ImagePixelFormatInternal::BlockDecompressionLODRangeHelper(
+						Result, Base, ResultLODBegin, BaseLODBegin, NumLODs, miro::ASTC4x4RGBAL_to_ASTC4x4RGBL);
+				break;
+			}
+			case EImageFormat::IF_L_UBYTE:
+			{
+				ImagePixelFormatInternal::BlockCompressionLODRangeHelper(
+						Result, Base, ResultLODBegin, BaseLODBegin, NumLODs, CompressionQuality, miro::L_to_ASTC4x4RGBL);
+				break;
+			}
+			default:
+			{
+				UE_LOG(LogMutableCore, Log, TEXT("Image format conversion not implemented. Expensive generic one used. "));
+
+				Ptr<Image> TempBase = ImagePixelFormat(CompressionQuality, Base, EImageFormat::IF_RGBA_UBYTE, -1);
+				ImagePixelFormatInternal::BlockCompressionLODRangeHelper(
+						Result, TempBase.get(), ResultLODBegin, BaseLODBegin, NumLODs, CompressionQuality, miro::RGBA_to_ASTC4x4RGBL);
+				ReleaseImage(TempBase);
+				break;
+			}
+			}
+
+			break;
+		}
+
+		case EImageFormat::IF_ASTC_4x4_RGBA_LDR:
+		{
+			switch (Base->GetFormat())
+			{
+			case EImageFormat::IF_RGBA_UBYTE:
+			{
+				ImagePixelFormatInternal::BlockCompressionLODRangeHelper(
+						Result, Base, ResultLODBegin, BaseLODBegin, NumLODs, CompressionQuality, miro::RGBA_to_ASTC4x4RGBAL);
+				break;
+			}
+			case EImageFormat::IF_RGB_UBYTE:
+			{
+				ImagePixelFormatInternal::BlockCompressionLODRangeHelper(
+						Result, Base, ResultLODBegin, BaseLODBegin, NumLODs, CompressionQuality, miro::RGB_to_ASTC4x4RGBAL);
+				break;
+			}
+			case EImageFormat::IF_ASTC_4x4_RGB_LDR:
+			{
+				ImagePixelFormatInternal::BlockDecompressionLODRangeHelper(
+						Result, Base, ResultLODBegin, BaseLODBegin, NumLODs, miro::ASTC4x4RGBL_to_ASTC4x4RGBAL);
+				break;
+			}
+			default:
+			{
+				UE_LOG(LogMutableCore, Log, TEXT("Image format conversion not implemented. Expensive generic one used. "));
+
+				Ptr<Image> TempBase = ImagePixelFormat(CompressionQuality, Base, EImageFormat::IF_RGBA_UBYTE, -1);
+				ImagePixelFormatInternal::BlockCompressionLODRangeHelper(
+						Result, TempBase.get(), ResultLODBegin, BaseLODBegin, NumLODs, CompressionQuality, miro::RGBA_to_ASTC4x4RGBAL);
+				ReleaseImage(TempBase);
+				break;
+			}
+			}
+
+			break;
+		}
+
+		case EImageFormat::IF_ASTC_4x4_RG_LDR:
+		{
+			switch (Base->GetFormat())
+			{
+			case EImageFormat::IF_RGBA_UBYTE:
+			{
+				ImagePixelFormatInternal::BlockCompressionLODRangeHelper(
+						Result, Base, ResultLODBegin, BaseLODBegin, NumLODs, CompressionQuality, miro::RGBA_to_ASTC4x4RGL);
+				break;
+			}
+			case EImageFormat::IF_RGB_UBYTE:
+			{
+				ImagePixelFormatInternal::BlockCompressionLODRangeHelper(
+						Result, Base, ResultLODBegin, BaseLODBegin, NumLODs, CompressionQuality, miro::RGB_to_ASTC4x4RGL);
+				break;
+			}
+			case EImageFormat::IF_ASTC_4x4_RG_LDR:
+			{
+				// Hack that actually works because of block size.
+				ImagePixelFormatInternal::BlockDecompressionLODRangeHelper(
+						Result, Base, ResultLODBegin, BaseLODBegin, NumLODs, miro::ASTC4x4RGBAL_to_ASTC4x4RGBL);
+				break;
+			}
+			default:
+			{
+				UE_LOG(LogMutableCore, Log, TEXT("Image format conversion not implemented. Expensive generic one used. "));
+
+				Ptr<Image> TempBase = ImagePixelFormat(CompressionQuality, Base, EImageFormat::IF_RGBA_UBYTE, -1);
+				ImagePixelFormatInternal::BlockCompressionLODRangeHelper(
+						Result, TempBase.get(), ResultLODBegin, BaseLODBegin, NumLODs, CompressionQuality, miro::RGBA_to_ASTC4x4RGL);
+				ReleaseImage(TempBase);
+				break;
+			}
+			}
+
+			break;
+		}
+
+		case EImageFormat::IF_ASTC_6x6_RGB_LDR:
+		{
+			switch (Base->GetFormat())
+			{
+			case EImageFormat::IF_RGBA_UBYTE:
+			{
+				ImagePixelFormatInternal::BlockCompressionLODRangeHelper(
+						Result, Base, ResultLODBegin, BaseLODBegin, NumLODs, CompressionQuality, miro::RGBA_to_ASTC6x6RGBL);
+				break;
+			}
+			case EImageFormat::IF_RGB_UBYTE:
+			{
+				ImagePixelFormatInternal::BlockCompressionLODRangeHelper(
+						Result, Base, ResultLODBegin, BaseLODBegin, NumLODs, CompressionQuality, miro::RGB_to_ASTC6x6RGBL);
+				break;
+			}
+			//case EImageFormat::IF_ASTC_6x6_RGBA_LDR:
+			//{
+			//  ImagePixelFormatInternal::BlockCompressionLODRangeHelper(
+			//  		Result, Base, ResultLODBegin, BaseLODBegin, NumLODs, CompressionQuality, miro::ASTC6x6RGBAL_to_ASTC6x6RGBL);
+			//	break;
+			//}
+			//case EImageFormat::IF_L_UBYTE:
+			//{
+			//  ImagePixelFormatInternal::BlockCompressionLODRangeHelper(
+			//  		Result, Base, ResultLODBegin, BaseLODBegin, NumLODs, CompressionQuality, miro::L_to_ASTC6x6RGBL);
+			//	break;
+			//}
+
+			default:
+			{
+				UE_LOG(LogMutableCore, Log, TEXT("Image format conversion not implemented. Expensive generic one used. "));
+
+				Ptr<Image> TempBase = ImagePixelFormat(CompressionQuality, Base, EImageFormat::IF_RGBA_UBYTE, -1);
+				ImagePixelFormatInternal::BlockCompressionLODRangeHelper(
+						Result, TempBase.get(), ResultLODBegin, BaseLODBegin, NumLODs, CompressionQuality, miro::RGBA_to_ASTC6x6RGBL);
+				ReleaseImage(TempBase);
+				break;
+			}
+			}
+
+			break;
+		}
+
+		case EImageFormat::IF_ASTC_6x6_RGBA_LDR:
+		{
+			switch (Base->GetFormat())
+			{
+			case EImageFormat::IF_RGBA_UBYTE:
+			{
+				ImagePixelFormatInternal::BlockCompressionLODRangeHelper(
+						Result, Base, ResultLODBegin, BaseLODBegin, NumLODs, CompressionQuality, miro::RGBA_to_ASTC6x6RGBAL);
+				break;
+			}
+			case EImageFormat::IF_RGB_UBYTE:
+			{
+				ImagePixelFormatInternal::BlockCompressionLODRangeHelper(
+						Result, Base, ResultLODBegin, BaseLODBegin, NumLODs, CompressionQuality, miro::RGB_to_ASTC6x6RGBAL);
+				break;
+			}
+			//case EImageFormat::IF_ASTC_6x6_RGB_LDR:
+			//{
+			//	ImagePixelFormatInternal::BlockCompressionLODRangeHelper(
+			//			Result, Base, ResultLODBegin, BaseLODBegin, NumLODs, CompressionQuality, miro::ASTC6x6RGBL_to_ASTC6x6RGBAL);
+			//	break;
+			//}
+			default:
+			{
+				UE_LOG(LogMutableCore, Log, TEXT("Image format conversion not implemented. Expensive generic one used. "));
+
+				Ptr<Image> TempBase = ImagePixelFormat(CompressionQuality, Base, EImageFormat::IF_RGBA_UBYTE, -1);
+				ImagePixelFormatInternal::BlockCompressionLODRangeHelper(
+						Result, TempBase.get(), ResultLODBegin, BaseLODBegin, NumLODs, CompressionQuality, miro::RGBA_to_ASTC6x6RGBAL);
+				ReleaseImage(TempBase);
+				break;
+			}
+			}
+
+			break;
+		}
+
+		case EImageFormat::IF_ASTC_6x6_RG_LDR:
+		{
+			switch (Base->GetFormat())
+			{
+			case EImageFormat::IF_RGBA_UBYTE:
+			{
+				ImagePixelFormatInternal::BlockCompressionLODRangeHelper(
+						Result, Base, ResultLODBegin, BaseLODBegin, NumLODs, CompressionQuality, miro::RGBA_to_ASTC6x6RGL);
+				break;
+			}
+
+			case EImageFormat::IF_RGB_UBYTE:
+			{
+				ImagePixelFormatInternal::BlockCompressionLODRangeHelper(
+						Result, Base, ResultLODBegin, BaseLODBegin, NumLODs, CompressionQuality, miro::RGB_to_ASTC6x6RGL);
+				break;
+			}
+			//case EImageFormat::IF_ASTC_6x6_RG_LDR:
+			//{
+			//	// Hack that actually works because of block size.
+			//	ImagePixelFormatInternal::BlockCompressionLODRangeHelper(
+			//			Result, Base, ResultLODBegin, BaseLODBegin, NumLODs, CompressionQuality, miro::ASTC6x6RGBAL_to_ASTC6x6RGBL);
+			//	break;
+			//}
+			default:
+			{
+				UE_LOG(LogMutableCore, Log, TEXT("Image format conversion not implemented. Expensive generic one used. "));
+
+				Ptr<Image> TempBase = ImagePixelFormat(CompressionQuality, Base, EImageFormat::IF_RGBA_UBYTE, -1);
+				ImagePixelFormatInternal::BlockCompressionLODRangeHelper(
+						Result, TempBase.get(), ResultLODBegin, BaseLODBegin, NumLODs, CompressionQuality, miro::RGBA_to_ASTC6x6RGL);
+				ReleaseImage(TempBase);
+				break;
+			}
+			}
+
+			break;
+		}
+
+		case EImageFormat::IF_ASTC_8x8_RGB_LDR:
+		{
+			switch (Base->GetFormat())
+			{
+			case EImageFormat::IF_RGBA_UBYTE:
+			{
+				ImagePixelFormatInternal::BlockCompressionLODRangeHelper(
+						Result, Base, ResultLODBegin, BaseLODBegin, NumLODs, CompressionQuality, miro::RGBA_to_ASTC8x8RGBL);
+				break;
+			}
+			case EImageFormat::IF_RGB_UBYTE:
+			{
+				ImagePixelFormatInternal::BlockCompressionLODRangeHelper(
+						Result, Base, ResultLODBegin, BaseLODBegin, NumLODs, CompressionQuality, miro::RGB_to_ASTC8x8RGBL);
+				break;
+			}
+			//case EImageFormat::IF_ASTC_8x8_RGBA_LDR:
+			//{
+			//	ImagePixelFormatInternal::BlockCompressionLODRangeHelper(
+			//			Result, Base, ResultLODBegin, BaseLODBegin, NumLODs, CompressionQuality, miro::ASTC8x8RGBAL_to_ASTC8x8RGBL);
+			//	break;
+			//}
+			//case EImageFormat::IF_L_UBYTE:
+			//{
+			//	ImagePixelFormatInternal::BlockCompressionLODRangeHelper(
+			//			Result, Base, ResultLODBegin, BaseLODBegin, NumLODs, CompressionQuality, miro::L_to_ASTC8x8RGBL);
+			//	break;
+			//}
+			default:
+			{
+				UE_LOG(LogMutableCore, Log, TEXT("Image format conversion not implemented. Expensive generic one used. "));
+
+				Ptr<Image> TempBase = ImagePixelFormat(CompressionQuality, Base, EImageFormat::IF_RGBA_UBYTE, -1);
+				ImagePixelFormatInternal::BlockCompressionLODRangeHelper(
+						Result, TempBase.get(), ResultLODBegin, BaseLODBegin, NumLODs, CompressionQuality, miro::RGBA_to_ASTC8x8RGBL);
+				ReleaseImage(TempBase);
+				break;
+			}
+			}
+
+			break;
+		}
+
+		case EImageFormat::IF_ASTC_8x8_RGBA_LDR:
+		{
+			switch (Base->GetFormat())
+			{
+			case EImageFormat::IF_RGBA_UBYTE:
+			{
+				ImagePixelFormatInternal::BlockCompressionLODRangeHelper(
+						Result, Base, ResultLODBegin, BaseLODBegin, NumLODs, CompressionQuality, miro::RGBA_to_ASTC8x8RGBAL);
+				break;
+			}
+			case EImageFormat::IF_RGB_UBYTE:
+			{
+				ImagePixelFormatInternal::BlockCompressionLODRangeHelper(
+						Result, Base, ResultLODBegin, BaseLODBegin, NumLODs, CompressionQuality, miro::RGB_to_ASTC8x8RGBAL);
+				break;
+			}
+			//case EImageFormat::IF_ASTC_8x8_RGB_LDR:
+			//{
+			//	ImagePixelFormatInternal::BlockCompressionLODRangeHelper(
+			//			Result, Base, ResultLODBegin, BaseLODBegin, NumLODs, CompressionQuality, miro::ASTC8x8RGBL_to_ASTC8x8RGBAL);
+			//	break;
+			//}
+			default:
+			{
+				UE_LOG(LogMutableCore, Log, TEXT("Image format conversion not implemented. Expensive generic one used. "));
+
+				Ptr<Image> TempBase = ImagePixelFormat(CompressionQuality, Base, EImageFormat::IF_RGBA_UBYTE, -1);
+				ImagePixelFormatInternal::BlockCompressionLODRangeHelper(
+						Result, TempBase.get(), ResultLODBegin, BaseLODBegin, NumLODs, CompressionQuality, miro::RGBA_to_ASTC8x8RGBAL);
+				ReleaseImage(TempBase);
+				break;
+			}
+			}
+
+			break;
+		}
+
+		case EImageFormat::IF_ASTC_8x8_RG_LDR:
+		{
+			switch (Base->GetFormat())
+			{
+			case EImageFormat::IF_RGBA_UBYTE:
+			{
+				ImagePixelFormatInternal::BlockCompressionLODRangeHelper(
+						Result, Base, ResultLODBegin, BaseLODBegin, NumLODs, CompressionQuality, miro::RGBA_to_ASTC8x8RGL);
+				break;
+			}
+			case EImageFormat::IF_RGB_UBYTE:
+			{
+				ImagePixelFormatInternal::BlockCompressionLODRangeHelper(
+						Result, Base, ResultLODBegin, BaseLODBegin, NumLODs, CompressionQuality, miro::RGB_to_ASTC8x8RGL);
+				break;
+			}
+			default:
+			{
+				UE_LOG(LogMutableCore, Log, TEXT("Image format conversion not implemented. Expensive generic one used. "));
+
+				Ptr<Image> TempBase = ImagePixelFormat(CompressionQuality, Base, EImageFormat::IF_RGBA_UBYTE, -1);
+				ImagePixelFormatInternal::BlockCompressionLODRangeHelper(
+						Result, TempBase.get(), ResultLODBegin, BaseLODBegin, NumLODs, CompressionQuality, miro::RGBA_to_ASTC8x8RGL);
+				ReleaseImage(TempBase);
+				break;
+			}
+			}
+			
+			break;
+		}
+
+		default:
+			// Case not implemented
+			check(false);
+		}
+	}
+}
+
+}
diff --git a/Engine/Plugins/Mutable/Source/MutableRuntime/Private/MuR/OpImageRasterMesh.h b/Engine/Plugins/Mutable/Source/MutableRuntime/Private/MuR/OpImageRasterMesh.h
new file mode 100644
--- /dev/null
+++ b/Engine/Plugins/Mutable/Source/MutableRuntime/Private/MuR/OpImageRasterMesh.h
@@ -0,0 +1,182 @@
+// Copyright Epic Games, Inc. All Rights Reserved.
+
+#pragma once
+
+#include "Async/ParallelFor.h"
+#include "MuR/ImagePrivate.h"
+#include "MuR/Platform.h"
+#include "MuR/MeshPrivate.h"
+#include "MuR/Raster.h"
+#include "MuR/ConvertData.h"
+
+namespace mu
+{
+
+	class WhitePixelProcessor
+	{
+	public:
+		inline void ProcessPixel(uint8* pBufferPos, float[1]) const
+		{
+			pBufferPos[0] = 255;
+		}
+
+		inline void operator()(uint8* BufferPos, float Interpolators[1]) const
+		{
+			ProcessPixel(BufferPos, Interpolators);
+		}
+	};
+
+
+	inline void ImageRasterMesh( const Mesh* pMesh, Image* pImage, int32 LayoutIndex, uint64 BlockId,
+		UE::Math::TIntVector2<uint16> CropMin, UE::Math::TIntVector2<uint16> UncroppedSize )
+	{
+		MUTABLE_CPUPROFILER_SCOPE(ImageRasterMesh);
+
+		if (pMesh->GetVertexCount() == 0)
+		{
+			return;
+		}
+
+		check( pImage->GetFormat()== EImageFormat::IF_L_UBYTE );
+
+		int32 sizeX = pImage->GetSizeX();
+		int32 sizeY = pImage->GetSizeY();
+
+		// Get the vertices
+		int32 vertexCount = pMesh->GetVertexCount();
+		TArray< RasterVertex<1> > vertices;
+		vertices.SetNumZeroed(vertexCount);
+
+		UntypedMeshBufferIteratorConst texIt( pMesh->GetVertexBuffers(), MBS_TEXCOORDS, 0 );
+		if (!texIt.ptr())
+		{
+			ensure(false);
+			return;
+		}
+
+		for ( int32 v=0; v<vertexCount; ++v )
+		{
+            float uv[2] = {0.0f,0.0f};
+			ConvertData( 0, uv, MBF_FLOAT32, texIt.ptr(), texIt.GetFormat() );
+			ConvertData( 1, uv, MBF_FLOAT32, texIt.ptr(), texIt.GetFormat() );
+
+			bool bUseCropping = UncroppedSize[0] > 0;
+			if (bUseCropping)
+			{
+				vertices[v].x = uv[0] * UncroppedSize[0] - CropMin[0];
+				vertices[v].y = uv[1] * UncroppedSize[1] - CropMin[1];
+			}
+			else
+			{
+				vertices[v].x = uv[0] * sizeX;
+				vertices[v].y = uv[1] * sizeY;
+			}
+			++texIt;
+		}
+
+		// Get the indices
+		int32 faceCount = pMesh->GetFaceCount();
+		TArray<int32> indices;
+		indices.SetNumZeroed(faceCount * 3);
+
+		UntypedMeshBufferIteratorConst indIt( pMesh->GetIndexBuffers(), MBS_VERTEXINDEX, 0 );
+		for ( int32 i=0; i<faceCount*3; ++i )
+		{
+            uint32_t index=0;
+			ConvertData( 0, &index, MBF_UINT32, indIt.ptr(), indIt.GetFormat() );
+
+			indices[i] = index;
+			++indIt;
+		}
+
+        UntypedMeshBufferIteratorConst bloIt( pMesh->GetVertexBuffers(), MBS_LAYOUTBLOCK, LayoutIndex );
+
+        if (BlockId== FLayoutBlock::InvalidBlockId || bloIt.GetElementSize()==0 )
+		{
+			// Raster all the faces
+            WhitePixelProcessor pixelProc;
+			const TArrayView<uint8> ImageData = pImage->DataStorage.GetLOD(0);
+
+			//for ( int f=0; f<faceCount; ++f )
+			const auto& ProcessFace = [
+				vertices, indices, ImageData, sizeX, sizeY, pixelProc
+			] (int32 f)
+			{
+				constexpr int32 NumInterpolators = 1;
+				Triangle<NumInterpolators>(ImageData.GetData(), ImageData.Num(),
+					sizeX, sizeY,
+					1,
+					vertices[indices[f * 3 + 0]],
+					vertices[indices[f * 3 + 1]],
+					vertices[indices[f * 3 + 2]],
+					pixelProc,
+					false);
+			};
+
+			ParallelFor(faceCount, ProcessFace);
+		}
+		else
+		{
+			// Raster only the faces in the selected block
+			check(bloIt.GetComponents() == 1);
+
+			// Get the block per vertex
+			TArray<uint64> VertexBlockIds;
+			VertexBlockIds.SetNumZeroed(vertexCount);
+
+			if (bloIt.GetFormat() == MBF_UINT16)
+			{
+				// Relative blocks.
+				const uint16* SourceIds = reinterpret_cast<const uint16*>(bloIt.ptr());
+				for (int32 i = 0; i < vertexCount; ++i)
+				{
+					uint64 Id = SourceIds[i];
+					Id = Id | (uint64(pMesh->MeshIDPrefix)<<32);
+					VertexBlockIds[i] = Id;
+				}
+			}
+			else if (bloIt.GetFormat() == MBF_UINT64)
+			{
+				// Absolute blocks.
+				const uint64* SourceIds = reinterpret_cast<const uint64*>(bloIt.ptr());
+				for (int32 i = 0; i < vertexCount; ++i)
+				{
+					uint64 Id = SourceIds[i];
+					VertexBlockIds[i] = Id;
+				}
+			}
+			else
+			{
+				// Format not supported
+				check(false);
+			}
+
+            WhitePixelProcessor pixelProc;
+
+			const TArrayView<uint8> ImageData = pImage->DataStorage.GetLOD(0); 
+
+			const auto& ProcessFace = [
+				vertices, indices, VertexBlockIds, BlockId, ImageData, sizeX, sizeY, pixelProc
+			] (int32 f)
+			{
+				// TODO: Select faces outside for loop?
+				if (VertexBlockIds[indices[f * 3 + 0]] == BlockId)
+				{
+					constexpr int32 NumInterpolators = 1;
+					Triangle<NumInterpolators>(ImageData.GetData(), ImageData.Num(),
+						sizeX, sizeY,
+						1,
+						vertices[indices[f * 3 + 0]],
+						vertices[indices[f * 3 + 1]],
+						vertices[indices[f * 3 + 2]],
+						pixelProc,
+						false);
+				}
+			};
+
+			ParallelFor(faceCount, ProcessFace);
+		}
+
+	}
+
+}
diff --git a/Engine/Plugins/Mutable/Source/MutableRuntime/Private/MuR/OpLayoutPack.h b/Engine/Plugins/Mutable/Source/MutableRuntime/Private/MuR/OpLayoutPack.h
new file mode 100644
--- /dev/null
+++ b/Engine/Plugins/Mutable/Source/MutableRuntime/Private/MuR/OpLayoutPack.h
@@ -0,0 +1,12 @@
+// Copyright Epic Games, Inc. All Rights Reserved.
+
+#pragma once
+
+
+namespace mu
+{
+	class Layout;
+
+	extern void LayoutPack3(Layout* pResult, const Layout* pSourceLayout);
+
+}
diff --git a/Engine/Plugins/Mutable/Source/MutableRuntime/Private/MuR/OpMeshApplyLayout.h b/Engine/Plugins/Mutable/Source/MutableRuntime/Private/MuR/OpMeshApplyLayout.h
new file mode 100644
--- /dev/null
+++ b/Engine/Plugins/Mutable/Source/MutableRuntime/Private/MuR/OpMeshApplyLayout.h
@@ -0,0 +1,227 @@
+// Copyright Epic Games, Inc. All Rights Reserved.
+
+#pragma once
+
+#include "MuR/MeshPrivate.h"
+#include "MuR/MutableRuntimeModule.h"
+#include "MuR/MutableTrace.h"
+#include "MuR/Platform.h"
+#include "MuR/OpMeshRemove.h"
+
+namespace mu
+{
+
+
+	//---------------------------------------------------------------------------------------------
+    inline void MeshApplyLayout( Mesh* pApplied, const Layout* pLayout, int32 texCoordsSet )
+	{
+		MUTABLE_CPUPROFILER_SCOPE(MeshApplyLayout);
+
+        int32 buffer = -1;
+		int32 channel = -1;
+		pApplied->GetVertexBuffers().FindChannel( MBS_TEXCOORDS, texCoordsSet, &buffer, &channel );
+
+        int32 layoutBuffer = -1;
+        int32 layoutChannel = -1;
+        pApplied->GetVertexBuffers().FindChannel( MBS_LAYOUTBLOCK, texCoordsSet, &layoutBuffer, &layoutChannel );
+
+		if (buffer < 0 || layoutBuffer < 0)
+		{
+			return;
+		}
+
+		// Get the information about the texture coordinates channel
+		EMeshBufferSemantic semantic;
+		int32 semanticIndex;
+		EMeshBufferFormat format;
+		int32 components;
+		int32 offset;
+		pApplied->GetVertexBuffers().GetChannel( buffer, channel, &semantic, &semanticIndex, &format, &components, &offset );
+		check( semantic == MBS_TEXCOORDS );
+
+        uint8* pData = pApplied->GetVertexBuffers().GetBufferData( buffer );
+		int32 elemSize = pApplied->GetVertexBuffers().GetElementSize( buffer );
+		int32 channelOffset = pApplied->GetVertexBuffers().GetChannelOffset( buffer, channel );
+		pData += channelOffset;
+
+		struct Box
+		{
+			FVector2f min, size;
+		};
+		TArray< Box > transforms;
+		transforms.SetNum(pLayout->GetBlockCount());
+		for ( int b=0; b<pLayout->GetBlockCount(); ++b )
+		{
+			FIntPoint grid = pLayout->GetGridSize();
+
+			box< FIntVector2 > block;
+			block.min = pLayout->Blocks[b].Min;
+			block.size = pLayout->Blocks[b].Size;
+
+			Box rect;
+			rect.min[0] = ( (float)block.min[0] ) / (float) grid[0];
+			rect.min[1] = ( (float)block.min[1] ) / (float) grid[1];
+			rect.size[0] = ( (float)block.size[0] ) / (float) grid[0];
+			rect.size[1] = ( (float)block.size[1] ) / (float) grid[1];
+			transforms[b] = rect;
+		}
+
+        check( layoutBuffer>=0 && layoutChannel>=0 );
+        check( pApplied->GetVertexBuffers().Buffers[layoutBuffer].Channels.Num()==1 );
+        check( pApplied->GetVertexBuffers().Buffers[layoutBuffer].Channels[layoutChannel].ComponentCount==1 );
+
+        const uint16* pLayoutData = reinterpret_cast<const uint16*>( pApplied->GetVertexBuffers().GetBufferData( layoutBuffer ) );
+		UntypedMeshBufferIterator ItLayoutData(pApplied->GetVertexBuffers(), MBS_LAYOUTBLOCK, texCoordsSet);
+
+		// In some corner case involving automatic LODs and remove meshes behaving differently among them we may need to remove vertices that don't 
+		// have any block in the current layout. Track them here.
+		TArray<int32> VerticesToRemove;
+
+        uint8* pVertices = pData;
+		for ( int32 v=0; v<pApplied->GetVertexBuffers().GetElementCount(); ++v )
+		{
+			uint64 BlockId = 0;
+			if (ItLayoutData.GetFormat() == MBF_UINT16)
+			{
+				// Relative blocks.
+				const uint16* SourceIds = reinterpret_cast<const uint16*>(ItLayoutData.ptr());
+				BlockId = *SourceIds;
+				BlockId = BlockId | (uint64(pApplied->MeshIDPrefix) << 32);
+			}
+			else if (ItLayoutData.GetFormat() == MBF_UINT64)
+			{
+				// Absolute blocks.
+				const uint64* SourceIds = reinterpret_cast<const uint64*>(ItLayoutData.ptr());
+				BlockId = *SourceIds;
+			}
+			else
+			{
+				// Format not supported
+				check(false);
+			}
+			++ItLayoutData;
+
+			// TODO: This could be optimised
+			int32 relBlock = pLayout->FindBlock(BlockId);
+
+			// This may still happen with lower LOD and "remove meshes" in a corner case:
+			// Auto LODs with "Remove Meshes" that behave differently across LODs, and leave geometry in a block that has been removed in the higher LOD.
+			if (relBlock < 0)
+			{
+				VerticesToRemove.Add(v); 
+				pVertices += elemSize;
+				continue;
+			}
+
+			switch( format )
+			{
+			case MBF_FLOAT32:
+			{
+				FVector2f uv;
+				FMemory::Memcpy( &uv, pVertices, sizeof(float)*2 );
+
+				uv = uv * transforms[ relBlock ].size + transforms[ relBlock ].min;
+				FMemory::Memcpy( pVertices, &uv, sizeof(float)*2 );
+				break;
+			}
+
+			case MBF_FLOAT16:
+			{
+				// TODO: Optimise
+				FFloat16* pUV = reinterpret_cast<FFloat16*>( pVertices );
+
+				FVector2f uv;
+				uv[0] = float( pUV[0] );
+				uv[1] = float( pUV[1] );
+				uv = uv * transforms[ relBlock ].size + transforms[ relBlock ].min;
+
+				pUV[0] = FFloat16( uv[0] );
+				pUV[1] = FFloat16( uv[1] );
+				break;
+			}
+
+			case MBF_NUINT32:
+			case MBF_NINT32:
+			case MBF_UINT32:
+			case MBF_INT32:
+			{
+				// TODO: Optimise
+                uint32* pUV = reinterpret_cast<uint32*>( pVertices );
+
+				for ( int c=0; c<2; ++c )
+				{
+                    uint64 u_32 = pUV[c];
+                    uint64 u_48 = u_32 * ((uint64)(((float)0xffff) * transforms[ relBlock ].size[c]));
+                    u_48 += (uint64)(((float)0xffffffffffULL) * transforms[ relBlock ].min[c]);
+                    pUV[c] = (uint16)( u_48 >> 16 );
+				}
+				break;
+			}
+
+			case MBF_NUINT16:
+			case MBF_NINT16:
+			case MBF_UINT16:
+			case MBF_INT16:
+			{
+				// TODO: Optimise
+                uint16* pUV = reinterpret_cast<uint16*>( pVertices );
+
+				for ( int c=0; c<2; ++c )
+				{
+                    uint32 u_16 = pUV[c];
+                    uint32 u_32 = u_16 * ((uint8)(((float)0xffff) * transforms[ relBlock ].size[c]));
+                    u_32 += (uint32)(((float)0xffffffff) * transforms[ relBlock ].min[c]);
+                    pUV[c] = (uint16)( u_32 >> 16 );
+				}
+
+				break;
+			}
+
+			case MBF_NUINT8:
+			case MBF_NINT8:
+			case MBF_UINT8:
+			case MBF_INT8:
+			{
+				// TODO: Optimise
+                uint8* pUV = reinterpret_cast<uint8*>( pVertices );
+
+				for ( int c=0; c<2; ++c )
+				{
+                    uint32 u_8 = pUV[c];
+                    uint32 u_24 = u_8 * ((uint32)(((float)0xffff) * transforms[ relBlock ].size[c]));
+                    u_24 += (uint32)(((float)0xffffff) * transforms[ relBlock ].min[c]);
+                    pUV[c] = (uint8)( u_24 >> 16 );
+				}
+				break;
+			}
+
+			default:
+				UE_LOG(LogMutableCore, Warning, TEXT("This case is not implemented.."));
+				check( false );
+				break;
+
+			}
+
+			pVertices += elemSize;
+		}
+
+		if (VerticesToRemove.Num())
+		{
+			// Unpack vertices into a mask
+			TBitArray<> VertexMask;
+			VertexMask.SetNum(pApplied->GetVertexCount(), false);
+			for (int32 VertexIndex : VerticesToRemove)
+			{
+				VertexMask[VertexIndex] = true;
+			}
+	
+			// Remove
+			MeshRemoveVerticesWithCullSet(pApplied, VertexMask);
+		}
+
+		//
+		pApplied->SetLayout( texCoordsSet, pLayout->Clone() );
+
+	}
+
+}
diff --git a/Engine/Plugins/Mutable/Source/MutableRuntime/Private/MuR/OpMeshClipWithMesh.h b/Engine/Plugins/Mutable/Source/MutableRuntime/Private/MuR/OpMeshClipWithMesh.h
new file mode 100644
--- /dev/null
+++ b/Engine/Plugins/Mutable/Source/MutableRuntime/Private/MuR/OpMeshClipWithMesh.h
@@ -0,0 +1,28 @@
+// Copyright Epic Games, Inc. All Rights Reserved.
+
+#pragma once
+
+#include "MuR/Platform.h"
+
+namespace mu
+{
+	class Mesh;
+	class Image;
+	class Layout;
+
+    /**  */
+	extern void MeshClipWithMesh(Mesh* Result, const Mesh* pBase, const Mesh* pClipMesh, bool& bOutSuccess);
+
+    /** Generate a mask mesh with the faces of the base mesh inside the clip mesh. */
+	extern void MeshMaskClipMesh(Mesh* Result, const Mesh* pBase, const Mesh* pClipMesh, bool& bOutSuccess);
+
+	/** Generate a mask mesh with the faces of the base mesh that have all 3 vertices marked in the given mask. */
+	extern void MakeMeshMaskFromUVMask(Mesh* Result, const Mesh* Base, const Mesh* BaseForUVs, const Image* Mask, uint8 LayoutIndex, bool& bOutSuccess);
+
+    /** Generate a mask mesh with the faces of the base mesh matching the fragment. */
+	extern void MeshMaskDiff(Mesh* Result, const Mesh* pBase, const Mesh* pFragment, bool& bOutSuccess);
+
+	/** Generate a mask mesh with the faces of the base mesh that have all 3 vertices inside any block of the given layout. */
+	extern void MakeMeshMaskFromLayout(Mesh* Result, const Mesh* Base, const Mesh* BaseForUVs, const Layout* Mask, uint8 LayoutIndex, bool& bOutSuccess);
+
+}
diff --git a/Engine/Plugins/Mutable/Source/MutableRuntime/Private/MuR/OpMeshInterpolate.h b/Engine/Plugins/Mutable/Source/MutableRuntime/Private/MuR/OpMeshInterpolate.h
new file mode 100644
--- /dev/null
+++ b/Engine/Plugins/Mutable/Source/MutableRuntime/Private/MuR/OpMeshInterpolate.h
@@ -0,0 +1,72 @@
+// Copyright Epic Games, Inc. All Rights Reserved.
+
+#pragma once
+
+#include "MuR/MeshPrivate.h"
+#include "MuR/Platform.h"
+
+
+namespace mu
+{
+
+
+	//---------------------------------------------------------------------------------------------
+    inline MeshPtr MeshInterpolate( const Mesh* pA, const Mesh* pB, float factor )
+	{
+		bool correct = ( pA->GetVertexBuffers().GetElementCount() == pB->GetVertexBuffers().GetElementCount() )
+					&& ( pA->GetIndexCount() == pB->GetIndexCount() )
+					&& ( pA->HasCompatibleFormat(pB) );
+
+		MeshPtr pDest = pA->Clone();
+
+		// Deal with degenerated data
+		if ( correct && pDest->GetVertexBuffers().GetElementCount() )
+		{
+			// For now, just interpolate the position and normal buffer.
+
+			// A mesh position
+			int aPosElementSize = 0;
+            uint8* pAPosBuf = 0;
+			GetMeshBuf( pDest.get(), MBS_POSITION, MBF_FLOAT32, 3, pAPosBuf, aPosElementSize );
+
+			// A mesh normal
+			int aNorElementSize = 0;
+            uint8* pANorBuf = 0;
+			GetMeshBuf( pDest.get(), MBS_NORMAL, MBF_FLOAT32, 3, pANorBuf, aNorElementSize );
+
+			// B mesh position
+			int bPosElementSize = 0;
+            const uint8* pBPosBuf = 0;
+			GetMeshBuf( pB, MBS_POSITION, MBF_FLOAT32, 3, pBPosBuf, bPosElementSize );
+
+			// B mesh normal
+			int bNorElementSize = 0;
+            const uint8* pBNorBuf = 0;
+			GetMeshBuf( pB, MBS_NORMAL, MBF_FLOAT32, 3, pBNorBuf, bNorElementSize );
+
+            uint32 vcount = pA->GetVertexBuffers().GetElementCount();
+            for ( uint32 vi=0; vi<vcount; ++vi )
+			{
+				FVector3f* aPos = reinterpret_cast<FVector3f* >( pAPosBuf );
+				const FVector3f* bPos = reinterpret_cast< const FVector3f* >( pBPosBuf );
+				*aPos += (*bPos - *aPos) * factor;
+
+				FVector3f* aNor = reinterpret_cast<FVector3f* >( pANorBuf );
+				const FVector3f* bNor = reinterpret_cast< const FVector3f* >( pBNorBuf );
+				*aNor += (*bNor - *aNor) * factor;
+
+				pAPosBuf += aPosElementSize;
+				pBPosBuf += bPosElementSize;
+				pANorBuf += aNorElementSize;
+				pBNorBuf += bNorElementSize;
+			}
+		}
+		else
+		{
+			// TODO report somehow
+		}
+
+		return pDest;
+	}
+
+}
