[mutable] Moved the Mutable plugin out of Experimental status into Beta.

diff --git a/Engine/Plugins/Mutable/Source/MutableRuntime/Private/MuR/OpMeshReshape.h b/Engine/Plugins/Mutable/Source/MutableRuntime/Private/MuR/OpMeshReshape.h
new file mode 100644
--- /dev/null
+++ b/Engine/Plugins/Mutable/Source/MutableRuntime/Private/MuR/OpMeshReshape.h
@@ -0,0 +1,975 @@
+// Copyright Epic Games, Inc. All Rights Reserved.
+
+#pragma once
+
+#include "MuR/ConvertData.h"
+#include "MuR/MeshPrivate.h"
+#include "MuR/Operations.h"
+#include "MuR/OpMeshBind.h"
+
+#include "Math/Plane.h"
+#include "Math/Ray.h"
+#include "Math/NumericLimits.h"
+
+#include "IndexTypes.h"
+#include "Spatial/MeshAABBTree3.h"
+#include "Distance/DistPoint3Triangle3.h"
+#include "LineTypes.h"
+#include "MuR/OpMeshSmoothing.h"
+#include "MuR/OpMeshComputeNormals.h"
+
+#include "Algo/AnyOf.h" 
+
+
+// TODO: Make the handling of rotations an option. It is more expensive on CPU and memory, and for some
+// cases it is not required at all.
+
+// TODO: Face stretch to scale the deformation per-vertex? 
+
+// TODO: Support multiple binding influences per vertex, to have smoother deformations.
+
+// TODO: Support multiple binding sets, to have different shapes deformations at once.
+
+// TODO: Deformation mask, to select the intensisty of the deformation per-vertex.
+
+// TODO: This is a reference implementation with ample roof for optimization.
+
+namespace mu
+{
+
+	struct FShapeMeshDescriptorApply
+	{
+		TArray<FVector3f> Positions;
+		TArray<FVector3f> Normals;
+		TArray<UE::Geometry::FIndex3i> Triangles;
+	};
+
+	// Methods to actually deform a point
+	inline void GetDeform( const FShapeMeshDescriptorApply& Shape, const FReshapeVertexBindingData& Binding, FVector3f& NewPosition, FVector3f& NewNormalPosition)
+	{
+		const UE::Geometry::FIndex3i& Triangle = Shape.Triangles[Binding.Triangle];
+	
+		FVector3f VA = Shape.Positions[Triangle.A] + Shape.Normals[Triangle.A]*Binding.D.X;
+		FVector3f VB = Shape.Positions[Triangle.B] + Shape.Normals[Triangle.B]*Binding.D.Y;
+		FVector3f VC = Shape.Positions[Triangle.C] + Shape.Normals[Triangle.C]*Binding.D.Z;
+
+		NewPosition = VA*Binding.S + VB*Binding.T + VC*(1.0f - Binding.S - Binding.T);
+	
+		FVector3f NVA = Shape.Positions[Triangle.A] + Shape.Normals[Triangle.A]*Binding.NormalD.X;
+		FVector3f NVB = Shape.Positions[Triangle.B] + Shape.Normals[Triangle.B]*Binding.NormalD.Y;
+		FVector3f NVC = Shape.Positions[Triangle.C] + Shape.Normals[Triangle.C]*Binding.NormalD.Z;
+
+		NewNormalPosition = (NVA*Binding.NS + NVB*Binding.NT + NVC*(1.0f - Binding.NS - Binding.NT));
+	}
+
+	inline void GetDeform( const FShapeMeshDescriptorApply& Shape, const FReshapePointBindingData& Binding, FVector3f& NewPosition)
+	{
+		const UE::Geometry::FIndex3i& Triangle = Shape.Triangles[Binding.Triangle];
+	
+		FVector3f VA = Shape.Positions[Triangle.A] + Shape.Normals[Triangle.A]*Binding.D.X;
+		FVector3f VB = Shape.Positions[Triangle.B] + Shape.Normals[Triangle.B]*Binding.D.Y;
+		FVector3f VC = Shape.Positions[Triangle.C] + Shape.Normals[Triangle.C]*Binding.D.Z;
+
+		NewPosition = VA*Binding.S + VB*Binding.T + VC*(1.0f - Binding.S - Binding.T);
+	}
+
+
+    //---------------------------------------------------------------------------------------------
+    //! Physics Bodies Reshape 
+    //---------------------------------------------------------------------------------------------
+	template<uint32 NumPoints>
+	inline bool GetDeformedPoints( const FShapeMeshDescriptorApply& Shape, const FReshapePointBindingData* BindingData,  TStaticArray<FVector3f, NumPoints>& OutPoints )
+	{
+		const int32 ShapeNumTris = Shape.Triangles.Num();
+		for (int32 I = 0; I < NumPoints; ++I)
+		{
+			const FReshapePointBindingData& BindingDataPoint = *(BindingData + I);
+			if ((BindingDataPoint.Triangle < 0) | (BindingDataPoint.Triangle >= ShapeNumTris))
+			{
+				return false;
+			}
+		}
+
+		for (int32 I = 0; I < NumPoints; ++I)
+		{
+			GetDeform(Shape, *(BindingData + I), OutPoints[I]);
+		}
+
+		return true;
+	}
+
+	inline void GetDeformedConvex( const FShapeMeshDescriptorApply& Shape, const FReshapePointBindingData* BindingData, TArrayView<FVector3f>& InOutDeformedVertices )
+	{
+		const int32 ShapeNumTris = Shape.Triangles.Num();
+		const int32 ConvexVertCount = InOutDeformedVertices.Num();
+
+		for (int32 I = 0; I < ConvexVertCount; ++I)
+		{
+			const FReshapePointBindingData& BindingDataPoint = *(BindingData + I);
+
+			if ((BindingDataPoint.Triangle < 0) | (BindingDataPoint.Triangle >= ShapeNumTris))
+			{
+				continue;
+			}
+
+			GetDeform(Shape, BindingDataPoint, InOutDeformedVertices[I]);	
+		}
+	}
+
+	inline void ComputeSphereFromDeformedPoints( 
+			const TStaticArray<FVector3f, 6>& Points, 
+			FVector3f& OutP, float& OutR, 
+			const FTransform3f& InvBoneT )
+	{	
+		FVector3f Centroid(ForceInitToZero);
+		for (const FVector3f& V : Points)
+		{
+			Centroid += V;	
+		}
+
+		Centroid *= (1.0f/6.0f);
+
+		OutP = Centroid + InvBoneT.GetTranslation();
+		
+		float Radius = 0.0f;
+		for (const FVector3f& V : Points)
+		{
+			Radius += (V - Centroid).Length();
+		}
+
+		Radius *= (1.0f/6.0f);
+		
+		OutR = Radius;
+	}
+
+	inline void ComputeBoxFromDeformedPoints( 
+			const TStaticArray<FVector3f, 14>& Points, 
+			FVector3f& OutP, FQuat4f& OutQ, FVector3f& OutS, 
+			const FTransform3f& InvBoneT )
+	{	
+		const FVector3f TopC = (Points[0] + Points[1] + Points[2] + Points[3])*0.25f;
+		const FVector3f BottomC = (Points[4+0] + Points[4+1] + Points[4+2] + Points[4+3])*0.25f;
+    
+		const FVector3f FrontC = (Points[0] + Points[1] + Points[4+0] + Points[4+1])*0.25f;
+		const FVector3f BackC = (Points[2] + Points[3] + Points[4+2] + Points[4+3])*0.25f;
+    
+		const FVector3f RightC =  (Points[1] + Points[2] + Points[4+1] + Points[4+2])*0.25f;
+		const FVector3f LeftC =  (Points[3] + Points[0] + Points[4+3] + Points[4+0])*0.25f;
+    
+		FVector3f ZB = (TopC - BottomC).GetSafeNormal();
+		FVector3f XB = (RightC - LeftC).GetSafeNormal();
+		FVector3f YB = (FrontC - BackC).GetSafeNormal();
+
+
+		// Pick the 2 most offaxis vectors and construct a rotation form those.	
+		// TODO: Find a better way of finding and orientation form ZB, XB, TN maybe by averaging somehow different bases created from the vectors, with quaternions?
+		FVector3f OF = FVector3f( XB.X, YB.Y, ZB.Y ).GetAbs();
+
+		float M0 = FMath::Max3( OF.X, OF.Y, OF.X);
+		float M1 = M0 == OF.X ? FMath::Max(OF.Y, OF.Z) : (M0 == OF.Y ? FMath::Max(OF.X,OF.Z) : FMath::Max(OF.Y, OF.X));
+		
+		FMatrix44f RotationMatrix = FMatrix44f::Identity;
+		if (M0 == OF.X)
+		{
+			RotationMatrix = M1 == OF.Y 
+					? FRotationMatrix44f::MakeFromXY(XB, YB)
+					: FRotationMatrix44f::MakeFromXZ(XB, ZB);
+		}
+		else if (M0 == OF.Y)
+		{
+			RotationMatrix = M1 == XB.X 
+					? FRotationMatrix44f::MakeFromYX(YB, XB)
+					: FRotationMatrix44f::MakeFromYZ(YB, ZB);
+		}
+		else
+		{
+			RotationMatrix = M1 == XB.X 
+					? FRotationMatrix44f::MakeFromZX(ZB, XB)
+					: FRotationMatrix44f::MakeFromZY(ZB, YB);	
+		} 
+	
+		FTransform3f ShapeToBone = FTransform3f(RotationMatrix.ToQuat(), (TopC + BottomC) * 0.5f) * InvBoneT;
+
+		OutQ = ShapeToBone.GetRotation();
+		OutP = ShapeToBone.GetTranslation();
+		OutS = FVector3f((RightC - LeftC).Size(), (FrontC - BackC).Size(), (TopC - BottomC).Size()) * 0.5f;	
+	}
+
+	inline void ComputeSphylFromDeformedPoints( 
+			const TStaticArray<FVector3f, 14>& Points, 
+			FVector3f& OutP, FQuat4f& OutQ, float& OutR, float& OutL, 
+			const FTransform3f& InvBoneT)
+	{
+		constexpr int32 NumCentroids = 5;
+		TStaticArray<FVector3f, NumCentroids> Centroids;
+		Centroids[0] = Points[0];
+		Centroids[1] = Points[1];
+
+		for (int32 I = 0; I < NumCentroids - 2; ++I)
+		{
+			Centroids[2 + I] = 
+				(Points[2 + I*4 + 0] + Points[2 + I*4 + 1] +
+				 Points[2 + I*4 + 2] + Points[2 + I*4 + 3] ) * 0.25f;
+		}
+
+		// Geometric linear regression of top, bottom and rings centroids.
+		FVector3f Centroid = FVector3f::Zero();
+		for (const FVector3f& C : Centroids) //-V1078
+		{
+			Centroid += C;
+		}
+
+		constexpr float OneOverNumCentroids = 1.0f / static_cast<float>(NumCentroids);
+		Centroid *= OneOverNumCentroids;
+		
+		for (FVector3f& C : Centroids) //-V1078
+		{
+			C -= Centroid;
+		}
+
+		constexpr int32 NumIters = 3;
+		FVector3f Direction = (Centroids[0] - Centroids[1]).GetSafeNormal();
+		for (int32 Iter = 0; Iter < NumIters; ++Iter)
+		{
+			FVector3f IterDirRefinement = Direction;
+			for (const FVector3f& C : Centroids) //-V1078
+			{
+				IterDirRefinement += C * FVector3f::DotProduct(Direction, C);
+			}
+
+			Direction = IterDirRefinement.GetSafeNormal();
+		}
+
+		// Project centroids to the line described by Direction and Centroid.
+		for (FVector3f& C : Centroids) //-V1078
+		{
+			C = Centroid + Direction * FVector3f::DotProduct(C, Direction);
+		}
+
+		// Quaternion form {0,0,1} to Direction.
+		const FQuat4f Rotation = FQuat4f(-Direction.Y, Direction.X, 0.0f, 1.0f + FMath::Max(Direction.Z, -1.0f + UE_SMALL_NUMBER)).GetNormalized(); 
+
+		FTransform3f ShapeToBone = FTransform3f(Rotation, Centroid) * InvBoneT;
+		OutQ = ShapeToBone.GetRotation();
+		OutP = ShapeToBone.GetTranslation();
+
+		// Project ring points to plane formed by ring centroid and direction to extract ring radius.
+		const auto ComputeRadiusContribution = [&](const FVector3f& P, const FVector3f& Origin, const FVector3f& Dir) -> float
+		{
+			return ((P + Dir * FVector3f::DotProduct(Dir, P - Origin)) - Origin).Length();
+		};
+
+		const float R0 = 
+			ComputeRadiusContribution(Points[2], Centroids[2], Direction) +
+			ComputeRadiusContribution(Points[3], Centroids[2], Direction) +
+			ComputeRadiusContribution(Points[4], Centroids[2], Direction) +
+			ComputeRadiusContribution(Points[5], Centroids[2], Direction);
+
+		const float R1 = 
+			ComputeRadiusContribution(Points[6], Centroids[3], Direction) +
+			ComputeRadiusContribution(Points[7], Centroids[3], Direction) +
+			ComputeRadiusContribution(Points[8], Centroids[3], Direction) +
+			ComputeRadiusContribution(Points[9], Centroids[3], Direction);
+
+		const float R2 = 
+			ComputeRadiusContribution(Points[10], Centroids[4], Direction) +
+			ComputeRadiusContribution(Points[11], Centroids[4], Direction) +
+			ComputeRadiusContribution(Points[12], Centroids[4], Direction) +
+			ComputeRadiusContribution(Points[13], Centroids[4], Direction);
+
+		OutR =  (R0 + R1 + R2) * (0.25f/3.0f);
+		OutL = FMath::Max(0.0f, (Centroids[0] - Centroids[1]).Length() - OutR*2.0f);
+	}
+
+	inline void ComputeTaperedCapsuleFromDeformedPoints( 
+			const TStaticArray<FVector3f, 14>& Points, 
+			FVector3f& OutP, FQuat4f& OutQ, float& OutR0, float& OutR1, float& OutL, 
+			const FTransform3f& InvBoneT)
+	{
+		constexpr int32 NumCentroids = 5;
+		TStaticArray<FVector3f, NumCentroids> Centroids;
+		Centroids[0] = Points[0];
+		Centroids[1] = Points[1];
+
+		for (int32 I = 0; I < NumCentroids - 2; ++I)
+		{
+			Centroids[2 + I] = 
+				(Points[2 + I*4 + 0] + Points[2 + I*4 + 1] +
+				 Points[2 + I*4 + 2] + Points[2 + I*4 + 3] ) * 0.25f;
+		}
+	
+		// Geometric linear regression of top, bottom and ring centroids.
+		FVector3f Centroid = FVector3f::Zero();
+		for (const FVector3f& C : Centroids) //-V1078
+		{
+			Centroid += C;
+		}
+
+		constexpr float OneOverNumCentroids = 1.0f / static_cast<float>(NumCentroids);
+		Centroid *= OneOverNumCentroids;
+		
+		for (FVector3f& C : Centroids) //-V1078
+		{
+			C -= Centroid;
+		}
+
+		constexpr int32 NumIters = 3;
+		FVector3f Direction = (Centroids[0] - Centroids[1]).GetSafeNormal();
+		for (int32 Iter = 0; Iter < NumIters; ++Iter)
+		{
+			FVector3f IterDirRefinement = Direction;
+			for (const FVector3f& C : Centroids) //-V1078
+			{
+				IterDirRefinement += C * FVector3f::DotProduct(Direction, C);
+			}
+
+			Direction = IterDirRefinement.GetSafeNormal();
+		}
+
+		// Project centroids to the line described by Direction and Centroid.
+		for (FVector3f& C : Centroids) //-V1078
+		{
+			C = Centroid + Direction * FVector3f::DotProduct(C, Direction);
+		}
+
+		// Quaternion form {0,0,1} to Direction.
+		const FQuat4f Rotation = FQuat4f(-Direction.Y, Direction.X, 0.0f, 1.0f + FMath::Max(Direction.Z, -1.0f + UE_SMALL_NUMBER)).GetNormalized(); 
+
+		FTransform3f ShapeToBone = FTransform3f(Rotation, Centroid) * InvBoneT;
+		OutQ = ShapeToBone.GetRotation();
+		OutP = ShapeToBone.GetTranslation();
+
+		// Project ring points to plane formed by ring centroid and direction to extract ring radius.
+		const auto ComputeRadiusContribution = [&](const FVector3f& P, const FVector3f& Origin, const FVector3f& Dir) -> float
+		{
+			return ((P + Dir * FVector3f::DotProduct(Dir, P - Origin)) - Origin).Length();
+		};
+
+		const float R0 = 
+			ComputeRadiusContribution(Points[2], Centroids[2], Direction) +
+			ComputeRadiusContribution(Points[3], Centroids[2], Direction) +
+			ComputeRadiusContribution(Points[4], Centroids[2], Direction) +
+			ComputeRadiusContribution(Points[5], Centroids[2], Direction);
+
+		const float R1 = 
+			ComputeRadiusContribution(Points[6], Centroids[3], Direction) +
+			ComputeRadiusContribution(Points[7], Centroids[3], Direction) +
+			ComputeRadiusContribution(Points[8], Centroids[3], Direction) +
+			ComputeRadiusContribution(Points[9], Centroids[3], Direction);
+
+		const float R2 = 
+			ComputeRadiusContribution(Points[10], Centroids[4], Direction) +
+			ComputeRadiusContribution(Points[11], Centroids[4], Direction) +
+			ComputeRadiusContribution(Points[12], Centroids[4], Direction) +
+			ComputeRadiusContribution(Points[13], Centroids[4], Direction);
+
+		// TODO: Ajust for R1, center ring radius.
+		OutR0 = R0*0.25f;
+		OutR1 = R2*0.25f;
+		OutL = FMath::Max(0.0f, (Centroids[0] - Centroids[1]).Length() - (OutR0 + OutR1));
+	}
+
+	inline void ApplyToVertices(Mesh* Mesh, TArrayView<const FReshapeVertexBindingData> BindingData, const FShapeMeshDescriptorApply& Shape, bool bSkipNormalReshape)
+	{
+		check(Mesh);
+		check(Mesh->GetVertexCount() == BindingData.Num());
+		
+		const UntypedMeshBufferIterator PositionIter(Mesh->GetVertexBuffers(), MBS_POSITION);
+		const UntypedMeshBufferIterator NormalIter(Mesh->GetVertexBuffers(), MBS_NORMAL);
+		const UntypedMeshBufferIterator TangentIter(Mesh->GetVertexBuffers(), MBS_TANGENT);
+		const UntypedMeshBufferIterator BiNormalIter(Mesh->GetVertexBuffers(), MBS_BINORMAL);
+
+        const EMeshBufferFormat NormalFormat   = NormalIter.GetFormat();
+        const EMeshBufferFormat TangentFormat  = TangentIter.GetFormat();
+        const EMeshBufferFormat BiNormalFormat = BiNormalIter.GetFormat();
+	
+        const int32 NormalComps   = NormalIter.GetComponents();
+        const int32 TangentComps  = TangentIter.GetComponents();
+        const int32 BiNormalComps = BiNormalIter.GetComponents();
+
+#if DO_CHECK
+		// checking if the Base shape has more triangles than the target shape
+		const bool bTriangleOutOfScopeFound = Algo::AnyOf(BindingData, 
+			[NumShapeTriangles = Shape.Triangles.Num()](const FReshapeVertexBindingData& B) { return B.Triangle >= NumShapeTriangles; });
+
+		if (bTriangleOutOfScopeFound)
+		{
+			UE_LOG(LogMutableCore, Warning, TEXT("Performing a Mesh Reshape where base shape and target shape do not have the same number of triangles."));
+		}	
+#endif
+
+		const int32 ShapeTriangleCount = Shape.Triangles.Num();
+		const int32 ShapePositionCount = Shape.Positions.Num();
+
+		const int32 MeshVertexCount = BindingData.Num();
+
+		const bool bComputeNormal = bSkipNormalReshape && NormalIter.ptr();
+		for (int32 VertexIndex = 0; VertexIndex < MeshVertexCount; ++VertexIndex)
+		{
+			const FReshapeVertexBindingData& Binding = BindingData[VertexIndex];
+
+			const bool bModified = (Binding.Triangle >= 0) & (Binding.Triangle < ShapeTriangleCount);
+			if (bModified)
+			{
+				FVector3f NewPosition;
+				FVector3f NewNormalPosition;
+				
+				GetDeform(Shape, Binding, NewPosition, NewNormalPosition);
+
+				const FVector3f OldPosition = (PositionIter + VertexIndex).GetAsVec3f();
+
+				FVector3f PositionDisplacement = (NewPosition - OldPosition) * Binding.Weight;
+				
+				// TODO: Separeat rigid from not rigid processing.
+				(PositionIter + VertexIndex).SetFromVec3f(OldPosition + PositionDisplacement);
+				
+				if (bComputeNormal)
+				{
+					const FVector3f OldNormal = (NormalIter + VertexIndex).GetAsVec3f();
+					const FVector3f OldNormalPosition = OldPosition + OldNormal;
+					
+					FVector3f NormalDisplacement = (NewNormalPosition - OldNormalPosition) * Binding.Weight;
+
+					FVector3f Normal = ((OldNormalPosition + NormalDisplacement) - NewPosition).GetSafeNormal();
+					FVector3f Tangent = TangentIter.ptr() ? (TangentIter + VertexIndex).GetAsVec3f() : FVector3f::XAxisVector;
+					FVector3f BiNormal = BiNormalIter.ptr() ? (BiNormalIter + VertexIndex).GetAsVec3f() : FVector3f::YAxisVector;
+				
+					OrthogonalizeTangentSpace(
+							&Normal, 
+							TangentIter.ptr() ? &Tangent : nullptr, 
+							BiNormalIter.ptr() ? &BiNormal : nullptr,
+							BiNormalIter.ptr() ? ComputeTangentBasisDeterminantSign(OldNormal, Tangent, BiNormal) : 0.0f);
+					
+					// Leave the tangent basis sign untouched for packed normals formats.
+					uint8 * const NormalElemPtr = (NormalIter + VertexIndex).ptr();
+					for (int32 C = 0; C < NormalComps && C < 3; ++C)
+					{
+						ConvertData(C, NormalElemPtr, NormalFormat, &Normal, MBF_FLOAT32);
+					}
+
+					if (TangentIter.ptr())
+					{
+						uint8 * const TangentElemPtr = (TangentIter + VertexIndex).ptr();
+						for (int32 C = 0; C < TangentComps && C < 3; ++C)
+						{
+							ConvertData(C, TangentElemPtr, TangentFormat, &Tangent, MBF_FLOAT32);
+						}
+					}
+					
+					if (BiNormalIter.ptr())
+					{
+						uint8 * const BiNormalElemPtr = (BiNormalIter + VertexIndex).ptr();
+						for (int32 C = 0; C < BiNormalComps && C < 3; ++C)
+						{
+							ConvertData(C, BiNormalElemPtr, BiNormalFormat, &BiNormal, MBF_FLOAT32);
+						}
+					}
+				}
+			}
+		}
+	}
+
+	inline void ApplyToPose(Mesh* Result, 
+			TArrayView<const FReshapePointBindingData> BindingData, 
+			TArrayView<const int32> BoneIndices,
+			const FShapeMeshDescriptorApply& Shape)
+	{
+		check(Result);
+
+#if DO_CHECK
+		// checking if the Base shape has more triangles than the target shape
+		const bool bTriangleOutOfScopeFound = Algo::AnyOf(BindingData, 
+			[NumShapeTriangles = Shape.Triangles.Num()](const FReshapePointBindingData& B) { return B.Triangle >= NumShapeTriangles; });
+
+		if (bTriangleOutOfScopeFound)
+		{
+			UE_LOG(LogMutableCore, Warning, TEXT("Performing a Mesh Reshape where base shape and target shape do not have the same number of triangles."));
+		}	
+#endif
+
+		const int32 NumShapeTris = Shape.Triangles.Num();
+		const int32 NumBoneIndices = BoneIndices.Num();
+		for (int32 BoneSelectionIndex = 0; BoneSelectionIndex < NumBoneIndices; ++BoneSelectionIndex)
+		{
+			int32 BoneIndex = BoneIndices[BoneSelectionIndex];
+
+			const FReshapePointBindingData& Binding = BindingData[BoneSelectionIndex];
+
+			check(!EnumHasAnyFlags(Result->BonePoses[BoneIndex].BoneUsageFlags, EBoneUsageFlags::Root));
+
+			FTransform3f& T = Result->BonePoses[BoneIndex].BoneTransform;
+
+			FVector3f NewPosition(ForceInitToZero);
+
+			const bool bModified = (Binding.Triangle >= 0) & (Binding.Triangle < NumShapeTris);
+			
+			if (bModified)
+			{
+				GetDeform(Shape, Binding, NewPosition);
+			}
+
+			const bool bHasChanged = bModified && FVector3f::DistSquared(NewPosition, T.GetLocation()) > UE_SMALL_NUMBER;
+			// Only set it if has actually moved.
+			if (bHasChanged)
+			{	
+				// Mark as reshaped. 
+				EnumAddFlags(Result->BonePoses[BoneIndex].BoneUsageFlags, EBoneUsageFlags::Reshaped);
+
+				// TODO: Review if the rotation also needs to be applied.
+				T.SetLocation(FMath::Lerp(T.GetLocation(), NewPosition, Binding.Weight));
+			}
+		}
+	}
+
+	inline void ApplyToPhysicsBodies(
+			PhysicsBody& PBody, int32& InOutNumProcessedBindPoints, 
+			const Mesh& BaseMesh, 
+			TArrayView<const FReshapePointBindingData> BindingData, 
+			TArrayView<const int32> UsedIndices, 
+			const FShapeMeshDescriptorApply& Shape)
+	{
+#if DO_CHECK
+		// checking if the Base shape has more triangles than the target shape
+		const bool bTriangleOutOfScopeFound = Algo::AnyOf(BindingData, 
+			[NumShapeTriangles = Shape.Triangles.Num()](const FReshapePointBindingData& B) { return B.Triangle >= NumShapeTriangles; });
+
+		if (bTriangleOutOfScopeFound)
+		{
+			UE_LOG(LogMutableCore, Warning, TEXT("Performing a Mesh Reshape where base shape and target shape do not have the same number of triangles."));
+		}
+#endif
+		
+		bool bAnyModified = false;
+
+		// Retrieve them in the same order the boxes where put in, so they can be linked to the physics body volumes.
+		for ( const int32 B : UsedIndices )
+		{	
+			int32 BoneIndex = BaseMesh.FindBonePose(PBody.GetBodyBoneId(B));
+			FTransform3f BoneTransform = FTransform3f::Identity;
+			if (BoneIndex >= 0)
+			{
+				BaseMesh.GetBonePoseTransform(BoneIndex, BoneTransform);
+			}
+			
+			FTransform3f InvBoneTransform = BoneTransform.Inverse();
+
+			const int32 NumSpheres = PBody.GetSphereCount(B); 
+			for ( int32 I = 0; I < NumSpheres; ++I )
+			{
+				FVector3f P;
+				float R;
+
+				TStaticArray<FVector3f, 6> Points; 
+				const bool bDeformed = GetDeformedPoints(Shape, &BindingData[InOutNumProcessedBindPoints], Points);
+
+				if (bDeformed)
+				{
+					ComputeSphereFromDeformedPoints(Points, P, R, InvBoneTransform);
+					PBody.SetSphere(B, I, P, R);
+					bAnyModified = true;
+				}
+
+				InOutNumProcessedBindPoints += Points.Num();
+			}
+
+			const int32 NumBoxes = PBody.GetBoxCount(B);
+			for ( int32 I = 0; I < NumBoxes; ++I )
+			{
+				TStaticArray<FVector3f, 14> Points;
+				const bool bDeformed = GetDeformedPoints(Shape, &BindingData[InOutNumProcessedBindPoints], Points);
+				
+				if (bDeformed)
+				{
+					FVector3f P;
+					FQuat4f Q;
+					FVector3f S;
+					
+					ComputeBoxFromDeformedPoints(Points, P, Q, S, InvBoneTransform);
+					PBody.SetBox(B, I, P, Q, S);
+					bAnyModified = true;
+				}
+
+				InOutNumProcessedBindPoints += Points.Num();
+			}
+			
+			const int32 NumSphyls = PBody.GetSphylCount(B);
+			for ( int32 I = 0; I < NumSphyls; ++I )
+			{
+				TStaticArray<FVector3f, 14> Points;
+				const bool bDeformed = GetDeformedPoints(Shape, &BindingData[InOutNumProcessedBindPoints], Points);
+
+				if ( bDeformed )
+				{
+					FVector3f P;
+					FQuat4f Q;
+					float R;
+					float L;						
+
+					ComputeSphylFromDeformedPoints(Points, P, Q, R, L, InvBoneTransform);
+					PBody.SetSphyl(B, I, P, Q, R, L);
+					bAnyModified = true;
+				}
+
+				InOutNumProcessedBindPoints += Points.Num();
+			}
+
+			const int32 NumTaperedCapsules = PBody.GetTaperedCapsuleCount(B);
+			for ( int32 I = 0; I < NumTaperedCapsules; ++I )
+			{
+				TStaticArray<FVector3f, 14> Points;
+				const bool bDeformed = GetDeformedPoints(Shape, &BindingData[InOutNumProcessedBindPoints], Points);
+
+				if ( bDeformed )
+				{
+					FVector3f P;
+					FQuat4f Q;
+					float R0;
+					float R1;
+					float L;
+				
+					ComputeTaperedCapsuleFromDeformedPoints(Points, P, Q, R0, R1, L, InvBoneTransform);
+					PBody.SetTaperedCapsule(B, I, P, Q, R0, R1, L);
+					bAnyModified = true;
+				}
+				
+				InOutNumProcessedBindPoints += Points.Num();
+			}
+
+			const int32 NumConvex = PBody.GetConvexCount(B);
+			for ( int32 I = 0; I < NumConvex; ++I )
+			{
+				TArrayView<FVector3f> VerticesView;
+				TArrayView<int32> IndicesView;
+				PBody.GetConvexMeshView(B, I, VerticesView, IndicesView);
+
+				FTransform3f ConvexTransform;
+				PBody.GetConvexTransform(B, I, ConvexTransform);
+
+				GetDeformedConvex(Shape, &BindingData[InOutNumProcessedBindPoints], VerticesView);
+
+				FTransform3f InvConvexT = InvBoneTransform * ConvexTransform.Inverse();
+				for ( FVector3f& V : VerticesView )
+				{
+					V = InvConvexT.TransformPosition( V );				
+				}
+				
+				InOutNumProcessedBindPoints += VerticesView.Num();
+				bAnyModified = true;
+			}
+		}
+
+		PBody.bBodiesModified = bAnyModified;
+	}
+
+	inline void ApplyToAllPhysicsBodies(
+		Mesh& OutNewMesh, const Mesh& BaseMesh, 
+		TArrayView<const FReshapePointBindingData> BindingData,
+		TArrayView<const int32> UsedIndices,
+		TArrayView<const int32> BodyOffsets,
+		const FShapeMeshDescriptorApply& Shape)
+	{
+		if (!BodyOffsets.Num())
+		{
+			return; // Nothing to do.
+		}
+
+		int32 NumProcessedBindPoints = 0;
+
+		auto ApplyPhysicsBody = [&](PhysicsBody& OutBody, int32 IndicesBegin, int32 IndicesEnd) -> void
+		{
+			TArrayView<const int32> BodyUsedIndices(UsedIndices.GetData() + IndicesBegin, IndicesEnd - IndicesBegin);
+
+			ApplyToPhysicsBodies(OutBody, NumProcessedBindPoints, BaseMesh, BindingData, BodyUsedIndices, Shape);
+		};
+
+
+		check(BodyOffsets.Num() > 1);
+		check(OutNewMesh.AdditionalPhysicsBodies.Num() + 1 == BodyOffsets.Num() - 1);
+		const int32 PhysicsBodiesNum = BodyOffsets.Num() - 1;
+		
+		// Apply main physics body
+		if (BaseMesh.PhysicsBody)
+		{
+			Ptr<PhysicsBody> NewBody = BaseMesh.PhysicsBody->Clone();
+			ApplyPhysicsBody(*NewBody, 0, BodyOffsets[1]);
+			OutNewMesh.PhysicsBody = NewBody;
+		}
+
+
+		// Apply additional physics bodies
+		for (int32 I = 1; I < PhysicsBodiesNum; ++I)
+		{	
+			Ptr<PhysicsBody> NewBody = BaseMesh.AdditionalPhysicsBodies[I - 1]->Clone();
+			ApplyPhysicsBody(*NewBody, BodyOffsets[I], BodyOffsets[I + 1]);
+			OutNewMesh.AdditionalPhysicsBodies[I - 1] = NewBody;
+		}
+
+		check(NumProcessedBindPoints == BindingData.Num());	
+	}
+
+	//---------------------------------------------------------------------------------------------
+	//! Rebuild the (previously bound) mesh data for a new shape.
+	//! Proof-of-concept implementation.
+	//---------------------------------------------------------------------------------------------
+	inline void MeshApplyShape(Mesh* Result, const Mesh* BaseMesh, const Mesh* ShapeMesh, EMeshBindShapeFlags BindFlags, bool& bOutSuccess)
+	{
+		MUTABLE_CPUPROFILER_SCOPE(MeshApplyReshape);
+
+		bOutSuccess = true;
+
+		if (!BaseMesh)
+		{
+			bOutSuccess = false;
+			return;
+		}
+		
+		const bool bReshapeVertices = EnumHasAnyFlags(BindFlags, EMeshBindShapeFlags::ReshapeVertices);
+		const bool bRecomputeNormals = EnumHasAnyFlags(BindFlags, EMeshBindShapeFlags::RecomputeNormals);
+		const bool bReshapeSkeleton = EnumHasAnyFlags(BindFlags, EMeshBindShapeFlags::ReshapeSkeleton);
+		const bool bReshapePhysicsVolumes = EnumHasAnyFlags(BindFlags, EMeshBindShapeFlags::ReshapePhysicsVolumes);
+		const bool bApplyLaplacian = EnumHasAnyFlags(BindFlags, EMeshBindShapeFlags::ApplyLaplacian);
+
+		// Early out if nothing will be modified and the vertices discarted.
+		const bool bSkeletonModification = BaseMesh->GetSkeleton() && bReshapeSkeleton;
+		const bool bPhysicsModification = (BaseMesh->GetPhysicsBody() || BaseMesh->AdditionalPhysicsBodies.Num()) && bReshapePhysicsVolumes;
+
+		if (!bReshapeVertices && !bSkeletonModification && !bPhysicsModification)
+		{
+			bOutSuccess = false;
+			return;
+		}
+	
+		// \TODO: Multiple binding data support
+		int32 BindingDataIndex = 0;
+
+		// If the base mesh has no binding data, just clone it.
+		int32 BarycentricDataBuffer = 0;
+		int32 BarycentricDataChannel = 0;
+		const FMeshBufferSet& VB = BaseMesh->GetVertexBuffers();
+		VB.FindChannel(MBS_BARYCENTRICCOORDS, BindingDataIndex, &BarycentricDataBuffer, &BarycentricDataChannel);
+		
+		Ptr<Mesh> TempMesh = nullptr;
+		Mesh* VerticesReshapeMesh = Result;
+		
+		if (bApplyLaplacian)
+		{
+			TempMesh = new Mesh();
+			VerticesReshapeMesh = TempMesh.get();
+		}
+
+		// Copy Without VertexBuffers or AdditionalBuffers
+		constexpr EMeshCopyFlags CopyFlags = ~(EMeshCopyFlags::WithVertexBuffers | EMeshCopyFlags::WithAdditionalBuffers);
+		VerticesReshapeMesh->CopyFrom(*BaseMesh, CopyFlags);
+	
+		FMeshBufferSet& ResultBuffers = VerticesReshapeMesh->GetVertexBuffers();
+
+		check(ResultBuffers.Buffers.Num() == 0);
+
+		// Copy buffers skipping binding data. 
+		ResultBuffers.ElementCount = VB.ElementCount;
+		// Remove one element to the number of buffers if BarycentricDataBuffer found. 
+		ResultBuffers.Buffers.SetNum(FMath::Max(0, VB.Buffers.Num() - int32(BarycentricDataBuffer >= 0)));
+		
+		for (int32 B = 0, R = 0; B < VB.Buffers.Num(); ++B)
+		{
+			if (B != BarycentricDataBuffer)
+			{
+				ResultBuffers.Buffers[R++] = VB.Buffers[B];
+			}
+		}
+
+		// Copy the additional buffers skipping binding data.
+		VerticesReshapeMesh->AdditionalBuffers.Reserve(BaseMesh->AdditionalBuffers.Num());
+		for (const TPair<EMeshBufferType, FMeshBufferSet>& A : BaseMesh->AdditionalBuffers)
+		{	
+			const bool bIsBindOpBuffer = 
+					A.Key == EMeshBufferType::SkeletonDeformBinding || 
+					A.Key == EMeshBufferType::PhysicsBodyDeformBinding ||
+					A.Key == EMeshBufferType::PhysicsBodyDeformSelection ||
+					A.Key == EMeshBufferType::PhysicsBodyDeformOffsets;
+	
+			if (!bIsBindOpBuffer)
+			{
+				VerticesReshapeMesh->AdditionalBuffers.Add(A);
+			}
+		}
+		if (!ShapeMesh)
+		{
+			return;
+		}
+
+		int32 ShapeVertexCount = ShapeMesh->GetVertexCount();
+		int32 ShapeTriangleCount = ShapeMesh->GetFaceCount();
+		if (!ShapeVertexCount || !ShapeTriangleCount)
+		{	
+			return;
+		}
+
+		// Generate the temp vertex query data for the shape
+		// TODO: Vertex data copy may be saved in most cases.
+		TArray<FVector3f> ShapePositions;
+		FShapeMeshDescriptorApply ShapeDescriptor;
+		{
+			MUTABLE_CPUPROFILER_SCOPE(GenerateVertexQueryData);
+
+			ShapeDescriptor.Positions.SetNum(ShapeVertexCount);
+			ShapeDescriptor.Normals.SetNum(ShapeVertexCount);
+
+			// \TODO: Simple but inefficient
+			const UntypedMeshBufferIteratorConst ItPosition(ShapeMesh->GetVertexBuffers(), MBS_POSITION);
+			const UntypedMeshBufferIteratorConst ItNormal(ShapeMesh->GetVertexBuffers(), MBS_NORMAL);
+			for (int32 ShapeVertexIndex = 0; ShapeVertexIndex < ShapeVertexCount; ++ShapeVertexIndex)
+			{
+				FVector3f Position = (ItPosition + ShapeVertexIndex).GetAsVec3f();
+				ShapeDescriptor.Positions[ShapeVertexIndex] = Position;
+
+				FVector3f Normal = (ItNormal + ShapeVertexIndex).GetAsVec3f();
+				ShapeDescriptor.Normals[ShapeVertexIndex] = Normal;
+			}
+		}
+
+		// Generate the temp face query data for the shape
+		// TODO: Index data copy may be saved in most cases.
+		{
+			MUTABLE_CPUPROFILER_SCOPE(GenerateTriangleQueryData);
+
+			ShapeDescriptor.Triangles.SetNum(ShapeTriangleCount);
+
+			// \TODO: Simple but inefficient
+			const UntypedMeshBufferIteratorConst ItIndices(ShapeMesh->GetIndexBuffers(), MBS_VERTEXINDEX);
+			for (int32 TriangleIndex = 0; TriangleIndex < ShapeTriangleCount; ++TriangleIndex)
+			{
+				UE::Geometry::FIndex3i Triangle;
+				Triangle.A = int32((ItIndices + TriangleIndex*3 + 0).GetAsUINT32());
+				Triangle.B = int32((ItIndices + TriangleIndex*3 + 1).GetAsUINT32());
+				Triangle.C = int32((ItIndices + TriangleIndex*3 + 2).GetAsUINT32());
+
+				ShapeDescriptor.Triangles[TriangleIndex] = Triangle;
+			}
+		}
+
+		if (bReshapeVertices && BarycentricDataBuffer >= 0)
+		{
+			{
+				MUTABLE_CPUPROFILER_SCOPE(ReshapeVertices);
+				// \TODO: More checks
+				check(BarycentricDataChannel == 0);
+				check(VB.GetElementSize(BarycentricDataBuffer) == (int)sizeof(FReshapeVertexBindingData));
+
+				TArrayView<const FReshapeVertexBindingData> VerticesBindingData(
+						(const FReshapeVertexBindingData*)VB.GetBufferData(BarycentricDataBuffer),
+						VB.GetElementCount());
+
+				ApplyToVertices(VerticesReshapeMesh, VerticesBindingData, ShapeDescriptor, bRecomputeNormals);
+			}
+
+			if (bApplyLaplacian)
+			{
+				check(Result && VerticesReshapeMesh);
+				// check result is empty at this point.
+				check(Result->GetVertexCount() == 0 && Result->GetIndexCount() == 0); 
+				SmoothMeshLaplacian(*Result, *VerticesReshapeMesh);
+			}
+
+			if (bRecomputeNormals)
+			{
+				ComputeMeshNormals(*Result);
+			}
+		}
+	
+		if (bReshapeSkeleton)
+		{
+			MUTABLE_CPUPROFILER_SCOPE(ReshapeSkeleton);
+
+			// If the base mesh has no binding data for the skeleton don't do anything.
+			const FMeshBufferSet* SkeletonBindBuffer = nullptr;
+			for ( const TPair<EMeshBufferType, FMeshBufferSet>& A : BaseMesh->AdditionalBuffers )
+			{
+				if ( A.Key == EMeshBufferType::SkeletonDeformBinding )
+				{
+					SkeletonBindBuffer = &A.Value;
+					break;
+				}
+			}
+			
+			BarycentricDataBuffer = -1;
+			if (SkeletonBindBuffer)
+			{
+				SkeletonBindBuffer->FindChannel(MBS_BARYCENTRICCOORDS, BindingDataIndex, &BarycentricDataBuffer, &BarycentricDataChannel);
+			}
+
+			if (SkeletonBindBuffer && BarycentricDataBuffer >= 0)
+			{
+				// \TODO: More checks
+				check(BarycentricDataChannel == 0);
+				check(SkeletonBindBuffer && SkeletonBindBuffer->GetElementSize(BarycentricDataBuffer) == (int)sizeof(FReshapePointBindingData));
+
+				TArrayView<const FReshapePointBindingData> SkeletonBindingData( 
+						(const FReshapePointBindingData*)SkeletonBindBuffer->GetBufferData(BarycentricDataBuffer),
+						SkeletonBindBuffer->GetElementCount());
+			
+				check(SkeletonBindBuffer->GetBufferCount() >= 2);
+				TArrayView<const int32> BoneIndices( 
+						(const int32*)SkeletonBindBuffer->GetBufferData(1), SkeletonBindBuffer->GetElementCount());
+
+				ApplyToPose(Result, SkeletonBindingData, BoneIndices, ShapeDescriptor);
+			}
+		}
+
+		// When transforming the physics volumes, the resulting pose of of the skeleton reshape operation will be used, so 
+		// order of operation is important.
+
+		// Transform physics volumes based on the deformed sampling points.
+		const PhysicsBody* OldPhysicsBody = Result->PhysicsBody.get();
+		const int32 AdditionalPhysicsBodiesNum = Result->AdditionalPhysicsBodies.Num();
+
+		if (bReshapePhysicsVolumes && (OldPhysicsBody || AdditionalPhysicsBodiesNum))
+		{	
+			MUTABLE_CPUPROFILER_SCOPE(ReshapePhysicsBodies);
+			
+			using BufferEntryType = TPair<EMeshBufferType, FMeshBufferSet>;
+			const BufferEntryType* FoundPhysicsBindBuffer = BaseMesh->AdditionalBuffers.FindByPredicate(
+					[](BufferEntryType& E){ return E.Key == EMeshBufferType::PhysicsBodyDeformBinding; });
+
+			const BufferEntryType* FoundPhysicsBindSelectionBuffer = BaseMesh->AdditionalBuffers.FindByPredicate(
+					[](BufferEntryType& E){ return E.Key == EMeshBufferType::PhysicsBodyDeformSelection; });
+	
+			const BufferEntryType* FoundPhysicsBindOffsetsBuffer = BaseMesh->AdditionalBuffers.FindByPredicate(
+					[](BufferEntryType& E){ return E.Key == EMeshBufferType::PhysicsBodyDeformOffsets; });
+
+			BarycentricDataBuffer = -1;
+			if (FoundPhysicsBindBuffer)
+			{
+				FoundPhysicsBindBuffer->Value.FindChannel(MBS_BARYCENTRICCOORDS, BindingDataIndex, &BarycentricDataBuffer, &BarycentricDataChannel);
+			}
+			
+			const bool bAllNeededBuffersFound =
+				FoundPhysicsBindBuffer && FoundPhysicsBindSelectionBuffer && FoundPhysicsBindOffsetsBuffer && BarycentricDataBuffer >= 0;
+
+			if (bAllNeededBuffersFound)
+			{
+				const FMeshBufferSet& PhysicsBindBuffer = FoundPhysicsBindBuffer->Value;
+				const FMeshBufferSet& PhyiscsBindSelectionBuffer = FoundPhysicsBindSelectionBuffer->Value;
+				const FMeshBufferSet& PhyiscsBindOffsetsBuffer = FoundPhysicsBindOffsetsBuffer->Value;
+
+				// \TODO: More checks
+				check(BarycentricDataChannel == 0);
+				check(PhysicsBindBuffer.GetElementSize(BarycentricDataBuffer) == (int)sizeof(FReshapePointBindingData));
+					
+				TArrayView<const FReshapePointBindingData> BindingData(  
+						(const FReshapePointBindingData*)PhysicsBindBuffer.GetBufferData(BarycentricDataBuffer),
+						PhysicsBindBuffer.GetElementCount() );
+
+				TArrayView<const int32> UsedIndices( 
+						(const int32*)PhyiscsBindSelectionBuffer.GetBufferData(0), 
+					    PhyiscsBindSelectionBuffer.GetElementCount());
+
+				TArrayView<const int32> Offsets(
+						(const int32*)PhyiscsBindOffsetsBuffer.GetBufferData(0), 
+					    PhyiscsBindOffsetsBuffer.GetElementCount());
+
+				ApplyToAllPhysicsBodies(*Result, *BaseMesh, BindingData, UsedIndices, Offsets, ShapeDescriptor);
+			}
+		}
+	}
+}
diff --git a/Engine/Plugins/Mutable/Source/MutableRuntime/Private/MuR/Operations.cpp b/Engine/Plugins/Mutable/Source/MutableRuntime/Private/MuR/Operations.cpp
new file mode 100644
--- /dev/null
+++ b/Engine/Plugins/Mutable/Source/MutableRuntime/Private/MuR/Operations.cpp
@@ -0,0 +1,1093 @@
+// Copyright Epic Games, Inc. All Rights Reserved.
+
+
+#include "MuR/Operations.h"
+
+#include "MuR/ModelPrivate.h"
+
+
+namespace mu
+{
+    MUTABLE_IMPLEMENT_POD_SERIALISABLE(OP);
+    MUTABLE_IMPLEMENT_POD_VECTOR_SERIALISABLE(OP);
+	
+	//---------------------------------------------------------------------------------------------
+	//---------------------------------------------------------------------------------------------
+	//---------------------------------------------------------------------------------------------
+
+    // clang-format off
+    static const FOpDesc s_opDescs[] =
+	{ 
+		// type				cached	supported base image formats
+        { DT_NONE,			false,	{ 0, 0, 0, 0, 0, 0, 0, 0, 0, 0, 0, 0, 0, 0, 0, 0, 0, 0, 0, 0, 0, 0, 0, 0, 0, 0, 0, 0, 0 }		},	// NONE
+
+		{ DT_BOOL,			false,	{ 0, 0, 0, 0, 0, 0, 0, 0, 0, 0, 0, 0, 0, 0, 0, 0, 0, 0, 0, 0, 0, 0, 0, 0, 0, 0, 0, 0, 0 }		},	// BO_CONSTANT
+		{ DT_INT,			false,	{ 0, 0, 0, 0, 0, 0, 0, 0, 0, 0, 0, 0, 0, 0, 0, 0, 0, 0, 0, 0, 0, 0, 0, 0, 0, 0, 0, 0, 0 }		},	// NU_CONSTANT
+		{ DT_SCALAR,		false,	{ 0, 0, 0, 0, 0, 0, 0, 0, 0, 0, 0, 0, 0, 0, 0, 0, 0, 0, 0, 0, 0, 0, 0, 0, 0, 0, 0, 0, 0 }		},	// SC_CONSTANT
+		{ DT_COLOUR,		false,	{ 0, 0, 0, 0, 0, 0, 0, 0, 0, 0, 0, 0, 0, 0, 0, 0, 0, 0, 0, 0, 0, 0, 0, 0, 0, 0, 0, 0, 0 }		},	// CO_CONSTANT
+		{ DT_IMAGE,			true,   { 1, 1, 1, 1, 1, 1, 1, 1, 1, 1, 1, 1, 0, 0, 0, 0, 0, 0, 0, 0, 0, 0, 0, 0, 0, 0, 0, 0, 0 }		},	// IM_CONSTANT
+		{ DT_MESH,			true,   { 0, 0, 0, 0, 0, 0, 0, 0, 0, 0, 0, 0, 0, 0, 0, 0, 0, 0, 0, 0, 0, 0, 0, 0, 0, 0, 0, 0, 0 }		},	// ME_CONSTANT
+		{ DT_LAYOUT,		false,	{ 0, 0, 0, 0, 0, 0, 0, 0, 0, 0, 0, 0, 0, 0, 0, 0, 0, 0, 0, 0, 0, 0, 0, 0, 0, 0, 0, 0, 0 }		},	// LA_CONSTANT
+		{ DT_PROJECTOR,		false,	{ 0, 0, 0, 0, 0, 0, 0, 0, 0, 0, 0, 0, 0, 0, 0, 0, 0, 0, 0, 0, 0, 0, 0, 0, 0, 0, 0, 0, 0 }		},	// PR_CONSTANT
+		{ DT_STRING,		false,	{ 0, 0, 0, 0, 0, 0, 0, 0, 0, 0, 0, 0, 0, 0, 0, 0, 0, 0, 0, 0, 0, 0, 0, 0, 0, 0, 0, 0, 0 }		},	// ST_CONSTANT
+		{ DT_EXTENSION_DATA,true,	{ 0, 0, 0, 0, 0, 0, 0, 0, 0, 0, 0, 0, 0, 0, 0, 0, 0, 0, 0, 0, 0, 0, 0, 0, 0, 0, 0, 0, 0 }		},	// ED_CONSTANT
+
+		{ DT_BOOL,			false,	{ 0, 0, 0, 0, 0, 0, 0, 0, 0, 0, 0, 0, 0, 0, 0, 0, 0, 0, 0, 0, 0, 0, 0, 0, 0, 0, 0, 0, 0 }		},	// BO_PARAMETER
+		{ DT_INT,			false,	{ 0, 0, 0, 0, 0, 0, 0, 0, 0, 0, 0, 0, 0, 0, 0, 0, 0, 0, 0, 0, 0, 0, 0, 0, 0, 0, 0, 0, 0 }		},	// NU_PARAMETER
+		{ DT_SCALAR,		false,	{ 0, 0, 0, 0, 0, 0, 0, 0, 0, 0, 0, 0, 0, 0, 0, 0, 0, 0, 0, 0, 0, 0, 0, 0, 0, 0, 0, 0, 0 }		},	// SC_PARAMETER
+		{ DT_COLOUR,		false,	{ 0, 0, 0, 0, 0, 0, 0, 0, 0, 0, 0, 0, 0, 0, 0, 0, 0, 0, 0, 0, 0, 0, 0, 0, 0, 0, 0, 0, 0 }		},	// CO_PARAMETER
+		{ DT_PROJECTOR,		false,	{ 0, 0, 0, 0, 0, 0, 0, 0, 0, 0, 0, 0, 0, 0, 0, 0, 0, 0, 0, 0, 0, 0, 0, 0, 0, 0, 0, 0, 0 }		},	// PR_PARAMETER
+		{ DT_IMAGE,			false,	{ 0, 0, 0, 0, 0, 0, 0, 0, 0, 0, 0, 0, 0, 0, 0, 0, 0, 0, 0, 0, 0, 0, 0, 0, 0, 0, 0, 0, 0 }		},	// IM_PARAMETER
+		{ DT_STRING,		false,	{ 0, 0, 0, 0, 0, 0, 0, 0, 0, 0, 0, 0, 0, 0, 0, 0, 0, 0, 0, 0, 0, 0, 0, 0, 0, 0, 0, 0, 0 }		},	// ST_PARAMETER
+
+		{ DT_IMAGE,			true,	{ 0, 0, 0, 0, 0, 0, 0, 0, 0, 0, 0, 0, 0, 0, 0, 0, 0, 0, 0, 0, 0, 0, 0, 0, 0, 0, 0, 0, 0 }		},	// IM_REFERENCE
+		{ DT_MESH,			true,	{ 0, 0, 0, 0, 0, 0, 0, 0, 0, 0, 0, 0, 0, 0, 0, 0, 0, 0, 0, 0, 0, 0, 0, 0, 0, 0, 0, 0, 0 }		},	// ME_REFERENCE
+
+		{ DT_INT,			false,	{ 0, 0, 0, 0, 0, 0, 0, 0, 0, 0, 0, 0, 0, 0, 0, 0, 0, 0, 0, 0, 0, 0, 0, 0, 0, 0, 0, 0, 0 }		},	// NU_CONDITIONAL
+		{ DT_SCALAR,		false,	{ 0, 0, 0, 0, 0, 0, 0, 0, 0, 0, 0, 0, 0, 0, 0, 0, 0, 0, 0, 0, 0, 0, 0, 0, 0, 0, 0, 0, 0 }		},	// SC_CONDITIONAL
+		{ DT_COLOUR,		false,	{ 0, 0, 0, 0, 0, 0, 0, 0, 0, 0, 0, 0, 0, 0, 0, 0, 0, 0, 0, 0, 0, 0, 0, 0, 0, 0, 0, 0, 0 }		},	// CO_CONDITIONAL
+		{ DT_IMAGE,			true,	{ 1, 1, 1, 1, 1, 1, 1, 1, 1, 1, 1, 1, 1, 1, 1, 1, 1, 1, 1, 1, 1, 1, 0, 0, 1, 1, 0, 0, 0 }		},	// IM_CONDITIONAL
+		{ DT_MESH,			true,	{ 0, 0, 0, 0, 0, 0, 0, 0, 0, 0, 0, 0, 0, 0, 0, 0, 0, 0, 0, 0, 0, 0, 0, 0, 0, 0, 0, 0, 0 }		},	// ME_CONDITIONAL
+		{ DT_LAYOUT,		false,  { 0, 0, 0, 0, 0, 0, 0, 0, 0, 0, 0, 0, 0, 0, 0, 0, 0, 0, 0, 0, 0, 0, 0, 0, 0, 0, 0, 0, 0 }		},	// LA_CONDITIONAL
+		{ DT_INSTANCE,		false,	{ 0, 0, 0, 0, 0, 0, 0, 0, 0, 0, 0, 0, 0, 0, 0, 0, 0, 0, 0, 0, 0, 0, 0, 0, 0, 0, 0, 0, 0 }		},	// IN_CONDITIONAL
+		{ DT_EXTENSION_DATA,false,	{ 0, 0, 0, 0, 0, 0, 0, 0, 0, 0, 0, 0, 0, 0, 0, 0, 0, 0, 0, 0, 0, 0, 0, 0, 0, 0, 0, 0, 0 }		},	// ED_CONDITIONAL
+
+		{ DT_INT,			false,	{ 0, 0, 0, 0, 0, 0, 0, 0, 0, 0, 0, 0, 0, 0, 0, 0, 0, 0, 0, 0, 0, 0, 0, 0, 0, 0, 0, 0, 0 }		},	// NU_SWITCH
+		{ DT_SCALAR,		false,	{ 0, 0, 0, 0, 0, 0, 0, 0, 0, 0, 0, 0, 0, 0, 0, 0, 0, 0, 0, 0, 0, 0, 0, 0, 0, 0, 0, 0, 0 }		},	// SC_SWITCH
+		{ DT_COLOUR,		false,	{ 0, 0, 0, 0, 0, 0, 0, 0, 0, 0, 0, 0, 0, 0, 0, 0, 0, 0, 0, 0, 0, 0, 0, 0, 0, 0, 0, 0, 0 }		},	// CO_SWITCH
+		{ DT_IMAGE,			true,	{ 1, 1, 1, 1, 1, 1, 1, 1, 1, 1, 1, 1, 1, 1, 1, 1, 1, 1, 1, 1, 1, 1, 0, 0, 1, 1, 0, 0, 0 }		},	// IM_SWITCH
+		{ DT_MESH,			true,	{ 0, 0, 0, 0, 0, 0, 0, 0, 0, 0, 0, 0, 0, 0, 0, 0, 0, 0, 0, 0, 0, 0, 0, 0, 0, 0, 0, 0, 0 }		},	// ME_SWITCH
+		{ DT_LAYOUT,		false,	{ 0, 0, 0, 0, 0, 0, 0, 0, 0, 0, 0, 0, 0, 0, 0, 0, 0, 0, 0, 0, 0, 0, 0, 0, 0, 0, 0, 0, 0 }		},	// LA_SWITCH
+		{ DT_INSTANCE,		false,	{ 0, 0, 0, 0, 0, 0, 0, 0, 0, 0, 0, 0, 0, 0, 0, 0, 0, 0, 0, 0, 0, 0, 0, 0, 0, 0, 0, 0, 0 }		},	// IN_SWITCH
+		{ DT_EXTENSION_DATA,false,	{ 0, 0, 0, 0, 0, 0, 0, 0, 0, 0, 0, 0, 0, 0, 0, 0, 0, 0, 0, 0, 0, 0, 0, 0, 0, 0, 0, 0, 0 }		},	// ED_SWITCH
+
+		{ DT_BOOL,			false,  { 0, 0, 0, 0, 0, 0, 0, 0, 0, 0, 0, 0, 0, 0, 0, 0, 0, 0, 0, 0, 0, 0, 0, 0, 0, 0, 0, 0, 0 }		},	// BO_LESS
+		{ DT_BOOL,			false,  { 0, 0, 0, 0, 0, 0, 0, 0, 0, 0, 0, 0, 0, 0, 0, 0, 0, 0, 0, 0, 0, 0, 0, 0, 0, 0, 0, 0, 0 }		},	// BO_EQUAL_SC_CONST
+		{ DT_BOOL,			false,  { 0, 0, 0, 0, 0, 0, 0, 0, 0, 0, 0, 0, 0, 0, 0, 0, 0, 0, 0, 0, 0, 0, 0, 0, 0, 0, 0, 0, 0 }		},	// BO_AND
+		{ DT_BOOL,			false,  { 0, 0, 0, 0, 0, 0, 0, 0, 0, 0, 0, 0, 0, 0, 0, 0, 0, 0, 0, 0, 0, 0, 0, 0, 0, 0, 0, 0, 0 }		},	// BO_OR
+		{ DT_BOOL,			false,  { 0, 0, 0, 0, 0, 0, 0, 0, 0, 0, 0, 0, 0, 0, 0, 0, 0, 0, 0, 0, 0, 0, 0, 0, 0, 0, 0, 0, 0 }		},	// BO_NOT
+
+		{ DT_SCALAR,		true,   { 0, 0, 0, 0, 0, 0, 0, 0, 0, 0, 0, 0, 0, 0, 0, 0, 0, 0, 0, 0, 0, 0, 0, 0, 0, 0, 0, 0, 0 }		},	// SC_MULTIPLYADD
+		{ DT_SCALAR,		true,   { 0, 0, 0, 0, 0, 0, 0, 0, 0, 0, 0, 0, 0, 0, 0, 0, 0, 0, 0, 0, 0, 0, 0, 0, 0, 0, 0, 0, 0 }		},	// SC_ARITHMETIC
+		{ DT_SCALAR,		true,   { 0, 0, 0, 0, 0, 0, 0, 0, 0, 0, 0, 0, 0, 0, 0, 0, 0, 0, 0, 0, 0, 0, 0, 0, 0, 0, 0, 0, 0 }		},	// SC_CURVE
+
+		{ DT_COLOUR,		false,	{ 0, 0, 0, 0, 0, 0, 0, 0, 0, 0, 0, 0, 0, 0, 0, 0, 0, 0, 0, 0, 0, 0, 0, 0, 0, 0, 0, 0, 0 }		},	// CO_SAMPLEIMAGE
+		{ DT_COLOUR,		false,	{ 0, 0, 0, 0, 0, 0, 0, 0, 0, 0, 0, 0, 0, 0, 0, 0, 0, 0, 0, 0, 0, 0, 0, 0, 0, 0, 0, 0, 0 }		},	// CO_SWIZZLE
+		{ DT_COLOUR,		false,	{ 0, 0, 0, 0, 0, 0, 0, 0, 0, 0, 0, 0, 0, 0, 0, 0, 0, 0, 0, 0, 0, 0, 0, 0, 0, 0, 0, 0, 0 }		},	// CO_FROMSCALARS
+		{ DT_COLOUR,		false,	{ 0, 0, 0, 0, 0, 0, 0, 0, 0, 0, 0, 0, 0, 0, 0, 0, 0, 0, 0, 0, 0, 0, 0, 0, 0, 0, 0, 0, 0 }		},	// CO_ARITHMETIC
+
+		{ DT_IMAGE,			true,	{ 0, 1, 1, 1, 0, 0, 0, 0, 0, 0, 0, 0, 0, 0, 0, 0, 0, 0, 0, 0, 0, 0, 0, 0, 0, 0, 0, 0, 0 }		},	// IM_LAYER
+		{ DT_IMAGE,			true,	{ 0, 1, 1, 1, 0, 0, 0, 0, 0, 0, 0, 0, 0, 0, 0, 0, 0, 0, 0, 0, 0, 0, 0, 0, 0, 0, 0, 0, 0 }		},	// IM_LAYERCOLOUR
+		{ DT_IMAGE,			true,	{ 1, 1, 1, 1, 1, 1, 1, 1, 1, 1, 1, 1, 0, 0, 0, 0, 0, 0, 0, 0, 0, 0, 0, 0, 0, 0, 0, 0, 0 }		},	// IM_PIXELFORMAT
+		{ DT_IMAGE,			true,	{ 1, 1, 1, 1, 0, 0, 0, 0, 0, 0, 0, 0, 0, 0, 0, 0, 0, 0, 0, 0, 0, 0, 0, 0, 0, 0, 0, 0, 0 }		},	// IM_MIPMAP
+		{ DT_IMAGE,			true,	{ 0, 1, 1, 1, 0, 0, 0, 0, 0, 0, 0, 0, 0, 0, 0, 0, 0, 0, 0, 0, 0, 0, 0, 0, 0, 0, 0, 0, 0 }		},	// IM_RESIZE
+		{ DT_IMAGE,			true,	{ 0, 1, 1, 1, 0, 0, 0, 0, 0, 0, 0, 0, 0, 0, 0, 0, 0, 0, 0, 0, 0, 0, 0, 0, 0, 0, 0, 0, 0 }		},	// IM_RESIZELIKE
+		{ DT_IMAGE,			true,	{ 0, 1, 1, 1, 0, 0, 0, 0, 0, 0, 0, 0, 0, 0, 0, 0, 0, 0, 0, 0, 0, 0, 0, 0, 0, 0, 0, 0, 0 }		},	// IM_RESIZEREL
+		{ DT_IMAGE,			true,   { 1, 1, 1, 1, 1, 1, 1, 1, 1, 1, 1, 1, 0, 0, 0, 0, 0, 0, 0, 0, 0, 0, 0, 0, 0, 0, 0, 0, 0 }		},	// IM_BLANKLAYOUT
+		{ DT_IMAGE,			true,	{ 1, 1, 1, 1, 1, 1, 1, 1, 1, 1, 1, 0, 0, 0, 1, 1, 1, 1, 1, 1, 1, 1, 0, 0, 1, 1, 0, 0, 0 }		},	// IM_COMPOSE
+		{ DT_IMAGE,			true,	{ 0, 1, 1, 1, 0, 0, 0, 0, 0, 0, 0, 0, 0, 0, 0, 0, 0, 0, 0, 0, 0, 0, 0, 0, 0, 0, 0, 0, 0 }		},	// IM_INTERPOLATE
+		{ DT_IMAGE,			true,	{ 0, 1, 1, 0, 0, 0, 0, 0, 0, 0, 0, 0, 0, 0, 0, 0, 0, 0, 0, 0, 0, 0, 0, 0, 0, 0, 0, 0, 0 }		},	// IM_SATURATE
+		{ DT_IMAGE,			true,	{ 0, 1, 1, 0, 0, 0, 0, 0, 0, 0, 0, 0, 0, 0, 0, 0, 0, 0, 0, 0, 0, 0, 0, 0, 0, 0, 0, 0, 0 }		},	// IM_LUMINANCE
+		{ DT_IMAGE,			true,	{ 0, 1, 1, 1, 0, 0, 0, 0, 0, 0, 0, 0, 0, 0, 0, 0, 0, 0, 0, 0, 0, 0, 0, 0, 0, 0, 0, 0, 0 }		},	// IM_SWIZZLE
+		{ DT_IMAGE,			true,	{ 0, 1, 1, 1, 0, 0, 0, 0, 0, 0, 0, 0, 0, 0, 0, 0, 0, 0, 0, 0, 0, 0, 0, 0, 0, 0, 0, 0, 0 }		},	// IM_COLOURMAP
+		{ DT_IMAGE,			true,	{ 0, 0, 0, 0, 0, 0, 0, 0, 0, 0, 0, 0, 0, 0, 0, 0, 0, 0, 0, 0, 0, 0, 0, 0, 0, 0, 0, 0, 0 }		},	// IM_GRADIENT
+		{ DT_IMAGE,			true,	{ 0, 1, 1, 1, 0, 0, 0, 0, 0, 0, 0, 0, 0, 0, 0, 0, 0, 0, 0, 0, 0, 0, 0, 0, 0, 0, 0, 0, 0 }		},	// IM_BINARISE
+		{ DT_IMAGE,			true,	{ 0, 0, 0, 0, 0, 0, 0, 0, 0, 0, 0, 0, 0, 0, 0, 0, 0, 0, 0, 0, 0, 0, 0, 0, 0, 0, 0, 0, 0 }		},	// IM_PLAINCOLOUR
+		{ DT_IMAGE,			true,	{ 0, 0, 0, 0, 0, 0, 0, 0, 0, 0, 0, 0, 0, 0, 0, 0, 0, 0, 0, 0, 0, 0, 0, 0, 0, 0, 0, 0, 0 }		},	// IM_CROP
+		{ DT_IMAGE,			true,	{ 0, 0, 0, 0, 0, 0, 0, 0, 0, 0, 0, 0, 0, 0, 0, 0, 0, 0, 0, 0, 0, 0, 0, 0, 0, 0, 0, 0, 0 }		},	// IM_PATCH
+		{ DT_IMAGE,			true,	{ 0, 0, 0, 0, 0, 0, 0, 0, 0, 0, 0, 0, 0, 0, 0, 0, 0, 0, 0, 0, 0, 0, 0, 0, 0, 0, 0, 0, 0 }		},	// IM_RASTERMESH
+		{ DT_IMAGE,			true,	{ 0, 0, 0, 0, 0, 0, 0, 0, 0, 0, 0, 0, 0, 0, 0, 0, 0, 0, 0, 0, 0, 0, 0, 0, 0, 0, 0, 0, 0 }		},	// IM_MAKEGROWMAP
+		{ DT_IMAGE,			true,	{ 0, 0, 0, 0, 0, 0, 0, 0, 0, 0, 0, 0, 0, 0, 0, 0, 0, 0, 0, 0, 0, 0, 0, 0, 0, 0, 0, 0, 0 }		},	// IM_DISPLACE
+		{ DT_IMAGE,			true,	{ 0, 0, 0, 0, 0, 0, 0, 0, 0, 0, 0, 0, 0, 0, 0, 0, 0, 0, 0, 0, 0, 0, 0, 0, 0, 0, 0, 0, 0 }		},	// IM_MULTILAYER
+		{ DT_IMAGE,			true,	{ 0, 1, 1, 1, 0, 0, 0, 0, 0, 0, 0, 0, 0, 0, 0, 0, 0, 0, 0, 0, 0, 0, 0, 0, 0, 0, 0, 0, 0 }		},	// IM_INVERT
+		{ DT_IMAGE,			true,	{ 0, 1, 1, 1, 0, 0, 0, 0, 0, 0, 0, 0, 0, 0, 0, 0, 0, 0, 0, 0, 0, 0, 0, 0, 0, 0, 0, 0, 0 }		},	// IM_NORMALCOMPOSITE
+		{ DT_IMAGE,			true,	{ 0, 1, 1, 1, 0, 0, 0, 0, 0, 0, 0, 0, 0, 0, 0, 0, 0, 0, 0, 0, 0, 0, 0, 0, 0, 0, 0, 0, 0 }		},	// IM_TRANSFORM
+
+		{ DT_MESH,			true,	{ 0, 0, 0, 0, 0, 0, 0, 0, 0, 0, 0, 0, 0, 0, 0, 0, 0, 0, 0, 0, 0, 0, 0, 0, 0, 0, 0, 0, 0 }		},	// ME_APPLYLAYOUT
+		{ DT_MESH,			true,	{ 0, 0, 0, 0, 0, 0, 0, 0, 0, 0, 0, 0, 0, 0, 0, 0, 0, 0, 0, 0, 0, 0, 0, 0, 0, 0, 0, 0, 0 }		},	// ME_DIFFERENCE
+		{ DT_MESH,			true,	{ 0, 0, 0, 0, 0, 0, 0, 0, 0, 0, 0, 0, 0, 0, 0, 0, 0, 0, 0, 0, 0, 0, 0, 0, 0, 0, 0, 0, 0 }		},	// ME_MORPH
+		{ DT_MESH,			true,	{ 0, 0, 0, 0, 0, 0, 0, 0, 0, 0, 0, 0, 0, 0, 0, 0, 0, 0, 0, 0, 0, 0, 0, 0, 0, 0, 0, 0, 0 }		},	// ME_MERGE
+		{ DT_MESH,			true,	{ 0, 0, 0, 0, 0, 0, 0, 0, 0, 0, 0, 0, 0, 0, 0, 0, 0, 0, 0, 0, 0, 0, 0, 0, 0, 0, 0, 0, 0 }		},	// ME_INTERPOLATE
+		{ DT_MESH,			true,	{ 0, 0, 0, 0, 0, 0, 0, 0, 0, 0, 0, 0, 0, 0, 0, 0, 0, 0, 0, 0, 0, 0, 0, 0, 0, 0, 0, 0, 0 }		},	// ME_MASKCLIPMESH
+		{ DT_MESH,			true,	{ 0, 0, 0, 0, 0, 0, 0, 0, 0, 0, 0, 0, 0, 0, 0, 0, 0, 0, 0, 0, 0, 0, 0, 0, 0, 0, 0, 0, 0 }		},	// ME_MASKCLIPUVMASK
+		{ DT_MESH,			true,	{ 0, 0, 0, 0, 0, 0, 0, 0, 0, 0, 0, 0, 0, 0, 0, 0, 0, 0, 0, 0, 0, 0, 0, 0, 0, 0, 0, 0, 0 }		},	// ME_MASKDIFF
+		{ DT_MESH,			true,	{ 0, 0, 0, 0, 0, 0, 0, 0, 0, 0, 0, 0, 0, 0, 0, 0, 0, 0, 0, 0, 0, 0, 0, 0, 0, 0, 0, 0, 0 }		},	// ME_REMOVEMASK
+		{ DT_MESH,			true,	{ 0, 0, 0, 0, 0, 0, 0, 0, 0, 0, 0, 0, 0, 0, 0, 0, 0, 0, 0, 0, 0, 0, 0, 0, 0, 0, 0, 0, 0 }		},	// ME_FORMAT
+		{ DT_MESH,			true,	{ 0, 0, 0, 0, 0, 0, 0, 0, 0, 0, 0, 0, 0, 0, 0, 0, 0, 0, 0, 0, 0, 0, 0, 0, 0, 0, 0, 0, 0 }		},	// ME_EXTRACTLAYOUTBLOCK
+		{ DT_MESH,			true,	{ 0, 0, 0, 0, 0, 0, 0, 0, 0, 0, 0, 0, 0, 0, 0, 0, 0, 0, 0, 0, 0, 0, 0, 0, 0, 0, 0, 0, 0 }		},	// ME_TRANSFORM
+		{ DT_MESH,			true,	{ 0, 0, 0, 0, 0, 0, 0, 0, 0, 0, 0, 0, 0, 0, 0, 0, 0, 0, 0, 0, 0, 0, 0, 0, 0, 0, 0, 0, 0 }		},	// ME_CLIPMORPHPLANE
+		{ DT_MESH,			true,	{ 0, 0, 0, 0, 0, 0, 0, 0, 0, 0, 0, 0, 0, 0, 0, 0, 0, 0, 0, 0, 0, 0, 0, 0, 0, 0, 0, 0, 0 }		},	// ME_CLIPWITHMESH
+		{ DT_MESH,			true,	{ 0, 0, 0, 0, 0, 0, 0, 0, 0, 0, 0, 0, 0, 0, 0, 0, 0, 0, 0, 0, 0, 0, 0, 0, 0, 0, 0, 0, 0 }		},	// ME_SETSKELETON
+		{ DT_MESH,			true,	{ 0, 0, 0, 0, 0, 0, 0, 0, 0, 0, 0, 0, 0, 0, 0, 0, 0, 0, 0, 0, 0, 0, 0, 0, 0, 0, 0, 0, 0 }		},	// ME_PROJECT
+		{ DT_MESH,			true,	{ 0, 0, 0, 0, 0, 0, 0, 0, 0, 0, 0, 0, 0, 0, 0, 0, 0, 0, 0, 0, 0, 0, 0, 0, 0, 0, 0, 0, 0 }		},	// ME_APPLYPOSE
+		{ DT_MESH,			true,	{ 0, 0, 0, 0, 0, 0, 0, 0, 0, 0, 0, 0, 0, 0, 0, 0, 0, 0, 0, 0, 0, 0, 0, 0, 0, 0, 0, 0, 0 }		},	// ME_GEOMETRYOPERATION
+		{ DT_MESH,			true,	{ 0, 0, 0, 0, 0, 0, 0, 0, 0, 0, 0, 0, 0, 0, 0, 0, 0, 0, 0, 0, 0, 0, 0, 0, 0, 0, 0, 0, 0 }		},	// ME_BINDSHAPE
+		{ DT_MESH,			true,	{ 0, 0, 0, 0, 0, 0, 0, 0, 0, 0, 0, 0, 0, 0, 0, 0, 0, 0, 0, 0, 0, 0, 0, 0, 0, 0, 0, 0, 0 }		},	// ME_APPLYSHAPE
+		{ DT_MESH,			true,	{ 0, 0, 0, 0, 0, 0, 0, 0, 0, 0, 0, 0, 0, 0, 0, 0, 0, 0, 0, 0, 0, 0, 0, 0, 0, 0, 0, 0, 0 }		},	// ME_CLIPDEFORM
+		{ DT_MESH,			true,	{ 0, 0, 0, 0, 0, 0, 0, 0, 0, 0, 0, 0, 0, 0, 0, 0, 0, 0, 0, 0, 0, 0, 0, 0, 0, 0, 0, 0, 0 }		},	// ME_MORPHRESHAPE
+		{ DT_MESH,			true,	{ 0, 0, 0, 0, 0, 0, 0, 0, 0, 0, 0, 0, 0, 0, 0, 0, 0, 0, 0, 0, 0, 0, 0, 0, 0, 0, 0, 0, 0 }		},	// ME_OPTIMIZESKINNING
+		{ DT_MESH,			true,	{ 0, 0, 0, 0, 0, 0, 0, 0, 0, 0, 0, 0, 0, 0, 0, 0, 0, 0, 0, 0, 0, 0, 0, 0, 0, 0, 0, 0, 0 }		},	// ME_ADDTAGS
+
+		{ DT_INSTANCE,		false,	{ 0, 0, 0, 0, 0, 0, 0, 0, 0, 0, 0, 0, 0, 0, 0, 0, 0, 0, 0, 0, 0, 0, 0, 0, 0, 0, 0, 0, 0 }		},	// IN_ADDMESH
+		{ DT_INSTANCE,		false,	{ 0, 0, 0, 0, 0, 0, 0, 0, 0, 0, 0, 0, 0, 0, 0, 0, 0, 0, 0, 0, 0, 0, 0, 0, 0, 0, 0, 0, 0 }		},	// IN_ADDIMAGE
+		{ DT_INSTANCE,		false,	{ 0, 0, 0, 0, 0, 0, 0, 0, 0, 0, 0, 0, 0, 0, 0, 0, 0, 0, 0, 0, 0, 0, 0, 0, 0, 0, 0, 0, 0 }		},	// IN_ADDVECTOR
+		{ DT_INSTANCE,		false,	{ 0, 0, 0, 0, 0, 0, 0, 0, 0, 0, 0, 0, 0, 0, 0, 0, 0, 0, 0, 0, 0, 0, 0, 0, 0, 0, 0, 0, 0 }		},	// IN_ADDSCALAR
+		{ DT_INSTANCE,		false,	{ 0, 0, 0, 0, 0, 0, 0, 0, 0, 0, 0, 0, 0, 0, 0, 0, 0, 0, 0, 0, 0, 0, 0, 0, 0, 0, 0, 0, 0 }		},	// IN_ADDSTRING
+		{ DT_INSTANCE,		false,	{ 0, 0, 0, 0, 0, 0, 0, 0, 0, 0, 0, 0, 0, 0, 0, 0, 0, 0, 0, 0, 0, 0, 0, 0, 0, 0, 0, 0, 0 }		},	// IN_ADDSURFACE
+		{ DT_INSTANCE,		false,	{ 0, 0, 0, 0, 0, 0, 0, 0, 0, 0, 0, 0, 0, 0, 0, 0, 0, 0, 0, 0, 0, 0, 0, 0, 0, 0, 0, 0, 0 }		},	// IN_ADDCOMPONENT
+		{ DT_INSTANCE,		false,	{ 0, 0, 0, 0, 0, 0, 0, 0, 0, 0, 0, 0, 0, 0, 0, 0, 0, 0, 0, 0, 0, 0, 0, 0, 0, 0, 0, 0, 0 }		},	// IN_ADDLOD
+		{ DT_INSTANCE,		false,	{ 0, 0, 0, 0, 0, 0, 0, 0, 0, 0, 0, 0, 0, 0, 0, 0, 0, 0, 0, 0, 0, 0, 0, 0, 0, 0, 0, 0, 0 }		},	// IN_ADDEXTENSIONDATA
+
+		{ DT_LAYOUT,		true,   { 0, 0, 0, 0, 0, 0, 0, 0, 0, 0, 0, 0, 0, 0, 0, 0, 0, 0, 0, 0, 0, 0, 0, 0, 0, 0, 0, 0, 0 }		},	// LA_PACK
+		{ DT_LAYOUT,		true,   { 0, 0, 0, 0, 0, 0, 0, 0, 0, 0, 0, 0, 0, 0, 0, 0, 0, 0, 0, 0, 0, 0, 0, 0, 0, 0, 0, 0, 0 }		},	// LA_MERGE
+		{ DT_LAYOUT,		true,   { 0, 0, 0, 0, 0, 0, 0, 0, 0, 0, 0, 0, 0, 0, 0, 0, 0, 0, 0, 0, 0, 0, 0, 0, 0, 0, 0, 0, 0 }		},	// LA_REMOVEBLOCKS
+		{ DT_LAYOUT,		true,   { 0, 0, 0, 0, 0, 0, 0, 0, 0, 0, 0, 0, 0, 0, 0, 0, 0, 0, 0, 0, 0, 0, 0, 0, 0, 0, 0, 0, 0 }		},	// LA_FROMMESH
+	};
+
+    // clang-format on
+
+    static_assert( sizeof(s_opDescs)/sizeof(FOpDesc) == (int32)OP_TYPE::COUNT, "OperationDescMismatch");
+
+	//---------------------------------------------------------------------------------------------
+	const FOpDesc& GetOpDesc( OP_TYPE type )
+	{
+        return s_opDescs[ (int32)type ];
+	}
+
+
+    //---------------------------------------------------------------------------------------------
+    void ForEachReference( OP& op, const TFunctionRef<void(OP::ADDRESS*)> f )
+    {
+		// Only operations that still use ASTOpFixed should be handled here.
+        switch ( op.type )
+        {
+        case OP_TYPE::NONE:
+        case OP_TYPE::BO_CONSTANT:
+        case OP_TYPE::NU_CONSTANT:
+        case OP_TYPE::SC_CONSTANT:
+        case OP_TYPE::CO_CONSTANT:
+        case OP_TYPE::IM_CONSTANT:
+        case OP_TYPE::ME_CONSTANT:
+        case OP_TYPE::LA_CONSTANT:
+        case OP_TYPE::PR_CONSTANT:
+        case OP_TYPE::BO_PARAMETER:
+        case OP_TYPE::NU_PARAMETER:
+        case OP_TYPE::SC_PARAMETER:
+        case OP_TYPE::CO_PARAMETER:
+        case OP_TYPE::PR_PARAMETER:
+        case OP_TYPE::IM_PARAMETER:
+            break;
+
+        //-------------------------------------------------------------------------------------
+        case OP_TYPE::BO_LESS:
+            f(&op.args.BoolLess.a );
+            f(&op.args.BoolLess.b );
+            break;
+
+        //-------------------------------------------------------------------------------------
+        case OP_TYPE::BO_EQUAL_INT_CONST:
+            f(&op.args.BoolEqualScalarConst.value );
+            break;
+
+        //-------------------------------------------------------------------------------------
+        case OP_TYPE::BO_AND:
+        case OP_TYPE::BO_OR:
+            f(&op.args.BoolBinary.a );
+            f(&op.args.BoolBinary.b );
+            break;
+
+        //-------------------------------------------------------------------------------------
+        case OP_TYPE::BO_NOT:
+            f(&op.args.BoolNot.source );
+            break;
+
+        //-------------------------------------------------------------------------------------
+        case OP_TYPE::SC_MULTIPLYADD:
+            f(&op.args.ScalarMultiplyAdd.factor0 );
+            f(&op.args.ScalarMultiplyAdd.factor1 );
+            f(&op.args.ScalarMultiplyAdd.add );
+            break;
+
+        //-------------------------------------------------------------------------------------
+        case OP_TYPE::SC_ARITHMETIC:
+            f(&op.args.ScalarArithmetic.a );
+            f(&op.args.ScalarArithmetic.b );
+            break;
+
+        //-------------------------------------------------------------------------------------
+        case OP_TYPE::CO_SAMPLEIMAGE:
+            f(&op.args.ColourSampleImage.image );
+            f(&op.args.ColourSampleImage.x );
+            f(&op.args.ColourSampleImage.y );
+            break;
+
+        case OP_TYPE::CO_SWIZZLE:
+            for (int32 t=0;t<MUTABLE_OP_MAX_SWIZZLE_CHANNELS;++t)
+            {
+                f(&op.args.ColourSwizzle.sources[t] );
+            }
+            break;
+
+        case OP_TYPE::CO_FROMSCALARS:
+			for (int32 t = 0; t < MUTABLE_OP_MAX_SWIZZLE_CHANNELS; ++t)
+			{
+				f(&op.args.ColourFromScalars.v[t]);
+			}
+            break;
+
+        case OP_TYPE::CO_ARITHMETIC:
+            f(&op.args.ColourArithmetic.a);
+            f(&op.args.ColourArithmetic.b);
+            break;
+
+        //-------------------------------------------------------------------------------------
+        case OP_TYPE::IM_RESIZE:
+            f(&op.args.ImageResize.source );
+            break;
+
+        case OP_TYPE::IM_RESIZELIKE:
+            f(&op.args.ImageResizeLike.source );
+            f(&op.args.ImageResizeLike.sizeSource );
+            break;
+
+        case OP_TYPE::IM_RESIZEREL:
+            f(&op.args.ImageResizeRel.source );
+            break;
+
+        case OP_TYPE::IM_BLANKLAYOUT:
+            f(&op.args.ImageBlankLayout.layout );
+            break;
+
+        case OP_TYPE::IM_INTERPOLATE:
+            f(&op.args.ImageInterpolate.factor );
+
+            for (int t=0;t<MUTABLE_OP_MAX_INTERPOLATE_COUNT;++t)
+            {
+                f(&op.args.ImageInterpolate.targets[t] );
+            }
+            break;
+
+        case OP_TYPE::IM_SATURATE:
+            f(&op.args.ImageSaturate.base );
+            f(&op.args.ImageSaturate.factor );
+            break;
+
+        case OP_TYPE::IM_LUMINANCE:
+            f(&op.args.ImageLuminance.base );
+            break;
+
+        case OP_TYPE::IM_COLOURMAP:
+            f(&op.args.ImageColourMap.base );
+            f(&op.args.ImageColourMap.mask );
+            f(&op.args.ImageColourMap.map );
+            break;
+
+        case OP_TYPE::IM_GRADIENT:
+            f(&op.args.ImageGradient.colour0 );
+            f(&op.args.ImageGradient.colour1 );
+            break;
+
+        case OP_TYPE::IM_BINARISE:
+            f(&op.args.ImageBinarise.base );
+            f(&op.args.ImageBinarise.threshold );
+            break;
+
+        case OP_TYPE::IM_PLAINCOLOUR:
+            f(&op.args.ImagePlainColour.colour );
+            break;
+
+        case OP_TYPE::IM_DISPLACE:
+            f(&op.args.ImageDisplace.source );
+            f(&op.args.ImageDisplace.displacementMap );
+            break;
+
+		case OP_TYPE::IM_INVERT:
+			f(&op.args.ImageInvert.base);
+			break;
+
+        //-------------------------------------------------------------------------------------
+        case OP_TYPE::ME_APPLYLAYOUT:
+            f(&op.args.MeshApplyLayout.layout );
+            f(&op.args.MeshApplyLayout.mesh );
+            break;
+
+        case OP_TYPE::ME_MERGE:
+            f(&op.args.MeshMerge.base );
+            f(&op.args.MeshMerge.added );
+            break;
+
+        case OP_TYPE::ME_INTERPOLATE:
+            f(&op.args.MeshInterpolate.factor );
+            f(&op.args.MeshInterpolate.base );
+
+            for (int t=0;t<MUTABLE_OP_MAX_INTERPOLATE_COUNT-1;++t)
+            {
+                f(&op.args.MeshInterpolate.targets[t] );
+            }
+            break;
+
+        case OP_TYPE::ME_MASKDIFF:
+            f(&op.args.MeshMaskDiff.source );
+            f(&op.args.MeshMaskDiff.fragment );
+            break;
+
+        case OP_TYPE::ME_CLIPMORPHPLANE:
+            f(&op.args.MeshClipMorphPlane.source);
+            break;
+
+        case OP_TYPE::ME_CLIPWITHMESH :
+            f(&op.args.MeshClipWithMesh.source);
+            f(&op.args.MeshClipWithMesh.clipMesh);
+            break;
+
+        case OP_TYPE::ME_SETSKELETON :
+            f(&op.args.MeshSetSkeleton.source);
+            f(&op.args.MeshSetSkeleton.skeleton);
+            break;
+
+		case OP_TYPE::ME_PROJECT:
+			f(&op.args.MeshProject.mesh);
+			f(&op.args.MeshProject.projector);
+			break;
+
+        default:
+			check( false );
+            break;
+        }
+
+    }
+
+
+
+    //---------------------------------------------------------------------------------------------
+    void ForEachReference( const FProgram& program, OP::ADDRESS at, const TFunctionRef<void(OP::ADDRESS)> f )
+    {
+        OP_TYPE type = program.GetOpType(at);
+        switch ( type )
+        {
+        case OP_TYPE::NONE:
+        case OP_TYPE::BO_CONSTANT:
+        case OP_TYPE::NU_CONSTANT:
+        case OP_TYPE::SC_CONSTANT:
+        case OP_TYPE::ST_CONSTANT:
+        case OP_TYPE::CO_CONSTANT:
+        case OP_TYPE::IM_CONSTANT:
+        case OP_TYPE::ME_CONSTANT:
+        case OP_TYPE::LA_CONSTANT:
+        case OP_TYPE::PR_CONSTANT:
+		case OP_TYPE::ED_CONSTANT:
+        case OP_TYPE::BO_PARAMETER:
+        case OP_TYPE::NU_PARAMETER:
+        case OP_TYPE::SC_PARAMETER:
+        case OP_TYPE::CO_PARAMETER:
+        case OP_TYPE::PR_PARAMETER:
+        case OP_TYPE::IM_PARAMETER:
+		case OP_TYPE::IM_REFERENCE:
+		case OP_TYPE::ME_REFERENCE:
+			break;
+
+        case OP_TYPE::SC_CURVE:
+        {
+			OP::ScalarCurveArgs args = program.GetOpArgs<OP::ScalarCurveArgs>(at);
+            f(args.time );
+            break;
+        }
+
+        case OP_TYPE::NU_CONDITIONAL:
+        case OP_TYPE::SC_CONDITIONAL:
+        case OP_TYPE::CO_CONDITIONAL:
+        case OP_TYPE::IM_CONDITIONAL:
+        case OP_TYPE::ME_CONDITIONAL:
+        case OP_TYPE::LA_CONDITIONAL:
+        case OP_TYPE::IN_CONDITIONAL:
+		case OP_TYPE::ED_CONDITIONAL:
+        {
+			OP::ConditionalArgs args = program.GetOpArgs<OP::ConditionalArgs>(at);
+            f(args.condition );
+            f(args.yes );
+            f(args.no );
+            break;
+        }
+
+        case OP_TYPE::NU_SWITCH:
+        case OP_TYPE::SC_SWITCH:
+        case OP_TYPE::CO_SWITCH:
+        case OP_TYPE::IM_SWITCH:
+        case OP_TYPE::ME_SWITCH:
+        case OP_TYPE::LA_SWITCH:
+        case OP_TYPE::IN_SWITCH:
+		case OP_TYPE::ED_SWITCH:
+        {
+			const uint8_t* data = program.GetOpArgsPointer(at);
+			
+			OP::ADDRESS VarAddress;
+			FMemory::Memcpy( &VarAddress, data, sizeof(OP::ADDRESS));
+			data += sizeof(OP::ADDRESS);
+
+			OP::ADDRESS DefAddress;
+			FMemory::Memcpy( &DefAddress, data, sizeof(OP::ADDRESS));
+			data += sizeof(OP::ADDRESS);
+
+			uint32 CaseCount;
+			FMemory::Memcpy( &CaseCount, data, sizeof(uint32));
+			data += sizeof(uint32);
+
+			f(VarAddress);
+			f(DefAddress);
+
+			for ( uint32 C = 0; C < CaseCount; ++C )
+			{	
+				//int32 Condition;
+				//FMemory::Memcpy( &Condition, data, sizeof(int32_t));
+				data += sizeof(int32);
+				
+				OP::ADDRESS At;
+				FMemory::Memcpy( &At, data, sizeof(OP::ADDRESS) );
+				data += sizeof(OP::ADDRESS);
+
+				if (At)
+				{
+					f(At);
+				}
+			}
+
+            break;
+        }
+
+        //-------------------------------------------------------------------------------------
+        case OP_TYPE::BO_LESS:
+        {
+			OP::BoolLessArgs args = program.GetOpArgs<OP::BoolLessArgs>(at);
+            f(args.a );
+            f(args.b );
+            break;
+        }
+
+        //-------------------------------------------------------------------------------------
+        case OP_TYPE::BO_EQUAL_INT_CONST:
+        {
+			OP::BoolEqualScalarConstArgs args = program.GetOpArgs<OP::BoolEqualScalarConstArgs>(at);
+            f(args.value );
+            break;
+        }
+
+        //-------------------------------------------------------------------------------------
+        case OP_TYPE::BO_AND:
+        case OP_TYPE::BO_OR:
+        {
+			OP::BoolBinaryArgs args = program.GetOpArgs<OP::BoolBinaryArgs>(at);
+            f(args.a );
+            f(args.b );
+            break;
+        }
+
+        //-------------------------------------------------------------------------------------
+        case OP_TYPE::BO_NOT:
+        {
+			OP::BoolNotArgs args = program.GetOpArgs<OP::BoolNotArgs>(at);
+            f(args.source );
+            break;
+        }
+
+        //-------------------------------------------------------------------------------------
+        case OP_TYPE::SC_MULTIPLYADD:
+        {
+			OP::ScalarMultiplyAddArgs args = program.GetOpArgs<OP::ScalarMultiplyAddArgs>(at);
+            f(args.factor0 );
+            f(args.factor1 );
+            f(args.add );
+            break;
+        }
+
+        //-------------------------------------------------------------------------------------
+        case OP_TYPE::SC_ARITHMETIC:
+        {
+			OP::ArithmeticArgs args = program.GetOpArgs<OP::ArithmeticArgs>(at);
+            f(args.a );
+            f(args.b );
+            break;
+        }
+
+        //-------------------------------------------------------------------------------------
+        case OP_TYPE::CO_SAMPLEIMAGE:
+        {
+			OP::ColourSampleImageArgs args = program.GetOpArgs<OP::ColourSampleImageArgs>(at);
+            f(args.image );
+            f(args.x );
+            f(args.y );
+            break;
+        }
+
+        case OP_TYPE::CO_SWIZZLE:
+        {
+			OP::ColourSwizzleArgs args = program.GetOpArgs<OP::ColourSwizzleArgs>(at);
+            for (int t=0;t<MUTABLE_OP_MAX_SWIZZLE_CHANNELS;++t)
+            {
+                f(args.sources[t] );
+            }
+            break;
+        }
+
+        case OP_TYPE::CO_FROMSCALARS:
+        {
+			OP::ColourFromScalarsArgs args = program.GetOpArgs<OP::ColourFromScalarsArgs>(at);
+			for (int32 t = 0; t < MUTABLE_OP_MAX_SWIZZLE_CHANNELS; ++t)
+			{
+				f(args.v[t]);
+			}
+            break;
+        }
+
+        case OP_TYPE::CO_ARITHMETIC:
+        {
+			OP::ArithmeticArgs args = program.GetOpArgs<OP::ArithmeticArgs>(at);
+            f(args.a);
+            f(args.b);
+            break;
+        }
+
+        //-------------------------------------------------------------------------------------
+        case OP_TYPE::IM_LAYER:
+        {
+			OP::ImageLayerArgs args = program.GetOpArgs<OP::ImageLayerArgs>(at);
+            f(args.base );
+            if ( args.mask )
+            {
+                f(args.mask );
+            }
+            f(args.blended );
+            break;
+        }
+
+        case OP_TYPE::IM_LAYERCOLOUR:
+        {
+			OP::ImageLayerColourArgs args = program.GetOpArgs<OP::ImageLayerColourArgs>(at);
+            f(args.base );
+            if ( args.mask )
+            {
+                f(args.mask );
+            }
+            f(args.colour );
+            break;
+        }
+
+        case OP_TYPE::IM_MULTILAYER:
+        {
+			OP::ImageMultiLayerArgs args = program.GetOpArgs<OP::ImageMultiLayerArgs>(at);
+            f(args.rangeSize );
+            f(args.base );
+            if ( args.mask )
+            {
+                f(args.mask );
+            }
+            f(args.blended );
+            break;
+        }
+
+		case OP_TYPE::IM_NORMALCOMPOSITE:
+		{
+			OP::ImageNormalCompositeArgs args = program.GetOpArgs<OP::ImageNormalCompositeArgs>(at);
+			f(args.base);
+			f(args.normal);
+
+			break;
+		}
+
+        case OP_TYPE::IM_PIXELFORMAT:
+        {
+			OP::ImagePixelFormatArgs args = program.GetOpArgs<OP::ImagePixelFormatArgs>(at);
+            f(args.source );
+            break;
+        }
+
+        case OP_TYPE::IM_MIPMAP:
+        {
+			OP::ImageMipmapArgs args = program.GetOpArgs<OP::ImageMipmapArgs>(at);
+            f(args.source );
+            break;
+        }
+
+        case OP_TYPE::IM_RESIZE:
+        {
+			OP::ImageResizeArgs args = program.GetOpArgs<OP::ImageResizeArgs>(at);
+            f(args.source );
+            break;
+        }
+
+        case OP_TYPE::IM_RESIZELIKE:
+        {
+			OP::ImageResizeLikeArgs args = program.GetOpArgs<OP::ImageResizeLikeArgs>(at);
+            f(args.source );
+            f(args.sizeSource );
+            break;
+        }
+
+        case OP_TYPE::IM_RESIZEREL:
+        {
+			OP::ImageResizeRelArgs args = program.GetOpArgs<OP::ImageResizeRelArgs>(at);
+            f(args.source );
+            break;
+        }
+
+        case OP_TYPE::IM_BLANKLAYOUT:
+        {
+			OP::ImageBlankLayoutArgs args = program.GetOpArgs<OP::ImageBlankLayoutArgs>(at);
+            f(args.layout );
+            break;
+        }
+
+        case OP_TYPE::IM_COMPOSE:
+        {
+			OP::ImageComposeArgs args = program.GetOpArgs<OP::ImageComposeArgs>(at);
+            f(args.layout );
+            f(args.base );
+            f(args.blockImage );
+            f(args.mask );
+            break;
+        }
+
+        case OP_TYPE::IM_INTERPOLATE:
+        {
+			OP::ImageInterpolateArgs args = program.GetOpArgs<OP::ImageInterpolateArgs>(at);
+            f(args.factor );
+
+            for (int t=0;t<MUTABLE_OP_MAX_INTERPOLATE_COUNT;++t)
+            {
+                f(args.targets[t] );
+            }
+            break;
+        }
+
+        case OP_TYPE::IM_SWIZZLE:
+        {
+			OP::ImageSwizzleArgs args = program.GetOpArgs<OP::ImageSwizzleArgs>(at);
+            for (int t=0;t<MUTABLE_OP_MAX_SWIZZLE_CHANNELS;++t)
+            {
+                f(args.sources[t] );
+            }
+            break;
+        }
+
+        case OP_TYPE::IM_SATURATE:
+        {
+			OP::ImageSaturateArgs args = program.GetOpArgs<OP::ImageSaturateArgs>(at);
+            f(args.base );
+            f(args.factor );
+            break;
+        }
+
+        case OP_TYPE::IM_LUMINANCE:
+        {
+			OP::ImageLuminanceArgs args = program.GetOpArgs<OP::ImageLuminanceArgs>(at);
+            f(args.base );
+            break;
+        }
+
+        case OP_TYPE::IM_COLOURMAP:
+        {
+			OP::ImageColourMapArgs args = program.GetOpArgs<OP::ImageColourMapArgs>(at);
+            f(args.base );
+            f(args.mask );
+            f(args.map );
+            break;
+        }
+
+        case OP_TYPE::IM_GRADIENT:
+        {
+			OP::ImageGradientArgs args = program.GetOpArgs<OP::ImageGradientArgs>(at);
+            f(args.colour0 );
+            f(args.colour1 );
+            break;
+        }
+
+        case OP_TYPE::IM_BINARISE:
+        {
+			OP::ImageBinariseArgs args = program.GetOpArgs<OP::ImageBinariseArgs>(at);
+            f(args.base );
+            f(args.threshold );
+            break;
+        }
+
+        case OP_TYPE::IM_PLAINCOLOUR:
+        {
+			OP::ImagePlainColourArgs args = program.GetOpArgs<OP::ImagePlainColourArgs>(at);
+            f(args.colour );
+            break;
+        }
+
+        case OP_TYPE::IM_CROP:
+        {
+			OP::ImageCropArgs args = program.GetOpArgs<OP::ImageCropArgs>(at);
+            f(args.source );
+            break;
+        }
+
+        case OP_TYPE::IM_PATCH:
+        {
+			OP::ImagePatchArgs args = program.GetOpArgs<OP::ImagePatchArgs>(at);
+            f(args.base );
+            f(args.patch );
+            break;
+        }
+
+        case OP_TYPE::IM_RASTERMESH:
+        {
+			OP::ImageRasterMeshArgs args = program.GetOpArgs<OP::ImageRasterMeshArgs>(at);
+            f(args.mesh );
+            f(args.image );
+            f(args.mask );
+            f(args.angleFadeProperties );
+            f(args.projector );
+            break;
+        }
+
+        case OP_TYPE::IM_MAKEGROWMAP:
+        {
+			OP::ImageMakeGrowMapArgs args = program.GetOpArgs<OP::ImageMakeGrowMapArgs>(at);
+            f(args.mask );
+            break;
+        }
+
+        case OP_TYPE::IM_DISPLACE:
+        {
+			OP::ImageDisplaceArgs args = program.GetOpArgs<OP::ImageDisplaceArgs>(at);
+            f(args.source );
+            f(args.displacementMap );
+            break;
+        }
+
+		case OP_TYPE::IM_INVERT:
+		{
+			OP::ImageInvertArgs args = program.GetOpArgs<OP::ImageInvertArgs>(at);
+			f(args.base);
+			break;
+		}
+
+		case OP_TYPE::IM_TRANSFORM:
+		{
+			OP::ImageTransformArgs Args = program.GetOpArgs<OP::ImageTransformArgs>(at);
+			f(Args.Base);
+			f(Args.OffsetX);
+			f(Args.OffsetY);
+			f(Args.ScaleX);
+			f(Args.ScaleY);
+			f(Args.Rotation);
+			break;
+		}
+
+        //-------------------------------------------------------------------------------------
+        case OP_TYPE::ME_APPLYLAYOUT:
+        {
+			OP::MeshApplyLayoutArgs args = program.GetOpArgs<OP::MeshApplyLayoutArgs>(at);
+            f(args.layout );
+            f(args.mesh );
+            break;
+        }
+
+        case OP_TYPE::ME_DIFFERENCE:
+        {
+			const uint8_t* data = program.GetOpArgsPointer(at);
+
+			OP::ADDRESS BaseAt = 0;
+			FMemory::Memcpy(&BaseAt, data, sizeof(OP::ADDRESS)); data += sizeof(OP::ADDRESS);
+			f(BaseAt);
+
+			OP::ADDRESS TargetAt = 0;
+			FMemory::Memcpy(&TargetAt, data, sizeof(OP::ADDRESS)); data += sizeof(OP::ADDRESS);
+			f(TargetAt);
+			break;
+        }
+
+        case OP_TYPE::ME_MORPH:
+        {
+			const uint8_t* data = program.GetOpArgsPointer(at);
+
+			OP::ADDRESS FactorAt = 0;
+			FMemory::Memcpy(&FactorAt, data, sizeof(OP::ADDRESS)); data += sizeof(OP::ADDRESS);
+			f(FactorAt);
+
+			OP::ADDRESS BaseAt = 0;
+			FMemory::Memcpy(&BaseAt, data, sizeof(OP::ADDRESS)); data += sizeof(OP::ADDRESS);
+			f(BaseAt);
+
+			OP::ADDRESS TargetAt = 0;
+			FMemory::Memcpy(&TargetAt, data, sizeof(OP::ADDRESS)); data += sizeof(OP::ADDRESS);
+			f(TargetAt);
+			break;
+        }
+
+        case OP_TYPE::ME_MERGE:
+        {
+			OP::MeshMergeArgs args = program.GetOpArgs<OP::MeshMergeArgs>(at);
+            f(args.base );
+            f(args.added );
+            break;
+        }
+
+        case OP_TYPE::ME_INTERPOLATE:
+        {
+			OP::MeshInterpolateArgs args = program.GetOpArgs<OP::MeshInterpolateArgs>(at);
+            f(args.factor );
+            f(args.base );
+
+            for (int t=0;t<MUTABLE_OP_MAX_INTERPOLATE_COUNT-1;++t)
+            {
+                f(args.targets[t] );
+            }
+            break;
+        }
+
+		case OP_TYPE::ME_MASKCLIPMESH:
+		{
+			OP::MeshMaskClipMeshArgs args = program.GetOpArgs<OP::MeshMaskClipMeshArgs>(at);
+			f(args.source);
+			f(args.clip);
+			break;
+		}
+
+		case OP_TYPE::ME_MASKCLIPUVMASK:
+		{
+			OP::MeshMaskClipUVMaskArgs args = program.GetOpArgs<OP::MeshMaskClipUVMaskArgs>(at);
+			f(args.Source);
+			f(args.MaskImage);
+			f(args.MaskLayout);
+			break;
+		}
+
+        case OP_TYPE::ME_MASKDIFF:
+        {
+			OP::MeshMaskDiffArgs args = program.GetOpArgs<OP::MeshMaskDiffArgs>(at);
+            f(args.source );
+            f(args.fragment );
+            break;
+        }
+
+        case OP_TYPE::ME_REMOVEMASK:
+        {
+            const uint8_t* data = program.GetOpArgsPointer(at);
+            mu::OP::ADDRESS source;
+            FMemory::Memcpy( &source, data, sizeof(OP::ADDRESS) ); data+=sizeof(OP::ADDRESS);
+            f(source);
+
+            uint16 removes = 0;
+			FMemory::Memcpy( &removes, data, sizeof(uint16) ); data+=sizeof(uint16);
+            for (uint16 r=0; r<removes; ++r)
+            {
+                mu::OP::ADDRESS condition;
+				FMemory::Memcpy( &condition, data, sizeof(OP::ADDRESS) ); data+=sizeof(OP::ADDRESS);
+                f(condition);
+
+                mu::OP::ADDRESS mask;
+				FMemory::Memcpy( &mask, data, sizeof(OP::ADDRESS) ); data+=sizeof(OP::ADDRESS);
+                f(mask);
+            }
+            break;
+        }
+
+		case OP_TYPE::ME_ADDTAGS:
+		{
+			const uint8_t* data = program.GetOpArgsPointer(at);
+			mu::OP::ADDRESS source;
+			FMemory::Memcpy(&source, data, sizeof(OP::ADDRESS)); data += sizeof(OP::ADDRESS);
+			f(source);
+			break;
+		}
+
+        case OP_TYPE::ME_FORMAT:
+        {
+			OP::MeshFormatArgs args = program.GetOpArgs<OP::MeshFormatArgs>(at);
+            f(args.source );
+            f(args.format );
+            break;
+        }
+
+        case OP_TYPE::ME_TRANSFORM:
+        {
+			OP::MeshTransformArgs args = program.GetOpArgs<OP::MeshTransformArgs>(at);
+            f(args.source );
+            break;
+        }
+
+        case OP_TYPE::ME_EXTRACTLAYOUTBLOCK:
+        {
+            const uint8_t* data = program.GetOpArgsPointer(at);
+            mu::OP::ADDRESS source;
+			FMemory::Memcpy( &source, data, sizeof(OP::ADDRESS) );
+            f(source);
+            break;
+        }
+
+        case OP_TYPE::ME_CLIPMORPHPLANE:
+        {
+			OP::MeshClipMorphPlaneArgs args = program.GetOpArgs<OP::MeshClipMorphPlaneArgs>(at);
+            f(args.source);
+            break;
+        }
+
+        case OP_TYPE::ME_CLIPWITHMESH :
+        {
+			OP::MeshClipWithMeshArgs args = program.GetOpArgs<OP::MeshClipWithMeshArgs>(at);
+            f(args.source);
+            f(args.clipMesh);
+            break;
+        }
+
+        case OP_TYPE::ME_CLIPDEFORM:
+        {
+			OP::MeshClipDeformArgs args = program.GetOpArgs<OP::MeshClipDeformArgs>(at);
+            f(args.mesh);
+            f(args.clipShape);
+            break;
+        }
+       
+		case OP_TYPE::ME_MORPHRESHAPE:
+		{
+			OP::MeshMorphReshapeArgs Args = program.GetOpArgs<OP::MeshMorphReshapeArgs>(at);
+			f(Args.Morph);
+			f(Args.Reshape);
+
+			break;
+		}
+ 
+        case OP_TYPE::ME_SETSKELETON :
+        {
+			OP::MeshSetSkeletonArgs args = program.GetOpArgs<OP::MeshSetSkeletonArgs>(at);
+            f(args.source);
+            f(args.skeleton);
+            break;
+        }
+
+        case OP_TYPE::ME_PROJECT :
+        {
+			OP::MeshProjectArgs args = program.GetOpArgs<OP::MeshProjectArgs>(at);
+            f(args.mesh);
+            f(args.projector);
+            break;
+        }
+
+		case OP_TYPE::ME_APPLYPOSE:
+		{
+			OP::MeshApplyPoseArgs args = program.GetOpArgs<OP::MeshApplyPoseArgs>(at);
+			f(args.base);
+			f(args.pose);
+			break;
+		}
+
+		case OP_TYPE::ME_GEOMETRYOPERATION:
+		{
+			OP::MeshGeometryOperationArgs args = program.GetOpArgs<OP::MeshGeometryOperationArgs>(at);
+			f(args.meshA);
+			f(args.meshB);
+			f(args.scalarA);
+			f(args.scalarB);
+			break;
+		}
+
+		case OP_TYPE::ME_BINDSHAPE:
+		{
+			OP::MeshBindShapeArgs args = program.GetOpArgs<OP::MeshBindShapeArgs>(at);
+			f(args.mesh);
+			f(args.shape);
+			break;
+		}
+
+		case OP_TYPE::ME_APPLYSHAPE:
+		{
+			OP::MeshApplyShapeArgs args = program.GetOpArgs<OP::MeshApplyShapeArgs>(at);
+			f(args.mesh);
+			f(args.shape);
+			break;
+		}
+
+		case OP_TYPE::ME_OPTIMIZESKINNING:
+		{
+			OP::MeshOptimizeSkinningArgs args = program.GetOpArgs<OP::MeshOptimizeSkinningArgs>(at);
+			f(args.source);
+			break;
+		}
+
+        //-------------------------------------------------------------------------------------
+        case OP_TYPE::IN_ADDMESH:
+        case OP_TYPE::IN_ADDIMAGE:
+        case OP_TYPE::IN_ADDVECTOR:
+        case OP_TYPE::IN_ADDSCALAR:
+        case OP_TYPE::IN_ADDSTRING:
+        case OP_TYPE::IN_ADDCOMPONENT:
+        case OP_TYPE::IN_ADDSURFACE:
+        {
+			OP::InstanceAddArgs args = program.GetOpArgs<OP::InstanceAddArgs>(at);
+            f(args.instance );
+            f(args.value );
+            break;
+        }
+
+        case OP_TYPE::IN_ADDLOD:
+        {
+			const uint8* Data = program.GetOpArgsPointer(at);
+
+			uint8 LODCount;
+			FMemory::Memcpy(&LODCount, Data, sizeof(uint8));
+			Data += sizeof(uint8);
+
+			for (int8 LODIndex = 0; LODIndex < LODCount; ++LODIndex)
+			{
+				OP::ADDRESS LODAddress;
+				FMemory::Memcpy(&LODAddress, Data, sizeof(OP::ADDRESS));
+				Data += sizeof(OP::ADDRESS);
+
+				f(LODAddress);
+            }
+            break;
+        }
+
+		case OP_TYPE::IN_ADDEXTENSIONDATA:
+		{
+			const OP::InstanceAddExtensionDataArgs Args = program.GetOpArgs<OP::InstanceAddExtensionDataArgs>(at);
+
+			f(Args.Instance);
+			f(Args.ExtensionData);
+
+			break;
+		}
+
+        //-------------------------------------------------------------------------------------
+        case OP_TYPE::LA_PACK:
+        {
+			OP::LayoutPackArgs args = program.GetOpArgs<OP::LayoutPackArgs>(at);
+            f(args.Source );
+            break;
+        }
+
+        case OP_TYPE::LA_MERGE:
+        {
+			OP::LayoutMergeArgs args = program.GetOpArgs<OP::LayoutMergeArgs>(at);
+            f(args.Base );
+            f(args.Added );
+            break;
+        }
+
+		case OP_TYPE::LA_REMOVEBLOCKS:
+		{
+			OP::LayoutRemoveBlocksArgs args = program.GetOpArgs<OP::LayoutRemoveBlocksArgs>(at);
+			f(args.Source);
+			f(args.ReferenceLayout);
+			break;
+		}
+
+		case OP_TYPE::LA_FROMMESH:
+		{
+			OP::LayoutFromMeshArgs args = program.GetOpArgs<OP::LayoutFromMeshArgs>(at);
+			f(args.Mesh);
+			break;
+		}
+
+        default:
+			check( false );
+            break;
+        }
+
+    }
+
+
+}
+
