[mutable] Moved the Mutable plugin out of Experimental status into Beta.

diff --git a/Engine/Plugins/Mutable/Source/MutableRuntime/Private/MuR/Parameters.cpp b/Engine/Plugins/Mutable/Source/MutableRuntime/Private/MuR/Parameters.cpp
new file mode 100644
--- /dev/null
+++ b/Engine/Plugins/Mutable/Source/MutableRuntime/Private/MuR/Parameters.cpp
@@ -0,0 +1,1264 @@
+// Copyright Epic Games, Inc. All Rights Reserved.
+
+
+#include "MuR/Parameters.h"
+
+#include "HAL/LowLevelMemTracker.h"
+#include "HAL/PlatformMath.h"
+#include "HAL/UnrealMemory.h"
+#include "Misc/AssertionMacros.h"
+#include "MuR/Model.h"
+#include "MuR/ModelPrivate.h"
+#include "MuR/ParametersPrivate.h"
+#include "MuR/Serialisation.h"
+#include "MuR/SerialisationPrivate.h"
+#include "MuR/System.h"
+
+
+namespace mu
+{
+    MUTABLE_IMPLEMENT_ENUM_SERIALISABLE(PARAMETER_TYPE)              
+    MUTABLE_IMPLEMENT_ENUM_SERIALISABLE(PROJECTOR_TYPE)
+
+	
+	//---------------------------------------------------------------------------------------------
+	Parameters::Parameters()
+	{
+		m_pD = new Private();
+	}
+
+
+	//---------------------------------------------------------------------------------------------
+	Parameters::~Parameters()
+	{
+        check( m_pD );
+		delete m_pD;
+		m_pD = 0;
+	}
+
+
+	//---------------------------------------------------------------------------------------------
+	void Parameters::Serialise( const Parameters* p, OutputArchive& arch )
+	{
+		arch << *p->m_pD;
+    }
+
+
+	//---------------------------------------------------------------------------------------------
+	ParametersPtr Parameters::StaticUnserialise( InputArchive& arch )
+	{
+		LLM_SCOPE_BYNAME(TEXT("MutableRuntime"));
+
+		ParametersPtr pResult = new Parameters();
+		arch >> *pResult->m_pD;
+		return pResult;
+	}
+
+
+    //---------------------------------------------------------------------------------------------
+    //void Parameters::SerialisePortable( const Parameters* params, OutputArchive& arch )
+    //{
+    //    const int32_t ver = 0;
+    //    arch << ver;
+
+    //    uint32 valueCount = params->GetPrivate()->m_values.Num();
+    //    arch << valueCount;
+
+    //    for (size_t p=0; p<valueCount; ++p)
+    //    {
+    //        const FParameterDesc& desc = params->GetPrivate()->m_pModel->GetPrivate()->m_program.m_parameters[p];
+    //        arch << desc.m_name;
+    //        arch << desc.m_uid;
+    //        arch << desc.m_type;
+    //        arch << params->GetPrivate()->m_values[p];
+    //    }
+    //}
+
+
+    //---------------------------------------------------------------------------------------------
+  //  ParametersPtr Parameters::UnserialisePortable( InputArchive& arch, TSharedPtr<const Model> pModel )
+  //  {
+		//LLM_SCOPE_BYNAME(TEXT("MutableRuntime"));
+
+  //      ParametersPtr pResult = Model::NewParameters(pModel);
+  //      size_t modelParameters = pModel->GetPrivate()->m_program.m_parameters.Num();
+
+  //      int32_t ver;
+  //      arch >> ver;
+  //      check( ver == 0 );
+
+  //      uint32_t valueCount = 0;
+  //      arch >> valueCount;
+
+  //      for (uint32_t p=0; p<valueCount; ++p)
+  //      {
+  //          // Read the parameter data.
+  //          std::string name;
+  //          arch >> name;
+
+  //          std::string uid;
+  //          arch >> uid;
+
+  //          PARAMETER_TYPE type;
+  //          arch >> type;
+
+  //          PARAMETER_VALUE value;
+  //          arch >> value;
+
+  //          int modelParam = -1;
+
+  //          // UIDs have higher priority
+  //          if ( !uid.empty() )
+  //          {
+  //              for (size_t mp=0; mp<modelParameters; ++mp)
+  //              {
+  //                  const FParameterDesc& desc = pModel->GetPrivate()->m_program.m_parameters[mp];
+  //                  if (desc.m_uid==uid)
+  //                  {
+  //                      modelParam = int(mp);
+  //                      break;
+  //                  }
+  //              }
+  //          }
+
+  //          // Try with name+type
+  //          if (modelParam<0)
+  //          {
+  //              for (size_t mp=0; mp<modelParameters; ++mp)
+  //              {
+  //                  const FParameterDesc& desc = pModel->GetPrivate()->m_program.m_parameters[mp];
+  //                  if (desc.m_name==name && desc.m_type==type)
+  //                  {
+  //                      modelParam = int(mp);
+  //                      break;
+  //                  }
+  //              }
+  //          }
+
+  //          // Try with name only, and maybe adapt type
+  //          if (modelParam<0)
+  //          {
+  //              for (size_t mp=0; mp<modelParameters; ++mp)
+  //              {
+  //                  const FParameterDesc& desc = pModel->GetPrivate()->m_program.m_parameters[mp];
+  //                  if (desc.m_name==name)
+  //                  {
+  //                      modelParam = int(mp);
+  //                      break;
+  //                  }
+  //              }
+  //          }
+
+  //          if (modelParam>=0)
+  //          {
+  //              // We found something. \todo: convert the type if necessary?
+  //              pResult->GetPrivate()->m_values[modelParam] = value;
+  //          }
+  //      }
+
+  //      return pResult;
+  //  }
+
+
+	//---------------------------------------------------------------------------------------------
+	Parameters::Private* Parameters::GetPrivate() const
+	{
+		return m_pD;
+	}
+
+
+	//---------------------------------------------------------------------------------------------
+	ParametersPtr Parameters::Clone() const
+	{
+		LLM_SCOPE_BYNAME(TEXT("MutableRuntime"));
+
+		ParametersPtr pRes = new Parameters();
+
+		pRes->m_pD->m_pModel = m_pD->m_pModel;
+		pRes->m_pD->m_values = m_pD->m_values;
+		pRes->m_pD->m_multiValues = m_pD->m_multiValues;
+
+		return pRes;
+	}
+
+
+	//---------------------------------------------------------------------------------------------
+	int Parameters::GetCount() const
+	{
+		return (int)m_pD->m_values.Num();
+	}
+
+
+	//---------------------------------------------------------------------------------------------
+	const FString& Parameters::GetName( int index ) const
+	{
+		const FProgram& program = m_pD->m_pModel->GetPrivate()->m_program;
+		check( index>=0 && index<(int)program.m_parameters.Num() );
+
+		return program.m_parameters[index].m_name;
+	}
+
+
+	//---------------------------------------------------------------------------------------------
+	const FGuid& Parameters::GetUid( int index ) const
+	{
+		const FProgram& program = m_pD->m_pModel->GetPrivate()->m_program;
+		check( index>=0 && index<(int)program.m_parameters.Num() );
+
+		return program.m_parameters[index].m_uid;
+	}
+
+
+	//---------------------------------------------------------------------------------------------
+	int Parameters::Find( const FString& strName ) const
+	{
+		return m_pD->Find( strName );
+	}
+
+
+	//---------------------------------------------------------------------------------------------
+	PARAMETER_TYPE Parameters::GetType( int index ) const
+	{
+		const FProgram& program = m_pD->m_pModel->GetPrivate()->m_program;
+		check( index>=0 && index<(int)program.m_parameters.Num() );
+
+		return program.m_parameters[index].m_type;
+	}
+
+
+    //---------------------------------------------------------------------------------------------
+    Ptr<RangeIndex> Parameters::NewRangeIndex( int paramIndex ) const
+    {
+		LLM_SCOPE_BYNAME(TEXT("MutableRuntime"));
+
+        const FProgram& program = m_pD->m_pModel->GetPrivate()->m_program;
+        check( paramIndex>=0 && paramIndex<int(program.m_parameters.Num()) );
+        Ptr<RangeIndex> range;
+
+        if ( paramIndex>=0 && paramIndex<int(program.m_parameters.Num()) )
+        {
+            size_t rangeCount = program.m_parameters[paramIndex].m_ranges.Num();
+            if ( rangeCount>0 )
+            {
+                range = new RangeIndex;
+                range->m_pD->m_pParameters = this;
+                range->m_pD->m_parameter = paramIndex;
+                range->m_pD->m_values.SetNumZeroed( rangeCount );
+            }
+        }
+
+        return range;
+    }
+
+
+    //---------------------------------------------------------------------------------------------
+    int Parameters::GetValueCount( int paramIndex ) const
+    {
+        if (paramIndex<0 || paramIndex >= int(m_pD->m_multiValues.Num()) )
+        {
+            return 0;
+        }
+
+        return int( m_pD->m_multiValues[paramIndex].Num() );
+    }
+
+
+    //---------------------------------------------------------------------------------------------
+    Ptr<RangeIndex> Parameters::GetValueIndex( int paramIndex, int valueIndex ) const
+    {
+        if (paramIndex<0 || paramIndex >= int(m_pD->m_multiValues.Num()) )
+        {
+            return nullptr;
+        }
+
+        if (valueIndex<0 || valueIndex >= int(m_pD->m_multiValues[paramIndex].Num()) )
+        {
+            return nullptr;
+        }
+
+		TMap< TArray<int32_t>, PARAMETER_VALUE >::TRangedForIterator it = m_pD->m_multiValues[paramIndex].begin();
+        for ( int i=0; i<valueIndex; ++i )
+        {
+            ++it;
+        }
+
+        Ptr<RangeIndex> result = NewRangeIndex( paramIndex );
+        result->m_pD->m_values = it->Key;
+
+        return result;
+    }
+
+
+    //---------------------------------------------------------------------------------------------
+    void Parameters::ClearAllValues( int paramIndex )
+    {
+        if (paramIndex<0 || paramIndex >= int(m_pD->m_multiValues.Num()) )
+        {
+            return;
+        }
+
+        m_pD->m_multiValues[paramIndex].Empty();
+    }
+
+
+	//---------------------------------------------------------------------------------------------
+    bool Parameters::GetBoolValue( int index,
+                                   const Ptr<const RangeIndex>& pos ) const
+	{
+		check( index>=0 && index<(int)m_pD->m_values.Num() );
+        check( GetType(index)==PARAMETER_TYPE::T_BOOL );
+
+        // Early out in case of invalid parameters
+        if ( index < 0
+             ||
+             index >= (int)m_pD->m_values.Num()
+             ||
+             GetType(index) != PARAMETER_TYPE::T_BOOL )
+        {
+            return false;
+        }
+
+        // Single value case
+        if (!pos)
+        {
+            // Return the single value
+            return m_pD->m_values[index].Get<ParamBoolType>();
+        }
+
+        // Multivalue case
+        check( pos->m_pD->m_parameter==index );
+
+        if ( index<int(m_pD->m_multiValues.Num()))
+        {
+            const TMap< TArray<int32_t>, PARAMETER_VALUE >& m = m_pD->m_multiValues[index];
+			const PARAMETER_VALUE* it = m.Find(pos->m_pD->m_values);
+            if (it)
+            {
+                return it->Get<ParamBoolType>();
+            }
+        }
+
+        // Multivalue parameter, but no multivalue set. Return single value.
+        return m_pD->m_values[index].Get<ParamBoolType>();
+	}
+
+
+	//---------------------------------------------------------------------------------------------
+    void Parameters::SetBoolValue( int index, bool value,
+                                   const Ptr<const RangeIndex>& pos )
+	{
+		check( index>=0 && index<(int)m_pD->m_values.Num() );
+		check( GetType(index)== PARAMETER_TYPE::T_BOOL );
+
+        // Early out in case of invalid parameters
+        if ( index < 0
+             ||
+             index >= (int)m_pD->m_values.Num()
+             ||
+             GetType(index) != PARAMETER_TYPE::T_BOOL )
+        {
+            return;
+        }
+
+        // Single value case
+        if (!pos)
+        {
+            // Clear multivalue, if set.
+            if (index<int(m_pD->m_multiValues.Num()))
+            {
+                m_pD->m_multiValues[index].Empty();
+            }
+
+            m_pD->m_values[index].Set<ParamBoolType>(value);
+        }
+
+        // Multivalue case
+        else
+        {
+            check( pos->m_pD->m_parameter==index );
+
+            if ( index>=int(m_pD->m_multiValues.Num()))
+            {
+                m_pD->m_multiValues.SetNum(index+1);
+            }
+
+			TMap< TArray<int32_t>, PARAMETER_VALUE >& m = m_pD->m_multiValues[index];
+			PARAMETER_VALUE& it = m.FindOrAdd(pos->m_pD->m_values);
+            it.Set<ParamBoolType>(value);
+        }
+	}
+
+
+	//---------------------------------------------------------------------------------------------
+	int Parameters::GetIntPossibleValueCount( int paramIndex ) const
+	{
+		const FProgram& program = m_pD->m_pModel->GetPrivate()->m_program;
+		check( paramIndex>=0 && paramIndex<(int)program.m_parameters.Num() );
+		return (int)program.m_parameters[paramIndex].m_possibleValues.Num();
+	}
+
+
+	//---------------------------------------------------------------------------------------------
+	int Parameters::GetIntPossibleValue( int paramIndex, int valueIndex ) const
+	{
+		const FProgram& program = m_pD->m_pModel->GetPrivate()->m_program;
+		check( paramIndex>=0
+				&& paramIndex<(int)program.m_parameters.Num() );
+		check( valueIndex>=0
+				&& valueIndex<(int)program.m_parameters[paramIndex].m_possibleValues.Num() );
+
+		return (int)program.m_parameters[paramIndex].m_possibleValues[valueIndex].m_value;
+	}
+
+
+	//---------------------------------------------------------------------------------------------
+	int Parameters::GetIntValueIndex(int paramIndex, const FString& ValueName) const
+	{
+		const FProgram& program = m_pD->m_pModel->GetPrivate()->m_program;
+		check(paramIndex >= 0
+			&& paramIndex < (int)program.m_parameters.Num());
+
+		int result = -1;
+		for (size_t v = 0; v < program.m_parameters[paramIndex].m_possibleValues.Num(); ++v)
+		{
+			if (program.m_parameters[paramIndex].m_possibleValues[v].m_name == ValueName)
+			{
+				result = (int)v;
+				break;
+			}
+		}
+		return result;
+	}
+
+
+	//---------------------------------------------------------------------------------------------
+	int Parameters::GetIntValueIndex(int paramIndex, int32 Value) const
+	{
+		const FProgram& program = m_pD->m_pModel->GetPrivate()->m_program;
+		check(paramIndex >= 0
+			&& paramIndex < (int)program.m_parameters.Num());
+
+		for (size_t v = 0; v < program.m_parameters[paramIndex].m_possibleValues.Num(); ++v)
+		{
+			if (program.m_parameters[paramIndex].m_possibleValues[v].m_value == Value)
+			{
+				return (int)v;
+			}
+		}
+		return -1;
+	}
+
+
+	//---------------------------------------------------------------------------------------------
+	const FString& Parameters::GetIntPossibleValueName( int paramIndex, int valueIndex ) const
+	{
+		const FProgram& program = m_pD->m_pModel->GetPrivate()->m_program;
+		check( paramIndex>=0
+				&& paramIndex<(int)program.m_parameters.Num() );
+		check( valueIndex>=0
+				&& valueIndex<(int)program.m_parameters[paramIndex].m_possibleValues.Num() );
+
+		return program.m_parameters[paramIndex].m_possibleValues[valueIndex].m_name;
+	}
+
+
+	//---------------------------------------------------------------------------------------------
+    int Parameters::GetIntValue( int index, const Ptr<const RangeIndex>& pos ) const
+	{
+		check( index>=0 && index<(int)m_pD->m_values.Num() );
+		check( GetType(index)== PARAMETER_TYPE::T_INT );
+
+        // Early out in case of invalid parameters
+        if ( index < 0
+             ||
+             index >= (int)m_pD->m_values.Num()
+             ||
+             GetType(index) != PARAMETER_TYPE::T_INT )
+        {
+            return 0;
+        }
+
+        // Single value case
+        if (!pos)
+        {
+            // Return the single value
+            return m_pD->m_values[index].Get<ParamIntType>();
+        }
+
+        // Multivalue case
+        check( pos->m_pD->m_parameter==index );
+
+        if ( index<int(m_pD->m_multiValues.Num()))
+        {
+            const TMap< TArray<int32_t>, PARAMETER_VALUE >& m = m_pD->m_multiValues[index];
+			const PARAMETER_VALUE* it = m.Find(pos->m_pD->m_values);
+            if (it)
+            {
+                return it->Get<ParamIntType>();
+            }
+        }
+
+        // Multivalue parameter, but no multivalue set. Return single value.
+        return m_pD->m_values[index].Get<ParamIntType>();
+    }
+
+
+	//---------------------------------------------------------------------------------------------
+    void Parameters::SetIntValue( int index, int value,
+                                  const Ptr<const RangeIndex>& pos )
+	{
+		check( index>=0 && index<(int)m_pD->m_values.Num() );
+		check( GetType(index)== PARAMETER_TYPE::T_INT );
+
+        // Early out in case of invalid parameters
+        if ( index < 0
+             ||
+             index >= (int)m_pD->m_values.Num()
+             ||
+             GetType(index) != PARAMETER_TYPE::T_INT )
+        {
+            return;
+        }
+
+        // Single value case
+        if (!pos)
+        {
+            // Clear multivalue, if set.
+            if (index<int(m_pD->m_multiValues.Num()))
+            {
+                m_pD->m_multiValues[index].Empty();
+            }
+
+            m_pD->m_values[index].Set<ParamIntType>(value);
+        }
+
+        // Multivalue case
+        else
+        {
+            check( pos->m_pD->m_parameter==index );
+
+            if ( index>=int(m_pD->m_multiValues.Num()))
+            {
+                m_pD->m_multiValues.SetNum(index+1);
+            }
+
+			TMap< TArray<int32_t>, PARAMETER_VALUE >& m = m_pD->m_multiValues[index];
+			PARAMETER_VALUE& it = m.FindOrAdd(pos->m_pD->m_values);
+            it.Set<ParamIntType>(value);
+        }
+    }
+
+
+	//---------------------------------------------------------------------------------------------
+    float Parameters::GetFloatValue( int index,
+                                     const Ptr<const RangeIndex>& pos ) const
+	{
+		check( index>=0 && index<(int)m_pD->m_values.Num() );
+		check( GetType(index)== PARAMETER_TYPE::T_FLOAT );
+
+        // Early out in case of invalid parameters
+        if ( index < 0
+             ||
+             index >= (int)m_pD->m_values.Num()
+             ||
+             GetType(index) != PARAMETER_TYPE::T_FLOAT )
+        {
+            return 0.0f;
+        }
+
+        // Single value case
+        if (!pos)
+        {
+            // Return the single value
+            return m_pD->m_values[index].Get<ParamFloatType>();
+        }
+
+        // Multivalue case
+        check( pos->m_pD->m_parameter==index );
+
+        if ( index<int(m_pD->m_multiValues.Num()))
+        {
+            const TMap< TArray<int32_t>, PARAMETER_VALUE >& m = m_pD->m_multiValues[index];
+			const PARAMETER_VALUE* it = m.Find(pos->m_pD->m_values);
+            if (it)
+            {
+                return it->Get<ParamFloatType>();
+            }
+        }
+
+        // Multivalue parameter, but no multivalue set. Return single value.
+        return m_pD->m_values[index].Get<ParamFloatType>();
+    }
+
+
+	//---------------------------------------------------------------------------------------------
+    void Parameters::SetFloatValue( int index, float value,
+                                    const Ptr<const RangeIndex>& pos )
+	{
+		check( index>=0 && index<(int)m_pD->m_values.Num() );
+		check( GetType(index)== PARAMETER_TYPE::T_FLOAT );
+
+        // Early out in case of invalid parameters
+        if ( index < 0
+             ||
+             index >= (int)m_pD->m_values.Num()
+             ||
+             GetType(index) != PARAMETER_TYPE::T_FLOAT )
+        {
+            return;
+        }
+
+        // Single value case
+        if (!pos)
+        {
+            // Clear multivalue, if set.
+            if (index<int(m_pD->m_multiValues.Num()))
+            {
+                m_pD->m_multiValues[index].Empty();
+            }
+
+            m_pD->m_values[index].Set<ParamFloatType>(value);
+        }
+
+        // Multivalue case
+        else
+        {
+            check( pos->m_pD->m_parameter==index );
+
+            if ( index>=int(m_pD->m_multiValues.Num()))
+            {
+                m_pD->m_multiValues.SetNum(index+1);
+            }
+
+			TMap< TArray<int32_t>, PARAMETER_VALUE >& m = m_pD->m_multiValues[index];
+			PARAMETER_VALUE& it = m.FindOrAdd(pos->m_pD->m_values);
+            it.Set<ParamFloatType>(value);
+        }
+    }
+
+
+    //---------------------------------------------------------------------------------------------
+    void Parameters::GetColourValue( int index, float* pR, float* pG, float* pB, float* pA,
+                                     const Ptr<const RangeIndex>& pos ) const
+    {
+        check( index>=0 && index<(int)m_pD->m_values.Num() );
+        check( GetType(index)== PARAMETER_TYPE::T_COLOUR );
+
+        // Early out in case of invalid parameters
+        if ( index < 0
+             ||
+             index >= (int)m_pD->m_values.Num()
+             ||
+             GetType(index) != PARAMETER_TYPE::T_COLOUR )
+        {
+            return;
+        }
+
+        // Single value case
+        if (!pos)
+        {
+            // Return the single value
+			if (pR)
+			{
+				*pR = m_pD->m_values[index].Get<ParamColorType>()[0];
+			}
+
+			if (pG)
+			{
+				*pG = m_pD->m_values[index].Get<ParamColorType>()[1];
+			}
+
+			if (pB)
+			{
+				*pB = m_pD->m_values[index].Get<ParamColorType>()[2];
+			}
+
+			if (pA)
+			{
+				*pA = m_pD->m_values[index].Get<ParamColorType>()[3];
+			}
+
+            return;
+        }
+
+        // Multivalue case
+        check( pos->m_pD->m_parameter==index );
+
+        if ( index<int(m_pD->m_multiValues.Num()))
+        {
+            const TMap< TArray<int32_t>, PARAMETER_VALUE >& m = m_pD->m_multiValues[index];
+			const PARAMETER_VALUE* it = m.Find(pos->m_pD->m_values);
+            if (it)
+            {
+				if (pR)
+				{
+					*pR = it->Get<ParamColorType>()[0];
+				}
+
+				if (pG)
+				{
+					*pG = it->Get<ParamColorType>()[1];
+				}
+
+				if (pB)
+				{
+					*pB = it->Get<ParamColorType>()[2];
+				}
+
+				if (pA)
+				{
+					*pA = it->Get<ParamColorType>()[3];
+				}
+
+                return;
+            }
+        }
+
+        // Multivalue parameter, but no multivalue set. Return single value.
+		if (pR)
+		{
+			*pR = m_pD->m_values[index].Get<ParamColorType>()[0];
+		}
+
+		if (pG)
+		{
+			*pG = m_pD->m_values[index].Get<ParamColorType>()[1];
+		}
+
+		if (pB)
+		{
+			*pB = m_pD->m_values[index].Get<ParamColorType>()[2];
+		}
+
+		if (pA)
+		{
+			*pA = m_pD->m_values[index].Get<ParamColorType>()[3];
+		}
+
+        return;
+    }
+
+
+    //---------------------------------------------------------------------------------------------
+    void Parameters::SetColourValue( int index, float R, float G, float B, float A,
+                                     const Ptr<const RangeIndex>& pos )
+    {
+        check( index>=0 && index<(int)m_pD->m_values.Num() );
+        check( GetType(index)== PARAMETER_TYPE::T_COLOUR );
+
+        // Early out in case of invalid parameters
+        if ( index < 0
+             ||
+             index >= (int)m_pD->m_values.Num()
+             ||
+             GetType(index) != PARAMETER_TYPE::T_COLOUR )
+        {
+            return;
+        }
+
+        // Single value case
+        if (!pos)
+        {
+            // Clear multivalue, if set.
+            if (index<int(m_pD->m_multiValues.Num()))
+            {
+                m_pD->m_multiValues[index].Empty();
+            }
+
+        	ParamColorType Value;
+            Value[0] = R;
+            Value[1] = G;
+            Value[2] = B;
+			Value[3] = A;
+        	
+        	m_pD->m_values[index].Set<ParamColorType>(Value);
+        }
+
+        // Multivalue case
+        else
+        {
+            check( pos->m_pD->m_parameter==index );
+
+            if ( index>=int(m_pD->m_multiValues.Num()))
+            {
+                m_pD->m_multiValues.SetNum(index+1);
+            }
+
+			TMap< TArray<int32_t>, PARAMETER_VALUE >& m = m_pD->m_multiValues[index];
+			PARAMETER_VALUE& it = m.FindOrAdd(pos->m_pD->m_values);
+
+        	ParamColorType Value;
+            Value[0] = R;
+            Value[1] = G;
+			Value[2] = B;
+			Value[3] = A;
+
+        	it.Set<ParamColorType>(Value);
+        }
+    }
+
+
+    //---------------------------------------------------------------------------------------------
+	FName Parameters::GetImageValue( int index, const Ptr<const RangeIndex>& pos ) const
+    {
+        check( index >= 0 && index < (int)m_pD->m_values.Num() );
+        check( GetType( index ) == PARAMETER_TYPE::T_IMAGE );
+
+		// Early out in case of invalid parameters
+        if ( index < 0
+             ||
+             index >= (int)m_pD->m_values.Num()
+             ||
+             GetType(index) != PARAMETER_TYPE::T_IMAGE )
+        {
+            return {};
+        }
+		
+		// Single value case
+		if (!pos)
+		{
+			return m_pD->m_values[index].Get<ParamImageType>();
+		}
+
+		// Multivalue case
+		check(pos->m_pD->m_parameter == index);
+
+		if (index < int(m_pD->m_multiValues.Num()))
+		{
+			const TMap< TArray<int32_t>, PARAMETER_VALUE >& m = m_pD->m_multiValues[index];
+			const PARAMETER_VALUE* it = m.Find(pos->m_pD->m_values);
+			if (it)
+			{
+				return it->Get<ParamImageType>();
+			}
+		}
+
+		// Multivalue parameter, but no multivalue set. Return single value.
+		return m_pD->m_values[index].Get<ParamImageType>();
+
+    }
+
+
+    //---------------------------------------------------------------------------------------------
+    void Parameters::SetImageValue( int index, FName id, const Ptr<const RangeIndex>& pos )
+    {
+        check( index >= 0 && index < (int)m_pD->m_values.Num() );
+        check( GetType( index ) == PARAMETER_TYPE::T_IMAGE );
+
+		// Single value case
+		if (!pos)
+		{
+			m_pD->m_values[index].Set<ParamImageType>(id);
+		}
+
+		// Multivalue case
+		else
+		{
+			check(pos->m_pD->m_parameter == index);
+
+			if (index >= int(m_pD->m_multiValues.Num()))
+			{
+				m_pD->m_multiValues.SetNum(index + 1);
+			}
+
+			TMap< TArray<int32_t>, PARAMETER_VALUE >& m = m_pD->m_multiValues[index];
+			PARAMETER_VALUE& it = m.FindOrAdd(pos->m_pD->m_values);
+			it.Set<ParamImageType>(id);
+		}
+    }
+
+
+    //---------------------------------------------------------------------------------------------
+    void Parameters::GetStringValue( int index, FString& OutValue, const Ptr<const RangeIndex>& pos ) const
+    {
+        check( index >= 0 && index < (int)m_pD->m_values.Num() );
+        check( GetType( index ) == PARAMETER_TYPE::T_STRING );
+
+        // Early out in case of invalid parameters
+        if ( index < 0 || index >= (int)m_pD->m_values.Num() ||
+             GetType( index ) != PARAMETER_TYPE::T_STRING )
+        {
+            OutValue = TEXT("");
+			return;
+        }
+
+        // Single value case
+        if ( !pos )
+        {
+            // Return the single value
+			OutValue = m_pD->m_values[index].Get<ParamStringType>();
+            return;
+        }
+
+        // Multivalue case
+        check( pos->m_pD->m_parameter == index );
+
+        if ( index < int( m_pD->m_multiValues.Num() ) )
+        {
+            const TMap< TArray<int32_t>, PARAMETER_VALUE >& m = m_pD->m_multiValues[index];
+			const PARAMETER_VALUE* it = m.Find( pos->m_pD->m_values );
+            if ( it )
+            {
+				OutValue = it->Get<ParamStringType>();
+                return;
+            }
+        }
+
+        // Multivalue parameter, but no multivalue set. Return single value.
+		OutValue = m_pD->m_values[index].Get<ParamStringType>();
+        return;
+    }
+
+
+    //---------------------------------------------------------------------------------------------
+    void Parameters::SetStringValue( int index, const FString& Value, const Ptr<const RangeIndex>& pos )
+    {
+		LLM_SCOPE_BYNAME(TEXT("MutableRuntime"));
+
+        check( index >= 0 && index < (int)m_pD->m_values.Num() );
+        check( GetType( index ) == PARAMETER_TYPE::T_FLOAT );
+
+        // Early out in case of invalid parameters
+        if ( index < 0 || index >= (int)m_pD->m_values.Num() ||
+             GetType( index ) != PARAMETER_TYPE::T_FLOAT )
+        {
+            return;
+        }
+
+        // Single value case
+        if ( !pos )
+        {
+            // Clear multivalue, if set.
+            if ( index < int( m_pD->m_multiValues.Num() ) )
+            {
+                m_pD->m_multiValues[index].Empty();
+            }
+
+        	m_pD->m_values[index].Set<ParamStringType>(Value);
+        }
+
+        // Multivalue case
+        else
+        {
+            check( pos->m_pD->m_parameter == index );
+
+            if ( index >= int( m_pD->m_multiValues.Num() ) )
+            {
+                m_pD->m_multiValues.SetNum( index + 1 );
+            }
+
+			TMap< TArray<int32_t>, PARAMETER_VALUE >& m = m_pD->m_multiValues[index];
+			PARAMETER_VALUE& it = m.FindOrAdd( pos->m_pD->m_values );
+        	it.Set<ParamStringType>(Value);
+        }
+    }
+
+
+    //---------------------------------------------------------------------------------------------
+    FProjector Parameters::Private::GetProjectorValue( int index, const Ptr<const RangeIndex>& pos ) const
+    {
+
+        const FProgram& program = m_pModel->GetPrivate()->m_program;
+
+        // Early out in case of invalid parameters
+        if ( index < 0
+             ||
+             index >= (int)m_values.Num()
+			 || 
+			 index >= (int)program.m_parameters.Num()
+             ||
+             program.m_parameters[index].m_type != PARAMETER_TYPE::T_PROJECTOR )
+        {
+			check(false);
+            return FProjector();
+        }
+
+        const FProjector* result = nullptr;
+
+        // Single value case
+        if (!pos)
+        {
+            // Return the single value
+            result = &m_values[index].Get<ParamProjectorType>();
+        }
+
+        // Multivalue case
+        if (!result)
+        {
+            check( pos->m_pD->m_parameter==index );
+
+            if ( index<m_multiValues.Num())
+            {
+                const PARAMETER_VALUE* it = m_multiValues[index].Find(pos->m_pD->m_values);
+                if (it)
+                {
+                    result = &it->Get<ParamProjectorType>();
+                }
+            }
+        }
+
+        // Multivalue parameter, but no multivalue set. Return single value.
+        if (!result)
+        {
+            result = &m_values[index].Get<ParamProjectorType>();
+        }
+
+        return *result;
+    }
+
+
+	//---------------------------------------------------------------------------------------------
+	void Parameters::GetProjectorValue( int ParameterIndex,
+		PROJECTOR_TYPE* OutType,
+		FVector3f* OutPos,
+		FVector3f* OutDir,
+		FVector3f* OutUp,
+		FVector3f* OutScale,
+		float* OutProjectionAngle,
+		const Ptr<const RangeIndex>& pos ) const
+	{
+		check(ParameterIndex >=0 && ParameterIndex <m_pD->m_values.Num() );
+		check( GetType(ParameterIndex)== PARAMETER_TYPE::T_PROJECTOR );
+
+        // Early out in case of invalid parameters
+        if (ParameterIndex < 0
+             ||
+			ParameterIndex >= m_pD->m_values.Num()
+             ||
+             GetType(ParameterIndex) != PARAMETER_TYPE::T_PROJECTOR )
+        {
+            return;
+        }
+
+        FProjector result = m_pD->GetProjectorValue(ParameterIndex, pos );
+
+        // Copy results
+        if (OutType) *OutType = result.type;
+
+        if (OutPos) *OutPos = result.position;
+        if (OutDir) *OutDir = result.direction;
+        if (OutUp) *OutUp = result.up;
+        if (OutScale) *OutScale = result.scale;
+
+        if (OutProjectionAngle) *OutProjectionAngle = result.projectionAngle;
+    }
+
+
+	//---------------------------------------------------------------------------------------------
+	void Parameters::SetProjectorValue( int ParameterIndex,
+		const FVector3f& pos,
+		const FVector3f& dir,
+		const FVector3f& up,
+		const FVector3f& scale,
+        float projectionAngle,
+        const Ptr<const RangeIndex>& RangePosition)
+	{
+		check(ParameterIndex >=0 && ParameterIndex <m_pD->m_values.Num() );
+		check( GetType(ParameterIndex)== PARAMETER_TYPE::T_PROJECTOR );
+
+        // Early out in case of invalid parameters
+        if (ParameterIndex < 0
+             ||
+			ParameterIndex >= m_pD->m_values.Num()
+             ||
+             GetType(ParameterIndex) != PARAMETER_TYPE::T_PROJECTOR )
+        {
+            return;
+        }
+
+	    // Parameters cannot change the projector type anymore
+		FProjector Value;
+        Value.type = PROJECTOR_TYPE::COUNT;
+        if (m_pD->m_pModel)
+        {
+            const FProgram& program = m_pD->m_pModel->GetPrivate()->m_program;
+            const FProjector& Projector = program.m_parameters[ParameterIndex].m_defaultValue.Get<ParamProjectorType>();
+            Value.type = Projector.type;
+        }
+
+        Value.position = pos;
+        Value.direction = dir;
+        Value.up = up;        
+        Value.scale = scale;
+
+        Value.projectionAngle = projectionAngle;
+		
+        // Single value case
+        if (!RangePosition)
+        {
+            // Clear multivalue, if set.
+            if (ParameterIndex<m_pD->m_multiValues.Num())
+            {
+                m_pD->m_multiValues[ParameterIndex].Empty();
+            }
+
+            m_pD->m_values[ParameterIndex].Set<ParamProjectorType>(Value);
+        }
+
+        // Multivalue case
+        else
+        {
+            check(RangePosition->m_pD->m_parameter== ParameterIndex);
+
+            if (ParameterIndex >=m_pD->m_multiValues.Num())
+            {
+                m_pD->m_multiValues.SetNum(ParameterIndex+1);
+            }
+
+			TMap< TArray<int32_t>, PARAMETER_VALUE >& m = m_pD->m_multiValues[ParameterIndex];
+			PARAMETER_VALUE& it = m.FindOrAdd(RangePosition->m_pD->m_values);
+            it.Set<ParamProjectorType>(Value);
+        }
+    }
+
+
+    //---------------------------------------------------------------------------------------------
+    bool Parameters::HasSameValue( int thisParamIndex,
+                                   const ParametersPtrConst& other,
+                                   int otherParamIndex ) const
+    {
+        if ( GetType(thisParamIndex) != other->GetType(otherParamIndex) )
+        {
+            return false;
+        }
+
+        if ( !(m_pD->m_values[thisParamIndex]==other->m_pD->m_values[otherParamIndex] ) )
+        {
+            return false;
+        }
+
+        size_t thisNumMultiValues = 0;
+        bool thisHasMultiValues = int(m_pD->m_multiValues.Num()) > thisParamIndex;
+        if (thisHasMultiValues)
+        {
+            thisNumMultiValues = m_pD->m_multiValues[thisParamIndex].Num();
+        }
+
+        size_t otherNumMultiValues = 0;
+        bool otherHasMultiValues = int(other->m_pD->m_multiValues.Num()) > otherParamIndex;
+        if (otherHasMultiValues)
+        {
+            otherNumMultiValues = other->m_pD->m_multiValues[otherParamIndex].Num();
+        }
+
+        if ( thisNumMultiValues != otherNumMultiValues )
+        {
+            return false;
+        }
+
+        if ( thisHasMultiValues
+             &&
+             otherHasMultiValues
+             &&
+             !(m_pD->m_multiValues[thisParamIndex]==other->m_pD->m_multiValues[otherParamIndex] ))
+        {
+            return false;
+        }
+
+        return true;
+    }
+
+
+	//---------------------------------------------------------------------------------------------
+	int32 Parameters::Private::Find( const FString& Name ) const
+	{
+		const FProgram& program = m_pModel->GetPrivate()->m_program;
+
+		int result = -1;
+
+		for( int i=0; result<0 && i<(int)program.m_parameters.Num(); ++i )
+		{
+			const FParameterDesc& p = program.m_parameters[i];
+
+			if (p.m_name == Name)
+			{
+				result = i;
+			}
+		}
+
+		return result;
+	}
+
+
+    //---------------------------------------------------------------------------------------------
+    //---------------------------------------------------------------------------------------------
+    //---------------------------------------------------------------------------------------------
+    RangeIndex::RangeIndex()
+    {
+        m_pD = new Private();
+    }
+
+
+    //---------------------------------------------------------------------------------------------
+    RangeIndex::~RangeIndex()
+    {
+        check( m_pD );
+        delete m_pD;
+        m_pD = 0;
+    }
+
+
+    //---------------------------------------------------------------------------------------------
+    RangeIndex::Private* RangeIndex::GetPrivate() const
+    {
+        return m_pD;
+    }
+
+
+    //---------------------------------------------------------------------------------------------
+    int RangeIndex::GetRangeCount() const
+    {
+        return int( m_pD->m_values.Num() );
+    }
+
+
+    //---------------------------------------------------------------------------------------------
+    const FString& RangeIndex::GetRangeName( int index ) const
+    {
+        check( index >= 0 && index < GetRangeCount() );
+        check( index < int(m_pD->m_pParameters->GetPrivate()->m_values.Num()) );
+        TSharedPtr<const Model> pModel = m_pD->m_pParameters->GetPrivate()->m_pModel;
+        check(pModel);
+        check( index < int(pModel->GetPrivate()->m_program.m_ranges.Num()) );
+
+        return pModel->GetPrivate()->m_program.m_ranges[index].m_name;
+    }
+
+
+    //---------------------------------------------------------------------------------------------
+    const FString& RangeIndex::GetRangeUid( int index ) const
+    {
+        check( index >= 0 && index < GetRangeCount() );
+        check( index < int(m_pD->m_pParameters->GetPrivate()->m_values.Num()) );
+		TSharedPtr<const Model> pModel = m_pD->m_pParameters->GetPrivate()->m_pModel;
+		check(pModel);
+        check( index < int(pModel->GetPrivate()->m_program.m_ranges.Num()) );
+
+        return pModel->GetPrivate()->m_program.m_ranges[index].m_uid;
+    }
+
+
+    //---------------------------------------------------------------------------------------------
+    void RangeIndex::SetPosition( int index, int position )
+    {
+        check( index >= 0 && index < GetRangeCount() );
+        if ( index >= 0 && index < GetRangeCount() )
+        {
+            m_pD->m_values[index] = position;
+        }
+    }
+
+
+    //---------------------------------------------------------------------------------------------
+    int RangeIndex::GetPosition( int index ) const
+    {
+        check( index >= 0 && index < GetRangeCount() );
+        if ( index >= 0 && index < GetRangeCount() )
+        {
+            return m_pD->m_values[index];
+        }
+        return 0;
+    }
+
+}
+
diff --git a/Engine/Plugins/Mutable/Source/MutableRuntime/Private/MuR/Settings.cpp b/Engine/Plugins/Mutable/Source/MutableRuntime/Private/MuR/Settings.cpp
new file mode 100644
--- /dev/null
+++ b/Engine/Plugins/Mutable/Source/MutableRuntime/Private/MuR/Settings.cpp
@@ -0,0 +1,40 @@
+// Copyright Epic Games, Inc. All Rights Reserved.
+
+#include "MuR/Settings.h"
+
+#include "Misc/AssertionMacros.h"
+
+namespace mu
+{
+  
+	Settings::Settings()
+    {
+    }
+
+
+    Settings::~Settings()
+    {
+    }
+
+
+    //---------------------------------------------------------------------------------------------
+    void Settings::SetProfile( bool bEnabled )
+    {
+        bProfile = bEnabled;
+    }
+
+
+    //---------------------------------------------------------------------------------------------
+    void Settings::SetWorkingMemoryBytes( uint64 Bytes )
+    {
+        WorkingMemoryBytes = Bytes;
+    }
+
+
+    //---------------------------------------------------------------------------------------------
+    void Settings::SetImageCompressionQuality( int32 Quality )
+    {
+        ImageCompressionQuality = Quality;
+    }
+
+}
