[mutable] Moved the Mutable plugin out of Experimental status into Beta.

diff --git a/Engine/Plugins/Mutable/Source/MutableRuntime/Private/MuR/System.cpp b/Engine/Plugins/Mutable/Source/MutableRuntime/Private/MuR/System.cpp
new file mode 100644
--- /dev/null
+++ b/Engine/Plugins/Mutable/Source/MutableRuntime/Private/MuR/System.cpp
@@ -0,0 +1,1946 @@
+// Copyright Epic Games, Inc. All Rights Reserved.
+
+#include "MuR/System.h"
+
+#include "Containers/Array.h"
+#include "Containers/Map.h"
+#include "Containers/Set.h"
+#include "Serialization/BitWriter.h"
+#include "HAL/IConsoleManager.h"
+#include "HAL/LowLevelMemTracker.h"
+#include "HAL/PlatformCrt.h"
+#include "HAL/PlatformMath.h"
+#include "Logging/LogCategory.h"
+#include "Logging/LogMacros.h"
+#include "Misc/AssertionMacros.h"
+#include "MuR/CodeRunner.h"
+#include "MuR/CodeVisitor.h"
+#include "MuR/InstancePrivate.h"
+#include "MuR/Mesh.h"
+#include "MuR/Model.h"
+#include "MuR/ModelPrivate.h"
+#include "MuR/MutableMath.h"
+#include "MuR/MutableString.h"
+#include "MuR/MutableTrace.h"
+#include "MuR/NullExtensionDataStreamer.h"
+#include "MuR/Operations.h"
+#include "MuR/Parameters.h"
+#include "MuR/ParametersPrivate.h"
+#include "MuR/Platform.h"
+#include "MuR/Serialisation.h"
+#include "MuR/SystemPrivate.h"
+#include "MuR/MutableRuntimeModule.h"
+#include "Templates/SharedPointer.h"
+#include "Templates/Tuple.h"
+#include "Trace/Detail/Channel.h"
+#include "ProfilingDebugging/CountersTrace.h"
+#include "PackedNormal.h"
+
+namespace 
+{
+
+bool bEnableDetailedMemoryBudgetExceededLogging = false;
+static FAutoConsoleVariableRef CVarEnableDetailedMemoryBudgetExceededLogging (
+	TEXT("mutable.EnableDetailedMemoryBudgetExceededLogging"),
+	bEnableDetailedMemoryBudgetExceededLogging,
+	TEXT("If set to true, enables a more detailed logging when memory budget is exceeded. Only for Debug and Development builds."),
+	ECVF_Default);
+}
+
+namespace mu
+{
+	MUTABLE_IMPLEMENT_ENUM_SERIALISABLE(ETextureCompressionStrategy);
+
+	TRACE_DECLARE_INT_COUNTER(MutableRuntime_LiveInstances,		TEXT("MutableRuntime/LiveInstances"));
+	TRACE_DECLARE_INT_COUNTER(MutableRuntime_Updates,			TEXT("MutableRuntime/Updates"));
+
+	TRACE_DECLARE_INT_COUNTER(MutableRuntime_MemInternal,	TEXT("MutableRuntime/MemInternal"));
+	TRACE_DECLARE_INT_COUNTER(MutableRuntime_MemTemp,		TEXT("MutableRuntime/MemTemp"));
+	TRACE_DECLARE_INT_COUNTER(MutableRuntime_MemPool,		TEXT("MutableRuntime/MemPool"));
+	TRACE_DECLARE_INT_COUNTER(MutableRuntime_MemCache,		TEXT("MutableRuntime/MemCache"));
+	TRACE_DECLARE_INT_COUNTER(MutableRuntime_MemRom,		TEXT("MutableRuntime/MemRom"));
+	TRACE_DECLARE_INT_COUNTER(MutableRuntime_MemImage,		TEXT("MutableRuntime/MemImage"));
+	TRACE_DECLARE_INT_COUNTER(MutableRuntime_MemMesh,		TEXT("MutableRuntime/MemMesh"));
+	TRACE_DECLARE_INT_COUNTER(MutableRuntime_MemStream,		TEXT("MutableRuntime/MemStream"));
+	TRACE_DECLARE_INT_COUNTER(MutableRuntime_MemTotal,		TEXT("MutableRuntime/MemTotal"));
+	TRACE_DECLARE_INT_COUNTER(MutableRuntime_MemBudget,		TEXT("MutableRuntime/MemBudget"));
+
+
+    //---------------------------------------------------------------------------------------------
+    System::System(const Ptr<Settings>& InSettings, const TSharedPtr<ExtensionDataStreamer>& DataStreamer)
+    {
+		LLM_SCOPE_BYNAME(TEXT("MutableRuntime"));
+
+        Ptr<Settings> pSettings = InSettings;
+
+        if ( !pSettings )
+        {
+            pSettings = new Settings;
+        }
+
+        // Choose the implementation
+        m_pD = new System::Private( pSettings, DataStreamer );
+    }
+
+
+    //---------------------------------------------------------------------------------------------
+    System::~System()
+    {
+		LLM_SCOPE_BYNAME(TEXT("MutableRuntime"));
+		MUTABLE_CPUPROFILER_SCOPE(SystemDestructor);
+
+        check( m_pD );
+
+        delete m_pD;
+        m_pD = nullptr;
+    }
+
+
+    //---------------------------------------------------------------------------------------------
+    System::Private* System::GetPrivate() const
+    {
+        return m_pD;
+    }
+
+
+    //---------------------------------------------------------------------------------------------
+    System::Private::Private( Ptr<mu::Settings> InSettings, const TSharedPtr<mu::ExtensionDataStreamer>& InDataStreamer)
+    {
+		LLM_SCOPE_BYNAME(TEXT("MutableRuntime"));
+
+        Settings = InSettings;
+	
+		WorkingMemoryManager.BudgetBytes = Settings->WorkingMemoryBytes;
+		WorkingMemoryManager.GeneratedResources.Reserve(WorkingMemoryManager.MaxGeneratedResourceCacheSize);
+
+		ExtensionDataStreamer = InDataStreamer;
+		if (!ExtensionDataStreamer)
+		{
+			ExtensionDataStreamer = MakeShared<NullExtensionDataStreamer>();
+		}
+	}
+
+
+    //---------------------------------------------------------------------------------------------
+    System::Private::~Private()
+    {
+		LLM_SCOPE_BYNAME(TEXT("MutableRuntime"));
+		MUTABLE_CPUPROFILER_SCOPE(SystemPrivateDestructor);
+
+		// Make it explicit to try to capture metrics
+        StreamInterface = nullptr;
+		ExtensionDataStreamer = nullptr;
+		ExternalResourceProvider = nullptr;
+    }
+
+
+    //---------------------------------------------------------------------------------------------
+    void System::SetStreamingInterface(const TSharedPtr<ModelReader>& InInterface )
+    {
+		LLM_SCOPE_BYNAME(TEXT("MutableRuntime"));
+       
+        m_pD->StreamInterface = InInterface;
+    }
+
+
+	//---------------------------------------------------------------------------------------------
+	void System::SetWorkingMemoryBytes(uint64 InBytes)
+	{
+		LLM_SCOPE_BYNAME(TEXT("MutableRuntime"));
+		MUTABLE_CPUPROFILER_SCOPE(SetWorkingMemoryBytes);
+
+		m_pD->WorkingMemoryManager.BudgetBytes = InBytes;
+		m_pD->WorkingMemoryManager.EnsureBudgetBelow(0);
+	}
+
+
+    //---------------------------------------------------------------------------------------------
+    void System::SetGeneratedCacheSize( uint32 InCount )
+    {
+		LLM_SCOPE_BYNAME(TEXT("MutableRuntime"));
+		MUTABLE_CPUPROFILER_SCOPE(SetGeneratedCacheSize);
+
+		m_pD->WorkingMemoryManager.MaxGeneratedResourceCacheSize = InCount;
+		m_pD->WorkingMemoryManager.GeneratedResources.Reserve(m_pD->WorkingMemoryManager.MaxGeneratedResourceCacheSize);
+		if (m_pD->WorkingMemoryManager.GeneratedResources.Num()>int32(InCount))
+		{
+			// Discard some random resource keys.
+			m_pD->WorkingMemoryManager.GeneratedResources.SetNum(InCount);
+		}
+	}
+
+	
+	//---------------------------------------------------------------------------------------------
+	void System::ClearWorkingMemory()
+    {
+		LLM_SCOPE_BYNAME(TEXT("MutableRuntime"));
+
+		UE_LOG(LogMutableCore, Log, TEXT("Forcing a working memory clear."));
+
+		// rom caches
+		for (FWorkingMemoryManager::FModelCacheEntry& ModelCache : m_pD->WorkingMemoryManager.CachePerModel)
+		{
+			if (const TSharedPtr<const Model> CacheModel = ModelCache.Model.Pin())
+			{
+				FProgram& Program = CacheModel->GetPrivate()->m_program;
+
+				for (int32 RomIndex = 0; RomIndex < Program.m_roms.Num(); ++RomIndex)
+				{
+					Program.UnloadRom(RomIndex);
+				}
+			}
+		}
+
+		m_pD->WorkingMemoryManager.PooledImages.Empty();
+		m_pD->WorkingMemoryManager.CacheResources.Empty();
+		check(m_pD->WorkingMemoryManager.TempImages.IsEmpty());
+		check(m_pD->WorkingMemoryManager.TempMeshes.IsEmpty());
+	}
+
+	
+    //---------------------------------------------------------------------------------------------
+    void System::SetExternalResourceProvider(const TSharedPtr<ExternalResourceProvider>& InInterface )
+    {
+		LLM_SCOPE_BYNAME(TEXT("MutableRuntime"));
+
+        m_pD->ExternalResourceProvider = InInterface;
+    }
+
+
+	//---------------------------------------------------------------------------------------------
+	void System::SetImagePixelConversionOverride(const FImageOperator::FImagePixelFormatFunc& FormatFunc)
+	{
+		LLM_SCOPE_BYNAME(TEXT("MutableRuntime"));
+
+		m_pD->ImagePixelFormatOverride = FormatFunc;
+	}
+
+
+    //---------------------------------------------------------------------------------------------
+    Instance::ID System::NewInstance( const TSharedPtr<const Model>& InModel )
+    {
+		LLM_SCOPE_BYNAME(TEXT("MutableRuntime"));
+		MUTABLE_CPUPROFILER_SCOPE(NewInstance);
+
+		FLiveInstance instanceData;
+		instanceData.InstanceID = ++m_pD->LastInstanceID;
+		instanceData.Instance = nullptr;
+		instanceData.Model = InModel;
+		instanceData.State = -1;
+		instanceData.Cache = MakeShared<FProgramCache>();
+		m_pD->WorkingMemoryManager.LiveInstances.Add(instanceData);
+
+		TRACE_COUNTER_SET(MutableRuntime_LiveInstances, m_pD->WorkingMemoryManager.LiveInstances.Num());
+
+		return instanceData.InstanceID;
+	}
+
+
+    //---------------------------------------------------------------------------------------------
+    Ptr<const Instance> System::BeginUpdate( Instance::ID InInstanceID,
+                                     const Ptr<const Parameters>& InParams,
+                                     int32 InStateIndex,
+                                     uint32 InLodMask )
+    {
+		LLM_SCOPE_BYNAME(TEXT("MutableRuntime"));
+		MUTABLE_CPUPROFILER_SCOPE(SystemBeginUpdate);
+		TRACE_COUNTER_INCREMENT(MutableRuntime_Updates);
+
+		if (!InParams)
+		{
+			UE_LOG(LogMutableCore, Error, TEXT("Invalid parameters in mutable update."));
+			return nullptr;
+		}
+
+		FLiveInstance* pLiveInstance = m_pD->FindLiveInstance(InInstanceID);
+		if (!pLiveInstance)
+		{
+			UE_LOG(LogMutableCore, Error, TEXT("Invalid instance id in mutable update."));
+			return nullptr;
+		}
+
+		m_pD->WorkingMemoryManager.CurrentInstanceCache = pLiveInstance->Cache;
+
+		FProgram& program = pLiveInstance->Model->GetPrivate()->m_program;
+
+		bool validState = InStateIndex >= 0 && InStateIndex < (int)program.m_states.Num();
+		if (!validState)
+		{
+			UE_LOG(LogMutableCore, Error, TEXT("Invalid state in mutable update."));
+			return nullptr;
+		}
+
+		// This may free resources that allow us to use less memory.
+		pLiveInstance->Instance = nullptr;
+
+		bool fullBuild = (InStateIndex != pLiveInstance->State);
+
+		pLiveInstance->State = InStateIndex;
+
+		// If we changed parameters that are not in this state, we need to rebuild all.
+		if (!fullBuild)
+		{
+			fullBuild = m_pD->CheckUpdatedParameters(pLiveInstance, InParams.get(), pLiveInstance->UpdatedParameters);
+		}
+
+		// Remove cached data
+		m_pD->WorkingMemoryManager.ClearCacheLayer0();
+		if (fullBuild)
+		{
+			m_pD->WorkingMemoryManager.ClearCacheLayer1();
+		}
+
+		m_pD->WorkingMemoryManager.BeginRunnerThread();
+
+		OP::ADDRESS rootAt = pLiveInstance->Model->GetPrivate()->m_program.m_states[InStateIndex].m_root;
+
+		// Prepare instance cache
+		m_pD->PrepareCache(pLiveInstance->Model.Get(), InStateIndex);
+		pLiveInstance->OldParameters = InParams->Clone();
+
+		// Ensure the model cache has been created
+		m_pD->WorkingMemoryManager.FindOrAddModelCache(pLiveInstance->Model);
+
+		m_pD->RunCode(pLiveInstance->Model, InParams.get(), rootAt, InLodMask);
+
+		Ptr<const Instance> Result = pLiveInstance->Cache->GetInstance(FCacheAddress(rootAt, 0, 0));
+
+		// Debug check to see if we managed the op-hit-counts correctly
+		pLiveInstance->Cache->CheckHitCountsCleared();
+
+		pLiveInstance->Instance = Result;
+		if (Result)
+		{
+			Result->GetPrivate()->Id = pLiveInstance->InstanceID;
+		}
+		else
+		{
+			// In case of failure return an empty instance, to prevent following code to have to check it every time
+			Result = new Instance;
+		}
+
+		m_pD->WorkingMemoryManager.EndRunnerThread();
+
+		m_pD->WorkingMemoryManager.CurrentInstanceCache = nullptr;
+
+		return Result;
+	}
+
+
+	Ptr<const Image> System::GetImageInline(Instance::ID instanceID, FResourceID ImageId, int32 MipsToSkip, int32 InImageLOD)
+	{
+		LLM_SCOPE_BYNAME(TEXT("MutableRuntime"));
+		MUTABLE_CPUPROFILER_SCOPE(SystemGetImage);
+
+		Ptr<const Image> pResult;
+
+		// Find the live instance
+		FLiveInstance* pLiveInstance = m_pD->FindLiveInstance(instanceID);
+		check(pLiveInstance);
+		m_pD->WorkingMemoryManager.CurrentInstanceCache = pLiveInstance->Cache;
+
+		OP::ADDRESS RootAddress = GetResourceIDRoot(ImageId);
+		pResult = m_pD->BuildImage(pLiveInstance->Model, pLiveInstance->OldParameters.get(), RootAddress, MipsToSkip, InImageLOD);
+
+		// We always need to return something valid.
+		if (!pResult)
+		{
+			pResult = new mu::Image(16, 16, 1, EImageFormat::IF_RGBA_UBYTE, EInitializationType::Black);
+		}
+
+		m_pD->WorkingMemoryManager.CurrentInstanceCache = nullptr;
+
+		return pResult;
+	}
+
+	UE::Tasks::TTask<Ptr<const Image>> System::GetImage(Instance::ID instanceID, FResourceID ImageId, int32 MipsToSkip, int32 InImageLOD)
+	{
+		LLM_SCOPE_BYNAME(TEXT("MutableRuntime"));
+		MUTABLE_CPUPROFILER_SCOPE(SystemGetImage);
+
+		Ptr<const Image> pResult;
+
+		// Find the live instance
+		FLiveInstance* pLiveInstance = m_pD->FindLiveInstance(instanceID);
+		check(pLiveInstance);
+		m_pD->WorkingMemoryManager.CurrentInstanceCache = pLiveInstance->Cache;
+
+		OP::ADDRESS RootAddress = GetResourceIDRoot(ImageId);
+		
+		m_pD->WorkingMemoryManager.BeginRunnerThread();
+		
+		mu::OP_TYPE OpType = pLiveInstance->Model->GetPrivate()->m_program.GetOpType(RootAddress);
+		if (GetOpDataType(OpType) != DT_IMAGE)
+		{
+			m_pD->WorkingMemoryManager.EndRunnerThread();
+			m_pD->WorkingMemoryManager.CurrentInstanceCache = nullptr;
+
+			return UE::Tasks::MakeCompletedTask<Ptr<const Image>>(
+					new mu::Image(16, 16, 1, EImageFormat::IF_RGBA_UBYTE, EInitializationType::Black));
+		}
+		
+		TSharedRef<CodeRunner> Runner = CodeRunner::Create(
+				m_pD->Settings, m_pD, EExecutionStrategy::MinimizeMemory, pLiveInstance->Model, pLiveInstance->OldParameters.get(), RootAddress, System::AllLODs, MipsToSkip, InImageLOD, FScheduledOp::EType::Full);
+	
+		
+		constexpr bool bForceInlineExecution = false;
+		UE::Tasks::FTask RunnerCompletionEvent = Runner->StartRun(bForceInlineExecution);
+
+		return UE::Tasks::Launch(TEXT("System::GetImageResultTask"),
+				[SystemPrivate = m_pD, Runner, RootAddress, MipsToSkip]() -> Ptr<const Image>
+				{
+					Ptr<const Image> Result;
+
+					SystemPrivate->bUnrecoverableError = Runner->bUnrecoverableError;
+					if (!Runner->bUnrecoverableError)
+					{
+						Result = SystemPrivate->WorkingMemoryManager.LoadImage(FCacheAddress(RootAddress, 0, MipsToSkip), true);
+					}
+
+					if (!Result)
+					{
+						Result = new mu::Image(16, 16, 1, EImageFormat::IF_RGBA_UBYTE, EInitializationType::Black);
+					}
+
+					SystemPrivate->WorkingMemoryManager.EndRunnerThread();
+					SystemPrivate->WorkingMemoryManager.CurrentInstanceCache = nullptr;
+					
+					return Result;
+				},
+				UE::Tasks::Prerequisites(RunnerCompletionEvent),
+				UE::Tasks::ETaskPriority::Inherit,
+				UE::Tasks::EExtendedTaskPriority::Inline);
+	}
+
+
+	// Temporarily make the Image DescCache clear at every image because otherwise it makes some textures 
+	// not evaluate their layout and be of size 0 and 0 lods, making them incorrectly evaluate MipsToSkip
+	static TAutoConsoleVariable<int32> CVarClearImageDescCache(
+		TEXT("mutable.ClearImageDescCache"),
+		1,
+		TEXT("If different than 0, clear the image desc cache at every image."),
+		ECVF_Scalability);
+
+	void System::GetImageDescInline(Instance::ID instanceID, FResourceID ImageId, FImageDesc& OutDesc)
+	{
+		LLM_SCOPE_BYNAME(TEXT("MutableRuntime"));
+		MUTABLE_CPUPROFILER_SCOPE(SystemGetImageDesc);
+
+		OutDesc = FImageDesc();
+
+		// Find the live instance
+		FLiveInstance* pLiveInstance = m_pD->FindLiveInstance(instanceID);
+		check(pLiveInstance);
+		m_pD->WorkingMemoryManager.CurrentInstanceCache = pLiveInstance->Cache;
+
+		OP::ADDRESS RootAddress = GetResourceIDRoot(ImageId);
+
+		const mu::Model* Model = pLiveInstance->Model.Get();
+		const mu::FProgram& program = Model->GetPrivate()->m_program;
+
+		// TODO: It should be possible to reuse this data if cleared in the correct places only, together with m_heapImageDesc.
+		int32 VarValue = CVarClearImageDescCache.GetValueOnAnyThread();
+		if (VarValue != 0)
+		{
+			m_pD->WorkingMemoryManager.CurrentInstanceCache->ClearDescCache();
+		}
+
+		mu::OP_TYPE opType = program.GetOpType(RootAddress);
+		if (GetOpDataType(opType) == DT_IMAGE)
+		{
+			// GetImageDesc may call normal execution paths where meshes are computed.
+			m_pD->WorkingMemoryManager.BeginRunnerThread();
+					
+			int8 executionOptions = 0;
+			TSharedRef<CodeRunner> Runner = CodeRunner::Create(
+					m_pD->Settings, m_pD, EExecutionStrategy::MinimizeMemory, pLiveInstance->Model, pLiveInstance->OldParameters.get(), RootAddress, System::AllLODs, executionOptions, 0, FScheduledOp::EType::ImageDesc);
+			
+			constexpr bool bForceInlineExecution = true;
+			UE::Tasks::FTask CompletionEvent = Runner->StartRun(bForceInlineExecution);
+			check(CompletionEvent.IsCompleted());
+
+			Runner->GetImageDescResult(OutDesc);
+
+			m_pD->WorkingMemoryManager.EndRunnerThread();
+		}
+
+		m_pD->WorkingMemoryManager.CurrentInstanceCache = nullptr;
+	}
+
+
+	UE::Tasks::TTask<FImageDesc> System::GetImageDesc(Instance::ID instanceID, FResourceID ImageId)
+	{
+		LLM_SCOPE_BYNAME(TEXT("MutableRuntime"));
+		MUTABLE_CPUPROFILER_SCOPE(SystemGetImageDesc);
+
+		// Find the live instance
+		FLiveInstance* pLiveInstance = m_pD->FindLiveInstance(instanceID);
+		check(pLiveInstance);
+		m_pD->WorkingMemoryManager.CurrentInstanceCache = pLiveInstance->Cache;
+
+		OP::ADDRESS RootAddress = GetResourceIDRoot(ImageId);
+
+		const mu::Model* Model = pLiveInstance->Model.Get();
+		const mu::FProgram& program = Model->GetPrivate()->m_program;
+
+		// TODO: It should be possible to reuse this data if cleared in the correct places only, together with m_heapImageDesc.
+		int32 VarValue = CVarClearImageDescCache.GetValueOnAnyThread();
+		if (VarValue != 0)
+		{
+			m_pD->WorkingMemoryManager.CurrentInstanceCache->ClearDescCache();
+		}
+
+		m_pD->WorkingMemoryManager.BeginRunnerThread();
+		mu::OP_TYPE OpType = program.GetOpType(RootAddress);
+		if (GetOpDataType(OpType) != DT_IMAGE)
+		{
+			m_pD->WorkingMemoryManager.EndRunnerThread();
+			m_pD->WorkingMemoryManager.CurrentInstanceCache = nullptr;
+
+			return UE::Tasks::MakeCompletedTask<FImageDesc>(); 
+		}
+
+		// GetImageDesc may call normal execution paths where meshes are computed.
+		int8 ExecutionOptions = 0;
+		TSharedRef<CodeRunner> Runner = CodeRunner::Create(
+				m_pD->Settings, m_pD, EExecutionStrategy::MinimizeMemory, pLiveInstance->Model, pLiveInstance->OldParameters.get(), RootAddress, System::AllLODs, ExecutionOptions, 0, FScheduledOp::EType::ImageDesc);
+		
+		constexpr bool bForceInlineExecution = false;
+		UE::Tasks::FTask RunnerCompletionEvent = Runner->StartRun(bForceInlineExecution);
+		
+		return UE::Tasks::Launch(TEXT("System::GetImageDescResultTask"),
+				[SystemPrivate = m_pD, Runner]() -> FImageDesc
+				{
+					FImageDesc Result;
+					Runner->GetImageDescResult(Result);
+
+					SystemPrivate->WorkingMemoryManager.EndRunnerThread();
+					SystemPrivate->WorkingMemoryManager.CurrentInstanceCache = nullptr;
+
+					return Result;
+				},
+				UE::Tasks::Prerequisites(RunnerCompletionEvent),
+				UE::Tasks::ETaskPriority::Inherit,
+				UE::Tasks::EExtendedTaskPriority::Inline);
+
+	}
+
+
+    Ptr<const Mesh> System::GetMeshInline(Instance::ID instanceID, FResourceID MeshId)
+    {
+		LLM_SCOPE_BYNAME(TEXT("MutableRuntime"));
+		MUTABLE_CPUPROFILER_SCOPE(SystemGetMesh);
+
+		Ptr<const Mesh> Result;
+
+		// Find the live instance
+		FLiveInstance* pLiveInstance = m_pD->FindLiveInstance(instanceID);
+		check(pLiveInstance);
+		m_pD->WorkingMemoryManager.CurrentInstanceCache = pLiveInstance->Cache;
+
+		OP::ADDRESS RootAddress = GetResourceIDRoot(MeshId);
+		Result = m_pD->BuildMesh(pLiveInstance->Model, pLiveInstance->OldParameters.get(), RootAddress);
+
+		// If the mesh is null it means empty, but we still need to return a valid one
+		if (!Result)
+		{
+			Result = new Mesh();
+		}
+
+		m_pD->WorkingMemoryManager.CurrentInstanceCache = nullptr;
+		return Result;
+	}
+
+	UE::Tasks::TTask<Ptr<const Mesh>> System::GetMesh(Instance::ID instanceID, FResourceID MeshId)
+    {
+		LLM_SCOPE_BYNAME(TEXT("MutableRuntime"));
+		MUTABLE_CPUPROFILER_SCOPE(SystemGetImage);
+
+		Ptr<const Mesh> ResultMesh;
+
+		// Find the live instance
+		FLiveInstance* pLiveInstance = m_pD->FindLiveInstance(instanceID);
+		check(pLiveInstance);
+		m_pD->WorkingMemoryManager.CurrentInstanceCache = pLiveInstance->Cache;
+		m_pD->WorkingMemoryManager.BeginRunnerThread();
+		
+		OP::ADDRESS RootAddress = GetResourceIDRoot(MeshId);
+			
+		mu::OP_TYPE OpType = pLiveInstance->Model->GetPrivate()->m_program.GetOpType(RootAddress);
+		if (GetOpDataType(OpType) != DT_MESH)
+		{
+			m_pD->WorkingMemoryManager.EndRunnerThread();
+			m_pD->WorkingMemoryManager.CurrentInstanceCache = nullptr;
+
+			return UE::Tasks::MakeCompletedTask<Ptr<const Mesh>>(new Mesh());
+		}
+		
+		TSharedRef<CodeRunner> Runner = CodeRunner::Create(
+				m_pD->Settings, m_pD, EExecutionStrategy::MinimizeMemory, pLiveInstance->Model, pLiveInstance->OldParameters.get(), RootAddress, System::AllLODs, 0, 0, FScheduledOp::EType::Full);
+		
+		constexpr bool bForceInlineExecution = false;
+		UE::Tasks::FTask RunnerCompletionEvent = Runner->StartRun(bForceInlineExecution);
+
+		return UE::Tasks::Launch(TEXT("System::GetMeshResultTask"),
+				[SystemPrivate = m_pD, Runner, RootAddress]() -> Ptr<const Mesh>
+				{
+					Ptr<const Mesh> Result;
+
+					SystemPrivate->bUnrecoverableError = Runner->bUnrecoverableError;
+					if (!Runner->bUnrecoverableError)
+					{
+						Result = SystemPrivate->WorkingMemoryManager.LoadMesh(FCacheAddress(RootAddress, 0, 0), true);
+					}
+					
+					SystemPrivate->WorkingMemoryManager.EndRunnerThread();
+					SystemPrivate->WorkingMemoryManager.CurrentInstanceCache = nullptr;
+					
+					return Result;
+				},
+				UE::Tasks::Prerequisites(RunnerCompletionEvent),
+				UE::Tasks::ETaskPriority::Inherit,
+				UE::Tasks::EExtendedTaskPriority::Inline);
+	}
+
+
+    //---------------------------------------------------------------------------------------------
+    void System::EndUpdate( Instance::ID instanceID )
+    {
+		LLM_SCOPE_BYNAME(TEXT("MutableRuntime"));
+		MUTABLE_CPUPROFILER_SCOPE(EndUpdate);
+
+		FLiveInstance* pLiveInstance = m_pD->FindLiveInstance(instanceID);
+		if (pLiveInstance)
+		{
+			pLiveInstance->Instance = nullptr;
+			
+			// Debug check to see if we managed the op-hit-counts correctly
+			pLiveInstance->Cache->CheckHitCountsCleared();
+
+			m_pD->WorkingMemoryManager.CurrentInstanceCache = pLiveInstance->Cache;
+
+			// We don't want to clear the cache layer 1 because it contains data that can be useful for a 
+			// future update (same states, just runtime parameters changed).
+			//m_pD->WorkingMemoryManager.ClearCacheLayer1();
+
+			// We need to clear the layer 0 cache, because it contains data that is only valid for the current 
+			// parameter values (unless it is data marked as state cache)
+			m_pD->WorkingMemoryManager.ClearCacheLayer0();
+
+			m_pD->WorkingMemoryManager.CurrentInstanceCache = nullptr;
+		}
+
+		// Reduce the cache until it fits the limit.
+		m_pD->WorkingMemoryManager.EnsureBudgetBelow(0);
+
+		// If we don't constrain the memory budget, free the pooled images or they may pile up.
+		if (m_pD->WorkingMemoryManager.BudgetBytes == 0)
+		{
+			m_pD->WorkingMemoryManager.PooledImages.Empty();
+		}
+
+	}
+
+
+    //---------------------------------------------------------------------------------------------
+    void System::ReleaseInstance( Instance::ID instanceID )
+    {
+		LLM_SCOPE_BYNAME(TEXT("MutableRuntime"));
+		MUTABLE_CPUPROFILER_SCOPE(ReleaseInstance);
+
+		for (int32 Index = 0; Index < m_pD->WorkingMemoryManager.LiveInstances.Num(); ++Index)
+		{
+			FLiveInstance& Instance = m_pD->WorkingMemoryManager.LiveInstances[Index];
+			if (Instance.InstanceID == instanceID)
+			{
+				// Make sure all the resources cached in the instance are removed from the tracking list
+				for (const FProgramCache::TResourceResult<Image>& Data : Instance.Cache->ImageResults)
+				{
+					if (Data.Value)
+					{
+						m_pD->WorkingMemoryManager.CacheResources.Remove(Data.Value);
+					}
+				}
+
+				for (const FProgramCache::TResourceResult<Mesh>& Data : Instance.Cache->MeshResults)
+				{
+					if (Data.Value)
+					{
+						m_pD->WorkingMemoryManager.CacheResources.Remove(Data.Value);
+					}
+				}
+
+				m_pD->WorkingMemoryManager.LiveInstances.RemoveAtSwap(Index);
+				break;
+			}
+		}
+
+ 		int32 Removed = m_pD->WorkingMemoryManager.LiveInstances.RemoveAllSwap(
+			[instanceID](const FLiveInstance& Instance)
+			{
+				return (Instance.InstanceID == instanceID);
+			});
+
+		TRACE_COUNTER_SET(MutableRuntime_LiveInstances, m_pD->WorkingMemoryManager.LiveInstances.Num());
+
+	}
+
+
+    //---------------------------------------------------------------------------------------------
+    class RelevantParameterVisitor : public UniqueDiscreteCoveredCodeVisitor<>
+    {
+    public:
+
+        RelevantParameterVisitor
+            (
+                System::Private* InSystem,
+				const TSharedPtr<const Model>& InModel,
+                const Ptr<const Parameters>& InParams,
+                bool* InFlags
+            )
+            : UniqueDiscreteCoveredCodeVisitor<>( InSystem, InModel, InParams, System::AllLODs )
+        {
+            Flags = InFlags;
+
+            FMemory::Memset( Flags, 0, sizeof(bool)*InParams->GetCount() );
+
+            OP::ADDRESS at = InModel->GetPrivate()->m_program.m_states[0].m_root;
+
+            Run( at );
+        }
+
+
+        bool Visit( OP::ADDRESS at, FProgram& program ) override
+        {
+            switch ( program.GetOpType(at) )
+            {
+            case OP_TYPE::BO_PARAMETER:
+            case OP_TYPE::NU_PARAMETER:
+            case OP_TYPE::SC_PARAMETER:
+            case OP_TYPE::CO_PARAMETER:
+            case OP_TYPE::PR_PARAMETER:
+            case OP_TYPE::IM_PARAMETER:
+            {
+				OP::ParameterArgs args = program.GetOpArgs<OP::ParameterArgs>(at);
+				OP::ADDRESS ParamIndex = args.variable;
+                Flags[ParamIndex] = true;
+                break;
+            }
+
+            default:
+                break;
+
+            }
+
+            return UniqueDiscreteCoveredCodeVisitor<>::Visit( at, program );
+        }
+
+    private:
+
+        //! Non-owned result buffer
+        bool* Flags;
+    };
+
+
+    //---------------------------------------------------------------------------------------------
+    void System::GetParameterRelevancy( Instance::ID InstanceID,
+                                        const Ptr<const Parameters>& Parameters,
+                                        bool* Flags )
+    {
+		LLM_SCOPE_BYNAME(TEXT("MutableRuntime"));
+
+		// Find the live instance
+		FLiveInstance* pLiveInstance = m_pD->FindLiveInstance(InstanceID);
+		check(pLiveInstance);
+		m_pD->WorkingMemoryManager.CurrentInstanceCache = pLiveInstance->Cache;
+		
+		RelevantParameterVisitor visitor( m_pD, pLiveInstance->Model, Parameters, Flags );
+
+		m_pD->WorkingMemoryManager.CurrentInstanceCache = nullptr;
+    }
+
+
+    //---------------------------------------------------------------------------------------------
+	inline FLiveInstance* System::Private::FindLiveInstance(Instance::ID id)
+	{
+		for (int32 i = 0; i < WorkingMemoryManager.LiveInstances.Num(); ++i)
+		{
+			if (WorkingMemoryManager.LiveInstances[i].InstanceID == id)
+			{
+				return &WorkingMemoryManager.LiveInstances[i];
+			}
+		}
+		return nullptr;
+	}
+
+
+	//---------------------------------------------------------------------------------------------
+	bool System::Private::CheckUpdatedParameters( const FLiveInstance* LiveInstance,
+                                 const Ptr<const Parameters>& Params,
+                                 uint64& UpdatedParameters)
+    {
+        bool bFullBuild = false;
+
+		if (!LiveInstance->OldParameters)
+		{
+			UpdatedParameters = AllParametersMask;
+			return true;
+		}
+
+        // check what parameters have changed
+		UpdatedParameters = 0;
+        const FProgram& program = LiveInstance->Model->GetPrivate()->m_program;
+        const TArray<int>& runtimeParams = program.m_states[ LiveInstance->State ].m_runtimeParameters;
+
+        check( Params->GetCount() == (int)program.m_parameters.Num() );
+        check( !LiveInstance->OldParameters
+			||
+			Params->GetCount() == LiveInstance->OldParameters->GetCount() );
+
+        for ( int32 p=0; p<program.m_parameters.Num() && !bFullBuild; ++p )
+        {
+            bool isRuntime = runtimeParams.Contains( p );
+            bool changed = !Params->HasSameValue( p, LiveInstance->OldParameters, p );
+
+            if (changed && isRuntime)
+            {
+				uint64 runtimeIndex = runtimeParams.IndexOfByKey(p);
+				UpdatedParameters |= uint64(1) << runtimeIndex;
+            }
+            else if (changed)
+            {
+                // A non-runtime parameter has changed, we need a full build.
+                // TODO: report, or log somehow.
+				bFullBuild = true;
+                UpdatedParameters = AllParametersMask;
+            }
+        }
+
+        return bFullBuild;
+    }
+
+
+	//---------------------------------------------------------------------------------------------
+	void System::Private::BeginBuild(const TSharedPtr<const Model>& InModel)
+	{
+		// We don't have a FLiveInstance, let's create the memory
+		// \TODO: There is no clear moment to remove this... EndBuild?
+		WorkingMemoryManager.CurrentInstanceCache = MakeShared<FProgramCache>();
+		WorkingMemoryManager.CurrentInstanceCache->Init(InModel->GetPrivate()->m_program.m_opAddress.Num());
+
+		// Ensure the model cache has been created
+		WorkingMemoryManager.FindOrAddModelCache(InModel);
+
+		PrepareCache(InModel.Get(), -1);
+	}
+
+
+	//---------------------------------------------------------------------------------------------
+	void System::Private::EndBuild()
+	{
+		WorkingMemoryManager.CurrentInstanceCache = nullptr;
+	}
+
+
+	//---------------------------------------------------------------------------------------------
+	void System::Private::RunCode(const TSharedPtr<const Model>& InModel,
+		const Parameters* InParameters, OP::ADDRESS InCodeRoot, uint32 InLODs, uint8 executionOptions, int32 InImageLOD)
+	{
+		TSharedRef<CodeRunner> Runner = CodeRunner::Create(Settings, this, EExecutionStrategy::MinimizeMemory, InModel, InParameters, InCodeRoot, InLODs,
+			executionOptions, InImageLOD, FScheduledOp::EType::Full);
+		
+		constexpr bool bForceInlineExecutution = true;
+		UE::Tasks::FTask RunnerCompletionEvent = Runner->StartRun(bForceInlineExecutution);
+		check(RunnerCompletionEvent.IsCompleted());
+
+		bUnrecoverableError = Runner->bUnrecoverableError;
+	}
+
+
+	//---------------------------------------------------------------------------------------------
+	bool System::Private::BuildBool(const TSharedPtr<const Model>& pModel, const Parameters* Params, OP::ADDRESS at)
+	{
+		WorkingMemoryManager.BeginRunnerThread();
+
+		RunCode(pModel, Params, at);
+
+		bool bResult = false;
+		if (!bUnrecoverableError)
+		{
+			bResult = WorkingMemoryManager.CurrentInstanceCache->GetBool(FCacheAddress(at, 0, 0));
+		}
+
+		WorkingMemoryManager.EndRunnerThread();
+
+		return bResult;
+	}
+
+
+	//---------------------------------------------------------------------------------------------
+	float System::Private::BuildScalar(const TSharedPtr<const Model>& pModel, const Parameters* Params, OP::ADDRESS at)
+	{
+		WorkingMemoryManager.BeginRunnerThread();
+
+		RunCode(pModel, Params, at);
+
+		float Result = 0.0f;		
+		if (!bUnrecoverableError)
+		{
+			Result = WorkingMemoryManager.CurrentInstanceCache->GetScalar(FCacheAddress(at, 0, 0));
+		}
+
+		WorkingMemoryManager.EndRunnerThread();
+
+		return Result;
+	}
+
+
+	//---------------------------------------------------------------------------------------------
+	int32 System::Private::BuildInt(const TSharedPtr<const Model>& pModel, const Parameters* Params, OP::ADDRESS at)
+	{
+		WorkingMemoryManager.BeginRunnerThread();
+
+		RunCode(pModel, Params, at);
+
+		int32 Result = 0;
+		if (!bUnrecoverableError)
+		{
+			Result = WorkingMemoryManager.CurrentInstanceCache->GetInt(FCacheAddress(at, 0, 0));
+		}
+
+		WorkingMemoryManager.EndRunnerThread();
+
+		return Result;
+	}
+
+
+	//---------------------------------------------------------------------------------------------
+	FVector4f System::Private::BuildColour(const TSharedPtr<const Model>& pModel, const Parameters* Params, OP::ADDRESS at)
+	{
+		WorkingMemoryManager.BeginRunnerThread();
+
+		FVector4f Result(0,0,0,1);
+
+		mu::OP_TYPE opType = pModel->GetPrivate()->m_program.GetOpType(at);
+		if (GetOpDataType(opType) == DT_COLOUR)
+		{
+			RunCode(pModel, Params, at);
+			if (!bUnrecoverableError)
+			{
+				Result = WorkingMemoryManager.CurrentInstanceCache->GetColour(FCacheAddress(at, 0, 0));
+			}
+		}
+
+		WorkingMemoryManager.EndRunnerThread();
+
+		return Result;
+	}
+
+	
+	//---------------------------------------------------------------------------------------------
+	FProjector System::Private::BuildProjector(const TSharedPtr<const Model>& pModel, const Parameters* Params, OP::ADDRESS at)
+	{
+		WorkingMemoryManager.BeginRunnerThread();
+
+    	RunCode(pModel, Params, at);
+
+		FProjector Result;
+		if (!bUnrecoverableError)
+		{
+			Result = WorkingMemoryManager.CurrentInstanceCache->GetProjector(FCacheAddress(at, 0, 0));
+		}
+
+		WorkingMemoryManager.EndRunnerThread();
+
+		return Result;
+	}
+
+	
+	Ptr<const Image> System::Private::BuildImage(const TSharedPtr<const Model>& pModel, const Parameters* Params, OP::ADDRESS at, int32 MipsToSkip, int32 InImageLOD)
+	{
+		WorkingMemoryManager.BeginRunnerThread();
+
+		Ptr<const Image> Result;
+
+		mu::OP_TYPE opType = pModel->GetPrivate()->m_program.GetOpType(at);
+		if (GetOpDataType(opType) == DT_IMAGE)
+		{
+			RunCode(pModel, Params, at, System::AllLODs, uint8(MipsToSkip), InImageLOD);
+			if (!bUnrecoverableError)
+			{
+				Result = WorkingMemoryManager.LoadImage(FCacheAddress(at, 0, MipsToSkip), true);
+			}			
+		}
+
+		WorkingMemoryManager.EndRunnerThread();
+
+		return Result;
+	}
+
+
+	Ptr<const Mesh> System::Private::BuildMesh(const TSharedPtr<const Model>& pModel, const Parameters* Params, OP::ADDRESS at)
+	{
+		WorkingMemoryManager.BeginRunnerThread();
+
+		Ptr<const Mesh> Result;
+
+		mu::OP_TYPE opType = pModel->GetPrivate()->m_program.GetOpType(at);
+		if (GetOpDataType(opType) == DT_MESH)
+		{
+			RunCode(pModel, Params, at);
+			if (!bUnrecoverableError)
+			{
+				Result = WorkingMemoryManager.LoadMesh(FCacheAddress(at, 0, 0), true);
+			}	
+		}
+
+		WorkingMemoryManager.EndRunnerThread();
+
+		return Result;
+	}
+
+
+	Ptr<const Instance> System::Private::BuildInstance(const TSharedPtr<const Model>& pModel, const Parameters* Params, OP::ADDRESS at)
+	{
+		WorkingMemoryManager.BeginRunnerThread();
+
+		Ptr<const Instance> Result;
+
+		mu::OP_TYPE opType = pModel->GetPrivate()->m_program.GetOpType(at);
+		if (GetOpDataType(opType) == DT_INSTANCE)
+		{
+			RunCode(pModel, Params, at);
+			if (!bUnrecoverableError)
+			{
+				Result = WorkingMemoryManager.CurrentInstanceCache->GetInstance(FCacheAddress(at, 0, 0));
+			}
+		}
+
+		WorkingMemoryManager.EndRunnerThread();
+
+		return Result;
+	}
+
+
+	//---------------------------------------------------------------------------------------------
+	Ptr<const Layout> System::Private::BuildLayout(const TSharedPtr<const Model>& pModel, const Parameters* Params, OP::ADDRESS at)
+	{
+		WorkingMemoryManager.BeginRunnerThread();
+
+		Ptr<const Layout> Result;
+
+		if (pModel->GetPrivate()->m_program.m_states[0].m_root)
+		{
+			mu::OP_TYPE opType = pModel->GetPrivate()->m_program.GetOpType(at);
+			if (GetOpDataType(opType) == DT_LAYOUT)
+			{
+				RunCode(pModel, Params, at);
+				if (!bUnrecoverableError)
+				{
+					Result = WorkingMemoryManager.CurrentInstanceCache->GetLayout(FCacheAddress(at, 0, 0));
+				}
+			}
+		}
+
+		WorkingMemoryManager.EndRunnerThread();
+
+		return Result;
+	}
+
+
+	//---------------------------------------------------------------------------------------------
+	Ptr<const String> System::Private::BuildString(const TSharedPtr<const Model>& pModel, const Parameters* Params, OP::ADDRESS at)
+	{
+		WorkingMemoryManager.BeginRunnerThread();
+
+		Ptr<const String> Result;
+
+		if (pModel->GetPrivate()->m_program.m_states[0].m_root)
+		{
+			mu::OP_TYPE opType = pModel->GetPrivate()->m_program.GetOpType(at);
+			if (GetOpDataType(opType) == DT_STRING)
+			{
+				RunCode(pModel, Params, at);
+				if (!bUnrecoverableError)
+				{
+					Result = WorkingMemoryManager.CurrentInstanceCache->GetString(FCacheAddress(at, 0, 0));
+				}
+			}
+		}
+
+		WorkingMemoryManager.EndRunnerThread();
+
+		return Result;
+	}
+
+
+	//---------------------------------------------------------------------------------------------
+	void System::Private::PrepareCache( const Model* InModel, int32 InState)
+	{
+		MUTABLE_CPUPROFILER_SCOPE(PrepareCache);
+
+		FProgram& Program = InModel->GetPrivate()->m_program;
+		int32 OpCount = Program.m_opAddress.Num();
+		WorkingMemoryManager.CurrentInstanceCache->Init(OpCount);
+
+		// Clear cache flags of existing data
+		CodeContainer<FProgramCache::FOpExecutionData>::iterator It = WorkingMemoryManager.CurrentInstanceCache->OpExecutionData.begin();
+		for (; It.IsValid(); ++It)
+		{
+			(*It).OpHitCount = 0; // This should already be 0, but just in case.
+			(*It).IsCacheLocked = false;
+		}
+
+		// Mark the resources that have to be cached to update the instance in this state
+		if (InState >= 0 && InState< Program.m_states.Num())
+		{
+			const FProgram::FState& State = Program.m_states[InState];
+			for (uint32 Address : State.m_updateCache)
+			{
+				WorkingMemoryManager.CurrentInstanceCache->SetForceCached(Address);
+			}
+		}
+	}
+
+
+	//---------------------------------------------------------------------------------------------
+	void FWorkingMemoryManager::LogWorkingMemory(const CodeRunner* CurrentRunner) const
+	{
+#if !(UE_BUILD_SHIPPING || UE_BUILD_TEST)
+
+		MUTABLE_CPUPROFILER_SCOPE(LogWorkingMemory);
+
+		// For now, we calculate these for every log. We will later track on resource creation, destruction or state change.
+		// All resource memory is tracked by the memory allocator, but that does not give information about where the memory is
+		// located. Keep the localized memory computation for now.   
+		const uint32 RomBytes = GetRomBytes();
+		const uint32 CacheBytes = GetCacheBytes();
+		const uint32 TrackedCacheBytes = GetTrackedCacheBytes();
+		const uint32 PoolBytes = GetPooledBytes();
+		const uint32 TempBytes = GetTempBytes();
+
+		// Get allocator counters.
+		const SSIZE_T ImageAllocBytes	 = MemoryCounters::FImageMemoryCounter::Counter.load(std::memory_order_relaxed);
+		const SSIZE_T MeshAllocBytes     = MemoryCounters::FMeshMemoryCounter::Counter.load(std::memory_order_relaxed);
+		const SSIZE_T StreamAllocBytes   = MemoryCounters::FStreamingMemoryCounter::Counter.load(std::memory_order_relaxed);
+		const SSIZE_T InternalAllocBytes = MemoryCounters::FInternalMemoryCounter::Counter.load(std::memory_order_relaxed);
+
+		SSIZE_T TotalBytes = ImageAllocBytes + MeshAllocBytes + StreamAllocBytes + InternalAllocBytes;
+
+		UE_LOG(LogMutableCore, Log, 
+			TEXT("Mem KB: ImageAlloc %7d, MeshAlloc %7d, StreamAlloc %7d, InternalAlloc %7d,  AllocTotal %7d / %7d. \
+				  Resources MemLoc: Temp %7d, Pool %7d, Cache0+1 %7d (%7d), Rom %7d."),
+			ImageAllocBytes/1024, MeshAllocBytes/1024, StreamAllocBytes/1024, InternalAllocBytes/1024, TotalBytes/1024, BudgetBytes/1024,  
+			TempBytes/1024, PoolBytes/1024, CacheBytes/1024, TrackedCacheBytes/1024, RomBytes/1024);
+
+		TRACE_COUNTER_SET(MutableRuntime_MemTemp,	  TempBytes);
+		TRACE_COUNTER_SET(MutableRuntime_MemPool,     PoolBytes);
+		TRACE_COUNTER_SET(MutableRuntime_MemCache,    TrackedCacheBytes);
+		TRACE_COUNTER_SET(MutableRuntime_MemRom,      RomBytes);
+		TRACE_COUNTER_SET(MutableRuntime_MemInternal, InternalAllocBytes);
+		TRACE_COUNTER_SET(MutableRuntime_MemMesh,     MeshAllocBytes);
+		TRACE_COUNTER_SET(MutableRuntime_MemImage,    ImageAllocBytes);
+		TRACE_COUNTER_SET(MutableRuntime_MemStream,   StreamAllocBytes);
+		TRACE_COUNTER_SET(MutableRuntime_MemTotal,    TotalBytes);
+		TRACE_COUNTER_SET(MutableRuntime_MemBudget,   BudgetBytes);
+#endif
+	}
+
+
+
+	//---------------------------------------------------------------------------------------------
+	FWorkingMemoryManager::FModelCacheEntry* FWorkingMemoryManager::FindModelCache(const Model* InModel)
+	{
+		check(InModel);
+
+		for (FModelCacheEntry& c : CachePerModel)
+		{
+			TSharedPtr<const Model> pCandidate = c.Model.Pin();
+			if (pCandidate)
+			{
+				if (pCandidate.Get() == InModel)
+				{
+					return &c;
+				}
+			}
+		}
+
+		return nullptr;
+	}
+
+    //---------------------------------------------------------------------------------------------
+	FWorkingMemoryManager::FModelCacheEntry& FWorkingMemoryManager::FindOrAddModelCache(const TSharedPtr<const Model>& InModel)
+    {
+        check(InModel);
+
+		// First clean stray data for models that may have been unloaded.
+		for (int32 CacheIndex=0; CacheIndex<CachePerModel.Num(); )
+		{
+			if (!CachePerModel[CacheIndex].Model.Pin())
+			{
+				CachePerModel.RemoveAtSwap(CacheIndex);
+			}
+			else
+			{
+				++CacheIndex;
+			}
+		}
+
+
+		FModelCacheEntry* ExistingCache = FindModelCache(InModel.Get());
+		if (ExistingCache)
+		{
+			return *ExistingCache;
+		}        
+
+        // Not found. Add new
+		FModelCacheEntry n;
+        n.Model = TWeakPtr<const Model>(InModel);
+		n.PendingOpsPerRom.SetNum(InModel->GetPrivate()->m_program.m_roms.Num());
+    	n.RomWeights.SetNumZeroed(InModel->GetPrivate()->m_program.m_roms.Num());
+    	CachePerModel.Add(n);
+        return CachePerModel.Last();
+    }
+
+
+	//---------------------------------------------------------------------------------------------
+	int64 FWorkingMemoryManager::GetCurrentMemoryBytes() const
+	{
+		MUTABLE_CPUPROFILER_SCOPE(GetCurrentMemoryBytes);
+
+		SSIZE_T TotalBytes = MemoryCounters::FImageMemoryCounter::Counter.load(std::memory_order_relaxed) + 
+						     MemoryCounters::FMeshMemoryCounter::Counter.load(std::memory_order_relaxed) +
+							 MemoryCounters::FStreamingMemoryCounter::Counter.load(std::memory_order_relaxed) +
+							 MemoryCounters::FInternalMemoryCounter::Counter.load(std::memory_order_relaxed);
+
+		return TotalBytes;
+	}
+
+
+	//---------------------------------------------------------------------------------------------
+	bool FWorkingMemoryManager::IsMemoryBudgetFull() const
+	{
+		// If we have 0 budget it means we have unlimited budget
+		if (BudgetBytes == 0)
+		{
+			return false;
+		}
+
+		uint64 CurrentBytes = GetCurrentMemoryBytes();
+		uint64 BudgetThresholdBytes = (BudgetBytes * 9) / 10;
+		
+		return CurrentBytes > BudgetThresholdBytes;
+	}
+
+	namespace Private
+	{
+		enum class EBudgetBelowSearchFlags : uint8
+		{
+			None = 0,
+			Keep = 1 << 1,
+			Visited = 1 << 2,
+			FirstOccurance = 1 << 3
+		};
+
+		ENUM_CLASS_FLAGS(EBudgetBelowSearchFlags);
+	}
+
+    //---------------------------------------------------------------------------------------------
+	bool FWorkingMemoryManager::EnsureBudgetBelow( uint64 AdditionalMemory )
+    {
+		MUTABLE_CPUPROFILER_SCOPE(EnsureBudgetBelow);
+
+		// If we have 0 budget it means we have unlimited budget
+		if (BudgetBytes == 0)
+		{
+			return true;
+		}		
+
+		int64 TotalBytes = GetCurrentMemoryBytes();
+
+		// Add the extra memory that we are trying to allocate when we return.
+		TotalBytes += AdditionalMemory;
+
+
+        bool bFinished = TotalBytes <= BudgetBytes;
+
+		// Try to free pooled resources first
+		if (!bFinished)
+		{
+			MUTABLE_CPUPROFILER_SCOPE(EnsureBudgetBelow_FreePooled);
+			while (PooledImages.Num() && !bFinished)
+			{
+				// TODO: Actually advancing index if possible after swap may be better to remove the oldest in the pool first.
+				int32 PooledResourceSize = PooledImages[0]->GetDataSize();
+				TotalBytes -= PooledResourceSize;				
+				PooledImages.RemoveAtSwap(0);
+				bFinished = TotalBytes <= BudgetBytes;
+			}
+		}
+		
+		// Try to free loaded roms
+		if (!bFinished)
+		{
+			MUTABLE_CPUPROFILER_SCOPE(EnsureBudgetBelow_FreeRoms);
+
+			struct FRomRef
+			{
+				const Model* Model = nullptr;
+				int32 RomIndex = 0;
+			};
+
+			TArray<TPair<float, FRomRef>> Candidates;
+			Candidates.Reserve(512);
+
+			for (FModelCacheEntry& ModelCache : CachePerModel)
+			{
+				TSharedPtr<const Model> CacheModel = ModelCache.Model.Pin();
+				if (CacheModel)
+				{
+					mu::FProgram& Program = CacheModel->GetPrivate()->m_program;
+					check(ModelCache.RomWeights.Num() == Program.m_roms.Num());
+	
+					check(Program.LoadedMemTrackedRoms.GetMaxIndex() <= Program.m_roms.Num());
+
+					for (TSparseArray<uint8>::TConstIterator Iter = Program.LoadedMemTrackedRoms.CreateConstIterator(); Iter; ++Iter)
+					{
+						const int32 RomIndex = Iter.GetIndex();
+
+						const FRomData& Rom = Program.m_roms[RomIndex];
+						check(Program.IsRomLoaded(RomIndex));
+						check(Rom.ResourceType == (uint32)*Iter);
+
+						// We cannot unload a rom if some operation is expecting it.
+						const bool bIsRomLocked = ModelCache.PendingOpsPerRom.IsValidIndex(RomIndex) && 
+												  ModelCache.PendingOpsPerRom[RomIndex] > 0;
+						if (!bIsRomLocked)
+						{
+							constexpr float FactorWeight = 100.0f;
+							constexpr float FactorTime = -1.0f;
+							float Priority = FactorWeight * float(ModelCache.RomWeights[RomIndex].Get<0>()) +
+											 FactorTime * float((RomTick - ModelCache.RomWeights[RomIndex].Get<1>()));
+
+							Candidates.Emplace(Priority, FRomRef{ CacheModel.Get(), RomIndex });
+						}
+					}
+				}
+			}
+
+			// Don't sort all candidates, make it a heap in O(N) time. We may not need to visit all elements.
+			auto CompareCandidates = [](const TPair<float, FRomRef>& A, const TPair<float, FRomRef>& B) { return A.Key < B.Key; };
+			Candidates.Heapify(CompareCandidates);
+
+			while (!bFinished && Candidates.Num())
+			{
+				MUTABLE_CPUPROFILER_SCOPE(EnsureBudgetBelow_UnloadRom);
+
+				TPair<float, FRomRef> Candidate;
+				Candidates.HeapPop(Candidate, CompareCandidates, EAllowShrinking::No);
+
+				// UE_LOG(LogMutableCore,Log, "Unloading rom because of memory budget: %d.", lowestPriorityRom);
+				int32 UnloadedSize = Candidate.Value.Model->GetPrivate()->m_program.UnloadRom(Candidate.Value.RomIndex);
+				TotalBytes -= UnloadedSize;
+				bFinished = TotalBytes <= BudgetBytes;
+			}
+		}
+	
+		// Try to free cache 1 memory
+		if (!bFinished)
+		{
+			MUTABLE_CPUPROFILER_SCOPE(EnsureBudgetBelow_FreeCached);
+
+			TSet<const Resource*> RemovedResources;
+			RemovedResources.Reserve(1024);
+			
+			// From other live instances first
+			for (const FLiveInstance& Instance : LiveInstances)
+			{
+				if (Instance.Cache == CurrentInstanceCache)
+				{
+					// Ignore the current live instance.
+					continue;
+				}
+
+				// Gather all data in the cache for this instance
+				{
+					MUTABLE_CPUPROFILER_SCOPE(EnsureBudgetBelow_FreeCached_GatherAndRemove_Other);
+
+					RemovedResources.Reset();
+					int32 ImageRemovedBytes = 0;
+					for (const FProgramCache::TResourceResult<Image>& ImageResult : Instance.Cache->ImageResults)
+					{
+						if (!ImageResult.Value)
+						{
+							continue;
+						}
+
+						if (!RemovedResources.Find(ImageResult.Value.get()))
+						{
+							ImageRemovedBytes += ImageResult.Value->GetDataSize();
+							bFinished = TotalBytes - ImageRemovedBytes <= BudgetBytes;
+
+							RemovedResources.Add(ImageResult.Value.get());
+							CacheResources.Remove(ImageResult.Value);
+						}
+
+						if (bFinished)
+						{
+							break;
+						}
+					}
+
+					if (ImageRemovedBytes > 0)
+					{
+						for (const FProgramCache::TResourceResult<Image>& ImageResult : Instance.Cache->ImageResults)
+						{
+							if (RemovedResources.Find(ImageResult.Value.get()))
+							{
+								Instance.Cache->SetUnused(Instance.Cache->OpExecutionData[ImageResult.OpAddress]);
+							}
+						}
+					}
+
+					TotalBytes -= ImageRemovedBytes;
+					if (bFinished)
+					{
+						break;
+					}
+
+					int32 MeshRemovedBytes = 0;
+					RemovedResources.Reset();
+					for (const FProgramCache::TResourceResult<Mesh>& MeshResult : Instance.Cache->MeshResults)
+					{
+						if (!MeshResult.Value)
+						{
+							continue;
+						}
+
+						if (!RemovedResources.Find(MeshResult.Value.get()))
+						{
+							MeshRemovedBytes = MeshResult.Value->GetDataSize();
+							bFinished = TotalBytes - MeshRemovedBytes <= BudgetBytes;
+
+							RemovedResources.Add(MeshResult.Value.get());
+							CacheResources.Remove(MeshResult.Value);
+						}
+
+						if (bFinished)
+						{
+							break;
+						}
+					}
+
+					if (MeshRemovedBytes > 0)
+					{
+						for (const FProgramCache::TResourceResult<Mesh>& MeshResult : Instance.Cache->MeshResults)
+						{
+							if (RemovedResources.Find(MeshResult.Value.get()))
+							{
+								Instance.Cache->SetUnused(Instance.Cache->OpExecutionData[MeshResult.OpAddress]);
+							}
+						}
+					}
+
+					TotalBytes -= MeshRemovedBytes;
+					if (bFinished)
+					{
+						break;
+					}
+				}
+			}
+		}
+
+		// From the current live instances. It is more involved: we have to make sure any data we want to free is not also
+		// in any cache (0 or 1) position with hit-count > 0.
+		if (!bFinished && CurrentInstanceCache)
+		{
+			MUTABLE_CPUPROFILER_SCOPE(EnsureBudgetBelow_FreeCached_Current);
+			// This removes images first, not sure if this can become a problem.
+			
+			using namespace Private;
+			using ESearchFlags = EBudgetBelowSearchFlags;
+			
+			auto SearchResourcesToRemove = 
+			[
+				TotalBytes = TotalBytes, 
+				BudgetBytes = BudgetBytes, 
+				CurrentInstanceCache = CurrentInstanceCache.Get()
+			](const auto& ResourceRange, TArray<ESearchFlags>& SearchFlags) -> int32
+			{
+				MUTABLE_CPUPROFILER_SCOPE(EnsureBudgetBelow_FreeCached_Current_SearchResources);
+
+				int32 RemovedBytes = 0;
+
+				const int32 NumResources = ResourceRange.Num();
+				for (int32 ResourceIndex = 0; ResourceIndex < NumResources; ++ResourceIndex)
+				{
+					// Check if visited.
+					if (EnumHasAnyFlags(SearchFlags[ResourceIndex], ESearchFlags::Visited))
+					{
+						continue;
+					}
+	
+					// Null values will not have any flag set.
+					if (!ResourceRange[ResourceIndex].Value)
+					{
+						continue;
+					}
+
+					SearchFlags[ResourceIndex] = (ESearchFlags::Visited | ESearchFlags::FirstOccurance);
+					
+					if (CurrentInstanceCache->OpExecutionData[ResourceRange[ResourceIndex].OpAddress].OpHitCount > 0)
+					{
+						// Mark all occurences as visited to keep. 
+						for (int32 I = ResourceIndex; I < NumResources; ++I)
+						{
+							if (ResourceRange[I].Value == ResourceRange[ResourceIndex].Value)
+							{
+								EnumAddFlags(SearchFlags[I], (ESearchFlags::Visited | ESearchFlags::Keep));
+							}
+						}
+					}
+					else
+					{
+						int32 I = ResourceIndex + 1;
+						for (; I < NumResources; ++I)
+						{
+							// Mark as visted.
+							if (ResourceRange[I].Value == ResourceRange[ResourceIndex].Value)
+							{
+								EnumAddFlags(SearchFlags[I], ESearchFlags::Visited);	
+								if (CurrentInstanceCache->OpExecutionData[ResourceRange[I].OpAddress].OpHitCount > 0)
+								{
+									// Still used, next step will mark to keep.
+									break;
+								}
+							}
+						}
+
+						if (I < NumResources)
+						{
+							// The image is still used, mark all occurences as visited and keep.
+							for (int32 J = ResourceIndex; J < NumResources; ++J)
+							{
+								if (ResourceRange[J].Value == ResourceRange[ResourceIndex].Value)
+								{
+									EnumAddFlags(SearchFlags[J], (ESearchFlags::Visited | ESearchFlags::Keep));	
+								}
+							}
+						}
+						else
+						{
+							// The resource is not used, see if we can stop searching. 
+							// In that case all occurences have been marked as visited.
+							RemovedBytes += ResourceRange[ResourceIndex].Value->GetDataSize();
+
+							if (TotalBytes - RemovedBytes <= BudgetBytes)
+							{
+								return RemovedBytes;
+							}
+						}
+					}
+				}
+
+				return RemovedBytes;
+			};
+
+			const int32 MaxNumResources = FMath::Max(CurrentInstanceCache->ImageResults.Num(), CurrentInstanceCache->MeshResults.Num());
+			TArray<ESearchFlags> SearchFlags;
+			SearchFlags.SetNumUninitialized(MaxNumResources);
+			
+			if (!bFinished)
+			{
+				FMemory::Memzero(SearchFlags.GetData(), MaxNumResources*sizeof(ESearchFlags));
+				
+				const TArrayView<FProgramCache::TResourceResult<Image>>& Images = MakeArrayView(CurrentInstanceCache->ImageResults);
+				const int32 RemovedBytes = SearchResourcesToRemove(Images, SearchFlags); 
+				
+				if (RemovedBytes > 0)
+				{	
+					MUTABLE_CPUPROFILER_SCOPE(EnsureBudgetBelow_FreeCached_Current_FreeResources);
+			
+					const int32 NumImages = Images.Num();
+					for (int32 I = 0; I < NumImages; ++I)
+					{
+						// Remove the first occurence of any visited not to keep. 
+						if (SearchFlags[I] == (ESearchFlags::FirstOccurance | ESearchFlags::Visited))
+						{
+							CacheResources.Remove(Images[I].Value);
+						}
+					
+						// Set unused all references visited not marked to keep. 
+						if ((SearchFlags[I] & (ESearchFlags::Visited | ESearchFlags::Keep)) == ESearchFlags::Visited)
+						{
+							CurrentInstanceCache->SetUnused(CurrentInstanceCache->OpExecutionData[Images[I].OpAddress]);
+						}
+					}	
+				}
+
+				TotalBytes -= RemovedBytes;
+				bFinished = TotalBytes <= BudgetBytes;
+			}
+			
+			if (!bFinished)
+			{
+				FMemory::Memzero(SearchFlags.GetData(), MaxNumResources*sizeof(ESearchFlags));
+				
+				TArrayView<FProgramCache::TResourceResult<Mesh>> Meshes = MakeArrayView(CurrentInstanceCache->MeshResults);
+				const int32 RemovedBytes = SearchResourcesToRemove(Meshes, SearchFlags); 
+
+				if (RemovedBytes > 0)
+				{	
+					MUTABLE_CPUPROFILER_SCOPE(EnsureBudgetBelow_FreeCached_Current_FreeResources);
+					
+					const int32 NumMeshes = Meshes.Num();
+					for (int32 I = 0; I < NumMeshes; ++I)
+					{
+						// Remove the first occurence of any visited not to keep. 
+						if (SearchFlags[I] == (ESearchFlags::FirstOccurance | ESearchFlags::Visited))
+						{
+							CacheResources.Remove(Meshes[I].Value);
+						}
+						
+						if ((SearchFlags[I] & (ESearchFlags::Visited | ESearchFlags::Keep)) == ESearchFlags::Visited)
+						{
+							CurrentInstanceCache->SetUnused(CurrentInstanceCache->OpExecutionData[Meshes[I].OpAddress]);
+						}
+					}
+				}
+
+				TotalBytes -= RemovedBytes;
+				bFinished = TotalBytes <= BudgetBytes;
+			}
+		}
+
+		if (!bFinished)
+		{
+			int64 ExcessBytes = TotalBytes - BudgetBytes;
+
+			if (ExcessBytes > BudgetExcessBytes)
+			{
+				BudgetExcessBytes = ExcessBytes;
+
+				// We failed to free enough memory. Log this, but try to continue anyway.
+				// This is a good place to insert a brakpoint to detect callstacks with memory peaks
+				UE_LOG(LogMutableCore, Log, TEXT("Failed to keep memory budget. Budget: %d, Current: %d, New: %d"),
+					BudgetBytes / 1024, (TotalBytes - AdditionalMemory) / 1024, AdditionalMemory / 1024);
+				
+				if (bEnableDetailedMemoryBudgetExceededLogging)
+				{
+					// We won't show correct internal or streaming buffer memory.
+					LogWorkingMemory(nullptr);
+				}
+			}
+
+		}
+
+        return bFinished;
+    }
+
+    //---------------------------------------------------------------------------------------------
+    void FWorkingMemoryManager::MarkRomUsed( int32 romIndex, const TSharedPtr<const Model>& pModel )
+    {
+        check(pModel);
+
+        // If budget is zero, we don't unload anything here, and we assume it is managed somewhere else.
+        if (!BudgetBytes)
+        {
+            return;
+        }
+
+        ++RomTick;
+
+        // Update current cache
+        {
+			FModelCacheEntry* ModelCache = FindModelCache(pModel.Get());
+        	
+            ModelCache->RomWeights[romIndex].Key++;
+            ModelCache->RomWeights[romIndex].Value = RomTick;
+        }
+    }
+
+
+	//---------------------------------------------------------------------------------------------
+	static void AddMultiValueKeys(FBitWriter& Blob, const TMap< TArray<int32>, PARAMETER_VALUE >& Multi)
+	{
+		uint16 Num = uint16(Multi.Num());
+		Blob.Serialize(&Num, 2);
+
+		for (const TPair< TArray<int32>, PARAMETER_VALUE >& v : Multi)
+		{
+			uint16 RangeNum = uint16(v.Key.Num());
+			Blob.Serialize(&RangeNum, 2);
+			Blob.Serialize((void*)v.Key.GetData(), RangeNum * sizeof(int32));
+		}
+	}
+
+
+	//---------------------------------------------------------------------------------------------
+	FResourceID FWorkingMemoryManager::GetResourceKey(const TSharedPtr<const Model>& Model, const Parameters* Params, uint32 ParamListIndex, OP::ADDRESS RootAt)
+	{
+		MUTABLE_CPUPROFILER_SCOPE(GetResourceKey);
+
+		constexpr uint32 ErrorId = 0xffff;
+
+		if (!Model)
+		{
+			return ErrorId;
+		}
+
+		const FProgram& Program = Model->GetPrivate()->m_program;
+
+		// Find the list of relevant parameters
+		const TArray<uint16>* RelevantParams = nullptr;
+		if (ParamListIndex < (uint32)Program.m_parameterLists.Num())
+		{
+			RelevantParams = &Program.m_parameterLists[ParamListIndex];
+		}
+		check(RelevantParams);
+		if (!RelevantParams)
+		{
+			return ErrorId;
+		}
+
+		// Generate the relevant parameters blob
+		FBitWriter Blob( 2048*8, true );
+
+		const TArray<FParameterDesc>& ParamDescs = Params->GetPrivate()->m_pModel->GetPrivate()->m_program.m_parameters;
+
+		// First make a mask with a bit for each relevant parameter. It will be on for parameters included in the blob.
+		// A parameter will be excluded from the blob if it has the deafult value, and no multivalues.
+		TBitArray IncludedParameters(0, RelevantParams->Num());
+		if (RelevantParams->Num())
+		{
+			for (int32 IndexIndex = 0; IndexIndex < RelevantParams->Num(); ++IndexIndex)
+			{
+				int32 ParamIndex = (*RelevantParams)[IndexIndex];
+				bool bInclude = Params->GetPrivate()->HasMultipleValues(ParamIndex);
+				if (!bInclude)
+				{
+					bInclude = !(
+						Params->GetPrivate()->m_values[ParamIndex]
+						==
+						ParamDescs[ParamIndex].m_defaultValue
+						);
+				}
+
+				IncludedParameters[IndexIndex] = bInclude;
+			}
+			Blob.SerializeBits(IncludedParameters.GetData(), IncludedParameters.Num());
+		}
+
+		// Second: serialize the value of the selected parameters.
+		for (int32 IndexIndex = 0; IndexIndex < RelevantParams->Num(); ++IndexIndex)
+		{
+			int32 ParamIndex = (*RelevantParams)[IndexIndex];
+			if (!IncludedParameters[IndexIndex])
+			{
+				continue;
+			}
+
+			int32 DataSize = 0;
+
+			switch (Program.m_parameters[ParamIndex].m_type)
+			{
+			case PARAMETER_TYPE::T_BOOL:
+				Blob.WriteBit(Params->GetPrivate()->m_values[ParamIndex].Get<ParamBoolType>() ? 1 : 0);
+
+				// Multi-values
+				if (Params->GetPrivate()->HasMultipleValues(ParamIndex))
+				{
+					const TMap< TArray<int32>, PARAMETER_VALUE >& Multi = Params->GetPrivate()->m_multiValues[ParamIndex];
+					AddMultiValueKeys(Blob, Multi);
+					for (const TPair< TArray<int32>, PARAMETER_VALUE >& v : Multi)
+					{
+						Blob.WriteBit(v.Value.Get<ParamBoolType>() ? 1 : 0);
+					}
+				}
+				break;
+
+			case PARAMETER_TYPE::T_INT:
+			{
+				int32 MaxValue = ParamDescs[ParamIndex].m_possibleValues.Num();
+				int32 Value = Params->GetPrivate()->m_values[ParamIndex].Get<ParamIntType>();
+				if (MaxValue)
+				{
+					// It is an enum
+					uint32 LimitedValue = FMath::Clamp( Params->GetIntValueIndex(ParamIndex,Value), 0, MaxValue-1 );
+					Blob.SerializeInt(LimitedValue, uint32(MaxValue));
+				}
+				else
+				{
+					// It may have any value
+					DataSize = sizeof(ParamIntType);
+					Blob.Serialize(&Value, DataSize);
+				}
+
+				// Multi-values
+				if (Params->GetPrivate()->HasMultipleValues(ParamIndex))
+				{
+					const TMap< TArray<int32>, PARAMETER_VALUE >& Multi = Params->GetPrivate()->m_multiValues[ParamIndex];
+					AddMultiValueKeys(Blob, Multi);
+					for (const TPair< TArray<int32>, PARAMETER_VALUE >& v : Multi)
+					{
+						Value = v.Value.Get<ParamIntType>();
+						if (MaxValue)
+						{
+							// It is an enum
+							uint32 LimitedValue = Value;
+							Blob.SerializeInt(LimitedValue, uint32(MaxValue));
+						}
+						else
+						{
+							// It may have any value
+							DataSize = sizeof(ParamIntType);
+							Blob.Serialize(&Value, DataSize);
+						}
+					}
+				}
+				break;
+			}
+
+			case PARAMETER_TYPE::T_FLOAT:
+				DataSize = sizeof(ParamFloatType);
+				Blob.Serialize(&Params->GetPrivate()->m_values[ParamIndex].Get<ParamFloatType>(), DataSize);
+
+				// Multi-values
+				if (Params->GetPrivate()->HasMultipleValues(ParamIndex))
+				{
+					const TMap< TArray<int32>, PARAMETER_VALUE >& Multi = Params->GetPrivate()->m_multiValues[ParamIndex];
+					AddMultiValueKeys(Blob, Multi);
+					for (const TPair< TArray<int32>, PARAMETER_VALUE >& v : Multi)
+					{
+						Blob.Serialize((void*)&v.Value.Get<ParamFloatType>(), DataSize);
+					}
+				}
+				break;
+
+			case PARAMETER_TYPE::T_COLOUR:
+				DataSize = sizeof(ParamColorType);
+				Blob.Serialize(&Params->GetPrivate()->m_values[ParamIndex].Get<ParamColorType>(), DataSize);
+
+				// Multi-values
+				if (Params->GetPrivate()->HasMultipleValues(ParamIndex))
+				{
+					const TMap< TArray<int32>, PARAMETER_VALUE >& Multi = Params->GetPrivate()->m_multiValues[ParamIndex];
+					AddMultiValueKeys(Blob, Multi);
+					for (const TPair< TArray<int32>, PARAMETER_VALUE >& v : Multi)
+					{
+						Blob.Serialize((void*)&v.Value.Get<ParamColorType>(), DataSize);
+					}
+				}
+				break;
+
+			case PARAMETER_TYPE::T_PROJECTOR:
+			{
+				FPackedNormal TempVec;
+				FProjector& Value = Params->GetPrivate()->m_values[ParamIndex].Get<ParamProjectorType>();
+				Blob.Serialize((void*)&Value.position, sizeof(FVector3f));
+				TempVec = FPackedNormal(Value.direction);
+				Blob.Serialize(&TempVec, sizeof(FPackedNormal));
+				TempVec = FPackedNormal(Value.up);
+				Blob.Serialize(&TempVec, sizeof(FPackedNormal));
+				Blob.Serialize((void*)&Value.scale, sizeof(FVector3f));
+				Blob.Serialize((void*)&Value.projectionAngle, sizeof(float));
+
+				// Multi-values
+				if (Params->GetPrivate()->HasMultipleValues(ParamIndex))
+				{
+					const TMap< TArray<int32>, PARAMETER_VALUE >& Multi = Params->GetPrivate()->m_multiValues[ParamIndex];
+					AddMultiValueKeys(Blob, Multi);
+					for (const TPair< TArray<int32>, PARAMETER_VALUE >& v : Multi)
+					{
+						Value = v.Value.Get<ParamProjectorType>();
+						Blob.Serialize((void*)&Value.position, sizeof(FVector3f));
+						TempVec = FPackedNormal(Value.direction);
+						Blob.Serialize(&TempVec, sizeof(FPackedNormal));
+						TempVec = FPackedNormal(Value.up);
+						Blob.Serialize(&TempVec, sizeof(FPackedNormal));
+						Blob.Serialize((void*)&Value.scale, sizeof(FVector3f));
+						Blob.Serialize((void*)&Value.projectionAngle, sizeof(float));
+					}
+				}
+				break;
+			}
+
+			case PARAMETER_TYPE::T_IMAGE:
+			{
+				DataSize = sizeof(FName);
+				Blob.Serialize((void*)&Params->GetPrivate()->m_values[ParamIndex].Get<ParamImageType>(), DataSize);
+
+				// Multi-values
+				if (Params->GetPrivate()->HasMultipleValues(ParamIndex))
+				{
+					const TMap< TArray<int32>, PARAMETER_VALUE >& Multi = Params->GetPrivate()->m_multiValues[ParamIndex];
+					AddMultiValueKeys(Blob, Multi);
+					for (const TPair< TArray<int32>, PARAMETER_VALUE >& v : Multi)
+					{
+						Blob.Serialize((void*)&v.Value.Get<ParamImageType>(), DataSize);
+					}
+				}
+				break;
+			}
+
+			default:
+				// unsupported parameter type
+				check(false);
+			}
+		}
+
+		// Increase the request id
+		++LastResourceResquestId;
+
+		FGeneratedResourceData NewKey;
+		int32 BlobBytes = Blob.GetNumBytes();
+		NewKey.ParameterValuesBlob.SetNum(BlobBytes);
+		FMemory::Memcpy(NewKey.ParameterValuesBlob.GetData(), Blob.GetData(), BlobBytes);
+
+		// See if we already have this id
+		int32 OldestCachePosition = 0;
+		uint32 OldestRequestId = 0;
+		for (int32 CacheIndex = 0; CacheIndex < GeneratedResources.Num(); ++CacheIndex)
+		{
+			FGeneratedResourceData& Data = GeneratedResources[CacheIndex];
+			bool bSameModel = Data.Model == Model;
+			bool bSameRoot = GetResourceIDRoot(Data.Id) == RootAt;
+			if (bSameModel
+				&&
+				bSameRoot
+				)
+			{
+				bool bSameBlob = Data.ParameterValuesBlob == NewKey.ParameterValuesBlob;
+				if (bSameBlob)
+				{
+					Data.LastRequestId = LastResourceResquestId;
+					return Data.Id;
+				}
+			}
+			
+			if (!Data.Model.Pin().Get() 
+				||
+				OldestRequestId > Data.LastRequestId)
+			{
+				OldestCachePosition = CacheIndex;
+				OldestRequestId = Data.Model.Pin() ? Data.LastRequestId : 0;
+			}
+		}
+
+		// Generate a new id
+		uint32 NewBlobId = ++LastResourceKeyId;
+		NewKey.Id = MakeResourceID( RootAt, NewBlobId);
+		NewKey.LastRequestId = LastResourceResquestId;
+		NewKey.Model = Model;
+
+		if (GeneratedResources.Num() >= MaxGeneratedResourceCacheSize)
+		{
+			GeneratedResources[OldestCachePosition] = MoveTemp(NewKey);
+		}
+		else
+		{
+			GeneratedResources.Add(MoveTemp(NewKey));
+		}
+
+		return NewKey.Id;
+	}
+
+}
diff --git a/Engine/Plugins/Mutable/Source/MutableRuntime/Public/MuR/ImageDataStorage.h b/Engine/Plugins/Mutable/Source/MutableRuntime/Public/MuR/ImageDataStorage.h
new file mode 100644
--- /dev/null
+++ b/Engine/Plugins/Mutable/Source/MutableRuntime/Public/MuR/ImageDataStorage.h
@@ -0,0 +1,407 @@
+// Copyright Epic Games, Inc. All Rights Reserved.
+
+#pragma once
+
+#include "MuR/ImageTypes.h"
+#include "MuR/MemoryTrackingAllocationPolicy.h"
+
+#include "Containers/Array.h"
+#include "Containers/ArrayView.h"
+#include "Containers/StaticArray.h"
+
+#include <atomic>
+
+namespace mu::MemoryCounters
+{
+	struct MUTABLERUNTIME_API FImageMemoryCounter
+	{
+		alignas(8) static inline std::atomic<SSIZE_T> Counter {0};
+	};
+}
+
+namespace mu
+{
+	using FImageArray = TArray<uint8, FDefaultMemoryTrackingAllocator<MemoryCounters::FImageMemoryCounter>>;
+}
+
+namespace mu::FImageDataStorageInternal
+{
+	FORCEINLINE int32 ComputeNumLODsForSize(FImageSize Size)
+	{
+		return static_cast<int32>(FMath::CeilLogTwo(FMath::Max(Size.X, Size.Y))) + 1;
+	}
+}
+
+namespace mu
+{
+class MUTABLERUNTIME_API FImageDataStorage
+{
+public:
+	TArray<FImageArray, TInlineAllocator<1>> Buffers;
+	
+	FImageSize ImageSize = FImageSize(0, 0);
+	EImageFormat ImageFormat = EImageFormat::IF_NONE;
+	uint8 NumLODs = 0;
+
+	static constexpr int32 NumLODsInCompactedTail = 7;
+	// This is needed for images its size cannot be known from dimensions and format, e.g., RLE compressed formats.
+	// It stores the offset to the end of the LOD so that CompactedTailOffsets[LOD] - CompactedTailOffsets[LOD - 1]
+	// is the size of LOD.
+	TStaticArray<uint16, NumLODsInCompactedTail> CompactedTailOffsets = MakeUniformStaticArray<uint16, NumLODsInCompactedTail>(0);
+
+public:
+
+	FImageDataStorage();
+	FImageDataStorage(const FImageDesc& Desc);
+
+	FImageDataStorage(const FImageDataStorage& Other);
+	FImageDataStorage& operator=(const FImageDataStorage& Other);
+	
+	FImageDataStorage(FImageDataStorage&& Other) = default;
+	FImageDataStorage& operator=(FImageDataStorage&& Other) = default;
+
+	bool operator==(const FImageDataStorage& Other) const;
+
+	FORCEINLINE FImageDesc MakeImageDesc() const
+	{
+		return FImageDesc(ImageSize, ImageFormat, NumLODs);
+	}
+
+	FORCEINLINE void InitVoid(const FImageDesc& Desc) 
+	{
+		Buffers.Empty();	
+		CompactedTailOffsets[0] = TNumericLimits<uint16>::Max();
+
+		ImageSize = Desc.m_size;
+		ImageFormat = Desc.m_format;
+		NumLODs = Desc.m_lods;
+	}
+
+	FORCEINLINE bool IsVoid() const
+	{
+		return Buffers.IsEmpty() && CompactedTailOffsets[0] == TNumericLimits<uint16>::Max();
+	}
+
+	void InitInternalArray(int32 Index, int32 Size, EInitializationType InitType);
+	void Init(const FImageDesc& ImageDesc, EInitializationType InitType);
+
+	void CopyFrom(const FImageDataStorage& Other);
+
+	FORCEINLINE int32 GetNumLODs() const
+	{
+		return NumLODs;
+	}
+
+	/**
+	* Get a const view to the data containing the LODIndex.
+	*/
+	FORCEINLINE TArrayView<const uint8> GetLOD(int32 LODIndex) const
+	{
+		check(!IsVoid());
+		if (LODIndex >= NumLODs)
+		{
+			//check(false);
+			return TArrayView<const uint8>();
+		}
+
+		const int32 FirstTailLOD = ComputeFirstCompactedTailLOD();
+		if (LODIndex >= FirstTailLOD)
+		{
+			const int32 LODIndexInTail = FMath::Max(0, LODIndex - FirstTailLOD);
+			
+			const int32 LODInTailBegin = LODIndexInTail == 0 ? 0 : CompactedTailOffsets[LODIndexInTail - 1];
+			const int32 LODInTailEnd = CompactedTailOffsets[LODIndexInTail];
+
+			const FImageArray& TailBuffer = Buffers.Last();
+			return MakeArrayView(TailBuffer.GetData() + LODInTailBegin, LODInTailEnd - LODInTailBegin);
+		}
+
+		return MakeArrayView(Buffers[LODIndex].GetData(), Buffers[LODIndex].Num());
+	}
+
+	/**
+	* Get a view to the data containing the LODIndex.
+	*/
+	FORCEINLINE TArrayView<uint8> GetLOD(int32 LODIndex)
+	{	
+		const TArrayView<const uint8> ConstLODView = const_cast<const FImageDataStorage*>(this)->GetLOD(LODIndex);
+		return TArrayView<uint8>(const_cast<uint8*>(ConstLODView.GetData()), ConstLODView.Num());
+	}
+
+	/**
+	 * Change the number of lods in the image. If InNumLODs is smaller than the current lod count, 
+	 * LODs are dropped from the tail up.
+	 */
+	void SetNumLODs(int32 InNumLODs, EInitializationType InitType = EInitializationType::NotInitialized);
+
+	/**
+	 * Remove NumLODsToDrop starting form LOD 0. The resulting image will be smaller. 
+	 */
+	void DropLODs(int32 NumLODsToDrop);
+
+	/**
+	 * Resizes LODIndex Data to NewSizeBytes
+	 */
+	void ResizeLOD(int32 LODIndex, int32 NewSizeBytes);
+
+	FORCEINLINE const FImageArray& GetInternalArray(int32 LODIndex) const
+	{	
+		check(!IsVoid());
+		check(LODIndex < NumLODs);
+		return Buffers[FMath::Min(LODIndex, ComputeFirstCompactedTailLOD())];
+	}
+
+	FORCEINLINE FImageArray& GetInternalArray(int32 LODIndex)
+	{	
+		check(!IsVoid());
+		check(LODIndex < NumLODs);
+		return Buffers[FMath::Min(LODIndex, ComputeFirstCompactedTailLOD())];
+	}
+
+	FORCEINLINE int32 ComputeFirstCompactedTailLOD() const
+	{
+		using namespace FImageDataStorageInternal;
+		return FMath::Max(0, ComputeNumLODsForSize(ImageSize) - NumLODsInCompactedTail); 
+	}
+
+	/**
+	 * Stateless iteration in batches.
+	 */
+	FORCEINLINE int32 GetNumBatches(int32 BatchSizeInElems, int32 BatchElemSizeInBytes) const
+	{	
+		return GetNumBatchesLODRange(BatchSizeInElems, BatchElemSizeInBytes, 0, NumLODs);
+	}
+
+	FORCEINLINE TArrayView<const uint8> GetBatch(int32 BatchId, int32 BatchSizeInElems, int32 BatchElemSizeInBytes) const
+	{
+		return GetBatchLODRange(BatchId, BatchSizeInElems, BatchElemSizeInBytes, 0, NumLODs);
+	}
+
+	FORCEINLINE TArrayView<uint8> GetBatch(int32 BatchId, int32 BatchSizeInElems, int32 BatchElemSizeInBytes)
+	{
+		// Use the const_cast idiom to generate the non-const operation. 
+		const TArrayView<const uint8> ConstBatchView = 
+				const_cast<const FImageDataStorage*>(this)->GetBatch(BatchId, BatchSizeInElems, BatchElemSizeInBytes);
+		return TArrayView<uint8>(const_cast<uint8*>(ConstBatchView.GetData()), ConstBatchView.Num());
+	}
+
+	/**
+	 * Returns the number of batches GetBatchFirtsLODOffet will admit and that will cover
+	 * the first LOD buffer with OffsetInBytes removed from the front. 
+	 * 
+	 * A batch cannot be larger than BatchElemsSize. The last batch of a buffer will be smaller if 
+	 * BatchSizeInElems*BatchElemSizeInBytes is not multiple of the buffer size.
+	 *
+	 **/
+	FORCEINLINE int32 GetNumBatchesFirstLODOffset(int32 BatchSizeInElems, int32 BatchElemSizeInBytes, int32 OffsetInBytes) const
+	{
+		check(!IsVoid());
+		TArrayView<const uint8> FirstLODView = GetLOD(0);
+		
+		check(OffsetInBytes % BatchElemSizeInBytes == 0);
+		check(FirstLODView.Num() % BatchElemSizeInBytes == 0);
+		check(FirstLODView.Num() > OffsetInBytes); 
+
+		return FMath::DivideAndRoundUp(FirstLODView.Num() - OffsetInBytes, BatchSizeInElems*BatchElemSizeInBytes);
+	}
+
+	/**
+	 * Returns a non modifiable view to the portion of the first LOD buffer with OffsetInBytes for the BatchId 
+	 *
+	 * A batch cannot be larger than BatchElemsSize. The last batch of a buffer will be smaller if 
+	 * BatchSizeInElems*BatchElemSizeInBytes is not multiple of the buffer size.
+	 *
+	 */
+	FORCEINLINE TArrayView<const uint8> GetBatchFirstLODOffset(int32 BatchId, int32 BatchSizeInElems, int32 BatchElemSizeInBytes, int32 OffsetInBytes) const
+	{
+		check(!IsVoid());
+		TArrayView<const uint8> FirstLODView = GetLOD(0);
+		
+		check(OffsetInBytes < FirstLODView.Num());
+		check(FirstLODView.Num() % BatchElemSizeInBytes == 0);
+		check(OffsetInBytes % BatchElemSizeInBytes == 0);
+	
+		TArrayView<const uint8> OffsetedLODView = MakeArrayView(FirstLODView.GetData() + OffsetInBytes, FirstLODView.Num() - OffsetInBytes);
+
+		const int32 BatchSizeInBytes = BatchSizeInElems * BatchElemSizeInBytes;
+		const int32 BatchOffset = BatchId * BatchSizeInBytes;
+
+		if (BatchOffset >= OffsetedLODView.Num())
+		{
+			return TArrayView<const uint8>();
+		}
+
+		return MakeArrayView(OffsetedLODView.GetData() + BatchOffset, FMath::Min(BatchSizeInBytes, OffsetedLODView.Num() - BatchOffset));
+	}
+
+	/**
+	 * Non const version of GetBatchFirstLODOffset(). 
+	 */
+	FORCEINLINE TArrayView<uint8> GetBatchFirstLODOffset(int32 BatchId, int32 BatchSizeInElems, int32 BatchElemSizeInBytes, int32 OffsetInBytes = 0)
+	{
+		const TArrayView<const uint8> ConstBatchView = 
+				const_cast<const FImageDataStorage*>(this)->GetBatchFirstLODOffset(BatchId, BatchSizeInElems, BatchElemSizeInBytes, OffsetInBytes);
+
+		return TArrayView<uint8>(const_cast<uint8*>(ConstBatchView.GetData()), ConstBatchView.Num());
+	}
+
+	/**
+	 * Returns the number of batches GetBatchLODRange will admit and that will cover
+	 * the [LODBegin, LODEnd) range.
+	 *
+	 **/
+	FORCEINLINE int32 GetNumBatchesLODRange(int32 BatchSizeInElems, int32 BatchElemSizeInBytes, int32 LODBegin, int32 LODEnd) const
+	{
+		check(!IsVoid());
+		check(LODBegin < LODEnd);
+		check(LODBegin >= 0 && LODEnd <= NumLODs);
+
+		const int32 BatchNumBytes = BatchSizeInElems*BatchElemSizeInBytes;
+		
+		const int32 FirstCompactedTailLOD = ComputeFirstCompactedTailLOD();
+		const int32 NumLODsInTail = FMath::Max(0, NumLODs - FirstCompactedTailLOD);
+		const int32 NumLODsNotInTail = FMath::Max(0, NumLODs - NumLODsInTail);
+
+		const int32 LODsNotInTailLODRangeEnd = FMath::Min(LODEnd, NumLODsNotInTail);
+
+		int32 NumBatches = 0;
+		for (int32 BufferIndex = LODBegin; BufferIndex < LODsNotInTailLODRangeEnd; ++BufferIndex)
+		{	
+			NumBatches += FMath::DivideAndRoundUp(Buffers[BufferIndex].Num(), BatchNumBytes);
+		}
+
+		if (FirstCompactedTailLOD < LODEnd)
+		{
+			TArrayView<const uint8> FirstLODInRangeView = GetLOD(FMath::Max(FirstCompactedTailLOD, LODBegin));
+			TArrayView<const uint8> LastLODInRangeView = GetLOD(LODEnd - 1);
+
+			const int32 TailInLODRangeNumBytes = LastLODInRangeView.GetData() - FirstLODInRangeView.GetData() + LastLODInRangeView.Num();
+			NumBatches += FMath::DivideAndRoundUp(TailInLODRangeNumBytes, BatchNumBytes);
+		}
+		return NumBatches;
+	}
+
+	/**
+	 * Returns a non modifiable view to the portion of the [LODBegin, LODEnd) for the BatchId. 
+	 *
+	 * A batch cannot be larger than BatchElemsSize. The last batch of a buffer will be smaller if 
+	 * BatchSizeInElems*BatchElemSizeInBytes is not multiple of the buffer size.
+	 *
+	 */
+	FORCEINLINE TArrayView<const uint8> GetBatchLODRange(int32 BatchId, int32 BatchSizeInElems, int32 BatchElemSizeInBytes, int32 LODBegin, int32 LODEnd) const
+	{
+		check(!IsVoid());
+		check(LODBegin < LODEnd);
+		check(LODBegin >= 0 && LODEnd <= NumLODs);
+
+		const int32 BatchNumBytes = BatchSizeInElems*BatchElemSizeInBytes;
+
+		const int32 FirstCompactedTailLOD = ComputeFirstCompactedTailLOD();
+		const int32 NonTailBuffersEnd = FMath::Min(LODEnd, FirstCompactedTailLOD);
+
+		int32 NumBatches = 0;	
+		for (int32 BufferIndex = LODBegin; BufferIndex < NonTailBuffersEnd; ++BufferIndex)
+		{
+			const TArrayView<const uint8> BufferView = MakeArrayView(Buffers[BufferIndex].GetData(), Buffers[BufferIndex].Num());
+			const int32 BufferNumBatches = FMath::DivideAndRoundUp(BufferView.Num(), BatchNumBytes); 
+			
+			check(BatchId >= NumBatches);
+			int32 BufferBatchId = BatchId - NumBatches;
+
+			if (BufferBatchId < BufferNumBatches)
+			{ 
+				const int32 BufferBatchBeginOffset = BufferBatchId * BatchNumBytes;
+
+				return TArrayView<const uint8>(
+						BufferView.GetData() + BufferBatchBeginOffset,
+						FMath::Min(BufferView.Num() - BufferBatchBeginOffset, BatchNumBytes));
+			}
+
+			NumBatches += BufferNumBatches;
+		}
+
+		if (FirstCompactedTailLOD < LODEnd)
+		{
+			TArrayView<const uint8> FirstLODInRangeView = GetLOD(FMath::Max(FirstCompactedTailLOD, LODBegin));
+			TArrayView<const uint8> LastLODInRangeView = GetLOD(LODEnd - 1);
+			
+			const int32 TailInLODRangeNumBytes = LastLODInRangeView.GetData() - FirstLODInRangeView.GetData() + LastLODInRangeView.Num();
+			TArrayView<const uint8> TailInLODRangeView = MakeArrayView(FirstLODInRangeView.GetData(), TailInLODRangeNumBytes);
+
+			check(BatchId >= NumBatches);
+			const int32 TailInLODRangeBatchId = BatchId - NumBatches;
+			const int32 TailInLODRangeNumBatches = FMath::DivideAndRoundUp(TailInLODRangeView.Num(), BatchNumBytes);
+
+			if (TailInLODRangeBatchId < TailInLODRangeNumBatches)
+			{
+				const int32 BufferBatchBeginOffset = TailInLODRangeBatchId * BatchNumBytes;
+				
+				return TArrayView<const uint8>(
+							TailInLODRangeView.GetData() + BufferBatchBeginOffset,
+							FMath::Min(TailInLODRangeView.Num() - BufferBatchBeginOffset, BatchNumBytes));
+			}
+		}
+		// If the BatchId is not found, return an empty view.
+		return TArrayView<const uint8>();
+	}
+	
+	/**
+	 * Non const version of GetBatchLODRange().
+	 */
+	FORCEINLINE TArrayView<uint8> GetBatchLODRange(int32 BatchId, int32 BatchSizeInElems, int32 BatchElemSizeInBytes, int32 LODBegin, int32 LODEnd)
+	{
+		const TArrayView<const uint8> ConstBatchView = 
+			const_cast<const FImageDataStorage*>(this)->GetBatchLODRange(
+					BatchId, BatchSizeInElems, BatchElemSizeInBytes, LODBegin, LODEnd);
+
+		return TArrayView<uint8>(const_cast<uint8*>(ConstBatchView.GetData()), ConstBatchView.Num());
+	}
+
+	FORCEINLINE int32 GetAllocatedSize() const
+	{ 
+		SIZE_T Result = 0;
+		for (const FImageArray& Buffer : Buffers)
+		{
+			Result += Buffer.GetAllocatedSize();
+		}
+
+		check(Result < TNumericLimits<int32>::Max())
+		return (int32)Result;
+	}
+
+	FORCEINLINE int32 GetDataSize() const
+	{
+		SIZE_T Result = 0;
+		for (const FImageArray& Buffer : Buffers)
+		{
+			Result += Buffer.Num();
+		}
+
+		check(Result < TNumericLimits<int32>::Max())
+		return (int32)Result;
+	}
+
+	/**
+	 * Returns true if all buffers are empty. This can be true after initialization for image formats that report 
+	 * BytesPerBlock 0, e.g., RLE formats. If not initialized will also return true.
+	 */
+	FORCEINLINE bool IsEmpty() const
+	{
+		check(!IsVoid());
+
+		for (const FImageArray& Buffer : Buffers)
+		{
+			if (!Buffer.IsEmpty())
+			{
+				return false;
+			}
+		}
+
+		return true;
+	}
+
+	void Serialise(OutputArchive& Arch) const;
+	void Unserialise(InputArchive& Arch);
+};
+
+}
diff --git a/Engine/Plugins/Mutable/Source/MutableRuntime/Public/MuR/Instance.h b/Engine/Plugins/Mutable/Source/MutableRuntime/Public/MuR/Instance.h
new file mode 100644
--- /dev/null
+++ b/Engine/Plugins/Mutable/Source/MutableRuntime/Public/MuR/Instance.h
@@ -0,0 +1,218 @@
+// Copyright Epic Games, Inc. All Rights Reserved.
+
+#pragma once
+
+#include "HAL/PlatformMath.h"
+#include "MuR/ExtensionData.h"
+#include "MuR/Ptr.h"
+#include "MuR/RefCounted.h"
+#include "MuR/Types.h"
+
+
+namespace mu
+{
+	class Instance;
+
+	typedef Ptr<Instance> InstancePtr;
+	typedef Ptr<const Instance> InstancePtrConst;
+
+
+    //! \brief A customised object created from a model and a set of parameter values.
+    //!
+    //! It corresponds to an "engine object" but the contents of its data depends on the Model, and
+    //! it may contain any number of LODs, components, surfaces, meshes and images, even none.
+	//! \ingroup runtime
+	class MUTABLERUNTIME_API Instance : public Resource
+	{
+	public:
+
+		//-----------------------------------------------------------------------------------------
+		// Life cycle
+		//-----------------------------------------------------------------------------------------
+
+		Instance();
+
+        //! Clone this instance
+        Ptr<Instance> Clone() const;
+
+		// Resource interface
+		int32 GetDataSize() const override;
+
+		//-----------------------------------------------------------------------------------------
+		// Own interface
+		//-----------------------------------------------------------------------------------------
+
+        //! Type for the instance unique identifiers.
+        typedef uint32 ID;
+
+        //! Get a unique identifier for this instance. It doesn't change during the entire
+        //! lifecycle of each instance. This identifier can be used in the System methods to update
+        //! or release the instance.
+        Instance::ID GetId() const;
+
+		//! Get the number of components of this instance.
+		int32 GetComponentCount() const;
+
+		//! Get the Id of a component
+		//! \param ComponentIndex Index of the component, from 0 to GetComponentCount()-1
+		uint16 GetComponentId(int32 ComponentIndex) const;
+
+        //! Get the number of LODs in a component
+        //! \param ComponentIndex Index of the component, from 0 to GetComponentCount()-1
+        int32 GetLODCount( int32 ComponentIndex ) const;
+
+        //! Get the number of surfaces in a component
+        //! \param ComponentIndex Index of the component, from 0 to GetComponentCount()-1
+		//! \param LODIndex Index of the level of detail, from 0 to GetLODCount()-1
+		int32 GetSurfaceCount( int32 ComponentIndex, int32 LODIndex) const;
+
+        //! Get an id that can be used to match the surface data with the mesh surface data.
+        //! \param ComponentIndex Index of the component, from 0 to GetComponentCount()-1
+		//! \param LODIndex Index of the level of detail, from 0 to GetLODCount()-1
+		//! \param SurfaceIndex Index of the surface, from 0 to GetSurfaceCount()-1
+        uint32 GetSurfaceId( int32 ComponentIndex, int32 LODIndex, int32 SurfaceIndex) const;
+
+        //! Find a surface index from the internal id (as returned by GetSurfaceId).
+        //! \param ComponentIndex Index of the component, from 0 to GetComponentCount()-1
+		//! \param LODIndex Index of the level of detail, from 0 to GetLODCount()-1
+		//! \param id ID of the surface to look for.
+        int32 FindSurfaceById( int32 ComponentIndex, int32 LODIndex, uint32 id ) const;
+
+		//! Find the base surface index and Lod index when reusing surfaces between LODs. Return the surface index
+		//! and the LOD it belongs to.
+		//! \param ComponentIndex - Index of the component, from 0 to GetComponentCount()-1
+		//! \param SharedSurfaceId - Id of the surface to look for (as returned by GetSharedSurfaceId).
+		//! \param OutSurfaceIndex - Index of the surface in the OutLODIndex lod. 
+		//! \param OutLODIndex - Index of the first LOD where the surface can be found. 
+		void FindBaseSurfaceBySharedId(int32 ComponentIndex, int32 SharedId, int32& OutSurfaceIndex, int32& OutLODIndex) const;
+
+		//! Get an id that can be used to find the same surface on other LODs
+		//! \param ComponentIndex Index of the component, from 0 to GetComponentCount()-1
+		//! \param lod Index of the level of detail, from 0 to GetLODCount()-1
+		//! \param SurfaceIndex Index of the surface, from 0 to GetSurfaceCount()-1
+		int32 GetSharedSurfaceId(int32 ComponentIndex, int32 LODIndex, int32 SurfaceIndex) const;
+
+        //! Get an optional, opaque application-defined identifier for this surface. The meaning of
+        //! this ID depends on each application, and it is specified when creating the source data
+        //! that generates this surface.
+        //! See NodeSurfaceNew::SetCustomID.
+		//! \param ComponentIndex Index of the component, from 0 to GetComponentCount()-1
+		//! \param LODIndex Index of the level of detail, from 0 to GetLODCount()-1
+		//! \param SurfaceIndex Index of the surface, from 0 to GetSurfaceCount()-1
+        uint32 GetSurfaceCustomId( int32 ComponentIndex, int32 LODIndex, int32 SurfaceIndex) const;
+
+        //! Get the mesh resource id from a component
+        //! \param ComponentIndex Index of the component, from 0 to GetComponentCount()-1
+		//! \param LODIndex Index of the level of detail, from 0 to GetLODCount()-1
+  		FResourceID GetMeshId( int32 ComponentIndex, int32 LODIndex ) const;
+
+		//! Get the number of images in a component
+		//! \param ComponentIndex Index of the component, from 0 to GetComponentCount()-1
+		//! \param LODIndex Index of the level of detail, from 0 to GetLODCount()-1
+		//! \param SurfaceIndex Index of the surface, from 0 to GetSurfaceCount()-1
+        int32 GetImageCount( int32 ComponentIndex, int32 LODIndex, int32 SurfaceIndex) const;
+
+        //! Get an image resource id from a component
+		//! \param ComponentIndex Index of the component, from 0 to GetComponentCount()-1
+		//! \param LODIndex Index of the level of detail, from 0 to GetLODCount()-1
+		//! \param SurfaceIndex Index of the surface, from 0 to GetSurfaceCount()-1
+        //! \param ImageIndex Index of the image, from 0 to GetImageCount()-1
+ 		FResourceID GetImageId( int32 ComponentIndex, int32 LODIndex, int32 SurfaceIndex, int32 ImageIndex) const;
+
+		//! Get the name of an image in a component
+		//! \param ComponentIndex Index of the component, from 0 to GetComponentCount()-1
+		//! \param LODIndex Index of the level of detail, from 0 to GetLODCount()-1
+		//! \param SurfaceIndex Index of the surface, from 0 to GetSurfaceCount()-1
+        //! \param ImageIndex Index of the image, from 0 to GetImageCount()-1
+		FName GetImageName( int32 ComponentIndex, int32 LODIndex, int32 SurfaceIndex, int32 ImageIndex ) const;
+
+		//! Get the number of vectors in a component
+		//! \param ComponentIndex Index of the component, from 0 to GetComponentCount()-1
+		//! \param LODIndex Index of the level of detail, from 0 to GetLODCount()-1
+		//! \param SurfaceIndex Index of the surface, from 0 to GetSurfaceCount()-1
+        int32 GetVectorCount( int32 ComponentIndex, int32 LODIndex, int32 SurfaceIndex) const;
+
+		//! Get a vector from a component
+		//! \param ComponentIndex Index of the component, from 0 to GetComponentCount()-1
+		//! \param LODIndex Index of the level of detail, from 0 to GetLODCount()-1
+		//! \param SurfaceIndex Index of the surface, from 0 to GetSurfaceCount()-1
+        //! \param VectorIndex Index of the vector, from 0 to GetVectorCount()-1
+        FVector4f GetVector( int32 ComponentIndex, int32 LODIndex, int32 SurfaceIndex, int32 VectorIndex) const;
+
+		//! Get the name of a vector in a component
+		//! \param ComponentIndex Index of the component, from 0 to GetComponentCount()-1
+		//! \param LODIndex Index of the level of detail, from 0 to GetLODCount()-1
+		//! \param SurfaceIndex Index of the surface, from 0 to GetSurfaceCount()-1
+        //! \param VectorIndex Index of the vector, from 0 to GetVectorCount()-1
+		FName GetVectorName( int32 ComponentIndex, int32 LODIndex, int32 SurfaceIndex, int32 VectorIndex ) const;
+
+        //! Get the number of scalar values in a component
+        //! \param ComponentIndex Index of the component, from 0 to GetComponentCount()-1
+		//! \param LODIndex Index of the level of detail, from 0 to GetLODCount()-1
+		//! \param SurfaceIndex Index of the surface, from 0 to GetSurfaceCount()-1
+        int32 GetScalarCount( int32 ComponentIndex, int32 LODIndex, int32 SurfaceIndex) const;
+
+        //! Get a scalar value from a component
+        //! \param ComponentIndex Index of the component, from 0 to GetComponentCount()-1
+		//! \param LODIndex Index of the level of detail, from 0 to GetLODCount()-1
+		//! \param SurfaceIndex Index of the surface, from 0 to GetSurfaceCount()-1
+        //! \param ScalarIndex Index of the scalar, from 0 to GetScalarCount()-1
+        float GetScalar( int32 ComponentIndex, int32 LODIndex, int32 SurfaceIndex, int32 ScalarIndex) const;
+
+        //! Get the name of a scalar from a component
+        //! \param ComponentIndex Index of the component, from 0 to GetComponentCount()-1
+		//! \param LODIndex Index of the level of detail, from 0 to GetLODCount()-1
+		//! \param SurfaceIndex Index of the surface, from 0 to GetSurfaceCount()-1
+        //! \param ScalarIndex Index of the scalar, from 0 to GetScalarCount()-1
+		FName GetScalarName( int32 ComponentIndex, int32 LODIndex, int32 SurfaceIndex, int32 ScalarIndex ) const;
+
+        //! Get the number of string values in a component
+        //! \param ComponentIndex Index of the component, from 0 to GetComponentCount()-1
+		//! \param LODIndex Index of the level of detail, from 0 to GetLODCount()-1
+		//! \param SurfaceIndex Index of the surface, from 0 to GetSurfaceCount()-1
+        int32 GetStringCount( int32 ComponentIndex, int32 LODIndex, int32 SurfaceIndex) const;
+
+        //! Get a string value from a component
+        //! \param ComponentIndex Index of the component, from 0 to GetComponentCount()-1
+		//! \param LODIndex Index of the level of detail, from 0 to GetLODCount()-1
+		//! \param SurfaceIndex Index of the surface, from 0 to GetSurfaceCount()-1
+        //! \param StringIndex Index of the string, from 0 to GetStringCount()-1
+        FString GetString( int32 ComponentIndex, int32 LODIndex, int32 SurfaceIndex, int32 StringIndex) const;
+
+        //! Get the name of a string from a component
+        //! \param ComponentIndex Index of the component, from 0 to GetComponentCount()-1
+		//! \param LODIndex Index of the level of detail, from 0 to GetLODCount()-1
+		//! \param SurfaceIndex Index of the surface, from 0 to GetSurfaceCount()-1
+        //! \param StringIndex Index of the string, from 0 to GetStringCount()-1
+        FName GetStringName( int32 ComponentIndex, int32 LODIndex, int32 SurfaceIndex, int32 StringIndex ) const;
+
+		//! Get the number of ExtensionData values in a component
+		int32 GetExtensionDataCount() const;
+
+		//! Get an ExtensionData value from a component
+		//! \param Index Index of the ExtensionData to fetch
+		//! \param OutExtensionData Receives the ExtensionData
+		//! \param OutName Receives the name associated with the ExtensionData. Guaranteed to be a valid string of non-zero length.
+		void GetExtensionData(int32 Index, Ptr<const class ExtensionData>& OutExtensionData, FName& OutName) const;
+
+        //-----------------------------------------------------------------------------------------
+		// Interface pattern
+		//-----------------------------------------------------------------------------------------
+		class Private;
+
+		Private* GetPrivate() const;
+
+	protected:
+
+		//! Forbidden. Manage with the Ptr<> template.
+		~Instance();
+
+	private:
+
+		Private* m_pD;
+
+	};
+
+
+}
+
