[mutable] Moved the Mutable plugin out of Experimental status into Beta.

diff --git a/Engine/Plugins/Mutable/Source/MutableRuntime/Public/MuR/Mesh.h b/Engine/Plugins/Mutable/Source/MutableRuntime/Public/MuR/Mesh.h
new file mode 100644
--- /dev/null
+++ b/Engine/Plugins/Mutable/Source/MutableRuntime/Public/MuR/Mesh.h
@@ -0,0 +1,577 @@
+// Copyright Epic Games, Inc. All Rights Reserved.
+
+#pragma once
+
+#include "Containers/Array.h"
+#include "Containers/Map.h"
+#include "HAL/PlatformMath.h"
+#include "Math/Transform.h"
+#include "Math/VectorRegister.h"
+#include "Misc/AssertionMacros.h"
+#include "Misc/EnumClassFlags.h"
+#include "MuR/Layout.h"
+#include "MuR/MeshBufferSet.h"
+#include "MuR/PhysicsBody.h"
+#include "MuR/Ptr.h"
+#include "MuR/RefCounted.h"
+#include "MuR/Serialisation.h"
+#include "MuR/Skeleton.h"
+#include "Templates/Tuple.h"
+
+#include <type_traits>
+
+class FString;
+
+namespace mu
+{
+
+	// Forward references
+	class Layout;
+    class Skeleton;
+    class PhysicsBody;
+
+	class Mesh;
+    typedef Ptr<Mesh> MeshPtr;
+    typedef Ptr<const Mesh> MeshPtrConst;
+
+	struct FSurfaceSubMesh
+	{
+		
+		int32 VertexBegin = 0;
+		int32 VertexEnd = 0;
+		int32 IndexBegin = 0;
+		int32 IndexEnd = 0;
+
+		uint32 ExternalId = 0;
+
+		friend bool operator==(const FSurfaceSubMesh& Lhs, const FSurfaceSubMesh& Rhs)
+		{
+			return FMemory::Memcmp(&Lhs, &Rhs, sizeof(FSurfaceSubMesh)) == 0;
+		}
+	};
+	static_assert(std::has_unique_object_representations_v<FSurfaceSubMesh>);
+	
+	MUTABLE_DEFINE_POD_SERIALISABLE(FSurfaceSubMesh);
+	MUTABLE_DEFINE_POD_VECTOR_SERIALISABLE(FSurfaceSubMesh);
+
+	struct FMeshSurface
+	{
+		TArray<FSurfaceSubMesh, TInlineAllocator<1>> SubMeshes; 
+
+		uint32 BoneMapIndex = 0;
+		uint32 BoneMapCount = 0;
+		uint32 Id = 0;
+
+		friend bool operator==(const FMeshSurface& Lhs, const FMeshSurface& Rhs)
+		{
+			return
+				Lhs.Id == Rhs.Id &&
+				Lhs.BoneMapIndex == Rhs.BoneMapIndex &&
+				Lhs.BoneMapCount == Rhs.BoneMapCount &&
+				Lhs.SubMeshes == Rhs.SubMeshes;
+		}
+
+		inline void Serialise(OutputArchive& Arch) const;
+		inline void Unserialise(InputArchive& Arch);
+	};
+
+
+	enum class EBoneUsageFlags : uint32
+	{
+		None		   = 0,
+		Root		   = 1 << 1,
+		Skinning	   = 1 << 2,
+		SkinningParent = 1 << 3,
+		Physics	       = 1 << 4,
+		PhysicsParent  = 1 << 5,
+		Deform         = 1 << 6,
+		DeformParent   = 1 << 7,
+		Reshaped       = 1 << 8	
+	};
+
+	ENUM_CLASS_FLAGS(EBoneUsageFlags);
+
+	//!
+	enum class EMeshBufferType
+	{
+		None,
+		SkeletonDeformBinding,
+		PhysicsBodyDeformBinding,
+		PhysicsBodyDeformSelection,
+		PhysicsBodyDeformOffsets,
+		MeshLaplacianData,
+		MeshLaplacianOffsets,
+		UniqueVertexMap
+	};
+
+	//!
+	enum class EShapeBindingMethod : uint32
+	{
+		ReshapeClosestProject = 0,
+		ClipDeformClosestProject = 1,
+		ClipDeformClosestToSurface = 2,
+		ClipDeformNormalProject = 3	
+	};
+
+	enum class EVertexColorUsage : uint32
+	{
+		None = 0,
+		ReshapeMaskWeight = 1,
+		ReshapeClusterId = 2
+	};
+
+	enum class EMeshCopyFlags : uint32
+	{
+		None = 0,
+		WithSkeletalMesh = 1 << 1,
+		WithSurfaces = 1 << 2,
+		WithSkeleton = 1 << 3,
+		WithPhysicsBody = 1 << 4,
+		WithFaceGroups = 1 << 5,
+		WithTags = 1 << 6,
+		WithVertexBuffers = 1 << 7,
+		WithIndexBuffers = 1 << 8,
+		// deprecated WithFaceBuffers = 1 << 9,
+		WithAdditionalBuffers = 1 << 10,
+		WithLayouts = 1 << 11,
+		WithPoses = 1 << 12,
+		WithBoneMap = 1 << 13,
+		WithSkeletonIDs = 1 << 14,
+		WithAdditionalPhysics = 1 << 15,
+		WithStreamedResources = 1 << 16,
+
+		AllFlags = 0xFFFFFFFF
+	};
+	
+	ENUM_CLASS_FLAGS(EMeshCopyFlags);
+
+
+	/** Optimised mesh formats that are identified in some operations to chose a faster version. */
+	enum class EMeshFlags : uint32
+	{
+		None = 0,
+
+		/** The mesh is formatted to be used for planar and cilyndrical projection */
+		ProjectFormat = 1 << 0,
+
+		/** The mesh is formatted to be used for wrapping projection */
+		ProjectWrappingFormat = 1 << 1,
+
+		/** The mesh is a reference to an external resource mesh. */
+		IsResourceReference = 1 << 2,
+
+		/** The mesh is a reference to an external resource mesh and must be loaded when first referenced. */
+		IsResourceForceLoad = 1 << 3,
+
+	};
+
+
+    //! \brief Mesh object containing any number of buffers with any number of channels.
+    //! The buffers can be per-index or per-vertex.
+    //! The mesh also includes layout information for every texture channel for internal usage, and
+    //! it can be ignored.
+    //! The meshes are always assumed to be triangle list primitives.
+    //! \ingroup runtime
+    class MUTABLERUNTIME_API Mesh : public Resource
+    {
+	public:
+
+		static constexpr uint64 InvalidVertexId = TNumericLimits<uint64>::Max();
+
+	public:
+
+        //-----------------------------------------------------------------------------------------
+        // Life cycle
+        //-----------------------------------------------------------------------------------------
+				
+		/** Create a new empty mesh that repreents an external resource mesh. */
+		static Ptr<Mesh> CreateAsReference(uint32 ID, bool bForceLoad);
+
+        /** Deep clone this mesh. */
+        Ptr<Mesh> Clone() const;
+		
+		/** Clone with flags allowing to not include some parts in the cloned mesh */
+		Ptr<Mesh> Clone(EMeshCopyFlags Flags) const;
+
+		/** Copy form another mesh. */
+		void CopyFrom(const Mesh& From, EMeshCopyFlags Flags = EMeshCopyFlags::AllFlags);
+
+        /** Serialisation */
+        static void Serialise(const Mesh* InMesh, OutputArchive& Arch );
+        static Ptr<Mesh> StaticUnserialise(InputArchive& Arch);
+
+		// Resource interface
+		int32 GetDataSize() const override;
+
+        //-----------------------------------------------------------------------------------------
+        // Own interface
+        //-----------------------------------------------------------------------------------------
+
+		/** Return true if this is a reference to an engine image. */
+		bool IsReference() const;
+
+		/** If true, this is a reference that must be resolved at compile time. */
+		bool IsForceLoad() const;
+
+		/** Return the id of the engine referenced mesh. Only valid if IsReference. */
+		uint32 GetReferencedMesh() const;
+
+        //! \name Buffers
+        //! \{
+
+        //!
+        int32 GetIndexCount() const;
+
+        /** Index buffers. They are owned by this mesh. */
+        FMeshBufferSet& GetIndexBuffers();
+        const FMeshBufferSet& GetIndexBuffers() const;
+
+        //
+        int32 GetVertexCount() const;
+
+        /** Vertex buffers. They are owned by this mesh. */
+        FMeshBufferSet& GetVertexBuffers();
+        const FMeshBufferSet& GetVertexBuffers() const;
+
+        int32 GetFaceCount() const;
+
+        /**
+		 * Get the number of surfaces defined in this mesh. Surfaces are buffer-contiguous mesh
+         * fragments that share common properties (usually material)
+		 */
+        int32 GetSurfaceCount() const;
+        void GetSurface(int32 SurfaceIndex,
+                        int32& OutFirstVertex, int32& OutVertexCount,
+                        int32& OutFirstIndex, int32& OutIndexCount,
+						int32& OutFirstBone, int32& OutBoneCount) const;
+
+        /**
+		 * Return an internal id that can be used to match mesh surfaces and instance surfaces.
+         * Only valid for meshes that are part of instances.
+		 */
+        uint32 GetSurfaceId(int32 SurfaceIndex) const;
+
+        //! \}
+
+		/** Return true if the mesh has unique vertex IDs and they stored in an implicit way. 
+		* This is relevant for some mesh operations that will need to make them explicit so that the result is still correct.
+		*/
+		bool AreVertexIdsImplicit() const;
+		bool AreVertexIdsExplicit() const;
+
+		/** Create an explicit vertex buffer for vertex IDs if they are implicit. */
+		void MakeVertexIdsRelative();
+
+		/** Ensure the format of an empty mesh includes explicit IDs. The mesh cannot have any vertex data. */
+		void MakeIdsExplicit();
+
+        //! \name Texture layouts
+        //! \{
+
+        //!
+        void AddLayout( Ptr<const Layout> pLayout );
+
+        //!
+        int32 GetLayoutCount() const;
+
+        //!
+        const Layout* GetLayout( int32 i ) const;
+
+        //!
+        void SetLayout( int32 i, Ptr<const Layout> );
+        //! \}
+
+        //! \name Skeleton information
+        //! \{
+
+        void SetSkeleton( Ptr<const Skeleton> );
+        Ptr<const Skeleton> GetSkeleton() const;
+
+        //! \}
+
+        //! \name PhysicsBody information
+        //! \{
+
+        void SetPhysicsBody( Ptr<const PhysicsBody> );
+        Ptr<const PhysicsBody> GetPhysicsBody() const;
+
+		int32 AddAdditionalPhysicsBody(Ptr<const PhysicsBody> Body);
+		Ptr<const PhysicsBody> GetAdditionalPhysicsBody(int32 I) const;
+		//int32 GetAdditionalPhysicsBodyExternalId(int32 I) const;
+
+        //! \}
+
+        //! \name Tags
+        //! \{
+
+        //!
+        void SetTagCount( int32 count );
+
+        //!
+        int32 GetTagCount() const;
+
+        //!
+        const FString& GetTag( int32 tagIndex ) const;
+
+        //!
+        void SetTag( int32 tagIndex, const FString& Name );
+
+		//!
+		void AddStreamedResource(uint64 ResourceId);
+
+		//!
+		const TArray<uint64>& GetStreamedResources() const;
+
+		//!
+		int32 FindBonePose(const FBoneName& BoneName) const;
+		
+		//!
+		void SetBonePoseCount(int32 count);
+
+		//!
+		int32 GetBonePoseCount() const;
+
+		//!
+		void SetBonePose(int32 Index, const FBoneName& BoneName, FTransform3f Transform, EBoneUsageFlags BoneUsageFlags);
+
+		//! @return - Bone identifier of the pose at 'Index'.
+		const FBoneName& GetBonePoseId(int32 BoneIndex) const;
+
+		//! Return a matrix stored per bone. It is a set of 16-float values.
+		void GetBonePoseTransform(int32 BoneIndex, FTransform3f& Transform) const;
+
+		//! 
+		EBoneUsageFlags GetBoneUsageFlags(int32 BoneIndex) const;
+
+		//! Set the bonemap of this mesh
+		void SetBoneMap(const TArray<FBoneName>& InBoneMap);
+
+		//! Return an array containing the bonemaps of all surfaces in the mesh.
+		const TArray<FBoneName>& GetBoneMap() const;
+
+		//!
+		int32 GetSkeletonIDsCount() const;
+
+		//!
+		int32 GetSkeletonID(int32 SkeletonIndex) const;
+
+		//!
+		void AddSkeletonID(int32 SkeletonID);
+
+        //! \}
+
+
+        /** 
+		 * Get an internal identifier used to reference this mesh in operations like deferred
+         * mesh building, or instance updating.
+		 */
+        uint32 GetId() const;
+
+
+    protected:
+
+        /** Forbidden. Manage with the Ptr<> template. */
+		~Mesh() {}
+
+    
+	public:
+
+		template<typename Type>
+		using TMemoryTrackedArray = TArray<Type, FDefaultMemoryTrackingAllocator<MemoryCounters::FMeshMemoryCounter>>;
+
+		/** Non-persistent internal id unique for a mesh generated for a specific state and parameter values. */
+		mutable uint32 InternalId = 0;
+
+		/** 
+		 * This is bit - mask on the EMeshFlags enumeration, marking what static formats are compatible with this one and other properties. 
+		 * It should be reset after any operation that modifies the format.
+		 */
+		mutable EMeshFlags Flags = EMeshFlags::None;
+
+		/** Only valid if the right flags are set, this identifies a referenced mesh. */
+		uint32 ReferenceID = 0;
+
+		/** 
+		 * Prefix for the unique IDs related to this mesh (vertices and layout blocks). Useful if the mesh stores them in an implicit, or relative way. 
+		 * See MeshVertexIdIterator for details.
+		 */
+		uint32 MeshIDPrefix = 0;
+
+		FMeshBufferSet VertexBuffers;
+
+		FMeshBufferSet IndexBuffers;
+
+		/** Additional buffers used for temporary or custom data in different algorithms. */
+		TArray<TPair<EMeshBufferType, FMeshBufferSet>> AdditionalBuffers;
+
+		TArray<FMeshSurface> Surfaces;
+
+		/** Externally provided SkeletonIDs of the skeletons required by this mesh. */
+		TArray<uint32> SkeletonIDs;
+
+		/** 
+		 * This skeleton and physics body are not owned and may be used by other meshes, so it cannot be modified
+		 * once the mesh has been fully created.
+		 */
+		Ptr<const Skeleton> Skeleton;
+		Ptr<const PhysicsBody> PhysicsBody;
+
+		/** Additional physics bodies referenced by the mesh that don't merge. */
+		TArray<Ptr<const mu::PhysicsBody>> AdditionalPhysicsBodies;
+
+		/** 
+		 * Texture Layout blocks attached to this mesh. They are const because they could be shared with
+		 * other meshes, so they need to be cloned and replaced if a modification is needed.
+		 */
+		TArray<Ptr<const Layout>> Layouts;		
+
+		TArray<FString> Tags;
+
+		/** Opaque handle to external resources. */
+		TArray<uint64> StreamedResources;
+
+		struct FBonePose
+		{
+			// Identifier built from the bone FName.
+			FBoneName BoneId;
+
+			EBoneUsageFlags BoneUsageFlags = EBoneUsageFlags::None;
+			FTransform3f BoneTransform;
+
+			inline void Serialise(OutputArchive& arch) const;
+			inline void Unserialise(InputArchive& arch);
+
+			inline bool operator==(const FBonePose& Other) const
+			{
+				return BoneUsageFlags == Other.BoneUsageFlags && BoneId == Other.BoneId;
+			}
+		};
+
+		/** 
+		 * This is the pose used by this mesh fragment, used to update the transforms of the final skeleton
+		 * taking into consideration the meshes being used.
+		 */
+		TMemoryTrackedArray<FBonePose> BonePoses;
+
+		/** Array containing the bonemaps of all surfaces in the mesh. */
+		TArray<FBoneName> BoneMap;
+
+		inline void Serialise(OutputArchive& arch) const;
+
+		inline void Unserialise(InputArchive& arch);
+
+
+		inline bool operator==(const Mesh& o) const
+		{
+			bool bEqual = true;
+
+			if (bEqual) bEqual = (ReferenceID == o.ReferenceID);
+			if (bEqual) bEqual = (MeshIDPrefix == o.MeshIDPrefix);
+			if (bEqual) bEqual = (IndexBuffers == o.IndexBuffers);
+			if (bEqual) bEqual = (VertexBuffers == o.VertexBuffers);
+			if (bEqual) bEqual = (Layouts.Num() == o.Layouts.Num());
+			if (bEqual) bEqual = (BonePoses.Num() == o.BonePoses.Num());
+			if (bEqual) bEqual = (BoneMap.Num() == o.BoneMap.Num());
+			if (bEqual && Skeleton != o.Skeleton)
+			{
+				if (Skeleton && o.Skeleton)
+				{
+					bEqual = (*Skeleton == *o.Skeleton);
+				}
+				else
+				{
+					bEqual = false;
+				}
+			}
+			if (bEqual) bEqual = (StreamedResources == o.StreamedResources);
+			if (bEqual) bEqual = (Surfaces == o.Surfaces);
+			if (bEqual) bEqual = (Tags == o.Tags);
+			if (bEqual) bEqual = (SkeletonIDs == o.SkeletonIDs);
+
+			for (int32 i = 0; bEqual && i < Layouts.Num(); ++i)
+			{
+				bEqual &= (*Layouts[i]) == (*o.Layouts[i]);
+			}
+
+			bEqual &= AdditionalBuffers.Num() == o.AdditionalBuffers.Num();
+			for (int32 i = 0; bEqual && i < AdditionalBuffers.Num(); ++i)
+			{
+				bEqual &= AdditionalBuffers[i] == o.AdditionalBuffers[i];
+			}
+
+			bEqual &= BonePoses.Num() == o.BonePoses.Num();
+			for (int32 i = 0; bEqual && i < BonePoses.Num(); ++i)
+			{
+				bEqual &= BonePoses[i] == o.BonePoses[i];
+			}
+
+			if (bEqual) bEqual = BoneMap == o.BoneMap;
+
+			bEqual &= AdditionalPhysicsBodies.Num() == o.AdditionalPhysicsBodies.Num();
+			for (int32 i = 0; bEqual && i < AdditionalPhysicsBodies.Num(); ++i)
+			{
+				bEqual &= *AdditionalPhysicsBodies[i] == *o.AdditionalPhysicsBodies[i];
+			}
+
+			return bEqual;
+		}
+
+		/** 
+		 * Compare the mesh with another one, but ignore internal data like generated vertex
+		 * indices.
+		 */
+		bool IsSimilar(const Mesh& Other, bool bCompareLayouts) const;
+
+
+		/**
+		 * Make a map from the vertices in this mesh to thefirst matching vertex of the given
+		 * mesh. If non is found, the index is set to -1.
+		 */
+		struct FVertexMatchMap
+		{
+			/** One for every vertex */
+			TArray<int32> FirstMatch;
+
+			/** The matches of every vertex in a sequence */
+			TArray<int32> Matches;
+
+			bool DoMatch(int32 Vertex, int32 OtherVertex) const;
+		};
+
+		void GetVertexMap(const Mesh& Other, FVertexMatchMap& VertexMap, float Tolerance = 1e-3f) const;
+
+		/** Compare the vertex attributes to check if they match. */
+		UE::Math::TIntVector3<uint32> GetFaceVertexIndices(int32 f) const;
+
+		/** 
+		 * Return true if the given mesh has the same vertex and index formats, and in the same
+		 * buffer structure.
+		 */
+		bool HasCompatibleFormat(const Mesh* Other) const;
+
+		/** Update the flags identifying the mesh format as some of the optimised formats. */
+		void ResetStaticFormatFlags() const;
+
+		/** Create the surface data if not present. */
+		void EnsureSurfaceData();
+
+		/** Check mesh buffer data for possible inconsistencies */
+		void CheckIntegrity() const;
+
+		/** 
+		 * Change the buffer descriptions so that all buffer indices start at 0 and are in the
+		 * same order than memory.
+		 */
+		void ResetBufferIndices();
+
+		/** Debug: get a text representation of the mesh */
+		void Log(FString& Out, int32 VertrexLimit);
+    };
+
+
+	MUTABLE_DEFINE_ENUM_SERIALISABLE(EBoneUsageFlags)
+	MUTABLE_DEFINE_ENUM_SERIALISABLE(EMeshBufferType)
+	MUTABLE_DEFINE_ENUM_SERIALISABLE(EShapeBindingMethod)
+	MUTABLE_DEFINE_ENUM_SERIALISABLE(EVertexColorUsage)
+}
+
diff --git a/Engine/Plugins/Mutable/Source/MutableRuntime/Public/MuR/Model.h b/Engine/Plugins/Mutable/Source/MutableRuntime/Public/MuR/Model.h
new file mode 100644
--- /dev/null
+++ b/Engine/Plugins/Mutable/Source/MutableRuntime/Public/MuR/Model.h
@@ -0,0 +1,224 @@
+// Copyright Epic Games, Inc. All Rights Reserved.
+
+#pragma once
+
+#include "MuR/Parameters.h"
+#include "MuR/Ptr.h"
+#include "MuR/RefCounted.h"
+#include "Templates/Function.h"
+#include "Templates/SharedPointer.h"
+
+
+namespace mu
+{
+	class InputArchive;
+	class ModelWriter;
+	class OutputArchive;
+
+	typedef Ptr<Parameters> ParametersPtr;
+	typedef Ptr<const Parameters> ParametersPtrConst;
+
+    class ModelParametersGenerator;
+
+    typedef Ptr<ModelParametersGenerator> ModelParametersGeneratorPtr;
+    typedef Ptr<const ModelParametersGenerator> ModelParametersGeneratorPtrConst;
+
+    class System;
+
+    typedef Ptr<System> SystemPtr;
+    typedef Ptr<const System> SystemPtrConst;
+
+
+	/** Flags for the model data chunks. */
+	enum class ERomFlags : uint16
+	{
+		/** Standard data. */
+		None = 0,
+
+		/** Bigger mips and mesh lods that are optional in some devices of a platform. */
+		HighRes = 1 << 0,
+	};
+
+
+    //! \brief A Model represents a customisable object with any number of parameters.
+    //!
+    //! When values are given to the parameters, specific Instances can be built, which hold the
+    //! built application-usable data.
+	//! \ingroup runtime
+    class MUTABLERUNTIME_API Model
+	{
+	public:
+
+		//-----------------------------------------------------------------------------------------
+		// Life cycle
+		//-----------------------------------------------------------------------------------------
+		Model();
+
+		//! Don't call directly. Manage with a TSharedPtr.
+		~Model();
+
+		static void Serialise( const Model* p, OutputArchive& arch );
+		static TSharedPtr<Model> StaticUnserialise( InputArchive& arch );
+
+		/** Special serialise operation that serialises the data in separate "files". An object
+         * with the ModelStreamer interface is responsible of storing this data and providing
+		 * the "file" concept. 
+		 * If bDropData is set to true, the rom data will be freed as it is serialized.
+		 */
+        static void Serialise( Model* p, ModelWriter& arch, bool bDropData=false );
+
+		//! Return true if the model has external data in other files. This kind of models will
+		//! require data streaming when used.
+		bool HasExternalData() const;
+
+#if WITH_EDITOR
+		//! Return true unless the streamed resources were destroyed, which could happen in the
+		//! editor after recompiling the CO.
+		bool IsValid() const;
+
+		//! Invalidate the Model. Compiling a compiled CO will invalidate the model kept by previously
+		//! generated resources, like streamed textures.
+		void Invalidate();
+#endif
+
+		//-----------------------------------------------------------------------------------------
+		// Own interface
+		//-----------------------------------------------------------------------------------------
+
+		//! Create a set of new parameters of the model with the default values.
+		//! If old parameters are provided, they will be reused when possible instead of the
+		//! default values.
+        static ParametersPtr NewParameters( TSharedPtr<const Model> Model, const Parameters* pOld = nullptr );
+
+		/** Return true if the parameter is multi-dimensional */
+		bool IsParameterMultidimensional(int32 ParameterIndex) const;
+
+		//! Get the number of states in the model.
+		int GetStateCount() const;
+
+		//! Get a state name by state index from 0 to GetStateCount-1
+		const FString& GetStateName( int32 StateIndex ) const;
+
+		//! Find a state index by state name
+		int32 FindState( const FString& Name ) const;
+
+		//! Get the number of parameters available in a particular state.
+		int GetStateParameterCount( int stateIndex ) const;
+
+		//! Get the index of one of the parameters in the given state. The index refers to the
+		//! parameters in a Parameters object obtained from this model with NewParameters.
+		int GetStateParameterIndex( int stateIndex, int paramIndex ) const;
+
+        //! Free memory used by streaming data that may be loaded again when needed.
+        void UnloadExternalData();
+    	
+		//! Return the default value of a boolean parameter.
+		//! \pre The parameter specified by index is a T_BOOL.
+		//! \param Index Index of the parameter from 0 to GetCount()-1
+    	bool GetBoolDefaultValue(int32 Index) const;
+
+   		//! Return the default value of a integer parameter.
+		//! \pre The parameter specified by index is a T_INT.
+		//! \param Index Index of the parameter from 0 to GetCount()-1
+    	int32 GetIntDefaultValue(int32 Index) const;
+
+		//! Return the default value of a float parameter.
+		//! \pre The parameter specified by index is a T_FLOAT.
+		//! \param Index Index of the parameter from 0 to GetCount()-1
+        float GetFloatDefaultValue(int32 Index) const;
+
+		//! Return the default value of a colour parameter.
+		//! \pre The parameter specified by index is a T_FLOAT.
+        //! \param Index Index of the parameter from 0 to GetCount()-1
+        //! \param R,G,B Pointers to values where every resulting colour channel will be stored
+    	void GetColourDefaultValue(int32 Index, float* R, float* G, float* B, float* A) const;
+
+    	//! Return the default value of a projector parameter, as a 4x4 matrix. The matrix is supposed to be
+		//! a linear transform in column-major.
+		//! \pre The parameter specified by index is a T_PROJECTOR.
+        //! \param Index Index of the parameter from 0 to GetCount()-1
+        //! \param OutPos Pointer to where the object-space position coordinates of the projector will be stored.
+        //! \param OutDir Pointer to where the object-space direction vector of the projector will be stored.
+        //! \param OutUp Pointer to where the object-space vertically up direction vector
+        //!         of the projector will be stored. This controls the "roll" angle of the
+        //!         projector.
+        //! \param OutScale Pointer to the projector-space scaling of the projector.
+    	void GetProjectorDefaultValue(int32 Index, PROJECTOR_TYPE* OutProjectionType, FVector3f* OutPos,
+    	 	FVector3f* OutDir, FVector3f* OutUp, FVector3f* OutScale, float* OutProjectionAngle) const;
+
+        //! Return the default value of an image parameter.
+        //! \pre The parameter specified by index is a T_IMAGE.
+        //! \param Index Index of the parameter from 0 to GetCount()-1
+		//! \return The externalId specified when setting the image value (\see SetImageValue)
+		FName GetImageDefaultValue(int32 Index) const;
+
+    	int32 GetRomCount() const;
+
+		uint32 GetRomId(int32 Index) const;
+
+		uint32 GetRomSize(int32 Index) const;
+
+		ERomFlags GetRomFlags(int32 Index) const;
+    	
+		//-----------------------------------------------------------------------------------------
+		// Interface pattern
+		//-----------------------------------------------------------------------------------------
+		class Private;
+
+		Private* GetPrivate() const;
+
+	private:
+
+		Private* m_pD;
+
+	};
+
+
+    //! \brief This class enumerates all possible model instances and generates the parameters for
+    //! each of them.
+    //!
+    //! This class can be used to generate all the possible instances of a model, or to generate
+    //! random instances. It only takes into account discrete parameters (bools and ints) but not
+    //! continuous parameters like colours and floats, which may be set to random values.
+    //!
+    //! \ingroup runtime
+    class MUTABLERUNTIME_API ModelParametersGenerator : public RefCounted
+    {
+    public:
+
+        //!
+        ModelParametersGenerator( TSharedPtr<const Model>, System* );
+
+        //! Return the number of different possible instances that can be built from the model.
+        int64 GetInstanceCount();
+
+        //! Return the parameters of one of the possible instances of the model.
+        //! \param index is an index in the range of 0 to GetInstanceCount-1
+        //! \param randomGenerator if not null, this function will be used to generate random values
+        //! for the continuous parameters of the instance.
+        ParametersPtr GetInstance( int64 index, TFunction<float()> RandomGenerator);
+
+        //! Return the parameters of one of a random instance of the model.
+        //! \param randomGenerator This function will be used to generate random values
+        //! for the continuous parameters of the instance.
+        ParametersPtr GetRandomInstance(TFunctionRef<float()> RandomGenerator );
+
+        //-----------------------------------------------------------------------------------------
+        // Interface pattern
+        //-----------------------------------------------------------------------------------------
+        class Private;
+
+        Private* GetPrivate() const;
+
+    protected:
+
+        //! Forbidden. Manage with the Ptr<> template.
+        ~ModelParametersGenerator();
+
+    private:
+
+        Private* m_pD;
+    };
+
+}
+
diff --git a/Engine/Plugins/Mutable/Source/MutableRuntime/Public/MuR/MutableMemory.h b/Engine/Plugins/Mutable/Source/MutableRuntime/Public/MuR/MutableMemory.h
new file mode 100644
--- /dev/null
+++ b/Engine/Plugins/Mutable/Source/MutableRuntime/Public/MuR/MutableMemory.h
@@ -0,0 +1,19 @@
+// Copyright Epic Games, Inc. All Rights Reserved.
+
+#pragma once
+
+#include "CoreTypes.h"
+
+namespace mu
+{
+	/** This should be called only before starting to use the library or any of its data. */
+	MUTABLERUNTIME_API extern void Initialize();
+
+    /** Shutdown everything related to mutable.
+     * This should be called when no other mutable objects will be used, created or destroyed,
+     * and all threads involving mutable tasks have been terminated.
+	 */
+	MUTABLERUNTIME_API extern void Finalize();
+
+}
+
diff --git a/Engine/Plugins/Mutable/Source/MutableRuntime/Public/MuR/MutableRuntimeModule.h b/Engine/Plugins/Mutable/Source/MutableRuntime/Public/MuR/MutableRuntimeModule.h
new file mode 100644
--- /dev/null
+++ b/Engine/Plugins/Mutable/Source/MutableRuntime/Public/MuR/MutableRuntimeModule.h
@@ -0,0 +1,31 @@
+// Copyright Epic Games, Inc. All Rights Reserved.
+
+#pragma once
+
+#include "Modules/ModuleInterface.h"
+#include "Logging/LogMacros.h"
+
+
+// This will use unreal's malloc instead of system malloc for mutable memory operations
+#define USE_UNREAL_ALLOC_IN_MUTABLE				1
+
+// This will enable the internal memory manager for mutable operations, reducing the number of operations
+// but increasing the memory usage slightly. This only affects packaged builds, and it is not used in editor.
+#define USE_INTERNAL_MEMORY_MANAGER_IN_MUTABLE	1
+
+// This slows down mutable, but gathers some data about memory usage for debug
+#define USE_STAT_MUTABLE_MEMORY					0
+
+
+MUTABLERUNTIME_API DECLARE_LOG_CATEGORY_EXTERN(LogMutableCore, Log, All);
+
+
+class MUTABLERUNTIME_API FMutableRuntimeModule : public IModuleInterface
+{
+public:
+
+	// IModuleInterface 
+	virtual void StartupModule() override;
+	virtual void ShutdownModule() override;
+
+};
diff --git a/Engine/Plugins/Mutable/Source/MutableRuntime/Public/MuR/MutableString.h b/Engine/Plugins/Mutable/Source/MutableRuntime/Public/MuR/MutableString.h
new file mode 100644
--- /dev/null
+++ b/Engine/Plugins/Mutable/Source/MutableRuntime/Public/MuR/MutableString.h
@@ -0,0 +1,91 @@
+// Copyright Epic Games, Inc. All Rights Reserved.
+
+#pragma once
+
+#include "Misc/AssertionMacros.h"
+#include "MuR/Ptr.h"
+#include "MuR/RefCounted.h"
+#include "MuR/Serialisation.h"
+#include "MuR/SerialisationPrivate.h"
+
+
+namespace mu
+{
+
+    //! \brief Container for text string data.
+	//! \ingroup runtime
+    class MUTABLERUNTIME_API String : public Resource
+	{
+	public:
+
+		//-----------------------------------------------------------------------------------------
+		// Life cycle
+		//-----------------------------------------------------------------------------------------
+
+		String(const FString&);
+
+		//! Deep clone this string.
+		Ptr<String> Clone() const;
+
+		//! Serialisation
+		//static void Serialise( const String* p, OutputArchive& arch );
+		//static Ptr<String> StaticUnserialise( InputArchive& arch );
+
+		// Resource interface
+		int32 GetDataSize() const override;
+
+		//-----------------------------------------------------------------------------------------
+		// Own interface
+		//-----------------------------------------------------------------------------------------
+
+		//! Get the string data.
+		const FString& GetValue() const;
+
+	protected:
+
+		//! Forbidden. Manage with the Ptr<> template.
+		~String() {}
+
+	public:
+
+		//!
+		FString Value;
+
+		//!
+		//void Serialise(OutputArchive& arch) const
+		//{
+		//	uint32 ver = 1;
+		//	arch << ver;
+
+		//	arch << Value;
+		//}
+
+		////!
+		//void Unserialise(InputArchive& arch)
+		//{
+		//	uint32 ver;
+		//	arch >> ver;
+		//	check(ver <= 1);
+
+		//	if (ver == 0)
+		//	{
+		//		std::string Temp;
+		//		arch >> Temp;
+		//		Value = Temp.c_str();
+		//	}
+		//	else
+		//	{
+		//		arch >> Value;
+		//	}
+		//}
+
+		//!
+		inline bool operator==(const String& o) const
+		{
+			return (Value == o.Value);
+		}
+
+	};
+
+}
+
diff --git a/Engine/Plugins/Mutable/Source/MutableRuntime/Public/MuR/Parameters.h b/Engine/Plugins/Mutable/Source/MutableRuntime/Public/MuR/Parameters.h
new file mode 100644
--- /dev/null
+++ b/Engine/Plugins/Mutable/Source/MutableRuntime/Public/MuR/Parameters.h
@@ -0,0 +1,386 @@
+// Copyright Epic Games, Inc. All Rights Reserved.
+
+#pragma once
+
+#include "HAL/Platform.h"
+#include "MuR/Ptr.h"
+#include "MuR/RefCounted.h"
+#include "MuR/Types.h"
+#include "Templates/SharedPointer.h"
+#include "VectorTypes.h"
+
+
+namespace mu
+{
+
+	//!
+	class InputArchive;
+	class OutputArchive;
+
+	class Parameters;
+    using ParametersPtr = Ptr<Parameters>;
+    using ParametersPtrConst = Ptr<const Parameters>;
+
+    class RangeIndex;
+
+	using RangeIndexPtr = Ptr<RangeIndex>;
+	using RangeIndexPtrConst = Ptr<const RangeIndex>;
+
+	class Model;
+
+	//! Default parameter values definition
+	static const FVector4f DefaultMutableColorValue(-UE_BIG_NUMBER);
+
+
+    //! \brief Model parameter types.
+	//! \ingroup runtime
+    enum class PARAMETER_TYPE : uint32
+    {
+        //! Undefined parameter type.
+        T_NONE,
+
+        //! Boolean parameter type (true or false)
+        T_BOOL,
+
+        //! Integer parameter type. It usually has a limited range of possible values that can be
+        //! queried in the Parameters object.
+        T_INT,
+
+        //! Floating point value in the range of 0.0 to 1.0
+        T_FLOAT,
+
+        //! Floating point RGBA colour, with each channel ranging from 0.0 to 1.0
+        T_COLOUR,
+
+        //! 3D Projector type, defining a position, scale and orientation. Basically used for
+        //! projected decals.
+        T_PROJECTOR,
+
+        //! An externally provided image.
+        T_IMAGE,
+
+        //! A text string.
+        T_STRING,
+
+        //! Utility enumeration value, not really a parameter type.
+        T_COUNT
+
+    };
+
+
+    //! \brief Types of projectors.
+    //! \ingroup runtime
+    enum class PROJECTOR_TYPE : uint32
+    {
+        //!
+        PLANAR,
+
+        //!
+        CYLINDRICAL,
+
+		//!
+        WRAPPING,
+
+        //! Utility enumeration value, not really a projector type.
+        COUNT
+    };
+
+
+    //! \brief Class used to set or read multi-dimensional parameter values.
+    //! If parameters have multiple values because of ranges, RangeIndex can be used to specify
+    //! what value is read or set in the methods of the Parameter class.
+    //! RangeIndex objects can be reused in multiple calls for the same parameter in the Parameter
+    //! class interface even after changing the position values.
+    //! \ingroup runtime
+    class MUTABLERUNTIME_API RangeIndex : public RefCounted
+    {
+    public:
+
+        //-----------------------------------------------------------------------------------------
+        // Life cycle
+        //-----------------------------------------------------------------------------------------
+
+
+        //-----------------------------------------------------------------------------------------
+        // Own interface
+        //-----------------------------------------------------------------------------------------
+
+        //! Return the number of ranges (or dimensions) used by this index
+        int GetRangeCount() const;
+
+        //! Return the name of a range.
+        //! \param index Index of the range from 0 to GetRangeCount()-1
+        const FString& GetRangeName( int index ) const;
+
+        //! Return the Guid of the parameter, resistant to parameter name changes.
+        //! \param index Index of the parameter from 0 to GetRangeCount()-1
+        const FString& GetRangeUid( int index ) const;
+
+        //! Set the position in one of the dimensions of the index.
+        //! \param index Index of the parameter from 0 to GetRangeCount()-1
+        void SetPosition( int index, int position );
+
+        //! Return the position in one of the dimensions of the index.
+        //! \param index Index of the parameter from 0 to GetRangeCount()-1
+        int GetPosition( int index ) const;
+
+
+        //-----------------------------------------------------------------------------------------
+        // Interface pattern
+        //-----------------------------------------------------------------------------------------
+        class Private;
+
+        Private* GetPrivate() const;
+
+    protected:
+
+        //! Forbid creation. Instances of this class must be obtained from a Parameters instance
+        RangeIndex();
+
+        //! Forbidden. Manage with the Ptr<> template.
+        ~RangeIndex() override;
+
+    private:
+
+        friend class Parameters;
+        Private* m_pD;
+
+    };
+
+
+    //! \brief  This class represents the parameters of a model including description, type and
+    //! value, and additional resources associated to the parameter.
+    //!
+    //! \warning Every object of this class holds a reference to the model that it was created from.
+    //! This implies that while any instance of a Parameters object is alive, its model will not
+    //! be freed.
+    //!
+	//! \ingroup runtime
+	class MUTABLERUNTIME_API Parameters : public RefCounted
+	{
+	public:
+
+		//-----------------------------------------------------------------------------------------
+		// Life cycle
+		//-----------------------------------------------------------------------------------------
+
+        //! Serialisation
+        static void Serialise( const Parameters* p, OutputArchive& arch );
+        static ParametersPtr StaticUnserialise( InputArchive& arch );
+
+        //! Portable serialisation, which tries to keep values when:
+        //! - new parameters have been added
+        //! - parameters have been removed
+        //! - parameter types have changed
+        //! - parameter names have changed, and the application is using UIDs
+        //static void SerialisePortable( const Parameters*, OutputArchive& );
+        //static ParametersPtr UnserialisePortable( InputArchive&, TSharedPtr<const Model> );
+
+		//! Deep clone this object.
+		ParametersPtr Clone() const;
+
+		//-----------------------------------------------------------------------------------------
+		// Own interface
+		//-----------------------------------------------------------------------------------------
+
+		//! Return the number of parameters
+		int GetCount() const;
+
+		//! Return the name of the parameter.
+		//! \param index Index of the parameter from 0 to GetCount()-1
+		const FString& GetName( int index ) const;
+
+		//! Return the Guid of the parameter, resistant to parameter name changes.
+		//! \param index Index of the parameter from 0 to GetCount()-1
+		const FGuid& GetUid( int index ) const;
+
+		//! Find the parameter index by name.
+		//! It returns -1 if the parameter doesn't exist.
+		int Find( const FString& Name ) const;
+
+		//! Return the type of the parameter.
+		//! \param index Index of the parameter from 0 to GetCount()-1
+		PARAMETER_TYPE GetType( int index ) const;
+
+        //! Create a new RangeIndex object to use to access a multi-dimensional parameter.
+        //! It will return nullptr if the parameter is not multidimensional.
+        Ptr<RangeIndex> NewRangeIndex( int paramIndex ) const;
+
+        //! Return the number of values other than the default that have been set to a specific
+        //! parameter.
+        int GetValueCount( int paramIndex ) const;
+
+        //! Return the RangeIndex of a value that has been set to a parameter.
+        //! \param paramIndex Index of the parameter from 0 to GetCount()-1
+        //! \param valueIndex Index of the value from 0 to GetValueCount()-1
+		Ptr<RangeIndex> GetValueIndex( int paramIndex, int valueIndex ) const;
+
+        //! Remove all the multidimensional values for a parameter. The current non-dimensional
+        //! value is kept.
+        void ClearAllValues( int paramIndex );
+
+		//! Return the value of a boolean parameter.
+		//! \pre The parameter specified by index is a T_BOOL.
+		//! \param index Index of the parameter from 0 to GetCount()-1
+        //! \param pos Only for multidimensional parameters: relevant position to get in the ranges
+        bool GetBoolValue( int index, const Ptr<const RangeIndex>& pos=nullptr ) const;
+
+		//! Set the value of a parameter of type boolean.
+		//! \pre The parameter specified by index is a T_BOOL.
+		//! \param index Index of the parameter from 0 to GetCount()-1
+		//! \param value new value of the parameter
+        //! \param pos optional parameter to set a specific value for the given multidimensional
+        //! position. If null, the value is set for all possible positions.
+        void SetBoolValue( int index, bool value, const Ptr<const RangeIndex>& pos=nullptr );
+
+		//! Get the number of possible values for this integer parameter.
+		//! If the number is zero, it means any integer value is accepted.
+		int GetIntPossibleValueCount( int paramIndex ) const;
+
+		//! Get the value of one of the possible values for this integer.
+		//! The paramIndex is in the range of 0 to GetIntPossibleValueCount()-1
+		int GetIntPossibleValue( int paramIndex, int valueIndex ) const;
+
+        //! Get the name of one of the possible values for this integer.
+        //! The paramIndex is in the range of 0 to GetIntPossibleValueCount()-1
+        const FString& GetIntPossibleValueName( int paramIndex, int valueIndex ) const;
+
+        //! Get the index of the value of one of the possible values for this integer.
+        //! The paramIndex is in the range of 0 to GetIntPossibleValueCount()-1
+		int GetIntValueIndex(int paramIndex, const FString& valueName) const;
+
+		//! Get the index of the value of one of the possible values for this integer.
+		//! The paramIndex is in the range of 0 to GetIntPossibleValueCount()-1
+		int GetIntValueIndex(int paramIndex, int32 Value) const;
+
+		//! Return the value of a integer parameter.
+		//! \pre The parameter specified by index is a T_INT.
+		//! \param index Index of the parameter from 0 to GetCount()-1
+        //! \param pos Only for multidimensional parameters: relevant position to get in the ranges
+        int GetIntValue( int index, const Ptr<const RangeIndex>& pos=nullptr ) const;
+
+		//! If the parameter is of the integer type, set its value.
+		//! \param index Index of the parameter from 0 to GetCount()-1
+		//! \param value new value of the parameter. It must be in the possible values for this
+		//! parameter (see GetIntPossibleValue), or the method will leave the value unchanged.
+        //! \param pos Only for multidimensional parameters: relevant position to set in the ranges
+        void SetIntValue( int index, int value, const Ptr<const RangeIndex>& pos=nullptr );
+
+		//! Return the value of a float parameter.
+		//! \pre The parameter specified by index is a T_FLOAT.
+		//! \param index Index of the parameter from 0 to GetCount()-1
+        //! \param pos Only for multidimensional parameters: relevant position to get in the ranges
+        float GetFloatValue( int index, const Ptr<const RangeIndex>& pos=nullptr ) const;
+
+		//! If the parameter is of the float type, set its value.
+		//! \param index Index of the parameter from 0 to GetCount()-1
+		//! \param value new value of the parameter
+        //! \param pos Only for multidimensional parameters: relevant position to set in the ranges
+        void SetFloatValue( int index, float value, const Ptr<const RangeIndex>& pos=nullptr );
+
+		//! Return the value of a colour parameter.
+		//! \pre The parameter specified by index is a T_FLOAT.
+        //! \param index Index of the parameter from 0 to GetCount()-1
+        //! \param pR,pG,pB Pointers to values where every resulting colour channel will be stored
+        //! \param pos Only for multidimensional parameters: relevant position to get in the ranges
+        void GetColourValue( int index, float* pR, float* pG, float* pB, float* pA,
+                             const Ptr<const RangeIndex>& pos=nullptr ) const;
+
+		//! If the parameter is of the colour type, set its value.
+		//! \param index Index of the parameter from 0 to GetCount()-1
+        //! \param r,g,b new value of the parameter
+        //! \param pos Only for multidimensional parameters: relevant position to set in the ranges
+        void SetColourValue( int index, float R, float G, float B, float A,
+                             const Ptr<const RangeIndex>& pos=nullptr );
+
+		//! Return the value of a projector parameter, as a 4x4 matrix. The matrix is supposed to be
+		//! a linear transform in column-major.
+		//! \pre The parameter specified by index is a T_PROJECTOR.
+        //! \param ParameterIndex Index of the parameter from 0 to GetCount()-1
+        //! \param OutPos Pointer to where the object-space position coordinates of the projector will be stored.
+        //! \param OutDir Pointer to where the object-space direction vector of the projector will be stored.
+        //! \param OutUp Pointer to where the object-space vertically up direction vector
+        //!         of the projector will be stored. This controls the "roll" angle of the
+        //!         projector.
+        //! \param OutScale Pointer to the projector-space scaling of the projector.
+        //! \param RangePosition Only for multidimensional parameters: relevant position to get in the ranges
+        void GetProjectorValue( int ParameterIndex,
+                                PROJECTOR_TYPE* OutProjectionType,
+								FVector3f* OutPos,
+								FVector3f* OutDir,
+								FVector3f* OutUp,
+								FVector3f* OutScale,
+                                float* OutProjectionAngle,
+                                const Ptr<const RangeIndex>& RangePosition=nullptr ) const;
+
+		//! If the parameter is of the projector type, set its value.
+		//! \param ParameterIndex Index of the parameter from 0 to GetCount()-1
+        //! \param pos Object-space position coordinates of the projector.
+        //! \param dir Object-space direction vector of the projector.
+        //! \param up Object-space vertically up direction vector of the projector.
+        //! \param scale Projector-space scaling of the projector.
+        //! \param projectionAngle [only for Cylindrical projectors], the angle in radians of the
+        //! projection area on the cylinder surface.
+        //! \param RangePosition Only for multidimensional parameters: relevant position to set in the ranges
+        void SetProjectorValue( int ParameterIndex,
+								const FVector3f& pos,
+								const FVector3f& dir,
+								const FVector3f& up,
+								const FVector3f& scale,
+                                float projectionAngle,
+                                const Ptr<const RangeIndex>& RangePosition=nullptr );
+
+        //! Return the value of an image parameter.
+        //! \pre The parameter specified by index is a T_IMAGE.
+        //! \param index Index of the parameter from 0 to GetCount()-1
+		//! \param pos Only for multidimensional parameters: relevant position to set in the ranges
+	    //! \return The externalId specified when setting the image value (\see SetImageValue)
+		FName GetImageValue( int index, const Ptr<const RangeIndex>& pos = nullptr) const;
+
+        //! If the parameter is of the image type, set its value.
+        //! \param index Index of the parameter from 0 to GetCount()-1
+        //! \param externalId Application-specific id used to identify this image during replication.
+		//! \param pos Only for multidimensional parameters: relevant position to set in the ranges
+		void SetImageValue( int index, FName externalId, const Ptr<const RangeIndex>& pos = nullptr);
+
+        //! Return the value of a float parameter.
+        //! \pre The parameter specified by index is a T_FLOAT.
+        //! \param index Index of the parameter from 0 to GetCount()-1
+        //! \param pos Only for multidimensional parameters: relevant position to get in the ranges
+        void GetStringValue( int index, FString& OutValue, const Ptr<const RangeIndex>& pos = nullptr) const;
+
+        //! If the parameter is of the float type, set its value.
+        //! \param index Index of the parameter from 0 to GetCount()-1
+        //! \param value new value of the parameter
+        //! \param pos Only for multidimensional parameters: relevant position to set in the ranges
+        void SetStringValue( int index, const FString& Value, const Ptr<const RangeIndex>& pos = nullptr );
+
+        //! Utility method to compare the values of a specific parameter with the values of another
+        //! Parameters object. It returns false if type or values are different.
+        bool HasSameValue( int thisParamIndex, const ParametersPtrConst& other, int otherParamIndex ) const;
+
+		//-----------------------------------------------------------------------------------------
+		// Interface pattern
+		//-----------------------------------------------------------------------------------------
+		class Private;
+
+		Private* GetPrivate() const;
+
+	protected:
+
+		//! Forbid creation. Instances of this class must be obtained from a Model instance
+		Parameters();
+
+		//! Forbidden. Manage with the Ptr<> template.
+		~Parameters();
+
+	private:
+
+		friend class Model;
+		Private* m_pD;
+
+	};
+
+
+}
+
diff --git a/Engine/Plugins/Mutable/Source/MutableRuntime/Public/MuR/PhysicsBody.h b/Engine/Plugins/Mutable/Source/MutableRuntime/Public/MuR/PhysicsBody.h
new file mode 100644
--- /dev/null
+++ b/Engine/Plugins/Mutable/Source/MutableRuntime/Public/MuR/PhysicsBody.h
@@ -0,0 +1,322 @@
+// Copyright Epic Games, Inc. All Rights Reserved.
+
+#pragma once
+
+#include <string>
+
+#include "MuR/Serialisation.h"
+#include "MuR/Ptr.h"
+#include "MuR/RefCounted.h"
+#include "MuR/Serialisation.h"
+#include "MuR/Skeleton.h"
+#include "Containers/Array.h"
+#include "Math/Quat.h"
+#include "Math/Transform.h"
+#include "Math/UnrealMathSSE.h"
+#include "Math/Vector.h"
+#include "Math/VectorRegister.h"
+#include "Misc/AssertionMacros.h"
+
+
+namespace mu
+{
+	struct FBodyShape
+	{
+		FString Name;
+		uint32 Flags = 0;
+
+		bool operator==(const FBodyShape& Other) const 
+		{
+			return Flags == Other.Flags && Name == Other.Name;
+		}
+	
+		//!
+		inline void Serialise( OutputArchive& arch ) const;
+
+		//!
+		inline void Unserialise( InputArchive& arch );
+	};
+
+	struct FSphereBody : FBodyShape
+	{
+		FVector3f Position = FVector3f::ZeroVector;
+		float Radius = 0.0f;
+
+		bool operator==( const FSphereBody& Other ) const
+		{
+			return FBodyShape::operator==(Other) && Position == Other.Position && Radius == Other.Radius;
+		}
+		 
+		//!
+		inline void Serialise( OutputArchive& arch ) const;
+
+		//!
+		inline void Unserialise( InputArchive& arch );
+	};
+	
+	struct FBoxBody : FBodyShape
+	{
+		FVector3f Position = FVector3f::ZeroVector;
+		FQuat4f Orientation = FQuat4f::Identity;
+		FVector3f Size = FVector3f::ZeroVector;
+		 
+		inline bool operator==( const FBoxBody& Other ) const 
+		{
+			return FBodyShape::operator==(Other) && Position == Other.Position && Orientation == Other.Orientation && Size == Other.Size;
+		}
+
+		//!
+		inline void Serialise( OutputArchive& arch ) const;
+
+		//!
+		inline void Unserialise( InputArchive& arch );
+	};
+
+	struct FSphylBody : FBodyShape
+	{
+		FVector3f Position = FVector3f::ZeroVector;
+		FQuat4f Orientation = FQuat4f::Identity;
+		float Radius = 0.0f;
+		float Length = 0.0f;
+
+		inline bool operator==( const FSphylBody& Other ) const 
+		{
+			return FBodyShape::operator==(Other) && 
+				Position == Other.Position && 
+				Orientation == Other.Orientation && 
+				Radius == Other.Radius && 
+				Length == Other.Length;
+		}
+		 
+		//!
+		inline void Serialise( OutputArchive& arch ) const;
+
+		//!
+		inline void Unserialise( InputArchive& arch );
+	};
+	
+	struct FTaperedCapsuleBody : FBodyShape
+	{
+		FVector3f Position;
+		FQuat4f Orientation;
+		float Radius0 = 0.0f;
+		float Radius1 = 0.0f;
+		float Length = 0.0f;
+		 
+		inline bool operator==( const FTaperedCapsuleBody& Other ) const 
+		{
+			return FBodyShape::operator==(Other) && 
+				Position == Other.Position && 
+				Orientation == Other.Orientation && 
+				Radius0 == Other.Radius0 && 
+				Radius1 == Other.Radius1 && 
+				Length == Other.Length;
+		}
+		 
+		//!
+		inline void Serialise( OutputArchive& arch ) const;
+
+		//!
+		inline void Unserialise( InputArchive& arch );
+	};
+	
+	struct FConvexBody : FBodyShape
+	{
+		TArray<FVector3f> Vertices;
+		TArray<int32> Indices;
+
+		FTransform3f Transform = FTransform3f::Identity;
+		
+		inline bool operator==( const FConvexBody& Other ) const
+		{
+			return FBodyShape::operator==(Other) && 
+				Vertices == Other.Vertices && 
+				Indices == Other.Indices && 
+				Transform.Equals(Other.Transform);
+		}
+
+		//!
+		inline void Serialise( OutputArchive& arch ) const;
+
+		//!
+		inline void Unserialise( InputArchive& arch );
+	};
+	
+	struct FPhysicsBodyAggregate
+	{
+		TArray<FSphereBody> Spheres;
+		TArray<FBoxBody> Boxes;
+		TArray<FConvexBody> Convex;
+		TArray<FSphylBody> Sphyls;
+		TArray<FTaperedCapsuleBody> TaperedCapsules;
+
+		inline bool operator==( const FPhysicsBodyAggregate& Other ) const
+		{
+			return Spheres == Other.Spheres &&
+					Boxes == Other.Boxes &&
+					Sphyls == Other.Sphyls &&
+					TaperedCapsules == Other.TaperedCapsules &&
+					Convex == Other.Convex;
+		}
+
+		 //!
+		inline void Serialise( OutputArchive& arch ) const;
+
+		//!
+		inline void Unserialise( InputArchive& arch );
+	};
+
+	class MUTABLERUNTIME_API PhysicsBody : public RefCounted
+	{
+	public:
+
+		PhysicsBody() 
+		{
+		}
+		
+		mu::Ptr<PhysicsBody> Clone() const; 
+
+		//! Serialisation
+        static void Serialise( const PhysicsBody* p, OutputArchive& arch );
+		static mu::Ptr<PhysicsBody> StaticUnserialise( InputArchive& arch );
+
+        //bool operator==( const PhysicsBody& Other ) const;
+
+		void SetCustomId(int32 Id);
+		int32 GetCustomId() const;
+
+		void SetBodyCount(int32 B);
+		int32 GetBodyCount() const;
+	
+		void SetBodyBoneId(int32 B, const FBoneName& BoneId);
+		const FBoneName& GetBodyBoneId(int32 B) const;
+		
+		void SetBodyCustomId(int32 B, int32 BodyCustomId);
+		int32 GetBodyCustomId(int32 B) const;
+		
+		int32 GetSphereCount(int32 B) const;
+		int32 GetBoxCount(int32 B) const;
+		int32 GetConvexCount(int32 B) const;
+		int32 GetSphylCount(int32 B) const;
+		int32 GetTaperedCapsuleCount(int32 B) const;
+		
+		void SetSphereCount(int32 B, int32 Count);
+		void SetBoxCount(int32 B, int32 Count);
+		void SetConvexCount(int32 B, int32 Count);
+		void SetSphylCount(int32 B, int32 Count);
+		void SetTaperedCapsuleCount(int32 B, int32 Count);
+		
+		void SetSphere( 
+				int32 B, int32 I, 
+				FVector3f Position, float Radius);
+
+		void SetBox( 
+				int32 B, int32 I, 
+				FVector3f Position, FQuat4f Orientation, FVector3f Size);
+
+		void SetConvexMesh( 
+				int32 B, int32 I,
+				TArrayView<const FVector3f> Vertices, TArrayView<const int32> Indices);
+
+		void SetConvexTransform( 
+				int32 B, int32 I, 
+				const FTransform3f& Transform);
+
+		void SetSphyl( 
+				int32 B, int32 I, 
+				FVector3f Position, FQuat4f Orientation, 
+				float Radius, float Length);
+
+		void SetTaperedCapsule( 
+				int32 B, int32 I, 
+				FVector3f Position, FQuat4f Orientation, 
+				float Radius0, float Radius1, float Length);
+
+		void SetSphereFlags(int32 B, int32 I, uint32 Flags);
+		void SetBoxFlags(int32 B, int32 I, uint32 Flags);
+		void SetConvexFlags(int32 B, int32 I, uint32 Flags);
+		void SetSphylFlags(int32 B, int32 I, uint32 Flags);
+		void SetTaperedCapsuleFlags(int32 B, int32 I, uint32 Flags);
+
+		void SetSphereName(int32 B, int32 I, const char* Name);
+		void SetBoxName(int32 B, int32 I, const char* Name);
+		void SetConvexName(int32 B, int32 I, const char* Name);
+		void SetSphylName(int32 B, int32 I, const char* Name);
+		void SetTaperedCapsuleName(int32 B, int32 I, const char* Name);
+	
+		void GetSphere( 
+				int32 B, int32 I, 
+				FVector3f& OutPosition, float& OutRadius) const;
+
+		void GetBox( 
+				int32 B, int32 I, 
+				FVector3f& OutPosition, FQuat4f& OutOrientation, FVector3f& OutSize) const;
+
+		void GetConvex( 
+				int32 B, int32 I, 
+				TArrayView<const FVector3f>& OutVertices, TArrayView<const int32>& OutIndices, FTransform3f& OutTransform) const;
+
+		void GetConvexMeshView(
+				int32 B, int32 I, 
+			    TArrayView<FVector3f>& OutVerticesView, TArrayView<int32>& OutIndicesView);
+
+		void GetConvexTransform(
+				int32 B, int32 I,
+				FTransform3f& OutTransform) const;
+
+		void GetSphyl( 
+				int32 B, int32 I, 
+				FVector3f& OutPosition, FQuat4f& OutOrientation, 
+				float& OutRadius, float& OutLength) const;
+
+		void GetTaperedCapsule( 
+				int32 B, int32 I, 
+				FVector3f& OutPosition, FQuat4f& OutOrientation, 
+				float& OutRadius0, float& OutRadius1, float& OutLength) const;
+		
+		uint32 GetSphereFlags(int32 B, int32 I) const;
+		uint32 GetBoxFlags(int32 B, int32 I) const;
+		uint32 GetConvexFlags(int32 B, int32 I) const;
+		uint32 GetSphylFlags(int32 B, int32 I) const;
+		uint32 GetTaperedCapsuleFlags(int32 B, int32 I) const;
+
+		const FString& GetSphereName(int32 B, int32 I) const;
+		const FString& GetBoxName(int32 B, int32 I) const;
+		const FString& GetConvexName(int32 B, int32 I) const;
+		const FString& GetSphylName(int32 B, int32 I) const;
+		const FString& GetTaperedCapsuleName(int32 B, int32 I) const;
+
+	protected:
+		//! Forbidden. Manage with the Ptr<> template.
+		~PhysicsBody()
+		{
+		}
+
+	public:
+		int32 CustomId = -1;
+
+		// Bone name the physics volume aggregate is bound to. 
+		TArray<FBoneName> BoneIds;
+		TArray<FPhysicsBodyAggregate> Bodies;
+		TArray<int32> BodiesCustomIds;
+
+
+		bool bBodiesModified = false;
+
+        //!
+        inline void Serialise( OutputArchive& arch ) const;
+
+		//!
+        inline void Unserialise( InputArchive& arch );
+
+		//!
+		inline bool operator==(const PhysicsBody& Other) const
+        {
+        	return CustomId        == Other.CustomId        &&
+				   BoneIds == Other.BoneIds					&&
+        		   Bodies          == Other.Bodies          &&
+        	       BodiesCustomIds == Other.BodiesCustomIds &&
+				   bBodiesModified == Other.bBodiesModified;
+        }
+	};
+}
diff --git a/Engine/Plugins/Mutable/Source/MutableRuntime/Public/MuR/Serialisation.h b/Engine/Plugins/Mutable/Source/MutableRuntime/Public/MuR/Serialisation.h
new file mode 100644
--- /dev/null
+++ b/Engine/Plugins/Mutable/Source/MutableRuntime/Public/MuR/Serialisation.h
@@ -0,0 +1,457 @@
+// Copyright Epic Games, Inc. All Rights Reserved.
+
+#pragma once
+
+#include <string>
+
+#include "HAL/Platform.h"
+#include "MuR/MutableMemory.h"
+#include "MuR/Ptr.h"
+#include "MuR/RefCounted.h"
+#include "MuR/MutableMath.h"
+
+#include "Curves/RichCurve.h"
+#include "Math/Vector.h"
+#include "Math/IntVector.h"
+#include "Math/Vector4.h"
+#include "Containers/Array.h"
+#include "Containers/Map.h"
+#include "Containers/ContainerAllocationPolicies.h"
+
+namespace mu
+{    
+    class Image;
+	class Model;
+
+#define MUTABLE_DEFINE_POD_SERIALISABLE(Type)					   \
+	void DLLEXPORT operator<<(OutputArchive& Arch, const Type& T); \
+	void DLLEXPORT operator>>(InputArchive& Arch, Type& T);		   \
+
+#define MUTABLE_DEFINE_POD_VECTOR_SERIALISABLE(Type)					\
+	template<typename Alloc>                                            \
+	void operator<<(OutputArchive& Arch, const TArray<Type, Alloc>& V);	\
+	                                                                    \
+	template<typename Alloc>                                            \
+	void operator>>(InputArchive& Arch, TArray<Type, Alloc>& V);        \
+
+#define MUTABLE_DEFINE_ENUM_SERIALISABLE(Type)					   \
+    void DLLEXPORT operator<<(OutputArchive& Arch, const Type& T); \
+    void DLLEXPORT operator>>(InputArchive& Arch, Type& T);		   \
+	
+	
+    /** This repesents a handle to Mutable resource. */
+    template<class R>
+    class ResourceProxy : public RefCounted
+    {
+    public:
+
+        virtual Ptr<const R> Get() = 0;
+
+    };
+
+
+	/** Proxy implementation that always has the resource loaded in memory. */
+	template<class R>
+	class ResourceProxyMemory : public ResourceProxy<R>
+	{
+	private:
+
+		Ptr<const R> m_resource;
+
+	public:
+
+		ResourceProxyMemory(const Ptr<const R>& i)
+		{
+			m_resource = i;
+		}
+
+		// ResourceProxy interface
+		Ptr<const R> Get() override
+		{
+			return m_resource;
+		}
+	};
+
+
+	/** */
+	class MUTABLERUNTIME_API ModelReader
+    {
+    public:
+
+        //-----------------------------------------------------------------------------------------
+        // Life cycle
+        //-----------------------------------------------------------------------------------------
+
+        //! Ensure virtual destruction.
+        virtual ~ModelReader() {}
+
+        //-----------------------------------------------------------------------------------------
+        // Reading interface
+        //-----------------------------------------------------------------------------------------
+
+        //! Identifier of reading data operations sent to this interface.
+		//! Negative values indicate an error.
+        typedef int32 OPERATION_ID;
+
+		//! \brief Start a data request operation.
+		//! \param Model.
+        //! \param key0 key identifying the model data fragment that is requested.
+        //!         This key interpretation depends on the implementation of the ModelStreamer,
+		//! \param Buffer is an already-allocated buffer big enough to receive the expected data.
+		//! \param size is the size of the pBuffer buffer, which must match the size of the data
+		//! requested with the key identifiers.
+		//! \param CompletionCallback Optional callback. Copied inside the called function. Will always be called.
+		//! \return a previously unused identifier, now used for this operation, that can be used in
+		//! calls to the other methods of this interface. If the return value is negative it indicates
+		//! an unrecoverable error.
+		virtual OPERATION_ID BeginReadBlock(const mu::Model*, uint32 BlockKey, void* Buffer, uint64 size, TFunction<void(bool bSuccess)>* CompletionCallback = nullptr) = 0;
+
+        //! Check if a data request operation has been completed.
+        //! This is a weak check than *may* return true if the given operation has completed, but
+        //! it is not mandatory. It is used as a hint by the System to optimise its opertaions.
+        //! There is no guarantee that this method will ever be called, and it is safe to always
+        //! return false.
+        virtual bool IsReadCompleted( OPERATION_ID ) = 0;
+
+        /** Complete a data request operation.This method has to block until a data request issued
+        * with BeginReadBlock has been completed. After returning from this call, the ID cannot be used
+        * any more to identify the same operation and becomes free.
+		* \return true if the data was loaded successfully.
+		*/
+        virtual bool EndRead( OPERATION_ID ) = 0;
+
+    };
+
+
+	/** */
+	class MUTABLERUNTIME_API ModelWriter
+	{
+	public:
+
+		//-----------------------------------------------------------------------------------------
+		// Life cycle
+		//-----------------------------------------------------------------------------------------
+
+		//! Ensure virtual destruction.
+		virtual ~ModelWriter() {}
+
+		//-----------------------------------------------------------------------------------------
+		// Writing interface
+		//-----------------------------------------------------------------------------------------
+
+		//! \brief Open a file for writing data.
+		//!
+		//! Only one file can be open a time for writing. This method is only required when
+		//! writing model data from tools.
+		//! \param strModelName Name of the model where the file to open belongs to. This string's
+		//!         meaning depends on the implementation of the ModelStreamer subclasses. It could
+		//!         be a file path, or a resource identifier, etc. It will be the same for all
+		//!         files in a model.
+		//! \param key key identifying the model data fragment that is requested.
+		//!         This key interpretation depends on the implementation of the ModelStreamer,
+		virtual void OpenWriteFile(uint32 BlockKey) = 0;
+
+		//! \brief Write a piece of data to the currently open file.
+		//!
+		//! There must be a file open with OpenWriteFile before calling this.
+		//! \param pBuffer pointer to the data that will be written to the file.
+		//! \param size size of the data to write on the file, in bytes.
+		virtual void Write(const void* pBuffer, uint64 size) = 0;
+
+		//! \brief Close the file open for writing in a previous call to OpenWriteFile in this
+		//! object.
+		virtual void CloseWriteFile() = 0;
+
+
+	};
+
+
+    /** Interface for any input stream to be use with InputArchives. */
+    class MUTABLERUNTIME_API InputStream
+    {
+    public:
+
+		/** Ensure virtual destruction. */
+		virtual ~InputStream() {}
+
+        /** Read a byte buffer
+         * \param pData destination buffer, must have at least size bytes allocated.
+         * \param size amount of bytes to read from the stream.
+		 */
+        virtual void Read( void* Data, uint64 Size ) = 0;
+    };
+
+
+    /** Interface for any output stream to be used with OutputArchives. */
+    class MUTABLERUNTIME_API OutputStream
+    {
+    public:
+
+        /** Ensure virtual destruction. */
+        virtual ~OutputStream() {}
+
+        /** Write a byte buffer
+         * \param pData source buffer where data will be read from.
+         * \param size amount of data to write to the stream.
+		 */
+        virtual void Write( const void* Data, uint64 Size ) = 0;
+
+    };
+
+
+    /** Archive containing data to be deserialised. */
+    class MUTABLERUNTIME_API InputArchive
+    {
+    public:
+
+        /** Construct form an input stream.The stream will not be owned by the archive and the
+         * caller must make sure it is not modified or destroyed while serialisation is happening.
+		 */
+        InputArchive( InputStream* );
+
+		/** Ensure virtual destruction. */
+		virtual ~InputArchive() {}
+
+        //
+        virtual Ptr<ResourceProxy<Image>> NewImageProxy();
+
+		/** Not owned. */
+		InputStream* Stream = nullptr;
+
+		/** Already read pointers. */
+		TArray< Ptr<RefCounted> > History;
+
+    };
+
+
+    /** Archive where data can be serialised to. */
+    class MUTABLERUNTIME_API OutputArchive
+    {
+    public:
+
+        /** Construct form an output stream.The stream will not be owned by the archive and the
+         * caller must make sure it is not modified or destroyed while serialisation is happening.
+		 */
+        OutputArchive( OutputStream* );
+
+		/** Not owned. */
+		OutputStream* Stream;
+
+		/** Already written pointers and their ids. */
+		TMap< const void*, int32 > History;
+
+    };
+
+
+    /** */
+    class MUTABLERUNTIME_API InputArchiveWithProxies : public InputArchive
+    {
+    public:
+
+        class ProxyFactory
+        {
+        public:
+          virtual ~ProxyFactory() {}
+          virtual Ptr<ResourceProxy<Image>> NewImageProxy(InputArchive& arch) = 0;
+        };
+
+    public:
+
+        /** Construct form an input stream.The stream will not be owned by the archive and the
+         * caller must make sure it is not modified or destroyed while serialisation is happening.
+		 */
+        InputArchiveWithProxies( InputStream*, ProxyFactory*  );
+
+        // InputArchive interface
+        virtual Ptr<ResourceProxy<Image>> NewImageProxy() override;
+
+    private:
+
+		TArray< Ptr<ResourceProxy<Image>> > ProxyHistory;
+		ProxyFactory* Factory = nullptr;
+	};
+
+
+    //!
+    class MUTABLERUNTIME_API OutputMemoryStream : public OutputStream
+    {
+    public:
+
+        /** Create the stream with an optional buffer size in bytes.
+		* The internal buffer will be enlarged as much as necessary.
+		*/
+        OutputMemoryStream( uint64 Reserve = 0 );
+
+        // OutputStream interface
+        virtual void Write( const void* Data, uint64 Size ) override;
+
+        // Own interface
+
+        /** Get the serialised data buffer pointer. This pointer invalidates after a Write
+         * operation has been done, and you need to get it again.
+		 */
+        const void* GetBuffer() const;
+
+        /** Get the amount of data in the stream, in bytes. */
+        uint64 GetBufferSize() const;
+
+		/** Clear the internal buffer. */
+		void Reset();
+
+    private:
+
+		TArray64<uint8> Buffer;
+
+    };
+
+
+	/** This stream doesn't store any data: it just counts the amount of data serialised. */
+	class MUTABLERUNTIME_API OutputSizeStream : public OutputStream
+	{
+	public:
+
+		// OutputStream interface
+		void Write(const void* pData, uint64 size) override;
+
+		// Own interface
+
+		/** Get the amount of data serialised, in bytes. */
+		uint64 GetBufferSize() const;
+
+	private:
+
+		uint64 WrittenBytes = 0;
+
+	};
+
+
+	/** This stream doesn't store any data: it just calculates of a hash of the data as it receives it. */
+	class MUTABLERUNTIME_API OutputHashStream : public OutputStream
+	{
+	public:
+
+		// OutputStream interface
+		void Write(const void* pData, uint64 size) override;
+
+		// Own interface
+
+		/** Return the hash of the data written so far. */
+		uint64 GetHash() const;
+
+	private:
+
+		uint64 Hash = 0;
+
+	};
+
+
+	template<typename Type>
+	void operator<<(OutputArchive& Arch, const Type& Value)
+	{
+        Value.Serialise(Arch);
+	}
+
+	template<typename Type>
+	void operator>>(InputArchive& Arch, Type& Value)
+	{
+        Value.Unserialise(Arch);
+	}
+
+	MUTABLE_DEFINE_POD_SERIALISABLE(float);
+	MUTABLE_DEFINE_POD_SERIALISABLE(double);
+
+    MUTABLE_DEFINE_POD_SERIALISABLE(int8);
+    MUTABLE_DEFINE_POD_SERIALISABLE(int16);
+    MUTABLE_DEFINE_POD_SERIALISABLE(int32);
+    MUTABLE_DEFINE_POD_SERIALISABLE(int64);
+
+    MUTABLE_DEFINE_POD_SERIALISABLE(uint8);
+    MUTABLE_DEFINE_POD_SERIALISABLE(uint16);
+    MUTABLE_DEFINE_POD_SERIALISABLE(uint32);
+	MUTABLE_DEFINE_POD_SERIALISABLE(uint64);
+
+	MUTABLE_DEFINE_POD_SERIALISABLE(FUintVector2);
+	MUTABLE_DEFINE_POD_SERIALISABLE(FIntVector2);
+	MUTABLE_DEFINE_POD_SERIALISABLE(UE::Math::TIntVector2<uint16>);
+	MUTABLE_DEFINE_POD_SERIALISABLE(UE::Math::TIntVector2<int16>);
+	MUTABLE_DEFINE_POD_SERIALISABLE(FVector2f);
+	MUTABLE_DEFINE_POD_SERIALISABLE(FVector4f);
+	MUTABLE_DEFINE_POD_SERIALISABLE(FMatrix44f);
+	MUTABLE_DEFINE_POD_SERIALISABLE(FRichCurveKey);
+
+	MUTABLE_DEFINE_POD_SERIALISABLE(FGuid);
+
+	MUTABLE_DEFINE_POD_VECTOR_SERIALISABLE(float);
+	MUTABLE_DEFINE_POD_VECTOR_SERIALISABLE(double);
+	MUTABLE_DEFINE_POD_VECTOR_SERIALISABLE(uint8);
+	MUTABLE_DEFINE_POD_VECTOR_SERIALISABLE(uint16);
+	MUTABLE_DEFINE_POD_VECTOR_SERIALISABLE(uint32);
+	MUTABLE_DEFINE_POD_VECTOR_SERIALISABLE(uint64);
+	MUTABLE_DEFINE_POD_VECTOR_SERIALISABLE(int8);
+	MUTABLE_DEFINE_POD_VECTOR_SERIALISABLE(int16);
+	MUTABLE_DEFINE_POD_VECTOR_SERIALISABLE(int32);
+	MUTABLE_DEFINE_POD_VECTOR_SERIALISABLE(int64);
+
+	MUTABLE_DEFINE_POD_VECTOR_SERIALISABLE(FVector2f);
+	MUTABLE_DEFINE_POD_VECTOR_SERIALISABLE(FMatrix44f);
+	MUTABLE_DEFINE_POD_VECTOR_SERIALISABLE(FIntVector2);
+	MUTABLE_DEFINE_POD_VECTOR_SERIALISABLE(TCHAR);
+	MUTABLE_DEFINE_POD_VECTOR_SERIALISABLE(FRichCurveKey);
+
+	MUTABLE_DEFINE_POD_VECTOR_SERIALISABLE(FUintVector2);
+	MUTABLE_DEFINE_POD_VECTOR_SERIALISABLE(UE::Math::TIntVector2<uint16>);
+	MUTABLE_DEFINE_POD_VECTOR_SERIALISABLE(UE::Math::TIntVector2<int16>);
+	MUTABLE_DEFINE_POD_VECTOR_SERIALISABLE(FVector4f);
+	
+	//---------------------------------------------------------------------------------------------
+	MUTABLERUNTIME_API void operator<<(OutputArchive& arch, const FString& t);
+	
+	MUTABLERUNTIME_API void operator>>(InputArchive& arch, FString& t);
+
+	//---------------------------------------------------------------------------------------------
+	MUTABLERUNTIME_API void operator<<(OutputArchive& arch, const FRichCurve& t);
+
+	MUTABLERUNTIME_API void operator>>(InputArchive& arch, FRichCurve& t);
+
+
+	//---------------------------------------------------------------------------------------------
+	MUTABLERUNTIME_API void operator<<(OutputArchive& arch, const FName& v);
+
+	MUTABLERUNTIME_API void operator>>(InputArchive& arch, FName& v);
+
+
+	//---------------------------------------------------------------------------------------------
+	template<typename T, typename Alloc> 
+	void operator<<(OutputArchive& Arch, const TArray<T, Alloc>& V)
+	{
+		const uint32 Num = (uint32)V.Num();
+		Arch << Num;
+		
+		for (SIZE_T Index = 0; Index < Num; ++Index)
+		{
+			Arch << V[Index];
+		}
+	}
+
+	template<typename T, typename Alloc> 
+	void operator>>(InputArchive& Arch, TArray<T, Alloc>& V)
+	{
+		uint32 Num;
+		Arch >> Num;
+		V.SetNum(Num);
+
+		for (SIZE_T Index = 0; Index < Num; ++Index)
+		{
+			Arch >> V[Index];
+		}
+	}
+
+	
+	MUTABLERUNTIME_API void operator>> ( InputArchive& arch, std::string& v );
+
+	// Bool size is not a standard
+	MUTABLERUNTIME_API void operator<<(OutputArchive& Arch, const bool& T);
+
+	MUTABLERUNTIME_API void operator>>(InputArchive& Arch, bool& T);
+}
+
diff --git a/Engine/Plugins/Mutable/Source/MutableRuntime/Public/MuR/Settings.h b/Engine/Plugins/Mutable/Source/MutableRuntime/Public/MuR/Settings.h
new file mode 100644
--- /dev/null
+++ b/Engine/Plugins/Mutable/Source/MutableRuntime/Public/MuR/Settings.h
@@ -0,0 +1,63 @@
+// Copyright Epic Games, Inc. All Rights Reserved.
+
+#pragma once
+
+#include "MuR/Ptr.h"
+#include "MuR/RefCounted.h"
+
+
+namespace mu
+{
+    /** Settings class that can be used to create a customised System instance. */
+    class MUTABLERUNTIME_API Settings : public RefCounted
+    {
+    public:
+
+        /** Create new settings with the default values as explained in each method below. */
+        Settings();
+
+        /** Record internal profiling data. Disabled by default. */
+        void SetProfile( bool bEnabled );
+
+        /** Limit the maximum memory in bytes used by the mutable core.A low value will force
+        * more streaming and higher instance construction times, but will use less memory while
+        * building objects.
+        * Defaults to 0, which disables the limit.
+		*/
+        void SetWorkingMemoryBytes( uint64 Bytes );
+
+        /** Set the quality for the image compression algorithms.
+        * \param quality Quality level of the compression. Higher quality increases the instance
+        * generation time. The value is in the range 0 to 4 with 0 being the fastest. It defaults
+        * to 0, which is the fastes for runtime. Tools may want ot set it higher, and low
+        * performance profiles to lower.
+        * Genreal rules are:
+        * 0 - Fastest for runtime
+        * 1 - Best for runtime
+        * 2 - Fast for tools
+        * 3 - Best for tools
+        * 4 - Maximum, with no time limits.
+        */
+        void SetImageCompressionQuality( int32 Quality );
+
+        // Interface pattern
+        Settings( const Settings& ) = delete;
+        Settings( Settings&& ) = delete;
+        Settings& operator=(const Settings&) = delete;
+        Settings& operator=(Settings&&) = delete;
+
+    protected:
+
+        /** Forbidden.Manage with the Ptr<> template. */
+        ~Settings() override;
+
+    public:
+
+		uint64 WorkingMemoryBytes = 0;
+		int32 ImageCompressionQuality = 0;
+		bool bProfile = false;
+
+    };
+
+}
+
diff --git a/Engine/Plugins/Mutable/Source/MutableRuntime/Public/MuR/Skeleton.h b/Engine/Plugins/Mutable/Source/MutableRuntime/Public/MuR/Skeleton.h
new file mode 100644
--- /dev/null
+++ b/Engine/Plugins/Mutable/Source/MutableRuntime/Public/MuR/Skeleton.h
@@ -0,0 +1,138 @@
+// Copyright Epic Games, Inc. All Rights Reserved.
+
+#pragma once
+
+#include "MuR/Serialisation.h"
+#include "Containers/Array.h"
+#include "HAL/PlatformMath.h"
+#include "Math/Transform.h"
+#include "Math/VectorRegister.h"
+#include "Misc/AssertionMacros.h"
+#include "MuR/Ptr.h"
+#include "MuR/RefCounted.h"
+
+// Remove when removin deprecated data
+#include <string>
+
+
+namespace mu
+{	
+	// Forward references
+    class Skeleton;
+
+    typedef Ptr<Skeleton> SkeletonPtr;
+    typedef Ptr<const Skeleton> SkeletonPtrConst;
+
+	// Bone name identifier
+	struct FBoneName
+	{
+		FBoneName() {};
+		FBoneName(uint32 InID) : Id(InID) {};
+
+		// Hash built from the bone name (FString)
+		uint32 Id = 0;
+
+		inline void Serialise(OutputArchive& arch) const
+		{
+			arch << Id;
+		}
+
+		inline void Unserialise(InputArchive& arch)
+		{
+			arch >> Id;
+		}
+
+		//!
+		inline bool operator==(const FBoneName& Other) const
+		{
+			return Id == Other.Id;
+		}
+	};
+
+	inline uint32 GetTypeHash(const FBoneName& Bone)
+	{
+		return Bone.Id;
+	}
+
+
+    //! \brief Skeleton object.
+	//! \ingroup runtime
+    class MUTABLERUNTIME_API Skeleton : public RefCounted
+	{
+	public:
+
+		//-----------------------------------------------------------------------------------------
+		// Life cycle
+		//-----------------------------------------------------------------------------------------
+
+		Skeleton() {}
+
+        //! Deep clone this skeleton.
+        Ptr<Skeleton> Clone() const;
+
+		//! Serialisation
+        static void Serialise( const Skeleton* p, OutputArchive& arch );
+        static Ptr<Skeleton> StaticUnserialise( InputArchive& arch );
+
+
+		//-----------------------------------------------------------------------------------------
+		// Own interface
+		//-----------------------------------------------------------------------------------------
+
+        //! @return - Number of bones in the Skeleton
+		int32 GetBoneCount() const;
+        void SetBoneCount(int32 c);
+
+		//! @return - FName of the bone at 'Index'. Only valid in the editor
+		const FName GetDebugName(int32 Index) const;
+		void SetDebugName(const int32 Index, const FName BoneName);
+
+        //! Get and set the parent bone of each bone. The parent can be -1 if the bone is a root.
+        int32 GetBoneParent(int32 boneIndex) const;
+        void SetBoneParent(int32 boneIndex, int32 parentBoneIndex);
+
+		//! @return - BoneName of the Bone at 'Index'.
+		const FBoneName& GetBoneName(int32 Index) const;
+		void SetBoneName(int32 Index, const FBoneName& BoneName);
+
+		//! @return - Index in the Skeleton. INDEX_NONE if not found.
+		int32 FindBone(const FBoneName& BoneName) const;
+		
+
+	protected:
+
+		//! Forbidden. Manage with the Ptr<> template.
+		~Skeleton() {}
+
+	public:
+
+		//! Deprecated
+		TArray<std::string> m_bones_DEPRECATED;
+		TArray<FTransform3f> m_boneTransforms_DEPRECATED;
+
+		//! DEBUG. FNames of the bones. Only valid in the editor. Do not serialize.
+		TArray<FName> DebugBoneNames;
+
+		//! Array of bone identifiers. 
+		TArray<FBoneName> BoneIds;
+
+		//! For each bone, index of the parent bone in the bone vectors. -1 means no parent.
+		//! This array must have the same size than the m_bones array.
+		TArray<int16> BoneParents;
+
+		//!
+		void Serialise(OutputArchive& arch) const;
+
+		//!
+		void Unserialise(InputArchive& arch);
+
+		//!
+		inline bool operator==(const Skeleton& o) const
+		{
+			return BoneIds == o.BoneIds
+				&& BoneParents == o.BoneParents;
+		}
+	};
+
+}
+
diff --git a/Engine/Plugins/Mutable/Source/MutableRuntime/Public/MuR/System.h b/Engine/Plugins/Mutable/Source/MutableRuntime/Public/MuR/System.h
new file mode 100644
--- /dev/null
+++ b/Engine/Plugins/Mutable/Source/MutableRuntime/Public/MuR/System.h
@@ -0,0 +1,238 @@
+// Copyright Epic Games, Inc. All Rights Reserved.
+
+#pragma once
+
+#include "HAL/Platform.h"
+#include "MuR/Image.h"
+#include "MuR/Instance.h"
+#include "MuR/MutableMemory.h"
+#include "MuR/Ptr.h"
+#include "MuR/RefCounted.h"
+#include "MuR/Settings.h"
+#include "MuR/Types.h"
+#include "Tasks/Task.h"
+#include "Templates/Tuple.h"
+
+
+/** If set to 1, this enables some expensive Unreal Insights traces, but can lead to 5x slower mutable operation. 
+* Other cheaper traces are enabled at all times.
+*/
+#define UE_MUTABLE_ENABLE_SLOW_TRACES	0
+
+
+#include "System.generated.h"
+
+
+/** Despite being an UEnum, this is not always version-serialized (in MutableTools).
+ * Beware of changing the enum options or order.
+ */
+UENUM()
+enum class ETextureCompressionStrategy : uint8
+{
+	/** Don't change the generated format. */
+	None,
+
+	/** If a texture depends on run-time parameters for an object state, don't compress. */
+	DontCompressRuntime,
+
+	/** Never compress the textures for this state. */
+	NeverCompress
+};
+
+namespace mu
+{
+	// Forward references
+	class Model;
+	class ModelReader;
+	class Parameters;
+    class Mesh;
+	class ExtensionDataStreamer;
+
+    class System;
+    using SystemPtr=Ptr<System>;
+    using SystemPtrConst=Ptr<const System>;
+
+
+	MUTABLE_DEFINE_ENUM_SERIALISABLE(ETextureCompressionStrategy);
+
+
+	/** */
+	enum class EExecutionStrategy : uint8
+	{
+		/** Undefined. */
+		None = 0,
+
+		/** Always try to run operations that reduce working memory first. */
+		MinimizeMemory,
+
+		/** Always try to run operations that unlock more operations first. */
+		MaximizeConcurrency,
+
+		/** Utility value with the number of error types. */
+		Count
+	};
+
+
+    /** Interface to request external images used as parameters. */
+    class MUTABLERUNTIME_API ExternalResourceProvider
+    {
+    public:
+
+        //! Ensure virtual destruction
+        virtual ~ExternalResourceProvider() = default;
+
+        /** Returns the completion event and a cleanup function that must be called once event is completed. */
+		virtual TTuple<UE::Tasks::FTask, TFunction<void()>> GetImageAsync(FName Id, uint8 MipmapsToSkip, TFunction<void(Ptr<Image>)>& ResultCallback) = 0;
+		virtual TTuple<UE::Tasks::FTask, TFunction<void()>> GetReferencedImageAsync(const void* ModelPtr, int32 Id, uint8 MipmapsToSkip, TFunction<void(Ptr<Image>)>& ResultCallback) { check(false); return {}; }
+
+		virtual mu::FImageDesc GetImageDesc(FName Id, uint8 MipmapsToSkip) = 0;
+
+		/** Returns the completion event and a cleanup function that must be called once event is completed. */
+		virtual TTuple<UE::Tasks::FTask, TFunction<void()>> GetMeshAsync(FName Id, TFunction<void(Ptr<Mesh>)>& ResultCallback) = 0;
+		virtual TTuple<UE::Tasks::FTask, TFunction<void()>> GetReferencedMeshAsync(const void* ModelPtr, int32 Id, TFunction<void(Ptr<Mesh>)>& ResultCallback) { check(false); return {}; }
+    };
+
+
+    //! \brief Main system class to load models and build instances.
+	//! \ingroup runtime
+	class MUTABLERUNTIME_API System : public RefCounted
+	{
+    public:
+
+        //! This constant can be used in place of the lodMask in methods like BeginUpdate
+        static constexpr uint32 AllLODs = 0xffffffff;
+
+    public:
+
+		//! Constructor of a system object to build data.
+        //! \param Settings Optional class with the settings to use in this system. The default
+        //! value configures a production-ready system.
+		//! \param DataStreamer Optional interface to allow the Model to stream in ExtensionData from disk
+        System( const Ptr<Settings>& Settings = nullptr, const TSharedPtr<ExtensionDataStreamer>& DataStreamer = nullptr );
+
+        //! Set a new provider for model data. 
+        void SetStreamingInterface(const TSharedPtr<ModelReader>& );
+
+        /** Set the working memory limit, overrding any set in the settings when the system was created.
+         * Refer to Settings::SetWorkingMemoryBudget for more information.
+		 */
+        void SetWorkingMemoryBytes( uint64 Bytes );
+
+        /** Removes all the possible working memory regardless of the budget set. This may make following 
+		* operations take longer.
+		*/
+		void ClearWorkingMemory();
+		
+		/** Set the amount of generated resources keys that will be stored for resource reusal. */
+		void SetGeneratedCacheSize(uint32 InCount);
+
+        /** Set a new provider for external image data. This is only necessary if image parameters are used in the models. */
+        void SetExternalResourceProvider(const TSharedPtr<ExternalResourceProvider>&);
+
+		/** Set a function that will be used to convert image pixel formats instead of the internal conversion. 
+		* \warning The provided function can be called from any thread, and also concurrently.
+		* If this function fails (returns false in the first parameter) the internal function is attempted next.
+		* This is useful to provide higher-quality external compressors in editor or when cooking.
+		*/
+		void SetImagePixelConversionOverride(const FImageOperator::FImagePixelFormatFunc&);
+
+        //! Create a new instance from the given model. The instance can then be configured through
+        //! calls to BeginUpdate/EndUpdate.
+        //! A call to NewInstance must be paired to a call to ReleasesInstance when the instance is
+        //! no longer needed.
+        //! \param pModel Model to build an instance of
+        //! \return An identifier that is always bigger than 0.
+        Instance::ID NewInstance(const TSharedPtr<const Model>& Model);
+
+        //! \brief Update an instance with a new parameter set and/or state.
+        //!
+        //! \warning a call to BeginUpdate must be paired with a call to EndUpdate once the returned
+        //! data has been processed.
+        //! \param InstanceID The id of the instance to update, as created by a NewInstance call.
+        //! \param Params The parameters that customise this instance.
+        //! \param StateIndex The index of the state this instance will be set to. The states range
+        //! from 0 to Model::GetStateCount-1
+        //! \param LodMask Bitmask selecting the levels of detail to build (i-th bit selects i-th lod).
+        //! \return the instance data with all the LOD, components, and ids to generate the meshes 
+        //! and images. The returned Instance is only valid until the next call to EndUpdate with 
+        //! the same instanceID parameter.
+        Ptr<const Instance> BeginUpdate(Instance::ID InstanceID, const Ptr<const Parameters>& Params, int32 StateIndex, uint32 LodMask);
+
+		//! Only valid between BeginUpdate and EndUpdate
+		//! Calculate the description of an image, without generating it.
+		UE::Tasks::TTask<FImageDesc> GetImageDesc(Instance::ID InstanceID, FResourceID ImageId);
+
+		//! Only valid between BeginUpdate and EndUpdate
+		//! \param MipsToSkip Number of mips to skip compared from the full image.
+		//! If 0, all mip levels will be generated. If more levels than possible to discard are specified, 
+		//! the image will still contain a minimum number of mips specified at model compile time.
+		UE::Tasks::TTask<Ptr<const Image>> GetImage(Instance::ID InstanceID, FResourceID ImageId, int32 MipsToSkip = 0, int32 LOD = 0);
+
+        //! Only valid between BeginUpdate and EndUpdate
+		UE::Tasks::TTask<Ptr<const Mesh>> GetMesh(Instance::ID InstanceID, FResourceID MeshId);
+
+		//! Only valid between BeginUpdate and EndUpdate
+		//! Calculate the description of an image, without generating it.
+		void GetImageDescInline(Instance::ID InstanceID, FResourceID ImageId, FImageDesc& OutDesc);
+
+		//! Only valid between BeginUpdate and EndUpdate
+		//! \param MipsToSkip Number of mips to skip compared from the full image.
+		//! If 0, all mip levels will be generated. If more levels than possible to discard are specified, 
+		//! the image will still contain a minimum number of mips specified at model compile time.
+		Ptr<const Image> GetImageInline(Instance::ID InstanceID, FResourceID ImageId, int32 MipsToSkip = 0, int32 LOD = 0);
+
+        //! Only valid between BeginUpdate and EndUpdate
+		Ptr<const Mesh> GetMeshInline(Instance::ID InstanceID, FResourceID MeshId);
+
+        //! Invalidate and free the last Instance data returned by a call to BeginUpdate with
+        //! the same instance index. After a call to this method, that Instance cannot be used any
+        //! more and its content is undefined.
+        //! \param instance The index of the instance whose last data will be invalidated.
+        void EndUpdate(Instance::ID InstanceID);
+
+        //! Completely destroy an instance. After a call to this method the given instance cannot be
+        //! updated any more, and its resources may have been freed.
+        //! \param instance The id of the instance to destroy.
+        void ReleaseInstance(Instance::ID InstanceID);
+
+		//! Calculate the relevancy of every parameter. Some parameters may be unused depending on
+		//! the values of other parameters. This method will set to true the flags for parameters
+		//! that are relevant, and to false otherwise. This is useful to hide irrelevant parameters
+		//! in dynamic user interfaces.
+        //! \param pModel The model used to create the Parameters instance.
+        //! \param pParameters Parameter set that we want to find the relevancy of.
+        //! \param pFlags is a pointer to a preallocated array of booleans that contains at least
+		//! pParameters->GetCount() elements.
+        void GetParameterRelevancy( Instance::ID InstanceID,
+									const Ptr<const Parameters>& Parameters,
+									bool* Flags );
+
+		//-----------------------------------------------------------------------------------------
+		// Interface pattern
+		//-----------------------------------------------------------------------------------------
+		class Private;
+
+		Private* GetPrivate() const;
+
+    public:
+
+        // Prevent copy, move and assignment.
+        System( const System& ) = delete;
+        System& operator=( const System& ) = delete;
+        System( System&& ) = delete;
+        System& operator=( System&& ) = delete;
+
+	protected:
+
+		//! Forbidden. Manage with the Ptr<> template.
+        ~System() override;
+
+	private:
+
+		Private* m_pD;
+
+	};
+
+
+}
+
diff --git a/Engine/Plugins/Mutable/Source/MutableRuntime/Public/MuR/Types.h b/Engine/Plugins/Mutable/Source/MutableRuntime/Public/MuR/Types.h
new file mode 100644
--- /dev/null
+++ b/Engine/Plugins/Mutable/Source/MutableRuntime/Public/MuR/Types.h
@@ -0,0 +1,47 @@
+// Copyright Epic Games, Inc. All Rights Reserved.
+
+#pragma once
+
+#include "HAL/Platform.h"
+#include "Containers/Map.h"
+#include "Containers/Array.h"
+#include "Containers/UnrealString.h"
+
+namespace mu
+{
+	class OutputArchive;
+	class InputArchive;
+
+	/** Handle used to identify separate resources generated by mutable: images and meshes included in instances. */
+    typedef uint64 FResourceID;
+
+	enum class EMemoryInitPolicy
+	{
+		Uninitialized,
+		Zeroed,
+	};
+
+}
+
+
+//
+template<typename T>
+inline uint32 GetTypeHash(const TArray<T>& Key)
+{
+	uint32 Hash = 0;
+	for (const T& e : Key)
+	{
+		Hash = HashCombine(Hash, GetTypeHash(e));
+	}
+	return Hash;
+}
+
+
+//
+template<typename K, typename V>
+inline bool operator==(const TMap<K,V>& A, const TMap<K,V>& B )
+{	
+	return A.OrderIndependentCompareEqual(B);
+}
+
+
diff --git a/Engine/Plugins/Mutable/Source/MutableTools/Internal/MuT/NodePrivate.h b/Engine/Plugins/Mutable/Source/MutableTools/Internal/MuT/NodePrivate.h
new file mode 100644
--- /dev/null
+++ b/Engine/Plugins/Mutable/Source/MutableTools/Internal/MuT/NodePrivate.h
@@ -0,0 +1,54 @@
+// Copyright Epic Games, Inc. All Rights Reserved.
+
+#pragma once
+
+#include "MuT/Node.h"
+#include "MuR/Operations.h"
+
+
+namespace mu
+{
+
+    class Node::Private
+    {
+    public:
+
+		/** Force a virtual destructor. */
+		virtual ~Private() = default; 
+
+		/** Generic pointer to the node owning this private. */
+		const Node* m_pNode = nullptr;
+    };
+
+
+
+#define MUTABLE_IMPLEMENT_NODE( N )				\
+    N::N()										\
+    {											\
+        m_pD = new Private();					\
+		m_pD->m_pNode = this;					\
+	}											\
+                                                \
+    N::~N()										\
+    {											\
+        check( m_pD );							\
+        delete m_pD;							\
+        m_pD = nullptr;							\
+    }											\
+                                                \
+    N::Private* N::GetPrivate() const			\
+    {											\
+        return m_pD;							\
+    }											\
+                                                \
+    const FNodeType* N::GetType() const			\
+    {											\
+        return GetStaticType();					\
+    }											\
+                                                \
+    const FNodeType* N::GetStaticType()			\
+    {											\
+        return &Private::s_type;				\
+    }
+
+}
