[mutable] Moved the Mutable plugin out of Experimental status into Beta.

diff --git a/Engine/Plugins/Mutable/Source/MutableTools/Private/MuT/AST.cpp b/Engine/Plugins/Mutable/Source/MutableTools/Private/MuT/AST.cpp
new file mode 100644
--- /dev/null
+++ b/Engine/Plugins/Mutable/Source/MutableTools/Private/MuT/AST.cpp
@@ -0,0 +1,1721 @@
+// Copyright Epic Games, Inc. All Rights Reserved.
+
+#include "MuT/AST.h"
+
+#include "Containers/Set.h"
+#include "Containers/Queue.h"
+#include "Logging/LogCategory.h"
+#include "Logging/LogMacros.h"
+#include "MuR/Mesh.h"
+#include "MuR/MutableTrace.h"
+#include "MuR/Platform.h"
+#include "MuT/ASTOpConstantResource.h"
+#include "MuT/ASTOpMeshRemoveMask.h"
+#include "MuT/ASTOpImageMultiLayer.h"
+#include "MuT/ASTOpParameter.h"
+#include "MuT/Platform.h"
+#include "MuT/ErrorLogPrivate.h"
+#include "Trace/Detail/Channel.h"
+
+#include <atomic>
+
+
+std::atomic<uint32> mu::ASTOp::s_lastTraverseIndex(1);
+
+
+namespace mu
+{
+
+//-------------------------------------------------------------------------------------------------
+//-------------------------------------------------------------------------------------------------
+//-------------------------------------------------------------------------------------------------
+ASTChild::ASTChild(ASTOp* p, const Ptr<ASTOp>& c)
+    : m_parent(p)
+    , m_child(c)
+{
+    if (m_parent && m_child.get())
+    {
+        AddParent();
+    }
+}
+
+ASTChild::ASTChild( const Ptr<ASTOp>& p, const Ptr<ASTOp>& c)
+    : ASTChild(p.get(),c)
+{
+}
+
+ASTChild::~ASTChild()
+{
+    if (m_child && m_parent)
+    {
+        ClearParent();
+    }
+}
+
+ASTChild& ASTChild::operator=( const Ptr<ASTOp>& c )
+{
+    if (c!=m_child)
+    {
+        if (m_child && m_parent)
+        {
+            ClearParent();
+        }
+
+        m_child = c;
+
+        if (m_child && m_parent)
+        {
+            AddParent();
+        }
+    }
+
+    return *this;
+}
+
+ASTChild& ASTChild::operator=( ASTChild&& rhs )
+{
+    m_parent = rhs.m_parent;
+    m_parentIndexInChild = rhs.m_parentIndexInChild;
+    m_child = rhs.m_child;
+    rhs.m_parent=nullptr;
+    rhs.m_child.reset();
+
+    return *this;
+}
+
+
+
+//-------------------------------------------------------------------------------------------------
+//-------------------------------------------------------------------------------------------------
+//-------------------------------------------------------------------------------------------------
+void ASTOp::ForEachParent(const TFunctionRef<void(ASTOp*)> f) const
+{
+	for (ASTOp* p : m_parents)
+	{
+		if (p)
+		{
+			f(p);
+		}
+	}
+}
+
+
+//-------------------------------------------------------------------------------------------------
+void ASTOp::RemoveChildren()
+{
+    // Actually destroyed when running out of scope
+    TArray<Ptr<ASTOp>> toDestroy;
+
+    // Try to make children destruction iterative
+    TArray<ASTOp*> pending;
+	pending.Reserve(1024);
+    pending.Add(this);
+
+    while (pending.Num())
+    {
+        ASTOp* n = pending.Pop(EAllowShrinking::No);
+
+        n->ForEachChild( [&](ASTChild& c)
+        {
+            if (c)
+            {
+                // Are we clearing the last reference?
+                if (c.child()->IsUnique())
+                {
+                    toDestroy.Add(c.child());
+                    pending.Add(c.child().get());
+                }
+
+                c = nullptr;
+            }
+        });
+    }
+}
+
+
+//-------------------------------------------------------------------------------------------------
+void ASTOp::Assert()
+{
+    // Check that every valid parent has us a child
+    // TODO: Should count the numbers, since a node may be child of another in multiple connections.
+    ForEachParent( [&](ASTOp*parent)
+    {
+        if(parent)
+        {
+            bool foundInParent = false;
+            parent->ForEachChild([&](ASTChild&c)
+            {
+                if (c && c.m_child.get()==this)
+                {
+                    foundInParent = true;
+                }
+            });
+
+            // If we hit this, we have a parent that doesn't know us.
+            check(foundInParent);
+        }
+    });
+
+    // Validate the children
+	ForEachChild( [this](ASTChild&c)
+    {
+        if(c)
+        {
+            // The child must have us as the parent.
+//            bool found = false;
+//            c.child()->ForEachParent( [&](ASTOp* childParent)
+//            {
+//                if (childParent==this)
+//                {
+//                    found = true;
+//                }
+//            });
+//            check(found);
+            check( c.m_parentIndexInChild < size_t(c.child()->m_parents.Num()) );
+            check( c.child()->m_parents[c.m_parentIndexInChild]==this );
+        }
+     });
+}
+
+
+//-------------------------------------------------------------------------------------------------
+bool ASTOp::operator==( const ASTOp& other ) const
+{
+//    if (typeid(*this) != typeid(other))
+//        return false;
+
+    return IsEqual(other);
+}
+
+
+//---------------------------------------------------------------------------------------------
+void ASTOp::FullAssert( const TArray<Ptr<ASTOp>>& Roots)
+{
+    MUTABLE_CPUPROFILER_SCOPE(AST_FullAssert);
+    Traverse_TopDown_Unique_Imprecise(Roots, [](const Ptr<ASTOp>& n)
+    {
+        n->Assert();
+        return true;
+    });
+}
+
+//-------------------------------------------------------------------------------------------------
+int32 ASTOp::CountNodes( const TArray<Ptr<ASTOp>>& Roots )
+{
+    MUTABLE_CPUPROFILER_SCOPE(AST_CountNodes);
+	int32 Count=0;
+    Traverse_TopRandom_Unique_NonReentrant(Roots, [&](const Ptr<ASTOp>&)
+    {
+        ++Count;
+        return true;
+    });
+    return Count;
+}
+
+
+//-------------------------------------------------------------------------------------------------
+mu::Ptr<ASTOp> ASTOp::DeepClone( const Ptr<ASTOp>& root )
+{
+    MUTABLE_CPUPROFILER_SCOPE(AST_DeepClone);
+
+    std::unordered_map<Ptr<const ASTOp>,Ptr<ASTOp>> visited;
+
+    MapChildFunc m = [&](const Ptr<ASTOp>&n)
+    {
+		if (!n)
+		{
+			return Ptr<ASTOp>();
+		}
+		std::unordered_map<Ptr<const ASTOp>, Ptr<ASTOp>>::const_iterator it = visited.find(n);
+        check(it!=visited.end());
+        return it->second;
+    };
+
+    Ptr<ASTOp> r = root;
+    Traverse_BottomUp_Unique( r, [&](Ptr<ASTOp> n)
+    {
+        Ptr<ASTOp> cloned = n->Clone( m );
+        visited[n] = cloned;
+    });
+
+	std::unordered_map<Ptr<const ASTOp>, Ptr<ASTOp>>::const_iterator it = visited.find(r);
+    check(it!=visited.end());
+    return it->second;
+}
+
+
+//-------------------------------------------------------------------------------------------------
+OP::ADDRESS ASTOp::FullLink( Ptr<ASTOp>& Root, FProgram& Program, FLinkerOptions* Options )
+{
+    MUTABLE_CPUPROFILER_SCOPE(AST_FullLink);
+
+    Traverse_BottomUp_Unique( Root,
+                              [&](Ptr<ASTOp> n){ n->Link(Program, Options); },
+                              [&](Ptr<const ASTOp> n){ return n->linkedAddress==0; });
+
+	OP::ADDRESS Result = Root->linkedAddress;
+
+	// This signals the caller that the Root pointer shouldn't be used anymore.
+	Root = nullptr;
+
+	return Result;
+}
+
+
+//---------------------------------------------------------------------------------------------
+void ASTOp::LogHistogram( ASTOpList& roots )
+{
+    (void)roots;
+
+ //   uint64 CountPerType[(int)OP_TYPE::COUNT];
+ //   FMemory::Memzero(CountPerType,sizeof(CountPerType));
+
+ //   size_t count=0;
+
+ //   Traverse_TopRandom_Unique_NonReentrant( roots,
+ //                            [&](const Ptr<ASTOp>& n)
+ //   {
+ //       ++count;
+	//	CountPerType[(int)n->GetOpType()]++;
+ //       return true;
+ //   });
+
+	//TArray< TPair<uint64, OP_TYPE> > Sorted;
+	//Sorted.SetNum((int)OP_TYPE::COUNT);
+ //   for (int i=0; i<(int)OP_TYPE::COUNT; ++i)
+ //   {
+ //       Sorted[i].Value = (OP_TYPE)i;
+ //       Sorted[i].Key = CountPerType[i];
+ //   }
+
+
+	//Sorted.Sort( [](const TPair<uint64, OP_TYPE>& a, const TPair<uint64, OP_TYPE>& b)
+ //   {
+ //       return a.Key >b.Key;
+ //   });
+
+ //   UE_LOG(LogMutableCore,Log, TEXT("Op histogram (%llu ops):"), count);
+ //   for(int i=0; i<8; ++i)
+ //   {
+ //       float p = Sorted[i].Key/float(count)*100.0f;
+ //       int OpType = (int)Sorted[i].Value;
+ //       UE_LOG(LogMutableCore, Log, TEXT("  %5.2f%% : %s"), p, s_opNames[OpType] );
+ //   }
+
+	//// Debug log part of the tree
+	//Traverse_TopDown_Unique( roots,
+	//	[&](const Ptr<ASTOp>& n)
+	//	{
+	//		if (n->GetOpType() == OP_TYPE::IM_MULTILAYER)
+	//		{
+	//			Ptr<const ASTOpImageMultiLayer> Typed = static_cast<const ASTOpImageMultiLayer*>(n.get());
+	//			Ptr<const ASTOp> Base = Typed->base.child();
+	//			Ptr<const ASTOpConstantResource> Constant = static_cast<const ASTOpConstantResource*>(Base.get());
+
+	//			// Log the op
+	//			UE_LOG(LogMutableCore, Log, TEXT("Multilayer at %x:"), n.get());
+	//			UE_LOG(LogMutableCore, Log, TEXT("    base is type %s:"), s_opNames[(int)Base->GetOpType()]);
+	//			if (Constant)
+	//			{
+	//				int Format = int( ((mu::Image*)Constant->GetValue().get())->GetFormat() );
+	//				UE_LOG(LogMutableCore, Log, TEXT("    constant image format is %d:"), Format);
+	//			}
+	//			else
+	//			{
+	//				FGetImageDescContext Context;
+	//				FImageDesc Desc = Base->GetImageDesc(true, &Context);
+	//				UE_LOG(LogMutableCore, Log, TEXT("    estimated image format is %d:"), int(Desc.m_format) );
+	//			}
+	//		}
+	//		return true;
+	//	});
+
+}
+
+
+//-------------------------------------------------------------------------------------------------
+void ASTOp::Traverse_TopDown_Unique( const TArray<Ptr<ASTOp>>& roots, TFunctionRef<bool(Ptr<ASTOp>&)> f )
+{
+    TQueue<Ptr<ASTOp>> pending;
+	for (const Ptr<ASTOp>& r : roots)
+	{
+		pending.Enqueue(r);
+	}
+
+    TSet<Ptr<const ASTOp>> traversed;
+
+    // We record the parents of all roots as traversed
+    for (const Ptr<ASTOp>& r: roots)
+    {
+        r->ForEachParent( [&]( const ASTOp* parent )
+        {
+            // If the parent is also a root, we will want to process it.
+            if ( !roots.Contains(parent) )
+            {
+                traversed.Add( parent );
+            }
+        });
+    }
+
+    while (!pending.IsEmpty())
+    {
+		Ptr<ASTOp> pCurrent;
+		pending.Dequeue(pCurrent);
+        if (!pCurrent)
+        {
+            continue;
+        }
+
+        // Did we traverse all parents?
+        bool parentsTraversed = true;
+
+        pCurrent->ForEachParent( [&]( const ASTOp* parent )
+        {
+            if (!traversed.Contains(parent))
+            {
+                // \todo Is the parent in the relevant subtree?
+
+
+                parentsTraversed = false;
+            }
+        });
+
+        if (!parentsTraversed)
+        {
+            pending.Enqueue(pCurrent);
+        }
+        else if (!traversed.Contains(pCurrent))
+        {
+            traversed.Add(pCurrent);
+
+            // Process
+            bool recurse = f(pCurrent);
+
+            // Recurse children
+            if (recurse)
+            {
+                pCurrent->ForEachChild([&]( ASTChild& c )
+                {
+                    if (c && !traversed.Contains(c.m_child))
+                    {
+                        pending.Enqueue( c.m_child );
+                    }
+                });
+            }
+        }
+    }
+}
+
+
+//-------------------------------------------------------------------------------------------------
+void ASTOp::Traverse_TopDown_Unique_Imprecise( const TArray<Ptr<ASTOp>>& roots, TFunctionRef<bool(Ptr<ASTOp>&)> f )
+{
+    TQueue<Ptr<ASTOp>> pending;
+	for (const Ptr<ASTOp>& r : roots)
+	{
+		pending.Enqueue(r);
+	}
+
+    TSet<Ptr<const ASTOp>> traversed;
+
+    while (!pending.IsEmpty())
+    {
+		Ptr<ASTOp> pCurrent;
+		pending.Dequeue(pCurrent);
+
+        // It could have been completed in another branch
+        if (pCurrent && !traversed.Contains(pCurrent))
+        {
+            traversed.Add(pCurrent);
+
+            // Process
+            bool recurse = f(pCurrent);
+
+            // Recurse children
+            if (recurse)
+            {
+                pCurrent->ForEachChild([&]( ASTChild& c )
+                {
+                    if (c && !traversed.Contains(c.m_child))
+                    {
+                        pending.Enqueue( c.m_child );
+                    }
+                });
+            }
+        }
+    }
+}
+
+
+//-------------------------------------------------------------------------------------------------
+void ASTOp::Traverse_TopRandom_Unique_NonReentrant( const TArray<Ptr<ASTOp>>& roots,
+	TFunctionRef<bool(Ptr<ASTOp>&)> f )
+{
+    ASTOpList pending;
+
+    uint32 traverseIndex = s_lastTraverseIndex++;
+
+    for (const Ptr<ASTOp>& r:roots)
+    {
+        if (r && r->m_traverseIndex!=traverseIndex )
+        {
+            r->m_traverseIndex = traverseIndex;
+            pending.Add( r );
+        }
+    }
+    for(const Ptr<ASTOp>& p : pending)
+    {
+        p->m_traverseIndex = traverseIndex-1;
+    }
+
+    while (pending.Num())
+    {
+        Ptr<ASTOp> pCurrent = pending.Pop();
+
+        // It could have been completed in another branch
+        if (pCurrent->m_traverseIndex!=traverseIndex)
+        {
+            pCurrent->m_traverseIndex = traverseIndex;
+
+            // Process
+            bool recurse = f(pCurrent);
+
+            // Recurse children
+            if (recurse)
+            {
+                pCurrent->ForEachChild([&]( ASTChild& c )
+                {
+                    if (c && c.m_child->m_traverseIndex!=traverseIndex)
+                    {
+                        pending.Add( c.m_child );
+                    }
+                });
+            }
+        }
+    }
+}
+
+
+
+//-------------------------------------------------------------------------------------------------
+//-------------------------------------------------------------------------------------------------
+//-------------------------------------------------------------------------------------------------
+void Visitor_TopDown_Unique_Cloning::Traverse( Ptr<ASTOp>& root )
+{
+    // Visit the given root
+    if (root)
+    {
+		m_pending.Add({ false,root });
+
+        Process();
+
+        root = GetOldToNew(root);
+    }
+}
+
+
+
+void Visitor_TopDown_Unique_Cloning::Process()
+{
+    while ( m_pending.Num() )
+    {
+        TPair<bool,Ptr<ASTOp>> item = m_pending.Pop();
+
+        Ptr<ASTOp> at = item.Value;
+
+		ASTOp::MapChildFunc Identity = [](const Ptr<ASTOp>& o) {return o; };
+		
+		if (item.Key)
+        {
+            // Item indicating we finished with all the children of this instruction
+			Ptr<ASTOp> cop = at->Clone(Identity);
+
+            // Fix the references to the children
+            bool childChanged = false;
+            cop->ForEachChild( [&](ASTChild& ref)
+            {
+				std::unordered_map<Ptr<ASTOp>, Ptr<ASTOp>>::iterator it = m_oldToNew.find(ref.m_child);
+                if ( ref && it!=m_oldToNew.end() && it->second.get()!=nullptr )
+                {
+                    auto oldRef = ref.m_child;
+                    ref=GetOldToNew(ref.m_child);
+                    if (ref.m_child!=oldRef)
+                    {
+                        childChanged = true;
+                    }
+                }
+            });
+
+            // If any child changed, we need to replace this instruction
+            if ( childChanged )
+            {
+                m_oldToNew[at]=cop;
+            }
+
+        }
+        else
+        {
+			std::unordered_map<Ptr<ASTOp>, Ptr<ASTOp>>::iterator it = m_oldToNew.find(at);
+            if (it==m_oldToNew.end())
+            {
+				Ptr<ASTOp> initialAt = at;
+
+                // Fix the references to the children, possibly adding a new instruction
+                {
+					Ptr<ASTOp> cop = at->Clone(Identity);
+                    bool childChanged = false;
+                    cop->ForEachChild( [&](ASTChild& ref)
+                    {
+						std::unordered_map<Ptr<ASTOp>, Ptr<ASTOp>>::iterator ito = m_oldToNew.find(ref.m_child);
+                        if ( ref && ito!=m_oldToNew.end() && ito->second.get()!=nullptr )
+                        {
+							Ptr<ASTOp> oldRef = ref.m_child;
+                            ref=GetOldToNew(ref.m_child);
+                            if (ref.m_child!=oldRef)
+                            {
+                                childChanged = true;
+                            }
+                        }
+                    });
+
+                    // If any child changed, we need to re-add this instruction
+                    if ( childChanged )
+                    {
+                        m_oldToNew[at]=cop;
+                        at = cop;
+                    }
+                }
+
+                //Ptr<ASTOp> test1 = at->Clone();
+                //check(*test1==*at);
+
+                bool processChildren = true;
+				Ptr<ASTOp> newAt = Visit( at, processChildren );
+                m_oldToNew[initialAt]=newAt;
+
+                //check(*test1==*at);
+
+                // Proceed with children
+                if (processChildren)
+                {
+					m_pending.Add({ true, newAt });
+
+                    at->ForEachChild( [&](ASTChild& ref)
+                    {
+                        if (ref && m_oldToNew.find(ref.m_child)==m_oldToNew.end())
+                        {
+							m_pending.Add({ false,ref.m_child });
+                        }
+                    });
+                }
+            }
+        }
+
+    }
+}
+
+
+//-------------------------------------------------------------------------------------------------
+//-------------------------------------------------------------------------------------------------
+//-------------------------------------------------------------------------------------------------
+void ASTOp::Traverse_BottomUp_Unique_NonReentrant
+(
+        ASTOpList& roots,
+        TFunctionRef<void(Ptr<ASTOp>&)> f
+)
+{
+    uint32 traverseIndex = s_lastTraverseIndex++;
+
+    TArray< std::pair<Ptr<ASTOp>,int32> > pending;
+    for (Ptr<ASTOp>& r:roots)
+    {
+        if (r && r->m_traverseIndex!=traverseIndex )
+        {
+            r->m_traverseIndex=traverseIndex;
+            pending.Add( std::make_pair<>(r,0) );
+        }
+    }
+    for(std::pair<Ptr<ASTOp>, int32>& p : pending)
+    {
+        p.first->m_traverseIndex = traverseIndex-1;
+    }
+
+    while (pending.Num())
+    {
+        int32 phase = pending.Last().second;
+        Ptr<ASTOp> pCurrent = pending.Last().first;
+        pending.Pop();
+
+        // It could have been completed in another branch
+        if (pCurrent->m_traverseIndex!=traverseIndex)
+        {
+            if (phase==0)
+            {
+                // Process this again...
+                pending.Add( std::make_pair<>(pCurrent,1) );
+
+                // ...after the children are processed
+                pCurrent->ForEachChild([&]( ASTChild& c )
+                {
+                    if (c && c.m_child->m_traverseIndex!=traverseIndex )
+                    {
+						std::pair<Ptr<ASTOp>, int32> e = std::make_pair<>(c.m_child,0);
+                        pending.Add( e );
+                    }
+                });
+            }
+            else
+            {
+                pCurrent->m_traverseIndex=traverseIndex;
+
+                // Children have been be completed
+                f(pCurrent);
+            }
+        }
+    }
+}
+
+
+//-------------------------------------------------------------------------------------------------
+//-------------------------------------------------------------------------------------------------
+//-------------------------------------------------------------------------------------------------
+void ASTOp::Traverse_BottomUp_Unique_NonReentrant
+(
+        ASTOpList& roots,
+        TFunctionRef<void(Ptr<ASTOp>&)> f,
+        TFunctionRef<bool(const ASTOp*)> accept
+)
+{
+    uint32 traverseIndex = s_lastTraverseIndex++;
+
+    TArray< std::pair<Ptr<ASTOp>,int32> > pending;
+    for (Ptr<ASTOp>& r:roots)
+    {
+        if (r && r->m_traverseIndex!=traverseIndex)
+        {
+            r->m_traverseIndex=traverseIndex;
+            pending.Add( std::make_pair<>(r,0) );
+        }
+    }
+    for(std::pair<Ptr<ASTOp>, int32>& p : pending)
+    {
+        p.first->m_traverseIndex = traverseIndex-1;
+    }
+
+    while (pending.Num())
+    {
+        int32 phase = pending.Last().second;
+        Ptr<ASTOp> pCurrent = pending.Last().first;
+        pending.Pop();
+
+        // It could have been completed in another branch
+        if (pCurrent->m_traverseIndex!=traverseIndex && accept(pCurrent.get()))
+        {
+            if (phase==0)
+            {
+                // Process this again...
+                pending.Add( std::make_pair<>(pCurrent,1) );
+
+                // ...after the children are processed
+                pCurrent->ForEachChild([&]( ASTChild& c )
+                {
+                    if (c && accept(c.m_child.get()) && c.m_child->m_traverseIndex!=traverseIndex )
+                    {
+                        pending.Add( std::make_pair<>(c.m_child,0) );
+                    }
+                });
+            }
+            else
+            {
+                pCurrent->m_traverseIndex=traverseIndex;
+
+                // Children have been be completed
+                f(pCurrent);
+            }
+        }
+    }
+}
+
+
+//-------------------------------------------------------------------------------------------------
+//-------------------------------------------------------------------------------------------------
+//-------------------------------------------------------------------------------------------------
+void ASTOp::Traverse_BottomUp_Unique
+(
+        ASTOpList& roots,
+        TFunctionRef<void(Ptr<ASTOp>&)> f,
+        TFunctionRef<bool(const ASTOp*)> accept
+)
+{
+    TSet<Ptr<ASTOp>> Traversed;
+    TArray< std::pair<Ptr<ASTOp>,int32> > Pending;
+    for (Ptr<ASTOp>& r:roots)
+    {
+        if (r)
+        {
+			std::pair<Ptr<ASTOp>, int32>* It = Pending.FindByPredicate([&](const std::pair<Ptr<ASTOp>, int>& p)
+				{
+					return r == p.first;
+				});
+            if ( !It )
+            {
+                Pending.Add( std::make_pair<>(r,0) );
+            }
+        }
+    }
+
+    while (!Pending.IsEmpty())
+    {
+        int phase = Pending.Last().second;
+        Ptr<ASTOp> pCurrent = Pending.Last().first;
+        Pending.Pop();
+
+        // It could have been completed in another branch
+        if (accept(pCurrent.get()) && !Traversed.Find(pCurrent))
+        {
+            if (phase==0)
+            {
+                // Process this again...
+				Pending.Add( std::make_pair<>(pCurrent,1) );
+
+                // ...after the children are processed
+                pCurrent->ForEachChild([&]( ASTChild& c )
+                {
+                    if (c && accept(c.m_child.get()) && !Traversed.Find(c.m_child) )
+                    {
+                        Pending.Add( std::make_pair<>(c.m_child,0) );
+                    }
+                });
+            }
+            else
+            {
+                Traversed.Add(pCurrent);
+
+                // Children have been be completed
+                f(pCurrent);
+            }
+        }
+    }
+}
+
+
+//-------------------------------------------------------------------------------------------------
+//-------------------------------------------------------------------------------------------------
+//-------------------------------------------------------------------------------------------------
+void ASTOp::Traverse_BottomUp_Unique
+(
+        Ptr<ASTOp>& root,
+        TFunctionRef<void(Ptr<ASTOp>&)> f,
+        TFunctionRef<bool(const ASTOp*)> accept
+)
+{
+    if (root)
+    {
+        ASTOpList roots;
+        roots.Add(root);
+        Traverse_BottomUp_Unique(roots,f,accept);
+    }
+}
+
+
+
+//-------------------------------------------------------------------------------------------------
+//-------------------------------------------------------------------------------------------------
+//-------------------------------------------------------------------------------------------------
+int ASTOp::GetParentCount() const
+{
+    int result=0;
+    ForEachParent( [&](const ASTOp* p)
+    {
+        if (p!=nullptr) ++result;
+    });
+    return result;
+}
+
+
+void ASTOp::Replace( const Ptr<ASTOp>& node, const Ptr<ASTOp>& other )
+{
+    if(other==node)
+    {
+        return;
+    }
+
+	TArray<ASTOp*, TInlineAllocator<4> > parentsCopy = node->m_parents;
+
+    for(ASTOp* p:parentsCopy)
+    {
+        if(p)
+        {
+            p->ForEachChild( [=](ASTChild& c)
+            {
+                if(c.m_child==node)
+                {
+                    c = other;
+                }
+            });
+        }
+    }
+}
+
+
+FImageDesc ASTOp::GetImageDesc( bool, FGetImageDescContext* ) const
+{
+    check(false);
+    return FImageDesc();
+}
+
+
+FSourceDataDescriptor ASTOp::GetSourceDataDescriptor(FGetSourceDataDescriptorContext*) const
+{
+	ensure(false);
+	return {};
+}
+
+
+bool ASTOp::IsImagePlainConstant(FVector4f&) const
+{
+	// Some image operations don't have this implemented and hit here.
+	// \TODO: Optimize for those cases.
+    //check(false);
+    return false;
+}
+
+
+bool ASTOp::IsColourConstant(FVector4f&) const
+{
+    check(false);
+    return false;
+}
+
+
+void ASTOp::GetBlockLayoutSize( uint64, int32*, int32*, FBlockLayoutSizeCache* )
+{
+    check(false);
+}
+
+void ASTOp::GetBlockLayoutSizeCached(uint64 blockId, int32* BlockX, int32* BlockY, FBlockLayoutSizeCache* cache)
+{
+	FBlockLayoutSizeCache::KeyType key(this,blockId);
+	FBlockLayoutSizeCache::ValueType* ValuePtr = cache->Find( key );
+    if (ValuePtr)
+    {
+        *BlockX = ValuePtr->Key;
+        *BlockY = ValuePtr->Value;
+        return;
+    }
+
+    GetBlockLayoutSize( blockId, BlockX, BlockY, cache );
+
+	cache->Add(key, FBlockLayoutSizeCache::ValueType( *BlockX, *BlockY) );
+}
+
+
+void ASTOp::GetLayoutBlockSize( int32*, int32* )
+{
+    check(false);
+}
+
+
+bool ASTOp::GetNonBlackRect( FImageRect& ) const
+{
+    return false;
+}
+
+
+void ASTOp::LinkRange(FProgram& program,
+	const FRangeData& range,
+	OP::ADDRESS& rangeSize,
+	uint16& rangeId)
+{
+	if (range.rangeSize)
+	{
+		if (range.rangeSize->linkedRange < 0)
+		{
+			check(program.m_ranges.Num() < 255);
+			range.rangeSize->linkedRange = int8(program.m_ranges.Num());
+			FRangeDesc rangeData;
+			rangeData.m_name = range.rangeName;
+			rangeData.m_uid = range.rangeUID;
+
+			// Try to see if a parameter directly controls de size of the range. This is used
+			// to store hint data for instance generation in tools or randomizers that want to
+			// support multilayer, but it is not critical otherwise.
+			int32 EstimatedSizeParameter = -1;
+			if ( range.rangeSize->GetOpType() == OP_TYPE::SC_PARAMETER
+				||
+				range.rangeSize->GetOpType() == OP_TYPE::NU_PARAMETER )
+			{
+				const ASTOpParameter* ParamOp = static_cast<const ASTOpParameter*>(range.rangeSize.child().get());
+
+				EstimatedSizeParameter = ParamOp->LinkedParameterIndex;
+			}
+			rangeData.m_dimensionParameter = EstimatedSizeParameter;
+
+			program.m_ranges.Add(rangeData);
+		}
+		rangeSize = range.rangeSize->linkedAddress;
+		rangeId = uint16(range.rangeSize->linkedRange);
+	}
+}
+
+
+//-------------------------------------------------------------------------------------------------
+//-------------------------------------------------------------------------------------------------
+//-------------------------------------------------------------------------------------------------
+ASTOpFixed::ASTOpFixed()
+    : children({{
+			ASTChild(this,nullptr),ASTChild(this,nullptr),ASTChild(this,nullptr),ASTChild(this,nullptr),
+			ASTChild(this,nullptr),ASTChild(this,nullptr),ASTChild(this,nullptr),ASTChild(this,nullptr)
+		}})
+{
+}
+
+
+ASTOpFixed::~ASTOpFixed()
+{
+    // Explicit call needed to avoid recursive destruction
+    ASTOp::RemoveChildren();
+}
+
+
+void ASTOpFixed::ForEachChild( const TFunctionRef<void(ASTChild&)> f )
+{
+    // ugly but temporary while ASTOpFixed exists
+    ForEachReference( op, [&](OP::ADDRESS* pAt)
+    {
+        if (pAt && *pAt)
+        {
+            check(children.size()>*pAt);
+            f(children[*pAt]);
+        }
+    });
+}
+
+
+void ASTOpFixed::Link( FProgram& program, FLinkerOptions* )
+{
+	MUTABLE_CPUPROFILER_SCOPE(ASTOpFixed_Link);
+
+    if (!linkedAddress)
+    {
+        OP lop = op;
+
+        // Fix the linked op and generate children
+        ForEachReference(lop, [&](OP::ADDRESS* pChildOp)
+        {
+            Ptr<ASTOp> pASTChild = children[*pChildOp].m_child;
+            if (pASTChild)
+            {
+                *pChildOp = pASTChild->linkedAddress;
+            }
+            else
+            {
+                *pChildOp = 0;
+            }
+        });
+
+        linkedAddress = (OP::ADDRESS)program.m_opAddress.Num();        
+        program.m_opAddress.Add((uint32)program.m_byteCode.Num());
+
+        AppendCode(program.m_byteCode,lop.type);
+        if (lop.type==OP_TYPE::ME_MERGE)
+        {
+            AppendCode(program.m_byteCode,lop.args.MeshMerge);
+        }
+        else
+        {
+            // Generic encoding
+            AppendCode(program.m_byteCode,lop.args);
+        }
+    }
+}
+
+
+bool ASTOpFixed::IsEqual(const ASTOp& otherUntyped) const
+{
+    if (otherUntyped.GetOpType()==GetOpType())
+    {
+		const ASTOpFixed* other = static_cast<const ASTOpFixed*>(&otherUntyped);
+        return op==other->op && children==other->children;
+    }
+    return false;
+}
+
+
+mu::Ptr<ASTOp> ASTOpFixed::Clone( MapChildFuncRef mapChild ) const
+{
+    Ptr<ASTOpFixed> n = new ASTOpFixed();
+    n->op = op;
+    n->childCount = childCount;
+    // Skip the children 0, used to represent null op.
+    for(uint8_t i=1; i<childCount; ++i)
+    {
+        n->children[i] = mapChild(children[i].child());
+    }
+    return n;
+}
+
+
+uint64 ASTOpFixed::Hash() const
+{
+	uint64 res = std::hash<uint64>()(uint64(op.type));
+    for (const ASTChild& c: children)
+    {
+        hash_combine( res, c.child().get() );
+    }
+    return res;
+}
+
+
+//!
+FImageDesc ASTOpFixed::GetImageDesc( bool returnBestOption, FGetImageDescContext* context ) const
+{
+    FImageDesc res;
+
+    // Local context in case it is necessary
+    FGetImageDescContext localContext;
+    if (!context)
+    {
+      context = &localContext;
+    }
+    else
+    {
+        // Cached result?
+		FImageDesc* PtrValue = context->m_results.Find(this);
+        if (PtrValue)
+        {
+            return *PtrValue;
+        }
+    }
+
+    OP_TYPE type = op.type;
+
+    switch ( type )
+    {
+
+    case OP_TYPE::NONE:
+        break;
+
+    case OP_TYPE::IM_SATURATE:
+        res = GetImageDesc( op.args.ImageSaturate.base, returnBestOption, context );
+        break;
+
+    case OP_TYPE::IM_LUMINANCE:
+        res = GetImageDesc( op.args.ImageLuminance.base,returnBestOption, context );
+        res.m_format = EImageFormat::IF_L_UBYTE;
+        break;
+
+    case OP_TYPE::IM_INTERPOLATE:
+        res = GetImageDesc(op.args.ImageInterpolate.targets[0], returnBestOption, context );
+        break;
+
+    case OP_TYPE::IM_PLAINCOLOUR:
+        res.m_format = EImageFormat(op.args.ImagePlainColour.format);
+        res.m_size[0] = op.args.ImagePlainColour.size[0];
+        res.m_size[1] = op.args.ImagePlainColour.size[1];
+		res.m_lods = op.args.ImagePlainColour.LODs;
+        check( res.m_format != EImageFormat::IF_NONE );
+        break;
+
+    case OP_TYPE::IM_RESIZE:
+        res = GetImageDesc( op.args.ImageResize.source, returnBestOption, context );
+
+        res.m_size = FImageSize
+            (
+                op.args.ImageResize.size[0],
+                op.args.ImageResize.size[1]
+            );
+        break;
+
+    case OP_TYPE::IM_RESIZEREL:
+        res = GetImageDesc( op.args.ImageResizeRel.source, returnBestOption, context );
+        res.m_size[0] = (uint16)( res.m_size[0]*op.args.ImageResizeRel.factor[0] );
+        res.m_size[1] = (uint16)( res.m_size[1]*op.args.ImageResizeRel.factor[1] );
+        break;
+
+    case OP_TYPE::IM_RESIZELIKE:
+        res = GetImageDesc( op.args.ImageResizeLike.source, returnBestOption, context );
+        if (op.args.ImageResizeLike.sizeSource)
+        {
+            res.m_size = children[op.args.ImageResizeLike.sizeSource]->GetImageDesc( returnBestOption, context ).m_size;
+        }
+        break;
+
+    case OP_TYPE::IM_GRADIENT:
+        res.m_size[0] = op.args.ImageGradient.size[0];
+        res.m_size[1] = op.args.ImageGradient.size[1];
+        res.m_format = EImageFormat::IF_RGB_UBYTE;
+        break;
+
+    case OP_TYPE::IM_BLANKLAYOUT:
+        // TODO: We would need to process the layout to find the grid size, and then use
+        // the block size with it.
+        res.m_size = FImageSize(0,0);
+        res.m_format = op.args.ImageBlankLayout.format;
+        break;
+
+    case OP_TYPE::IM_BINARISE:
+        res = GetImageDesc( op.args.ImageBinarise.base, returnBestOption, context );
+        res.m_format = EImageFormat::IF_L_UBYTE;
+        break;
+
+    case OP_TYPE::IM_DISPLACE:
+        res = GetImageDesc( op.args.ImageDisplace.source, returnBestOption, context );
+        break;
+
+	case OP_TYPE::IM_INVERT:
+		res = GetImageDesc(op.args.ImageInvert.base, returnBestOption, context);
+		break;
+
+	case OP_TYPE::IM_COLOURMAP:
+		res = GetImageDesc( op.args.ImageColourMap.base, returnBestOption, context );
+		break;
+
+    default:
+        check(false);
+    }
+
+    // Cache the result
+    if (context)
+    {
+        context->m_results.Add(this,res);
+    }
+
+    return res;
+}
+
+
+void ASTOpFixed::GetLayoutBlockSize( int* pBlockX, int* pBlockY )
+{
+    switch ( op.type )
+    {
+    
+	case OP_TYPE::IM_RESIZEREL:
+	{
+		GetLayoutBlockSize(op.args.ImageResizeRel.source, pBlockX, pBlockY);
+		(*pBlockX) = int(float(*pBlockX) * op.args.ImageResizeRel.factor[0]);
+		(*pBlockY) = int(float(*pBlockX) * op.args.ImageResizeRel.factor[1]);
+
+		break;
+	}
+
+	case OP_TYPE::IM_RESIZE:
+	{
+		GetLayoutBlockSize(op.args.ImageResize.source, pBlockX, pBlockY);
+
+		if (*pBlockX > 0 && *pBlockY> 0)
+		{
+			FImageDesc sourceDesc = GetImageDesc(op.args.ImageResize.source);
+			if (sourceDesc.m_size[0]>0 && sourceDesc.m_size[1]>0)
+			{
+				float factorX = float(op.args.ImageResize.size[0]) / float(sourceDesc.m_size[0]);
+				float factorY = float(op.args.ImageResize.size[1]) / float(sourceDesc.m_size[1]);
+				*pBlockX = int(*pBlockX * factorX);
+				*pBlockY = int(*pBlockX * factorY);
+			}
+			else
+			{
+				*pBlockX = 0;
+				*pBlockY = 0;
+			}
+		}
+
+		break;
+	}
+
+	case OP_TYPE::IM_BLANKLAYOUT:
+	{
+		*pBlockX = int(op.args.ImageBlankLayout.blockSize[0]);
+		*pBlockY = int(op.args.ImageBlankLayout.blockSize[1]);
+		break;
+	}
+
+	case OP_TYPE::IM_PLAINCOLOUR:
+	{
+		*pBlockX = 0;
+		*pBlockY = 0;
+		break;
+	}
+
+    default:
+        checkf( false, TEXT("Instruction not supported") );
+    }
+
+    //check( *pBlockX!=0 && *pBlockY!=0 );
+}
+
+
+ASTOp::FBoolEvalResult ASTOpFixed::EvaluateBool( ASTOpList& facts, FEvaluateBoolCache* cache ) const
+{
+    FEvaluateBoolCache localCache;
+    if (!cache)
+    {
+        cache = &localCache;
+    }
+    else
+    {
+        // Is this in the cache?
+		FEvaluateBoolCache::iterator it = cache->find(this);
+        if (it!=cache->end())
+        {
+            return it->second;
+        }
+    }
+
+    FBoolEvalResult result = BET_UNKNOWN;
+    switch(op.type)
+    {
+    case OP_TYPE::BO_NOT:
+    {
+        if (children[op.args.BoolNot.source])
+        {
+            FBoolEvalResult child = children[op.args.BoolNot.source]->EvaluateBool(facts,cache);
+            if (child==BET_TRUE)
+            {
+                result = BET_FALSE;
+            }
+            else if (child==BET_FALSE)
+            {
+                result = BET_TRUE;
+            }
+        }
+        break;
+    }
+
+    case OP_TYPE::BO_EQUAL_INT_CONST:
+    {
+        int32 intValue = op.args.BoolEqualScalarConst.constant;
+        const Ptr<ASTOp>& intExp = children[op.args.BoolEqualScalarConst.value].child();
+        bool intUnknown = true;
+        int intResult = intExp->EvaluateInt( facts, intUnknown );
+        if (intUnknown)
+        {
+            result = BET_UNKNOWN;
+        }
+        else if (intResult==intValue)
+        {
+            result = BET_TRUE;
+        }
+        else
+        {
+            result = BET_FALSE;
+        }
+        break;
+    }
+
+    case OP_TYPE::BO_AND:
+    {
+        const Ptr<ASTOp>& a = children[op.args.BoolBinary.a].child();
+        const Ptr<ASTOp>& b = children[op.args.BoolBinary.b].child();
+        FBoolEvalResult resultA = BET_UNKNOWN;
+        FBoolEvalResult resultB = BET_UNKNOWN;
+        for ( size_t f=0; f<facts.Num(); ++f )
+        {
+            if ( a && resultA == BET_UNKNOWN )
+            {
+                resultA = a->EvaluateBool( facts, cache );
+
+                if ( resultA==BET_TRUE && resultB==BET_TRUE )
+                {
+                    result = BET_TRUE;
+                    break;
+                }
+                if ( resultA==BET_FALSE || resultB==BET_FALSE )
+                {
+                    result = BET_FALSE;
+                    break;
+                }
+            }
+            if ( b && resultB == BET_UNKNOWN )
+            {
+                resultB = b->EvaluateBool( facts, cache );
+
+                if ( resultA==BET_TRUE && resultB==BET_TRUE )
+                {
+                    result = BET_TRUE;
+                    break;
+                }
+                if ( resultA==BET_FALSE || resultB==BET_FALSE )
+                {
+                    result = BET_FALSE;
+                    break;
+                }
+            }
+        }
+        break;
+    }
+
+    case OP_TYPE::BO_OR:
+    {
+        const Ptr<ASTOp>& a = children[op.args.BoolBinary.a].child();
+        const Ptr<ASTOp>& b = children[op.args.BoolBinary.b].child();
+        FBoolEvalResult resultA = BET_UNKNOWN;
+        FBoolEvalResult resultB = BET_UNKNOWN;
+        for ( size_t f=0; f<facts.Num(); ++f )
+        {
+            if ( a && resultA == BET_UNKNOWN )
+            {
+                resultA = a->EvaluateBool( facts, cache );
+
+                if ( resultA==BET_TRUE || resultB==BET_TRUE )
+                {
+                    result = BET_TRUE;
+                    break;
+                }
+                if ( resultA==BET_FALSE && resultB==BET_FALSE )
+                {
+                    result = BET_FALSE;
+                    break;
+                }
+            }
+            if ( b && resultB == BET_UNKNOWN )
+            {
+                resultB = b->EvaluateBool( facts, cache );
+
+                if ( resultA==BET_TRUE || resultB==BET_TRUE )
+                {
+                    result = BET_TRUE;
+                    break;
+                }
+                if ( resultA==BET_FALSE && resultB==BET_FALSE )
+                {
+                    result = BET_FALSE;
+                    break;
+                }
+            }
+        }
+        break;
+    }
+
+    default:
+        checkf( false, TEXT("Instruction not supported") );
+    }
+
+    (*cache)[this] = result;
+
+    return result;
+}
+
+
+int ASTOpFixed::EvaluateInt( ASTOpList& /*facts*/, bool &unknown ) const
+{
+    unknown = true;
+
+    switch (GetOpType())
+    {
+    case OP_TYPE::NU_CONSTANT:
+        unknown = false;
+        return op.args.IntConstant.value;
+        break;
+
+    case OP_TYPE::SC_CONSTANT:
+        unknown = false;
+        return (int)op.args.ScalarConstant.value;
+        break;
+
+        // TODO
+//    case IntExpression::IET_PARAMETER:
+//        for ( size_t f=0; (result==BET_UNKNOWN) && f<facts.Num(); ++f )
+//        {
+//            result = EvaluateIntEquality(facts[f].get(),intExp,value);
+//        }
+//        break;
+
+    default:
+        break;
+
+    }
+
+
+    return 0;
+}
+
+
+bool ASTOpFixed::IsImagePlainConstant(FVector4f& colour) const
+{
+    bool res = false;
+    switch( op.type )
+    {
+
+    case OP_TYPE::IM_BLANKLAYOUT:
+        res = true;
+        colour[0] = 0;
+        colour[1] = 0;
+        colour[2] = 0;
+        colour[3] = 0;
+        break;
+
+    case OP_TYPE::IM_RESIZE:
+        res = children[op.args.ImageResize.source]->IsImagePlainConstant( colour );
+        break;
+
+    case OP_TYPE::IM_RESIZEREL:
+        res = children[op.args.ImageResizeRel.source]->IsImagePlainConstant( colour );
+        break;
+
+    case OP_TYPE::IM_RESIZELIKE:
+        res = children[op.args.ImageResizeLike.source]->IsImagePlainConstant(  colour );
+        break;
+
+    case OP_TYPE::IM_PLAINCOLOUR:
+        res = children[op.args.ImagePlainColour.colour]->IsColourConstant( colour );
+        break;
+
+    default:
+        // TODO: Improve this test with more operations
+        //check( false );
+        res = false;
+        break;
+    }
+
+    return res;
+}
+
+
+//-------------------------------------------------------------------------------------------------
+bool ASTOpFixed::IsColourConstant(FVector4f& colour) const
+{
+    bool res = false;
+    switch( op.type )
+    {
+    case OP_TYPE::CO_CONSTANT:
+        res = true;
+        colour[0] = op.args.ColourConstant.value[0];
+        colour[1] = op.args.ColourConstant.value[1];
+        colour[2] = op.args.ColourConstant.value[2];
+        colour[3] = op.args.ColourConstant.value[3];
+        break;
+
+    case OP_TYPE::CO_SAMPLEIMAGE:
+    case OP_TYPE::CO_SWIZZLE:
+    case OP_TYPE::CO_FROMSCALARS:
+    case OP_TYPE::CO_ARITHMETIC:
+
+    default:
+        // TODO: Improve this test with more operations
+        //check( false );
+        res = false;
+        break;
+    }
+
+    return res;
+}
+
+
+//-------------------------------------------------------------------------------------------------
+mu::Ptr<ImageSizeExpression> ASTOpFixed::GetImageSizeExpression() const
+{
+    Ptr<ImageSizeExpression> pRes = new ImageSizeExpression;
+
+    OP_TYPE type = GetOpType();
+    switch ( type )
+    {
+    case OP_TYPE::NONE:
+    {
+        pRes->type = ImageSizeExpression::ISET_CONSTANT;
+        pRes->size[0]=0;
+        pRes->size[1]=0;
+        break;
+    }
+
+    case OP_TYPE::IM_RESIZE:
+        pRes->type = ImageSizeExpression::ISET_CONSTANT;
+        pRes->size[0] = op.args.ImageResize.size[0];
+        pRes->size[1] = op.args.ImageResize.size[1];
+        break;
+
+    case OP_TYPE::IM_RESIZEREL:
+        if ( children[op.args.ImageResizeRel.source] )
+        {
+            pRes = children[op.args.ImageResizeRel.source].child()->GetImageSizeExpression();
+            if (pRes->type == ImageSizeExpression::ISET_CONSTANT)
+            {
+                pRes->size[0] = uint16( pRes->size[0] * op.args.ImageResizeRel.factor[0] );
+                pRes->size[1] = uint16( pRes->size[1] * op.args.ImageResizeRel.factor[1] );
+            }
+            else
+            {
+                // TODO: Proportional factor
+                pRes = new ImageSizeExpression();
+                pRes->type = ImageSizeExpression::ISET_UNKNOWN;
+            }
+        }
+        break;
+
+    case OP_TYPE::IM_RESIZELIKE:
+        if ( children[op.args.ImageResizeLike.sizeSource] )
+        {
+            pRes = children[op.args.ImageResizeLike.sizeSource].child()->GetImageSizeExpression();
+        }
+        break;
+
+    case OP_TYPE::IM_PLAINCOLOUR:
+        pRes->type = ImageSizeExpression::ISET_CONSTANT;
+        pRes->size[0] = op.args.ImagePlainColour.size[0];
+        pRes->size[1] = op.args.ImagePlainColour.size[1];
+        break;
+
+    case OP_TYPE::IM_BLANKLAYOUT:
+        pRes->type = ImageSizeExpression::ISET_LAYOUTFACTOR;
+        pRes->layout = children[op.args.ImageBlankLayout.layout].child();
+        pRes->factor[0] = op.args.ImageBlankLayout.blockSize[0];
+        pRes->factor[1] = op.args.ImageBlankLayout.blockSize[1];
+        break;
+
+    case OP_TYPE::IM_INTERPOLATE:
+        if ( children[op.args.ImageInterpolate.targets[0]] )
+        {
+            pRes = children[op.args.ImageInterpolate.targets[0]].child()->GetImageSizeExpression();
+        }
+        break;
+
+    case OP_TYPE::IM_SATURATE:
+        if ( children[op.args.ImageSaturate.base] )
+        {
+            pRes = children[op.args.ImageSaturate.base].child()->GetImageSizeExpression();
+        }
+        break;
+
+    case OP_TYPE::IM_LUMINANCE:
+        if ( children[op.args.ImageLuminance.base] )
+        {
+            pRes = children[op.args.ImageLuminance.base].child()->GetImageSizeExpression();
+        }
+        break;
+
+    case OP_TYPE::IM_BINARISE:
+        if ( children[op.args.ImageBinarise.base] )
+        {
+            pRes = children[op.args.ImageBinarise.base].child()->GetImageSizeExpression();
+        }
+        break;
+
+    case OP_TYPE::IM_COLOURMAP:
+        if ( children[op.args.ImageColourMap.base] )
+        {
+            pRes = children[op.args.ImageColourMap.base].child()->GetImageSizeExpression();
+        }
+        break;
+
+    case OP_TYPE::IM_GRADIENT:
+        pRes->type = ImageSizeExpression::ISET_CONSTANT;
+        pRes->size[0] = op.args.ImageGradient.size[0];
+        pRes->size[1] = op.args.ImageGradient.size[1];
+        break;
+
+    case OP_TYPE::IM_DISPLACE:
+        if ( children[op.args.ImageDisplace.source] )
+        {
+            pRes = children[op.args.ImageDisplace.source].child()->GetImageSizeExpression();
+        }
+        break;
+
+	case OP_TYPE::IM_INVERT:
+		if (children[op.args.ImageInvert.base])
+		{
+			pRes = children[op.args.ImageInvert.base].child()->GetImageSizeExpression();
+		}
+		break;
+ 
+    default:
+    	check( false );
+    	break;
+    }
+
+    return pRes;
+}
+
+
+
+
+//!
+FSourceDataDescriptor ASTOpFixed::GetSourceDataDescriptor(FGetSourceDataDescriptorContext* Context) const
+{
+	FSourceDataDescriptor  Result;
+
+	// Cache management
+	TUniquePtr<FGetSourceDataDescriptorContext> LocalContext;
+	if (!Context)
+	{
+		LocalContext.Reset(new FGetSourceDataDescriptorContext);
+		Context = LocalContext.Get();
+	}
+
+	FSourceDataDescriptor* Found = Context->Cache.Find(this);
+	if (Found)
+	{
+		return *Found;
+	}
+
+	// Calculate
+	OP_TYPE Type = op.type;
+
+	switch (Type)
+	{
+
+	case OP_TYPE::NONE:
+	case OP_TYPE::IM_PLAINCOLOUR:
+	case OP_TYPE::IM_GRADIENT:
+	case OP_TYPE::IM_BLANKLAYOUT:
+		break;
+
+	case OP_TYPE::IM_SATURATE:
+		Result = GetSourceDataDescriptor(op.args.ImageSaturate.base, Context);
+		break;
+
+	case OP_TYPE::IM_LUMINANCE:
+		Result = GetSourceDataDescriptor(op.args.ImageLuminance.base, Context);
+		break;
+
+	case OP_TYPE::IM_INTERPOLATE:
+	{
+		for (int32 SourceIndex = 0; SourceIndex < MUTABLE_OP_MAX_INTERPOLATE_COUNT; ++SourceIndex)
+		{
+			FSourceDataDescriptor SourceDesc = GetSourceDataDescriptor(op.args.ImageInterpolate.targets[SourceIndex], Context);
+			Result.CombineWith(SourceDesc);
+		}
+		break;
+	}
+
+	case OP_TYPE::IM_RESIZE:
+		Result = GetSourceDataDescriptor(op.args.ImageResize.source, Context);
+		break;
+
+	case OP_TYPE::IM_RESIZEREL:
+		Result = GetSourceDataDescriptor(op.args.ImageResizeRel.source, Context);
+		break;
+
+	case OP_TYPE::IM_RESIZELIKE:
+		Result = GetSourceDataDescriptor(op.args.ImageResizeLike.source, Context);
+		break;
+
+	case OP_TYPE::IM_BINARISE:
+		Result = GetSourceDataDescriptor(op.args.ImageBinarise.base, Context);
+		break;
+
+	case OP_TYPE::IM_DISPLACE:
+		Result = GetSourceDataDescriptor(op.args.ImageDisplace.source, Context);
+		break;
+
+	case OP_TYPE::IM_INVERT:
+		Result = GetSourceDataDescriptor(op.args.ImageInvert.base, Context);
+		break;
+
+	case OP_TYPE::IM_COLOURMAP:
+		Result = GetSourceDataDescriptor(op.args.ImageColourMap.base, Context);
+		break;
+
+	default:
+		// This should probably be implemented
+		ensure(false);
+		break;
+	}
+
+	// Cache the result
+	Context->Cache.Add(this, Result);
+
+	return Result;
+}
+
+
+}
diff --git a/Engine/Plugins/Mutable/Source/MutableTools/Private/MuT/ASTOpAddLOD.cpp b/Engine/Plugins/Mutable/Source/MutableTools/Private/MuT/ASTOpAddLOD.cpp
new file mode 100644
--- /dev/null
+++ b/Engine/Plugins/Mutable/Source/MutableTools/Private/MuT/ASTOpAddLOD.cpp
@@ -0,0 +1,115 @@
+// Copyright Epic Games, Inc. All Rights Reserved.
+
+#include "MuT/ASTOpAddLOD.h"
+
+#include "HAL/PlatformMath.h"
+#include "MuR/ModelPrivate.h"
+#include "MuR/RefCounted.h"
+
+
+namespace mu
+{
+
+
+	//-------------------------------------------------------------------------------------------------
+	//-------------------------------------------------------------------------------------------------
+	//-------------------------------------------------------------------------------------------------
+	ASTOpAddLOD::ASTOpAddLOD()
+	{
+	}
+
+
+	ASTOpAddLOD::~ASTOpAddLOD()
+	{
+		// Explicit call needed to avoid recursive destruction
+		ASTOp::RemoveChildren();
+	}
+
+
+	bool ASTOpAddLOD::IsEqual(const ASTOp& otherUntyped) const
+	{
+		if (otherUntyped.GetOpType() == GetOpType())
+		{
+			const ASTOpAddLOD* other = static_cast<const ASTOpAddLOD*>(&otherUntyped);
+			return lods == other->lods;
+		}
+		return false;
+	}
+
+
+	uint64 ASTOpAddLOD::Hash() const
+	{
+		uint64 res = std::hash<uint64>()(uint64(OP_TYPE::IN_ADDLOD));
+		for (auto& c : lods)
+		{
+			hash_combine(res, c.child().get());
+		}
+		return res;
+	}
+
+
+	mu::Ptr<ASTOp> ASTOpAddLOD::Clone(MapChildFuncRef mapChild) const
+	{
+		Ptr<ASTOpAddLOD> n = new ASTOpAddLOD();
+		for (auto& c : lods)
+		{
+			n->lods.Emplace(n, mapChild(c.child()));
+		}
+		return n;
+	}
+
+
+	void ASTOpAddLOD::ForEachChild(const TFunctionRef<void(ASTChild&)> f)
+	{
+		for (auto& l : lods)
+		{
+			f(l);
+		}
+	}
+
+
+	void ASTOpAddLOD::Link(FProgram& program, FLinkerOptions*)
+	{
+		// Already linked?
+		if (!linkedAddress)
+		{
+			linkedAddress = (OP::ADDRESS)program.m_opAddress.Num();
+			program.m_opAddress.Add((uint32)program.m_byteCode.Num());
+			AppendCode(program.m_byteCode, OP_TYPE::IN_ADDLOD);
+
+			// Calculate LODCount mimicking previous behaviour.
+			// It may be better to leave null LODs if they are null in the AST.
+			uint8 LODCount = 0;
+			for (const ASTChild& LOD : lods)
+			{
+				OP::ADDRESS LODAddress = 0;
+				if (LOD)
+				{
+					++LODCount;
+					check(LODCount < 255);
+					if (LODCount == 255)
+					{
+						break;
+					}
+				}
+			}
+
+			AppendCode(program.m_byteCode, LODCount);
+			for (const ASTChild& LOD : lods)
+			{
+				OP::ADDRESS LODAddress = 0;
+				if (LOD)
+				{
+					LODAddress = LOD->linkedAddress;
+					AppendCode(program.m_byteCode, LODAddress);
+					if (LODCount == 255)
+					{
+						break;
+					}
+				}
+			}
+		}
+
+	}
+
+}
diff --git a/Engine/Plugins/Mutable/Source/MutableTools/Private/MuT/ASTOpAddLOD.h b/Engine/Plugins/Mutable/Source/MutableTools/Private/MuT/ASTOpAddLOD.h
new file mode 100644
--- /dev/null
+++ b/Engine/Plugins/Mutable/Source/MutableTools/Private/MuT/ASTOpAddLOD.h
@@ -0,0 +1,40 @@
+// Copyright Epic Games, Inc. All Rights Reserved.
+
+#pragma once
+
+#include "HAL/Platform.h"
+#include "MuR/Operations.h"
+#include "MuR/Ptr.h"
+#include "MuT/AST.h"
+
+
+namespace mu
+{
+struct FProgram;
+
+
+	/** Operations to add elements to a LOD. */
+	class ASTOpAddLOD final : public ASTOp
+	{
+	public:
+
+		TArray< ASTChild > lods;
+
+	public:
+
+		ASTOpAddLOD();
+		ASTOpAddLOD(const ASTOpAddLOD&) = delete;
+		~ASTOpAddLOD();
+
+		OP_TYPE GetOpType() const override { return OP_TYPE::IN_ADDLOD; }
+		uint64 Hash() const override;
+		void ForEachChild(const TFunctionRef<void(ASTChild&)>) override;
+		bool IsEqual(const ASTOp& otherUntyped) const override;
+		Ptr<ASTOp> Clone(MapChildFuncRef mapChild) const override;
+		void Link(FProgram& program, FLinkerOptions* Options) override;
+	};
+
+
+
+}
+
diff --git a/Engine/Plugins/Mutable/Source/MutableTools/Private/MuT/ASTOpConstantBool.cpp b/Engine/Plugins/Mutable/Source/MutableTools/Private/MuT/ASTOpConstantBool.cpp
new file mode 100644
--- /dev/null
+++ b/Engine/Plugins/Mutable/Source/MutableTools/Private/MuT/ASTOpConstantBool.cpp
@@ -0,0 +1,78 @@
+// Copyright Epic Games, Inc. All Rights Reserved.
+
+#include "MuT/ASTOpConstantBool.h"
+
+#include "HAL/PlatformMath.h"
+#include "MuR/ModelPrivate.h"
+#include "MuR/RefCounted.h"
+#include "MuR/Types.h"
+
+
+namespace mu
+{
+
+
+	//-------------------------------------------------------------------------------------------------
+	//-------------------------------------------------------------------------------------------------
+	//-------------------------------------------------------------------------------------------------
+	ASTOpConstantBool::ASTOpConstantBool(bool v)
+	{
+		value = v;
+	}
+
+
+	void ASTOpConstantBool::ForEachChild(const TFunctionRef<void(ASTChild&)>)
+	{
+	}
+
+
+	bool ASTOpConstantBool::IsEqual(const ASTOp& otherUntyped) const
+	{
+		if (otherUntyped.GetOpType() == GetOpType())
+		{
+			const ASTOpConstantBool* other = static_cast<const ASTOpConstantBool*>(&otherUntyped);
+			return value == other->value;
+		}
+		return false;
+	}
+
+
+	uint64 ASTOpConstantBool::Hash() const
+	{
+		uint64 res = std::hash<uint64>()(uint64(OP_TYPE::BO_CONSTANT));
+		hash_combine(res, value);
+		return res;
+	}
+
+
+	mu::Ptr<ASTOp> ASTOpConstantBool::Clone(MapChildFuncRef) const
+	{
+		Ptr<ASTOpConstantBool> n = new ASTOpConstantBool();
+		n->value = value;
+		return n;
+	}
+
+
+	void ASTOpConstantBool::Link(FProgram& program, FLinkerOptions*)
+	{
+		if (!linkedAddress)
+		{
+			OP::BoolConstantArgs args;
+			memset(&args, 0, sizeof(args));
+			args.value = value;
+
+			linkedAddress = (OP::ADDRESS)program.m_opAddress.Num();
+			//        program.m_code.push_back(op);
+			program.m_opAddress.Add((uint32_t)program.m_byteCode.Num());
+			AppendCode(program.m_byteCode, OP_TYPE::BO_CONSTANT);
+			AppendCode(program.m_byteCode, args);
+		}
+	}
+
+
+	ASTOp::FBoolEvalResult ASTOpConstantBool::EvaluateBool(ASTOpList&, FEvaluateBoolCache*) const
+	{
+		return value ? BET_TRUE : BET_FALSE;
+	}
+
+}
diff --git a/Engine/Plugins/Mutable/Source/MutableTools/Private/MuT/ASTOpConstantBool.h b/Engine/Plugins/Mutable/Source/MutableTools/Private/MuT/ASTOpConstantBool.h
new file mode 100644
--- /dev/null
+++ b/Engine/Plugins/Mutable/Source/MutableTools/Private/MuT/ASTOpConstantBool.h
@@ -0,0 +1,41 @@
+// Copyright Epic Games, Inc. All Rights Reserved.
+
+#pragma once
+
+#include "HAL/Platform.h"
+#include "MuR/Operations.h"
+#include "MuR/Ptr.h"
+#include "MuT/AST.h"
+
+
+namespace mu
+{
+struct FProgram;
+
+
+	//---------------------------------------------------------------------------------------------
+	//!
+	//---------------------------------------------------------------------------------------------
+	class ASTOpConstantBool final : public ASTOp
+	{
+	public:
+
+		//!
+		bool value;
+
+	public:
+
+		ASTOpConstantBool(bool value = true);
+
+		OP_TYPE GetOpType() const override { return OP_TYPE::BO_CONSTANT; }
+		uint64 Hash() const override;
+		void ForEachChild(const TFunctionRef<void(ASTChild&)>) override;
+		bool IsEqual(const ASTOp& otherUntyped) const override;
+		Ptr<ASTOp> Clone(MapChildFuncRef mapChild) const override;
+		void Link(FProgram& program, FLinkerOptions* Options) override;
+		FBoolEvalResult EvaluateBool(ASTOpList& facts, FEvaluateBoolCache* cache) const override;
+	};
+
+
+}
+
diff --git a/Engine/Plugins/Mutable/Source/MutableTools/Private/MuT/ASTOpConstantProjector.cpp b/Engine/Plugins/Mutable/Source/MutableTools/Private/MuT/ASTOpConstantProjector.cpp
new file mode 100644
--- /dev/null
+++ b/Engine/Plugins/Mutable/Source/MutableTools/Private/MuT/ASTOpConstantProjector.cpp
@@ -0,0 +1,65 @@
+// Copyright Epic Games, Inc. All Rights Reserved.
+
+#include "MuT/ASTOpConstantProjector.h"
+
+#include "HAL/PlatformMath.h"
+#include "MuR/ModelPrivate.h"
+#include "MuR/RefCounted.h"
+
+
+namespace mu
+{
+
+
+	//-------------------------------------------------------------------------------------------------
+	//-------------------------------------------------------------------------------------------------
+	//-------------------------------------------------------------------------------------------------
+	void ASTOpConstantProjector::ForEachChild(const TFunctionRef<void(ASTChild&)>)
+	{
+	}
+
+
+	uint64 ASTOpConstantProjector::Hash() const
+	{
+		uint64 res = std::hash<float>()(value.position[0]);
+		hash_combine(res, value.direction[0]);
+		return res;
+	}
+
+
+	bool ASTOpConstantProjector::IsEqual(const ASTOp& otherUntyped) const
+	{
+		if (otherUntyped.GetOpType() == GetOpType())
+		{
+			const ASTOpConstantProjector* other = static_cast<const ASTOpConstantProjector*>(&otherUntyped);
+			return value == other->value;
+		}
+		return false;
+	}
+
+
+	mu::Ptr<ASTOp> ASTOpConstantProjector::Clone(MapChildFuncRef) const
+	{
+		Ptr<ASTOpConstantProjector> n = new ASTOpConstantProjector();
+		n->value = value;
+		return n;
+	}
+
+
+	void ASTOpConstantProjector::Link(FProgram& program, FLinkerOptions*)
+	{
+		if (!linkedAddress)
+		{
+			OP::ResourceConstantArgs args;
+			memset(&args, 0, sizeof(args));
+			args.value = program.AddConstant(value);
+
+			linkedAddress = (OP::ADDRESS)program.m_opAddress.Num();
+			//program.m_code.push_back(op);
+			program.m_opAddress.Add((uint32_t)program.m_byteCode.Num());
+			AppendCode(program.m_byteCode, OP_TYPE::PR_CONSTANT);
+			AppendCode(program.m_byteCode, args);
+		}
+	}
+
+}
diff --git a/Engine/Plugins/Mutable/Source/MutableTools/Private/MuT/ASTOpConstantResource.cpp b/Engine/Plugins/Mutable/Source/MutableTools/Private/MuT/ASTOpConstantResource.cpp
new file mode 100644
--- /dev/null
+++ b/Engine/Plugins/Mutable/Source/MutableTools/Private/MuT/ASTOpConstantResource.cpp
@@ -0,0 +1,783 @@
+// Copyright Epic Games, Inc. All Rights Reserved.
+
+#include "MuT/ASTOpConstantResource.h"
+
+#include "MuT/CompilerPrivate.h"
+#include "MuR/Layout.h"
+#include "MuR/Mesh.h"
+#include "MuR/ModelPrivate.h"
+#include "MuR/ImagePrivate.h"
+#include "MuR/MutableMath.h"
+#include "MuR/PhysicsBody.h"
+#include "MuR/Serialisation.h"
+#include "MuR/Skeleton.h"
+#include "MuR/Types.h"
+
+#include "Containers/Array.h"
+#include "HAL/PlatformMath.h"
+#include "HAL/PlatformFileManager.h"
+#include "Hash/CityHash.h"
+#include "Misc/AssertionMacros.h"
+#include "GenericPlatform/GenericPlatformFile.h"
+#include "Compression/OodleDataCompression.h"
+
+#include <inttypes.h> // Required for 64-bit printf macros
+
+namespace mu
+{
+
+	/** Proxy class for a temporary resource while compiling. 
+	* The resource may be stored in different ways:
+	* - as is, in memory with its own pointer.
+	* - in a compressed buffer
+	* - saved to a disk file compressed or uncompressed.
+	*/
+	template<class R>
+	class MUTABLETOOLS_API ResourceProxyTempFile : public ResourceProxy<R>
+	{
+	private:
+
+		/** Actual resource to store. If the pointer is valid, it wasn't worth dumping to disk or compressing. */
+		Ptr<const R> Resource;
+
+		/** Temp filename used if it was necessary. */
+		FString FileName;
+
+		/** Size of the resource in memory. */
+		uint32 UncompressedSize = 0;
+
+		/** Size of the saved file. It may be the size of the resource in memory, or its compressed size. */
+		uint32 FileSize = 0;
+
+		/** Valid if the resource was compressed and stored in memory instead of dumped to disk. */
+		TArray<uint8> CompressedBuffer;
+
+		/** Shared context with cache settings and stats. */
+		FProxyFileContext& Options;
+
+		/** Prevent concurrent access to a signel resource. */
+		FCriticalSection Mutex;
+
+	public:
+
+		ResourceProxyTempFile(const R* InResource, FProxyFileContext& InOptions)
+			: Options(InOptions)
+		{
+			if (!InResource)
+			{
+				return;
+			}
+
+			IPlatformFile& PlatformFile = FPlatformFileManager::Get().GetPlatformFile();
+
+			OutputMemoryStream stream(128*1024);
+			OutputArchive arch(&stream);
+			R::Serialise(InResource, arch);
+
+			UncompressedSize = stream.GetBufferSize();
+
+			if (stream.GetBufferSize() <= Options.MinProxyFileSize)
+			{
+				// Not worth compressing or caching to disk
+				Resource = InResource;
+			}
+			else
+			{
+				// Compress
+				int64 CompressedSize = 0;
+				constexpr bool bEnableCompression = true;
+				if (bEnableCompression)
+				{
+					int64 CompressedBufferSize = FOodleDataCompression::CompressedBufferSizeNeeded(stream.GetBufferSize());
+					CompressedBufferSize = FMath::Max(CompressedBufferSize, int64(stream.GetBufferSize() / 2));
+					CompressedBuffer.SetNumUninitialized(CompressedBufferSize);
+
+					CompressedSize = FOodleDataCompression::CompressParallel(
+						CompressedBuffer.GetData(), CompressedBufferSize,
+						stream.GetBuffer(), stream.GetBufferSize(),
+						FOodleDataCompression::ECompressor::Kraken,
+						FOodleDataCompression::ECompressionLevel::SuperFast,
+						true // CompressIndependentChunks
+					);
+				}
+
+				bool bCompressed = CompressedSize != 0;
+
+				if (bCompressed && uint64(CompressedSize) <= Options.MinProxyFileSize)
+				{
+					// Keep the compressed data, and don't store to file
+					CompressedBuffer.SetNum(CompressedSize,EAllowShrinking::Yes);
+				}
+				else
+				{
+					// Save
+					FString Prefix = FPlatformProcess::UserTempDir();
+
+					uint32 PID = FPlatformProcess::GetCurrentProcessId();
+					Prefix += FString::Printf(TEXT("mut.temp.%u"), PID);
+
+					FString FinalTempPath;
+					IFileHandle* ResourceFile = nullptr;
+					uint64 AttemptCount = 0;
+					while (!ResourceFile && AttemptCount < Options.MaxFileCreateAttempts)
+					{
+						uint64 ThisThreadFileIndex = Options.CurrentFileIndex.load();
+						while (!Options.CurrentFileIndex.compare_exchange_strong(ThisThreadFileIndex, ThisThreadFileIndex + 1));
+
+						FinalTempPath = Prefix + FString::Printf(TEXT(".%.16" PRIx64), ThisThreadFileIndex);
+						ResourceFile = PlatformFile.OpenWrite(*FinalTempPath);
+						++AttemptCount;
+					}
+
+					if (!ResourceFile)
+					{
+						UE_LOG(LogMutableCore, Error, TEXT("Failed to create temporary file. Disk full?"));
+						check(false);
+					}
+
+					if (bCompressed)
+					{
+						FileSize = CompressedSize;
+						ResourceFile->Write(CompressedBuffer.GetData(), FileSize);
+					}
+					else
+					{
+						FileSize = UncompressedSize;
+						ResourceFile->Write((const uint8*)stream.GetBuffer(), FileSize);
+					}
+
+					CompressedBuffer.SetNum(0, EAllowShrinking::Yes);
+
+					delete ResourceFile;
+
+					FileName = FinalTempPath;
+					Options.FilesWritten++;
+					Options.BytesWritten += FileSize;
+				}
+			}
+		}
+
+		~ResourceProxyTempFile()
+		{
+			FScopeLock Lock(&Mutex);
+
+			if (!FileName.IsEmpty())
+			{
+				// Delete temp file
+				FPlatformFileManager::Get().GetPlatformFile().DeleteFile(*FileName);
+				FileName.Empty();
+			}
+		}
+
+		Ptr<const R> Get() override
+		{
+			FScopeLock Lock(&Mutex);
+
+			Ptr<const R> Result;
+			if (Resource)
+			{
+				// Cached as is
+				Result = Resource;
+			}
+			else if (!CompressedBuffer.Num() && !FileName.IsEmpty())
+			{
+				IFileHandle* resourceFile = FPlatformFileManager::Get().GetPlatformFile().OpenRead(*FileName);
+				check(resourceFile);
+
+				CompressedBuffer.SetNumUninitialized(FileSize);
+				resourceFile->Read(CompressedBuffer.GetData(), FileSize);
+				delete resourceFile;
+
+				bool bCompressed = FileSize != UncompressedSize;
+
+				if (!bCompressed)
+				{
+					InputMemoryStream stream(CompressedBuffer.GetData(), FileSize);
+					InputArchive arch(&stream);
+					Result = R::StaticUnserialise(arch);
+
+					CompressedBuffer.SetNum(0, EAllowShrinking::Yes);
+				}
+
+				Options.FilesRead++;
+				Options.BytesRead += FileSize;
+			}
+
+			if (CompressedBuffer.Num())
+			{
+				// Cached compressed
+				TArray<uint8> UncompressedBuf;
+				UncompressedBuf.SetNumUninitialized(UncompressedSize);
+
+				bool bSuccess = FOodleDataCompression::DecompressParallel(
+					UncompressedBuf.GetData(), UncompressedSize,
+					CompressedBuffer.GetData(), CompressedBuffer.Num());
+				check(bSuccess);
+
+				if (bSuccess)
+				{
+					InputMemoryStream stream(UncompressedBuf.GetData(), UncompressedSize);
+					InputArchive arch(&stream);
+					Result = R::StaticUnserialise(arch);
+				}
+
+				if (!FileName.IsEmpty())
+				{
+					CompressedBuffer.SetNum(0, EAllowShrinking::Yes);
+				}
+			}
+
+			return Result;
+		}
+
+	};
+
+
+	//-------------------------------------------------------------------------------------------------
+	//-------------------------------------------------------------------------------------------------
+	//-------------------------------------------------------------------------------------------------
+	void ASTOpConstantResource::ForEachChild(const TFunctionRef<void(ASTChild&)>)
+	{
+	}
+
+
+	//-------------------------------------------------------------------------------------------------
+	bool ASTOpConstantResource::IsEqual(const ASTOp& OtherUntyped) const
+	{
+		if (OtherUntyped.GetOpType()==GetOpType())
+		{
+			const ASTOpConstantResource* Other = static_cast<const ASTOpConstantResource*>(&OtherUntyped);
+			return Type == Other->Type && ValueHash == Other->ValueHash &&
+				LoadedValue == Other->LoadedValue && Proxy == Other->Proxy
+				&& SourceDataDescriptor == Other->SourceDataDescriptor;
+		}
+		return false;
+	}
+
+
+	//-------------------------------------------------------------------------------------------------
+	mu::Ptr<ASTOp> ASTOpConstantResource::Clone(MapChildFuncRef) const
+	{
+		Ptr<ASTOpConstantResource> n = new ASTOpConstantResource();
+		n->Type = Type;
+		n->Proxy = Proxy;
+		n->LoadedValue = LoadedValue;
+		n->ValueHash = ValueHash;
+		n->SourceDataDescriptor = SourceDataDescriptor; 
+		return n;
+	}
+
+
+	//-------------------------------------------------------------------------------------------------
+	uint64 ASTOpConstantResource::Hash() const
+	{
+		uint64 res = std::hash<uint64>()(uint64(Type));
+		hash_combine(res, ValueHash);
+		return res;
+	}
+
+
+	namespace
+	{
+		/** Adds a constant image data to a program and returns its constant index. */
+		int32 AddConstantImage(FProgram& Program, const Ptr<const Image>& pImage, FLinkerOptions& Options)
+		{
+			MUTABLE_CPUPROFILER_SCOPE(AddConstantImage);
+
+			check(pImage->GetSizeX() * pImage->GetSizeY() > 0);
+			
+			// Mips to store
+			int32 MipsToStore = 1;
+
+			int32 FirstLODIndexIndex = Program.ConstantImageLODIndices.Num();
+
+			FImageOperator& ImOp = Options.ImageOperator;
+			Ptr<const Image> pMip;
+
+			if (!Options.bSeparateImageMips)
+			{
+				pMip = pImage;
+			}
+			else
+			{
+				// We may want the full mipmaps for fragments of images, regardless of the resident mip size, for intermediate operations.
+				// \TODO: Calculate the mip ranges that makes sense to store.
+				int32 MaxMipmaps = Image::GetMipmapCount(pImage->GetSizeX(), pImage->GetSizeY());
+				MipsToStore = MaxMipmaps;
+
+				// Some images cannot be resized or mipmaped
+				bool bCannotBeScaled = pImage->m_flags & Image::IF_CANNOT_BE_SCALED;
+				if (bCannotBeScaled)
+				{
+					// Store only the mips that we have already calculated. We assume we have calculated them correctly.
+					MipsToStore = pImage->GetLODCount();
+				}
+
+				if (pImage->GetLODCount() == 1)
+				{
+					pMip = pImage;
+				}
+				else
+				{
+					pMip = ImOp.ExtractMip(pImage.get(), 0);
+				}
+			}
+
+			// Temporary uncompressed version of the image, if we need to generate the mips and the source is compressed.
+			Ptr<const Image> UncompressedMip;
+			EImageFormat UncompressedFormat = GetUncompressedFormat( pMip->GetFormat() );
+
+			for (int32 Mip = 0; Mip < MipsToStore; ++Mip)
+			{
+				check(pMip->GetFormat() == pImage->GetFormat());
+
+				// Ensure unique at mip level
+				int32 MipIndex = -1;
+
+				// Use a map-based deduplication only if we are splitting mips.
+				if (Options.bSeparateImageMips)
+				{
+					MUTABLE_CPUPROFILER_SCOPE(Deduplicate);
+
+					const int32* IndexPtr = Options.ImageConstantMipMap.Find(pMip);
+					if (IndexPtr)
+					{
+						MipIndex = *IndexPtr;
+					}
+				}
+
+				if (MipIndex<0)
+				{
+					MipIndex = Program.ConstantImageLODs.Add(TPair<int32, Ptr<const Image>>(-1, pMip));
+					Options.ImageConstantMipMap.Add(pMip, MipIndex);
+				}
+
+				Program.ConstantImageLODIndices.Add(uint32(MipIndex));
+
+				// Generate next mip if necessary
+				if (Mip + 1 < MipsToStore)
+				{
+					Ptr<Image> NewMip;
+					if (Mip+1 < pImage->GetLODCount())
+					{
+						// Extract directly from source image
+						NewMip = ImOp.ExtractMip(pImage.get(), Mip + 1);
+					}
+					else
+					{
+						// Generate from the last mip.
+						if (UncompressedFormat!=pMip->GetFormat())
+						{
+							int32 Quality = 4; // TODO
+
+							if (!UncompressedMip)
+							{
+								UncompressedMip = ImOp.ImagePixelFormat(Quality, pMip.get(), UncompressedFormat);
+							}
+
+							UncompressedMip = ImOp.ExtractMip(UncompressedMip.get(), 1);
+							NewMip = ImOp.ImagePixelFormat(Quality, UncompressedMip.get(), pMip->GetFormat());
+						}
+						else
+						{
+							NewMip = ImOp.ExtractMip(pMip.get(), 1);
+						}
+					}
+					check(NewMip);
+
+					pMip = NewMip;
+				}
+			}
+
+			FImageLODRange LODRange;
+			LODRange.FirstIndex = FirstLODIndexIndex;
+			LODRange.LODCount = MipsToStore;
+			LODRange.ImageFormat = pImage->GetFormat();
+			LODRange.ImageSizeX = pImage->GetSizeX();
+			LODRange.ImageSizeY = pImage->GetSizeY();
+			int32 ImageIndex = Program.ConstantImages.Add(LODRange);
+			return ImageIndex;
+		}
+	}
+
+	//-------------------------------------------------------------------------------------------------
+	void ASTOpConstantResource::Link(FProgram& program, FLinkerOptions* Options)
+	{
+		MUTABLE_CPUPROFILER_SCOPE(ASTOpConstantResource_Link);
+
+		if (!linkedAddress && !bLinkedAndNull)
+		{
+			if (Type == OP_TYPE::ME_CONSTANT)
+			{
+				OP::MeshConstantArgs args;
+				FMemory::Memset(&args, 0, sizeof(args));
+
+				Ptr<Mesh> MeshData = static_cast<const Mesh*>(GetValue().get())->Clone();
+				check(MeshData);
+
+				args.skeleton = -1;
+				if (Ptr<const Skeleton> pSkeleton = MeshData->GetSkeleton())
+				{
+					args.skeleton = program.AddConstant(pSkeleton.get());
+					MeshData->SetSkeleton(nullptr);
+				}
+
+				args.physicsBody = -1;
+				if (Ptr<const PhysicsBody> pPhysicsBody = MeshData->GetPhysicsBody())
+				{
+					args.physicsBody = program.AddConstant(pPhysicsBody.get());
+					MeshData->SetPhysicsBody(nullptr);
+				}
+
+				// Use a map-based deduplication
+				mu::Ptr<const mu::Mesh> Key = MeshData;
+				const int32* IndexPtr = Options->MeshConstantMap.Find(Key);
+				if (!IndexPtr)
+				{
+					args.value = program.AddConstant(MeshData.get());
+					Options->MeshConstantMap.Add(MeshData, int32(args.value));
+				}
+				else
+				{
+					args.value = *IndexPtr;
+				}
+
+				linkedAddress = (OP::ADDRESS)program.m_opAddress.Num();
+				program.m_opAddress.Add((uint32_t)program.m_byteCode.Num());
+				AppendCode(program.m_byteCode, Type);
+				AppendCode(program.m_byteCode, args);
+			}
+			else
+			{
+				OP::ResourceConstantArgs args;
+				FMemory::Memset(&args, 0, sizeof(args));
+
+				bool bValidData = true;
+
+				switch (Type)
+				{
+				case OP_TYPE::IM_CONSTANT:
+				{
+					Ptr<const Image> pTyped = static_cast<const Image*>(GetValue().get());
+					check(pTyped);
+
+					if (pTyped->GetSizeX() * pTyped->GetSizeY() == 0)
+					{
+						// It's an empty or degenerated image, return a null operation.
+						bValidData = false;
+					}
+					else
+					{
+						args.value = AddConstantImage( program, pTyped, *Options);
+
+						int32 DataDescIndex = Options->AdditionalData.SourceImagePerConstant.Add(SourceDataDescriptor);
+						check(DataDescIndex == args.value);
+					}
+
+					break;
+				}
+				case OP_TYPE::LA_CONSTANT:
+				{
+					Ptr<const Layout> pTyped = static_cast<const Layout*>(GetValue().get());
+					check(pTyped);
+					args.value = program.AddConstant(pTyped);
+					break;
+				}
+				default:
+					check(false);
+				}
+
+				if (bValidData)
+				{
+					linkedAddress = (OP::ADDRESS)program.m_opAddress.Num();
+					program.m_opAddress.Add((uint32)program.m_byteCode.Num());
+					AppendCode(program.m_byteCode, Type);
+					AppendCode(program.m_byteCode, args);
+				}
+				else
+				{
+					// Null op
+					linkedAddress = 0;
+					bLinkedAndNull = true;
+				}
+			}
+
+			// Clear stored value to reduce memory usage.
+			LoadedValue = nullptr;
+			Proxy = nullptr;
+		}
+	}
+
+
+	//-------------------------------------------------------------------------------------------------
+	FImageDesc ASTOpConstantResource::GetImageDesc(bool, class FGetImageDescContext*) const
+	{
+		FImageDesc Result;
+
+		if (Type == OP_TYPE::IM_CONSTANT)
+		{
+			// TODO: cache to avoid disk loading
+			Ptr<const Image> ConstImage = static_cast<const Image*>(GetValue().get());
+			Result.m_format = ConstImage->GetFormat();
+			Result.m_lods = ConstImage->GetLODCount();
+			Result.m_size = ConstImage->GetSize();
+		}
+		else
+		{
+			check(false);
+		}
+
+		return Result;
+	}
+
+	//-------------------------------------------------------------------------------------------------
+	void ASTOpConstantResource::GetBlockLayoutSize(uint64 BlockId, int32* BlockX, int32* BlockY, FBlockLayoutSizeCache*)
+	{
+		switch (Type)
+		{
+		case OP_TYPE::LA_CONSTANT:
+		{
+			Ptr<const Layout> pLayout = static_cast<const Layout*>(GetValue().get());
+			check(pLayout);
+
+			if (pLayout)
+			{
+				int relId = pLayout->FindBlock(BlockId);
+				if (relId >= 0)
+				{
+					*BlockX = pLayout->Blocks[relId].Size[0];
+					*BlockY = pLayout->Blocks[relId].Size[1];
+				}
+				else
+				{
+					*BlockX = 0;
+					*BlockY = 0;
+				}
+			}
+
+			break;
+		}
+		default:
+			check(false);
+		}
+	}
+
+
+	//-------------------------------------------------------------------------------------------------
+	void ASTOpConstantResource::GetLayoutBlockSize(int32* pBlockX, int32* pBlockY)
+	{
+		switch (Type)
+		{
+
+		case OP_TYPE::IM_CONSTANT:
+		{
+			// We didn't find any layout.
+			*pBlockX = 0;
+			*pBlockY = 0;
+			break;
+		}
+
+		default:
+			checkf(false, TEXT("Instruction not supported"));
+		}
+	}
+
+
+	//-------------------------------------------------------------------------------------------------
+	bool ASTOpConstantResource::GetNonBlackRect(FImageRect& maskUsage) const
+	{
+		if (Type == OP_TYPE::IM_CONSTANT)
+		{
+			// TODO: cache
+			Ptr<const Image> pMask = static_cast<const Image*>(GetValue().get());
+			pMask->GetNonBlackRect(maskUsage);
+			return true;
+		}
+
+		return false;
+	}
+
+
+	//-------------------------------------------------------------------------------------------------
+	bool ASTOpConstantResource::IsImagePlainConstant(FVector4f& colour) const
+	{
+		bool res = false;
+		switch (Type)
+		{
+
+		case OP_TYPE::IM_CONSTANT:
+		{
+			Ptr<const Image> pImage = static_cast<const Image*>(GetValue().get());
+			if (pImage->GetSizeX() <= 0 || pImage->GetSizeY() <= 0)
+			{
+				res = true;
+				colour = FVector4f(0.0f,0.0f,0.0f,1.0f);
+			}
+			else if (pImage->m_flags & Image::IF_IS_PLAIN_COLOUR_VALID)
+			{
+				if (pImage->m_flags & Image::IF_IS_PLAIN_COLOUR)
+				{
+					res = true;
+					colour = pImage->Sample(FVector2f(0, 0));
+				}
+				else
+				{
+					res = false;
+				}
+			}
+			else
+			{
+				if (pImage->IsPlainColour(colour))
+				{
+					res = true;
+					pImage->m_flags |= Image::IF_IS_PLAIN_COLOUR;
+				}
+
+				pImage->m_flags |= Image::IF_IS_PLAIN_COLOUR_VALID;
+			}
+			break;
+		}
+
+		default:
+			break;
+		}
+
+		return res;
+	}
+
+
+	//-------------------------------------------------------------------------------------------------
+	ASTOpConstantResource::~ASTOpConstantResource()
+	{
+	}
+
+
+	//-------------------------------------------------------------------------------------------------
+	uint64 ASTOpConstantResource::GetValueHash() const
+	{
+		return ValueHash;
+	}
+
+
+	//-------------------------------------------------------------------------------------------------
+	mu::Ptr<const RefCounted> ASTOpConstantResource::GetValue() const
+	{
+		if (LoadedValue)
+		{
+			return LoadedValue;
+		}
+		else
+		{
+			switch (Type)
+			{
+
+			case OP_TYPE::IM_CONSTANT:
+			{
+				Ptr<ResourceProxy<Image>> typedProxy = static_cast<ResourceProxy<Image>*>(Proxy.get());
+				Ptr<const Image> r = typedProxy->Get();
+				return r;
+			}
+
+			default:
+				check(false);
+				break;
+			}
+		}
+
+		return nullptr;
+	}
+
+
+	//-------------------------------------------------------------------------------------------------
+	void ASTOpConstantResource::SetValue(const Ptr<const RefCounted>& v, FProxyFileContext* DiskCacheContext)
+	{
+		MUTABLE_CPUPROFILER_SCOPE(ASTOpConstantResource_SetValue);
+
+		switch (Type)
+		{
+		case OP_TYPE::IM_CONSTANT:
+		{
+			Ptr<const Image> r = static_cast<const Image*>(v.get());
+
+			OutputHashStream stream;
+			{
+				MUTABLE_CPUPROFILER_SCOPE(Serialize);
+				OutputArchive arch(&stream);
+				Image::Serialise(r.get(), arch);
+			}
+
+			ValueHash = stream.GetHash();
+
+			if (DiskCacheContext)
+			{
+				Proxy = new ResourceProxyTempFile<Image>(r.get(), *DiskCacheContext);
+			}
+			else
+			{
+				LoadedValue = r;
+			}
+			break;
+		}
+
+		case OP_TYPE::ME_CONSTANT:
+		{
+			Ptr<const Mesh> r = static_cast<const Mesh*>(v.get());
+
+			OutputHashStream stream;
+			{
+				MUTABLE_CPUPROFILER_SCOPE(Serialize);
+				OutputArchive arch(&stream);
+				Mesh::Serialise(r.get(), arch);
+			}
+
+			ValueHash = stream.GetHash();
+
+			LoadedValue = v;
+			break;
+		}
+
+		case OP_TYPE::LA_CONSTANT:
+		{
+			Ptr<const Layout> r = static_cast<const Layout*>(v.get());
+
+			OutputHashStream stream;
+			{
+				MUTABLE_CPUPROFILER_SCOPE(Serialize);
+				OutputArchive arch(&stream);
+				Layout::Serialise(r.get(), arch);
+			}
+
+			ValueHash = stream.GetHash();
+
+			LoadedValue = v;
+			break;
+		}
+
+		default:
+			LoadedValue = v;
+			break;
+		}
+	}
+
+
+	mu::Ptr<ImageSizeExpression> ASTOpConstantResource::GetImageSizeExpression() const
+	{
+		if (Type==OP_TYPE::IM_CONSTANT)
+		{
+			Ptr<ImageSizeExpression> pRes = new ImageSizeExpression;
+			pRes->type = ImageSizeExpression::ISET_CONSTANT;
+			Ptr<const Image> pConst = static_cast<const Image*>(GetValue().get());
+			pRes->size = pConst->GetSize();
+			return pRes;
+		}
+
+		return nullptr;
+	}
+
+
+	FSourceDataDescriptor ASTOpConstantResource::GetSourceDataDescriptor(FGetSourceDataDescriptorContext*) const
+	{
+		return SourceDataDescriptor;
+	}
+
+}
diff --git a/Engine/Plugins/Mutable/Source/MutableTools/Private/MuT/ASTOpConstantString.cpp b/Engine/Plugins/Mutable/Source/MutableTools/Private/MuT/ASTOpConstantString.cpp
new file mode 100644
--- /dev/null
+++ b/Engine/Plugins/Mutable/Source/MutableTools/Private/MuT/ASTOpConstantString.cpp
@@ -0,0 +1,64 @@
+// Copyright Epic Games, Inc. All Rights Reserved.
+
+#include "MuT/ASTOpConstantString.h"
+
+#include "HAL/PlatformMath.h"
+#include "MuR/ModelPrivate.h"
+#include "MuR/RefCounted.h"
+#include "MuR/Types.h"
+
+
+namespace mu
+{
+
+
+	//-------------------------------------------------------------------------------------------------
+	//-------------------------------------------------------------------------------------------------
+	//-------------------------------------------------------------------------------------------------
+	void ASTOpConstantString::ForEachChild(const TFunctionRef<void(ASTChild&)>)
+	{
+	}
+
+
+	uint64 ASTOpConstantString::Hash() const
+	{
+		uint64 res = std::hash<int32>()(value.Len());
+		return res;
+	}
+
+
+	bool ASTOpConstantString::IsEqual(const ASTOp& otherUntyped) const
+	{
+		if (otherUntyped.GetOpType() == GetOpType())
+		{
+			const ASTOpConstantString* other = static_cast<const ASTOpConstantString*>(&otherUntyped);
+			return value == other->value;
+		}
+		return false;
+	}
+
+
+	mu::Ptr<ASTOp> ASTOpConstantString::Clone(MapChildFuncRef) const
+	{
+		Ptr<ASTOpConstantString> n = new ASTOpConstantString();
+		n->value = value;
+		return n;
+	}
+
+
+	void ASTOpConstantString::Link(FProgram& program, FLinkerOptions*)
+	{
+		if (!linkedAddress)
+		{
+			OP::ResourceConstantArgs args;
+			FMemory::Memset(&args, 0, sizeof(args));
+			args.value = program.AddConstant(value);
+
+			linkedAddress = (OP::ADDRESS)program.m_opAddress.Num();
+			program.m_opAddress.Add((uint32_t)program.m_byteCode.Num());
+			AppendCode(program.m_byteCode, OP_TYPE::ST_CONSTANT);
+			AppendCode(program.m_byteCode, args);
+		}
+	}
+
+}
\ No newline at end of file
diff --git a/Engine/Plugins/Mutable/Source/MutableTools/Private/MuT/ASTOpImageCompose.cpp b/Engine/Plugins/Mutable/Source/MutableTools/Private/MuT/ASTOpImageCompose.cpp
new file mode 100644
--- /dev/null
+++ b/Engine/Plugins/Mutable/Source/MutableTools/Private/MuT/ASTOpImageCompose.cpp
@@ -0,0 +1,339 @@
+// Copyright Epic Games, Inc. All Rights Reserved.
+
+#include "MuT/ASTOpImageCompose.h"
+
+#include "Containers/Array.h"
+#include "Containers/Map.h"
+#include "HAL/PlatformMath.h"
+#include "MuR/Layout.h"
+#include "MuR/ModelPrivate.h"
+#include "MuR/MutableMath.h"
+#include "MuR/MutableTrace.h"
+#include "MuR/ImagePrivate.h"
+#include "MuR/RefCounted.h"
+#include "MuR/Types.h"
+#include "MuT/ASTOpConstantResource.h"
+#include "MuT/ASTOpImagePixelFormat.h"
+
+
+namespace mu
+{
+
+
+//-------------------------------------------------------------------------------------------------
+//-------------------------------------------------------------------------------------------------
+//-------------------------------------------------------------------------------------------------
+ASTOpImageCompose::ASTOpImageCompose()
+    : Layout(this)
+    , Base(this)
+	, BlockImage(this)
+	, Mask(this)
+{
+}
+
+
+//-------------------------------------------------------------------------------------------------
+ASTOpImageCompose::~ASTOpImageCompose()
+{
+    // Explicit call needed to avoid recursive destruction
+    ASTOp::RemoveChildren();
+}
+
+
+//-------------------------------------------------------------------------------------------------
+bool ASTOpImageCompose::IsEqual(const ASTOp& otherUntyped) const
+{
+	if (otherUntyped.GetOpType() == GetOpType())
+    {
+		const ASTOpImageCompose* other = static_cast<const ASTOpImageCompose*>(&otherUntyped);
+		return Layout == other->Layout &&
+			Base ==other->Base &&
+			BlockImage == other->BlockImage &&
+			Mask == other->Mask &&
+			BlockId == other->BlockId;
+    }
+    return false;
+}
+
+
+//-------------------------------------------------------------------------------------------------
+uint64 ASTOpImageCompose::Hash() const
+{
+	uint64 res = std::hash<OP_TYPE>()(OP_TYPE::IM_COMPOSE);
+    hash_combine( res, Layout.child().get() );
+	hash_combine(res, Base.child().get());
+	hash_combine(res, BlockImage.child().get());
+	hash_combine(res, Mask.child().get());
+	hash_combine(res, BlockId);
+	return res;
+}
+
+
+//-------------------------------------------------------------------------------------------------
+mu::Ptr<ASTOp> ASTOpImageCompose::Clone(MapChildFuncRef mapChild) const
+{
+	mu::Ptr<ASTOpImageCompose> n = new ASTOpImageCompose();
+    n->Layout = mapChild(Layout.child());
+	n->Base = mapChild(Base.child());
+	n->BlockImage = mapChild(BlockImage.child());
+	n->Mask = mapChild(Mask.child());
+	n->BlockId = BlockId;
+    return n;
+}
+
+
+//-------------------------------------------------------------------------------------------------
+void ASTOpImageCompose::ForEachChild(const TFunctionRef<void(ASTChild&)> f )
+{
+	f(Layout);
+	f(Base);
+	f(BlockImage);
+	f(Mask);
+}
+
+
+//-------------------------------------------------------------------------------------------------
+void ASTOpImageCompose::Link( FProgram& program, FLinkerOptions*)
+{
+    // Already linked?
+    if (!linkedAddress)
+    {
+        OP::ImageComposeArgs args;
+        memset( &args,0, sizeof(args) );
+
+		if (Layout) args.layout = Layout->linkedAddress;
+		if (Base) args.base = Base->linkedAddress;
+		if (BlockImage) args.blockImage = BlockImage->linkedAddress;
+		if (Mask) args.mask = Mask->linkedAddress;
+        args.BlockId = BlockId;
+
+        linkedAddress = (OP::ADDRESS)program.m_opAddress.Num();
+        program.m_opAddress.Add((uint32_t)program.m_byteCode.Num());
+        AppendCode(program.m_byteCode,OP_TYPE::IM_COMPOSE);
+        AppendCode(program.m_byteCode,args);
+    }
+
+}
+
+
+FImageDesc ASTOpImageCompose::GetImageDesc( bool returnBestOption, FGetImageDescContext* context ) const
+{
+    FImageDesc res;
+
+    // Local context in case it is necessary
+    FGetImageDescContext localContext;
+    if (!context)
+    {
+      context = &localContext;
+    }
+    else
+    {
+        // Cached result?
+		FImageDesc* PtrValue = context->m_results.Find(this);
+		if (PtrValue)
+		{
+			return *PtrValue;
+		}
+	}
+
+    // Actual work
+    if (Base)
+    {
+        res = Base->GetImageDesc( returnBestOption, context );
+    }
+
+	if (BlockImage)
+	{
+		FImageDesc BlockDesc = BlockImage->GetImageDesc(returnBestOption, context);
+		res.m_format = GetMostGenericFormat(res.m_format,BlockDesc.m_format);
+	}
+
+    // Cache th result
+    if (context)
+    {
+		context->m_results.Add(this, res);
+	}
+
+    return res;
+}
+
+
+mu::Ptr<ImageSizeExpression> ASTOpImageCompose::GetImageSizeExpression() const
+{
+    if ( Base )
+    {
+        return Base->GetImageSizeExpression();
+    }
+
+    return nullptr;
+}
+
+
+bool ASTOpImageCompose::IsImagePlainConstant(FVector4f& colour) const
+{
+	bool res = false;
+
+	if (BlockImage.child())
+	{
+		BlockImage->IsImagePlainConstant(colour);
+	}
+
+	if (res && Base.child())
+	{
+		FVector4f baseColour;
+		res = Base->IsImagePlainConstant(baseColour);
+		res &= (colour == baseColour);
+	}
+
+	return res;
+}
+
+
+void ASTOpImageCompose::GetLayoutBlockSize(int32* pBlockX, int32* pBlockY)
+{
+	// Try to follow the base image of the compose, which is the most stable.
+	if (Base)
+	{
+		Base->GetLayoutBlockSize(pBlockX,pBlockY);
+	}
+
+	// We can only follow the block if the base is empty, since the first block will set the block size.	
+	if (*pBlockX == 0)
+	{
+		// Let's try the block approach: We need the block size and the layout blocks
+		int32 layoutBlocksX = 0;
+		int32 layoutBlocksY = 0;
+		if (Layout.child())
+		{
+			MUTABLE_CPUPROFILER_SCOPE(GetLayoutBlockSize_GetBlockLayoutSize);
+			FBlockLayoutSizeCache cache;
+			Layout->GetBlockLayoutSizeCached(BlockId, &layoutBlocksX, &layoutBlocksY, &cache);
+		}
+
+		if (layoutBlocksX > 0 && layoutBlocksY > 0 && BlockImage.child())
+		{
+			FImageDesc blockDesc = BlockImage->GetImageDesc();
+
+			*pBlockX = blockDesc.m_size[0] / layoutBlocksX;
+			*pBlockY = blockDesc.m_size[1] / layoutBlocksY;
+		}
+		else
+		{
+			*pBlockX = 0;
+			*pBlockY = 0;
+		}
+	}
+}
+
+
+mu::Ptr<ASTOp> ASTOpImageCompose::OptimiseSemantic(const FModelOptimizationOptions& options, int32 Pass) const
+{
+	mu::Ptr<ASTOp> at;
+
+	mu::Ptr<ASTOp> baseAt = Base.child();
+	mu::Ptr<ASTOp> blockAt = BlockImage.child();
+	mu::Ptr<ASTOp> layoutAt = Layout.child();
+	if (layoutAt
+		&&
+		layoutAt->GetOpType() == OP_TYPE::LA_CONSTANT
+		&&
+		baseAt
+		&&
+		blockAt)
+	{
+		const ASTOpConstantResource* typedLayout = static_cast<const ASTOpConstantResource*>(layoutAt.get());
+		mu::Ptr<const mu::Layout> pLayout = static_cast<const mu::Layout*>(typedLayout->GetValue().get());
+
+		// Constant single-block full layout?
+		if (pLayout->GetBlockCount() == 1
+			&&
+			pLayout->Blocks[0].Min == FIntVector2(0, 0)
+			&&
+			pLayout->Blocks[0].Size == pLayout->Size
+			&&
+			pLayout->Blocks[0].Id == BlockId
+			)
+		{
+			// We could only take the block, but we must make sure it will have the format
+			// and size of the base.
+			mu::FImageDesc baseDesc = baseAt->GetImageDesc(true);
+			mu::FImageDesc blockDesc = blockAt->GetImageDesc(true);
+
+			at = blockAt;
+
+			// \todo isn't this a common operation?
+			if (baseDesc.m_format != blockDesc.m_format
+				&&
+				baseDesc.m_format != EImageFormat::IF_NONE)
+			{
+				mu::Ptr<ASTOpImagePixelFormat> reformat = new ASTOpImagePixelFormat;
+				reformat->Format = baseDesc.m_format;
+				reformat->FormatIfAlpha = EImageFormat::IF_NONE;
+				reformat->Source = at;
+				at = reformat;
+			}
+
+			if (baseDesc.m_size != blockDesc.m_size
+				&&
+				baseDesc.m_size[0] != 0
+				&&
+				baseDesc.m_size[1] != 0)
+			{
+				mu::Ptr<ASTOpFixed> resize = new ASTOpFixed;
+				resize->op.type = OP_TYPE::IM_RESIZE;
+				resize->op.args.ImageResize.size[0] = baseDesc.m_size[0];
+				resize->op.args.ImageResize.size[1] = baseDesc.m_size[1];
+				resize->SetChild(resize->op.args.ImageResize.source, at);
+				at = resize;
+			}
+		}
+	}
+
+	return at;
+}
+
+
+FSourceDataDescriptor ASTOpImageCompose::GetSourceDataDescriptor(FGetSourceDataDescriptorContext* Context) const
+{
+	// Cache management
+	TUniquePtr<FGetSourceDataDescriptorContext> LocalContext;
+	if (!Context)
+	{
+		LocalContext.Reset(new FGetSourceDataDescriptorContext);
+		Context = LocalContext.Get();
+	}
+
+	FSourceDataDescriptor* Found = Context->Cache.Find(this);
+	if (Found)
+	{
+		return *Found;
+	}
+
+	// Not cached: calculate
+	FSourceDataDescriptor Result;
+
+	if (Base)
+	{
+		FSourceDataDescriptor SourceDesc = Base->GetSourceDataDescriptor(Context);
+		Result.CombineWith(SourceDesc);
+	}
+
+	if (BlockImage)
+	{
+		FSourceDataDescriptor SourceDesc = BlockImage->GetSourceDataDescriptor(Context);
+		Result.CombineWith(SourceDesc);
+	}
+
+	if (Mask)
+	{
+		FSourceDataDescriptor SourceDesc = Mask->GetSourceDataDescriptor(Context);
+		Result.CombineWith(SourceDesc);
+	}
+
+	Context->Cache.Add(this, Result);
+
+	return Result;
+}
+
+
+}
\ No newline at end of file
diff --git a/Engine/Plugins/Mutable/Source/MutableTools/Private/MuT/ASTOpImagePatch.h b/Engine/Plugins/Mutable/Source/MutableTools/Private/MuT/ASTOpImagePatch.h
new file mode 100644
--- /dev/null
+++ b/Engine/Plugins/Mutable/Source/MutableTools/Private/MuT/ASTOpImagePatch.h
@@ -0,0 +1,47 @@
+// Copyright Epic Games, Inc. All Rights Reserved.
+
+#pragma once
+
+#include "HAL/PlatformMath.h"
+#include "MuR/Image.h"
+#include "MuR/MutableMath.h"
+#include "MuR/Operations.h"
+#include "MuR/Ptr.h"
+#include "MuT/AST.h"
+
+
+namespace mu
+{
+struct FProgram;
+
+	class ASTOpImagePatch final : public ASTOp
+	{
+	public:
+
+		ASTChild base;
+		ASTChild patch;
+		UE::Math::TIntVector2<uint16> location = UE::Math::TIntVector2<uint16>(0, 0);
+
+	public:
+
+		ASTOpImagePatch();
+		ASTOpImagePatch(const ASTOpImagePatch&) = delete;
+		~ASTOpImagePatch();
+
+		virtual OP_TYPE GetOpType() const override { return OP_TYPE::IM_PATCH; }
+		virtual uint64 Hash() const override;
+		virtual bool IsEqual(const ASTOp& otherUntyped) const override;
+		virtual Ptr<ASTOp> Clone(MapChildFuncRef mapChild) const override;
+		virtual void ForEachChild(const TFunctionRef<void(ASTChild&)>) override;
+		virtual void Link(FProgram& program, FLinkerOptions* Options) override;
+		virtual FImageDesc GetImageDesc(bool returnBestOption, FGetImageDescContext* context) const override;
+		virtual Ptr<ImageSizeExpression> GetImageSizeExpression() const override;
+		virtual void GetLayoutBlockSize(int* pBlockX, int* pBlockY) override;
+		//TODO: virtual bool IsImagePlainConstant(FVector4f& colour) const override;
+		virtual FSourceDataDescriptor GetSourceDataDescriptor(FGetSourceDataDescriptorContext*) const override;
+
+	};
+
+
+}
+
diff --git a/Engine/Plugins/Mutable/Source/MutableTools/Private/MuT/ASTOpInstanceAdd.cpp b/Engine/Plugins/Mutable/Source/MutableTools/Private/MuT/ASTOpInstanceAdd.cpp
new file mode 100644
--- /dev/null
+++ b/Engine/Plugins/Mutable/Source/MutableTools/Private/MuT/ASTOpInstanceAdd.cpp
@@ -0,0 +1,169 @@
+// Copyright Epic Games, Inc. All Rights Reserved.
+
+#include "MuT/ASTOpInstanceAdd.h"
+
+#include "Misc/AssertionMacros.h"
+#include "MuR/ModelPrivate.h"
+#include "MuR/ParametersPrivate.h"
+#include "MuR/RefCounted.h"
+#include "MuR/Types.h"
+#include "MuT/CodeOptimiser.h"
+
+
+namespace mu
+{
+
+
+	//-------------------------------------------------------------------------------------------------
+	ASTOpInstanceAdd::ASTOpInstanceAdd()
+		: instance(this)
+		, value(this)
+	{
+	}
+
+
+	//-------------------------------------------------------------------------------------------------
+	ASTOpInstanceAdd::~ASTOpInstanceAdd()
+	{
+		// Explicit call needed to avoid recursive destruction
+		ASTOp::RemoveChildren();
+	}
+
+
+	//-------------------------------------------------------------------------------------------------
+	bool ASTOpInstanceAdd::IsEqual(const ASTOp& otherUntyped) const
+	{
+		if (otherUntyped.GetOpType()==GetOpType())
+		{
+			const ASTOpInstanceAdd* other = static_cast<const ASTOpInstanceAdd*>(&otherUntyped);
+			return type == other->type &&
+				instance == other->instance &&
+				value == other->value &&
+				id == other->id &&
+				ExternalId == other->ExternalId &&
+				SharedSurfaceId == other->SharedSurfaceId &&
+				name == other->name;
+		}
+		return false;
+	}
+
+
+	//-------------------------------------------------------------------------------------------------
+	mu::Ptr<ASTOp> ASTOpInstanceAdd::Clone(MapChildFuncRef mapChild) const
+	{
+		Ptr<ASTOpInstanceAdd> n = new ASTOpInstanceAdd();
+		n->type = type;
+		n->instance = mapChild(instance.child());
+		n->value = mapChild(value.child());
+		n->id = id;
+		n->ExternalId = ExternalId;
+		n->SharedSurfaceId = SharedSurfaceId;
+		n->name = name;
+		return n;
+	}
+
+
+	//-------------------------------------------------------------------------------------------------
+	uint64 ASTOpInstanceAdd::Hash() const
+	{
+		uint64 res = std::hash<size_t>()(size_t(type));
+		hash_combine(res, instance.child().get());
+		hash_combine(res, value.child().get());
+		return res;
+	}
+
+
+	//-------------------------------------------------------------------------------------------------
+	void ASTOpInstanceAdd::Assert()
+	{
+		switch (type)
+		{
+		case OP_TYPE::IN_ADDMESH:
+		case OP_TYPE::IN_ADDIMAGE:
+		case OP_TYPE::IN_ADDVECTOR:
+		case OP_TYPE::IN_ADDSCALAR:
+		case OP_TYPE::IN_ADDSTRING:
+		case OP_TYPE::IN_ADDSURFACE:
+		case OP_TYPE::IN_ADDCOMPONENT:
+		case OP_TYPE::IN_ADDLOD:
+			break;
+		default:
+			// Unexpected type
+			check(false);
+			break;
+		}
+
+		ASTOp::Assert();
+	}
+
+
+	//-------------------------------------------------------------------------------------------------
+	void ASTOpInstanceAdd::ForEachChild(const TFunctionRef<void(ASTChild&)> f)
+	{
+		f(instance);
+		f(value);
+	}
+
+
+	//-------------------------------------------------------------------------------------------------
+	void ASTOpInstanceAdd::Link(FProgram& program, FLinkerOptions*)
+	{
+		// Already linked?
+		if (!linkedAddress)
+		{
+			OP::InstanceAddArgs args;
+			FMemory::Memzero(&args, sizeof(args));
+			args.id = id;
+			args.ExternalId = ExternalId;
+			args.SharedSurfaceId = SharedSurfaceId;
+			args.name = program.AddConstant(name);
+
+			if (instance) args.instance = instance->linkedAddress;
+			if (value) args.value = value->linkedAddress;
+
+			if (type == OP_TYPE::IN_ADDIMAGE ||
+				type == OP_TYPE::IN_ADDMESH)
+			{
+				// Find out relevant parameters. \todo: this may be optimised by reusing partial
+				// values in a LINK_CONTEXT or similar
+				SubtreeRelevantParametersVisitorAST visitor;
+				visitor.Run(value.child());
+
+				TArray<uint16> params;
+				for (const FString& paramName : visitor.m_params)
+				{
+					for (int32 i = 0; i < program.m_parameters.Num(); ++i)
+					{
+						const auto& param = program.m_parameters[i];
+						if (param.m_name == paramName)
+						{
+							params.Add(uint16(i));
+							break;
+						}
+					}
+				}
+
+				params.Sort();
+
+				auto it = program.m_parameterLists.Find(params);
+
+				if (it != INDEX_NONE)
+				{
+					args.relevantParametersListIndex = it;
+				}
+				else
+				{
+					args.relevantParametersListIndex = uint32_t(program.m_parameterLists.Num());
+					program.m_parameterLists.Add(params);
+				}
+			}
+
+			linkedAddress = (OP::ADDRESS)program.m_opAddress.Num();
+			program.m_opAddress.Add((uint32_t)program.m_byteCode.Num());
+			AppendCode(program.m_byteCode, type);
+			AppendCode(program.m_byteCode, args);
+		}
+
+	}
+
+}
diff --git a/Engine/Plugins/Mutable/Source/MutableTools/Private/MuT/ASTOpInstanceAdd.h b/Engine/Plugins/Mutable/Source/MutableTools/Private/MuT/ASTOpInstanceAdd.h
new file mode 100644
--- /dev/null
+++ b/Engine/Plugins/Mutable/Source/MutableTools/Private/MuT/ASTOpInstanceAdd.h
@@ -0,0 +1,51 @@
+// Copyright Epic Games, Inc. All Rights Reserved.
+
+#pragma once
+
+#include "HAL/PlatformMath.h"
+#include "MuR/Operations.h"
+#include "MuR/Ptr.h"
+#include "MuT/AST.h"
+
+
+namespace mu
+{
+struct FProgram;
+
+
+	//---------------------------------------------------------------------------------------------
+	//! Operations to add elements to an instance
+	//---------------------------------------------------------------------------------------------
+	class ASTOpInstanceAdd final : public ASTOp
+	{
+	public:
+
+		//! Type of switch
+		OP_TYPE type = OP_TYPE::NONE;
+
+		ASTChild instance;
+		ASTChild value;
+		uint32_t id = 0;
+		uint32_t ExternalId = 0;
+		int32 SharedSurfaceId = INDEX_NONE;
+		FString name;
+
+	public:
+
+		ASTOpInstanceAdd();
+		ASTOpInstanceAdd(const ASTOpInstanceAdd&) = delete;
+		~ASTOpInstanceAdd();
+
+		OP_TYPE GetOpType() const override { return type; }
+
+		void ForEachChild(const TFunctionRef<void(ASTChild&)>) override;
+		bool IsEqual(const ASTOp& otherUntyped) const override;
+		Ptr<ASTOp> Clone(MapChildFuncRef mapChild) const override;
+		uint64 Hash() const override;
+		void Assert() override;
+		void Link(FProgram& program, FLinkerOptions* Options) override;
+	};
+
+
+}
+
