[mutable] Moved the Mutable plugin out of Experimental status into Beta.

diff --git a/Engine/Plugins/Mutable/Source/MutableTools/Private/MuT/ASTOpMeshApplyPose.cpp b/Engine/Plugins/Mutable/Source/MutableTools/Private/MuT/ASTOpMeshApplyPose.cpp
new file mode 100644
--- /dev/null
+++ b/Engine/Plugins/Mutable/Source/MutableTools/Private/MuT/ASTOpMeshApplyPose.cpp
@@ -0,0 +1,121 @@
+// Copyright Epic Games, Inc. All Rights Reserved.
+
+#include "MuT/ASTOpMeshApplyPose.h"
+#include "MuT/ASTOpMeshAddTags.h"
+
+#include "HAL/PlatformMath.h"
+#include "MuR/ModelPrivate.h"
+#include "MuR/RefCounted.h"
+#include "MuR/Types.h"
+
+namespace mu
+{
+
+
+	//-------------------------------------------------------------------------------------------------
+	ASTOpMeshApplyPose::ASTOpMeshApplyPose()
+		: base(this)
+		, pose(this)
+	{
+	}
+
+
+	ASTOpMeshApplyPose::~ASTOpMeshApplyPose()
+	{
+		// Explicit call needed to avoid recursive destruction
+		ASTOp::RemoveChildren();
+	}
+
+
+	bool ASTOpMeshApplyPose::IsEqual(const ASTOp& otherUntyped) const
+	{
+		if (otherUntyped.GetOpType()==GetOpType())
+		{
+			const ASTOpMeshApplyPose* other = static_cast<const ASTOpMeshApplyPose*>(&otherUntyped);
+			return base == other->base &&
+				pose == other->pose;
+		}
+		return false;
+	}
+
+
+	uint64 ASTOpMeshApplyPose::Hash() const
+	{
+		uint64 res = std::hash<void*>()(base.child().get());
+		hash_combine(res, pose.child().get());
+		return res;
+	}
+
+
+	mu::Ptr<ASTOp> ASTOpMeshApplyPose::Clone(MapChildFuncRef mapChild) const
+	{
+		Ptr<ASTOpMeshApplyPose> n = new ASTOpMeshApplyPose();
+		n->base = mapChild(base.child());
+		n->pose = mapChild(pose.child());
+		return n;
+	}
+
+
+	void ASTOpMeshApplyPose::ForEachChild(const TFunctionRef<void(ASTChild&)> f)
+	{
+		f(base);
+		f(pose);
+	}
+
+
+	void ASTOpMeshApplyPose::Link(FProgram& program, FLinkerOptions*)
+	{
+		// Already linked?
+		if (!linkedAddress)
+		{
+			OP::MeshApplyPoseArgs args;
+			memset(&args, 0, sizeof(args));
+
+			if (base) args.base = base->linkedAddress;
+			if (pose) args.pose = pose->linkedAddress;
+
+			linkedAddress = (OP::ADDRESS)program.m_opAddress.Num();
+			//program.m_code.push_back(op);
+			program.m_opAddress.Add((uint32_t)program.m_byteCode.Num());
+			AppendCode(program.m_byteCode, OP_TYPE::ME_APPLYPOSE);
+			AppendCode(program.m_byteCode, args);
+		}
+
+	}
+
+	Ptr<ASTOp> ASTOpMeshApplyPose::OptimiseSink(const FModelOptimizationOptions&, FOptimizeSinkContext&) const
+	{
+		Ptr<ASTOp> NewOp;
+
+		Ptr<ASTOp> MeshAt = base.child();
+
+		if (!MeshAt)
+		{
+			return nullptr;
+		}
+
+		OP_TYPE MeshType = MeshAt->GetOpType();
+
+		switch (MeshType)
+		{
+			case OP_TYPE::ME_ADDTAGS:
+			{
+				Ptr<ASTOpMeshAddTags> New = mu::Clone<ASTOpMeshAddTags>(MeshAt);
+				if (New->Source)
+				{
+					Ptr<ASTOpMeshApplyPose> NewApplyPose = mu::Clone<ASTOpMeshApplyPose>(this);
+					NewApplyPose->base = New->Source.child();
+					New->Source = NewApplyPose;
+				}
+
+				NewOp = New;
+				break;
+			}
+
+			default:
+				break;
+		}
+
+		return NewOp;
+	}
+}
diff --git a/Engine/Plugins/Mutable/Source/MutableTools/Private/MuT/ASTOpMeshApplyPose.h b/Engine/Plugins/Mutable/Source/MutableTools/Private/MuT/ASTOpMeshApplyPose.h
new file mode 100644
--- /dev/null
+++ b/Engine/Plugins/Mutable/Source/MutableTools/Private/MuT/ASTOpMeshApplyPose.h
@@ -0,0 +1,43 @@
+// Copyright Epic Games, Inc. All Rights Reserved.
+
+#pragma once
+
+#include "HAL/Platform.h"
+#include "MuR/Operations.h"
+#include "MuR/Ptr.h"
+#include "MuT/AST.h"
+
+
+namespace mu
+{
+struct FProgram;
+
+
+	//---------------------------------------------------------------------------------------------
+	//!
+	//---------------------------------------------------------------------------------------------
+	class ASTOpMeshApplyPose final : public ASTOp
+	{
+	public:
+
+		ASTChild base;
+		ASTChild pose;
+
+	public:
+
+		ASTOpMeshApplyPose();
+		ASTOpMeshApplyPose(const ASTOpMeshApplyPose&) = delete;
+		~ASTOpMeshApplyPose();
+
+		OP_TYPE GetOpType() const override { return OP_TYPE::ME_APPLYPOSE; }
+		uint64 Hash() const override;
+		bool IsEqual(const ASTOp& otherUntyped) const override;
+		Ptr<ASTOp> Clone(MapChildFuncRef mapChild) const override;
+		void ForEachChild(const TFunctionRef<void(ASTChild&)>) override;
+		void Link(FProgram& program, FLinkerOptions* Options) override;
+		Ptr<ASTOp> OptimiseSink(const FModelOptimizationOptions&, FOptimizeSinkContext&) const override;
+	};
+
+
+}
+
diff --git a/Engine/Plugins/Mutable/Source/MutableTools/Private/MuT/ASTOpMeshApplyShape.cpp b/Engine/Plugins/Mutable/Source/MutableTools/Private/MuT/ASTOpMeshApplyShape.cpp
new file mode 100644
--- /dev/null
+++ b/Engine/Plugins/Mutable/Source/MutableTools/Private/MuT/ASTOpMeshApplyShape.cpp
@@ -0,0 +1,117 @@
+// Copyright Epic Games, Inc. All Rights Reserved.
+
+#include "MuT/ASTOpMeshApplyShape.h"
+
+#include "HAL/PlatformMath.h"
+#include "MuR/ModelPrivate.h"
+#include "MuR/RefCounted.h"
+#include "MuR/Types.h"
+
+
+namespace mu
+{
+
+
+	ASTOpMeshApplyShape::ASTOpMeshApplyShape()
+		: Mesh(this)
+		, Shape(this)
+		, bRecomputeNormals(false)
+		, bReshapeSkeleton(false)
+		, bReshapePhysicsVolumes(false)
+		, bReshapeVertices(true)
+		, bApplyLaplacian(false)
+	{
+	}
+
+
+	ASTOpMeshApplyShape::~ASTOpMeshApplyShape()
+	{
+		// Explicit call needed to avoid recursive destruction
+		ASTOp::RemoveChildren();
+	}
+
+
+	bool ASTOpMeshApplyShape::IsEqual(const ASTOp& OtherUntyped) const
+	{
+		if (OtherUntyped.GetOpType()==GetOpType())
+		{
+			const ASTOpMeshApplyShape* Other = static_cast<const ASTOpMeshApplyShape*>(&OtherUntyped);
+
+			const bool bSameFlags =
+				bRecomputeNormals == Other->bRecomputeNormals &&
+				bReshapePhysicsVolumes == Other->bReshapePhysicsVolumes &&
+				bReshapeSkeleton == Other->bReshapeSkeleton &&
+				bReshapeVertices == Other->bReshapeVertices &&
+				bApplyLaplacian == Other->bApplyLaplacian;
+
+			return Mesh == Other->Mesh && Shape == Other->Shape && bSameFlags;
+		}
+		return false;
+	}
+
+
+	uint64 ASTOpMeshApplyShape::Hash() const
+	{
+		uint64 Result = std::hash<void*>()(Mesh.child().get());
+		hash_combine(Result, Shape.child().get());
+		hash_combine(Result, bool(bRecomputeNormals));
+		hash_combine(Result, bool(bReshapeSkeleton));
+		hash_combine(Result, bool(bReshapePhysicsVolumes));
+		hash_combine(Result, bool(bReshapeVertices));
+		hash_combine(Result, bool(bApplyLaplacian));
+
+		return Result;
+	}
+
+
+	mu::Ptr<ASTOp> ASTOpMeshApplyShape::Clone(MapChildFuncRef MapChild) const
+	{
+		Ptr<ASTOpMeshApplyShape> NewOp = new ASTOpMeshApplyShape();
+		NewOp->Mesh = MapChild(Mesh.child());
+		NewOp->Shape = MapChild(Shape.child());
+		NewOp->bRecomputeNormals = bRecomputeNormals;
+		NewOp->bReshapeSkeleton = bReshapeSkeleton;
+		NewOp->bReshapePhysicsVolumes = bReshapePhysicsVolumes;
+		NewOp->bReshapeVertices = bReshapeVertices;
+		NewOp->bApplyLaplacian = bApplyLaplacian;
+		return NewOp;
+	}
+
+
+	void ASTOpMeshApplyShape::ForEachChild(const TFunctionRef<void(ASTChild&)> Func)
+	{
+		Func(Mesh);
+		Func(Shape);
+	}
+
+
+	void ASTOpMeshApplyShape::Link(FProgram& program, FLinkerOptions*)
+	{
+		// Already linked?
+		if (!linkedAddress)
+		{
+			OP::MeshApplyShapeArgs Args;
+			FMemory::Memzero(&Args, sizeof(Args));
+
+			constexpr EMeshBindShapeFlags NoFlags = EMeshBindShapeFlags::None;
+			EMeshBindShapeFlags BindFlags = NoFlags;
+			EnumAddFlags(BindFlags, bRecomputeNormals ? EMeshBindShapeFlags::RecomputeNormals : NoFlags);
+			EnumAddFlags(BindFlags, bReshapeSkeleton ? EMeshBindShapeFlags::ReshapeSkeleton : NoFlags);
+			EnumAddFlags(BindFlags, bReshapePhysicsVolumes ? EMeshBindShapeFlags::ReshapePhysicsVolumes : NoFlags);
+			EnumAddFlags(BindFlags, bReshapeVertices ? EMeshBindShapeFlags::ReshapeVertices : NoFlags);
+			EnumAddFlags(BindFlags, bApplyLaplacian ? EMeshBindShapeFlags::ApplyLaplacian : NoFlags);
+
+			Args.flags = static_cast<uint32>(BindFlags);
+
+			Args.mesh = Mesh ? Mesh->linkedAddress : 0;
+			Args.shape = Shape ? Shape->linkedAddress : 0;
+
+			linkedAddress = (OP::ADDRESS)program.m_opAddress.Num();
+			program.m_opAddress.Add((uint32_t)program.m_byteCode.Num());
+			AppendCode(program.m_byteCode, OP_TYPE::ME_APPLYSHAPE);
+			AppendCode(program.m_byteCode, Args);
+		}
+
+	}
+
+}
diff --git a/Engine/Plugins/Mutable/Source/MutableTools/Private/MuT/ASTOpMeshBindShape.cpp b/Engine/Plugins/Mutable/Source/MutableTools/Private/MuT/ASTOpMeshBindShape.cpp
new file mode 100644
--- /dev/null
+++ b/Engine/Plugins/Mutable/Source/MutableTools/Private/MuT/ASTOpMeshBindShape.cpp
@@ -0,0 +1,473 @@
+// Copyright Epic Games, Inc. All Rights Reserved.
+
+#include "MuT/ASTOpMeshBindShape.h"
+
+#include "HAL/PlatformMath.h"
+#include "HAL/UnrealMemory.h"
+#include "MuR/ModelPrivate.h"
+#include "MuR/RefCounted.h"
+#include "MuR/Types.h"
+#include "MuT/ASTOpConditional.h"
+#include "MuT/ASTOpSwitch.h"
+#include "MuT/ASTOpMeshRemoveMask.h"
+#include "MuT/ASTOpMeshAddTags.h"
+
+
+namespace mu
+{
+
+
+	ASTOpMeshBindShape::ASTOpMeshBindShape()
+		: Mesh(this)
+		, Shape(this)
+		, bRecomputeNormals(false)
+		, bReshapeSkeleton(false)
+		, bReshapePhysicsVolumes(false)
+		, bReshapeVertices(false)
+		, bApplyLaplacian(false)
+	{
+	}
+
+
+	ASTOpMeshBindShape::~ASTOpMeshBindShape()
+	{
+		// Explicit call needed to avoid recursive destruction
+		ASTOp::RemoveChildren();
+	}
+
+
+	bool ASTOpMeshBindShape::IsEqual(const ASTOp& OtherUntyped) const
+	{
+		if (OtherUntyped.GetOpType()==GetOpType())
+		{
+			const ASTOpMeshBindShape* Other = static_cast<const ASTOpMeshBindShape*>(&OtherUntyped);
+
+			const bool bSameFlags =
+				bRecomputeNormals == Other->bRecomputeNormals &&
+				bReshapeSkeleton == Other->bReshapeSkeleton	&&
+				bReshapePhysicsVolumes == Other->bReshapePhysicsVolumes &&
+				bReshapeVertices == Other->bReshapeVertices &&
+				bApplyLaplacian == Other->bApplyLaplacian;
+
+			return bSameFlags &&
+				Mesh == Other->Mesh &&
+				Shape == Other->Shape &&
+				BonesToDeform == Other->BonesToDeform &&
+				PhysicsToDeform == Other->PhysicsToDeform &&
+				BindingMethod == Other->BindingMethod &&
+				RChannelUsage == Other->RChannelUsage && 
+				GChannelUsage == Other->GChannelUsage && 
+				BChannelUsage == Other->BChannelUsage && 
+				AChannelUsage == Other->AChannelUsage;
+		}
+
+		return false;
+	}
+
+
+	uint64 ASTOpMeshBindShape::Hash() const
+	{
+		uint64 Result = std::hash<void*>()(Mesh.child().get());
+		hash_combine(Result, Shape.child().get());
+		hash_combine(Result, bool(bRecomputeNormals));
+		hash_combine(Result, bool(bReshapeSkeleton));
+		hash_combine(Result, bool(bReshapePhysicsVolumes));
+		hash_combine(Result, bool(bReshapeVertices));
+		hash_combine(Result, bool(bApplyLaplacian));
+		hash_combine(Result, bool(BindingMethod));
+
+		hash_combine(Result, static_cast<uint32>(RChannelUsage));
+		hash_combine(Result, static_cast<uint32>(GChannelUsage));
+		hash_combine(Result, static_cast<uint32>(BChannelUsage));
+		hash_combine(Result, static_cast<uint32>(AChannelUsage));
+
+		for (const FBoneName& S : BonesToDeform)
+		{
+			hash_combine(Result, S.Id);
+		}
+
+		for (const FBoneName& S : PhysicsToDeform)
+		{
+			hash_combine(Result, S.Id);
+		}
+
+		return Result;
+	}
+
+
+	mu::Ptr<ASTOp> ASTOpMeshBindShape::Clone(MapChildFuncRef mapChild) const
+	{
+		Ptr<ASTOpMeshBindShape> NewOp = new ASTOpMeshBindShape();
+		NewOp->Mesh = mapChild(Mesh.child());
+		NewOp->Shape = mapChild(Shape.child());
+		NewOp->bRecomputeNormals = bRecomputeNormals;
+		NewOp->bReshapeSkeleton	= bReshapeSkeleton;
+		NewOp->bReshapePhysicsVolumes = bReshapePhysicsVolumes;
+		NewOp->bReshapeVertices = bReshapeVertices;
+		NewOp->bApplyLaplacian = bApplyLaplacian;
+		NewOp->BonesToDeform = BonesToDeform;
+		NewOp->PhysicsToDeform = PhysicsToDeform;
+		NewOp->BindingMethod = BindingMethod;
+
+		NewOp->RChannelUsage = RChannelUsage;
+		NewOp->GChannelUsage = GChannelUsage;
+		NewOp->BChannelUsage = BChannelUsage;
+		NewOp->AChannelUsage = AChannelUsage;
+
+		return NewOp;
+	}
+
+
+	void ASTOpMeshBindShape::ForEachChild(const TFunctionRef<void(ASTChild&)> Func)
+	{
+		Func(Mesh);
+		Func(Shape);
+	}
+
+
+	void ASTOpMeshBindShape::Link(FProgram& program, FLinkerOptions*)
+	{
+		// Already linked?
+		if (!linkedAddress)
+		{
+			OP::MeshBindShapeArgs Args;
+			FMemory::Memzero(&Args, sizeof(Args));
+
+			constexpr EMeshBindShapeFlags NoFlags = EMeshBindShapeFlags::None;
+			EMeshBindShapeFlags BindFlags = NoFlags;
+			EnumAddFlags(BindFlags, bRecomputeNormals ? EMeshBindShapeFlags::RecomputeNormals : NoFlags);
+			EnumAddFlags(BindFlags, bReshapeSkeleton ? EMeshBindShapeFlags::ReshapeSkeleton : NoFlags);
+			EnumAddFlags(BindFlags, bReshapePhysicsVolumes ? EMeshBindShapeFlags::ReshapePhysicsVolumes : NoFlags);
+			EnumAddFlags(BindFlags, bReshapeVertices ? EMeshBindShapeFlags::ReshapeVertices : NoFlags);
+			EnumAddFlags(BindFlags, bApplyLaplacian ? EMeshBindShapeFlags::ApplyLaplacian : NoFlags);
+
+			{
+				auto ConvertColorUsage = [](EVertexColorUsage Usage)
+				{
+					switch (Usage)
+					{
+					case EVertexColorUsage::None:			   return EMeshBindColorChannelUsage::None;
+					case EVertexColorUsage::ReshapeClusterId:  return EMeshBindColorChannelUsage::ClusterId;
+					case EVertexColorUsage::ReshapeMaskWeight: return EMeshBindColorChannelUsage::MaskWeight;
+					default: check(false); return EMeshBindColorChannelUsage::None;
+					};
+				};
+	
+				const FMeshBindColorChannelUsages ColorUsages = {
+					ConvertColorUsage(RChannelUsage),
+					ConvertColorUsage(GChannelUsage),
+					ConvertColorUsage(BChannelUsage),
+					ConvertColorUsage(AChannelUsage) };
+
+				FMemory::Memcpy(&Args.ColorUsage, &ColorUsages, sizeof(Args.ColorUsage));
+				static_assert(sizeof(Args.ColorUsage) == sizeof(ColorUsages));
+			}
+
+			Args.flags = static_cast<uint32>(BindFlags);
+
+			Args.bindingMethod = BindingMethod;
+
+			if (Mesh)
+			{
+				Args.mesh = Mesh->linkedAddress;
+			}
+
+			if (Shape)
+			{
+				Args.shape = Shape->linkedAddress;
+			}
+
+			linkedAddress = (OP::ADDRESS)program.m_opAddress.Num();
+			program.m_opAddress.Add((uint32_t)program.m_byteCode.Num());
+			AppendCode(program.m_byteCode, OP_TYPE::ME_BINDSHAPE);
+			AppendCode(program.m_byteCode, Args);
+
+			AppendCode(program.m_byteCode, (int32)BonesToDeform.Num());
+			for (const FBoneName& S : BonesToDeform)
+			{
+				AppendCode(program.m_byteCode, S);
+			}
+
+			AppendCode(program.m_byteCode, (int32)PhysicsToDeform.Num());
+			for (const FBoneName& S : PhysicsToDeform)
+			{
+				AppendCode(program.m_byteCode, S);
+			}
+		}
+	}
+
+
+	//-------------------------------------------------------------------------------------------------
+	mu::Ptr<ASTOp> ASTOpMeshBindShape::OptimiseSink(const FModelOptimizationOptions&, FOptimizeSinkContext&) const
+	{
+		Ptr<ASTOp> NewOp;
+
+		Ptr<ASTOp> MeshAt = Mesh.child();
+		if (!MeshAt)
+		{
+			return nullptr;
+		}
+
+		Ptr<ASTOp> ShapeAt = Shape.child();
+		if (!ShapeAt)
+		{
+			return nullptr;
+		}
+
+		OP_TYPE MeshType = MeshAt->GetOpType();
+		OP_TYPE ShapeType = ShapeAt->GetOpType();
+
+		// See if both mesh and shape have an operation that can be optimized in a combined way
+		if (MeshType == ShapeType)
+		{
+			switch (MeshType)
+			{
+
+			case OP_TYPE::ME_SWITCH:
+			{
+				// If the switch variable and structure is the same
+				const ASTOpSwitch* MeshSwitch = static_cast<const ASTOpSwitch*>(MeshAt.get());
+				const ASTOpSwitch* ShapeSwitch = static_cast<const ASTOpSwitch*>(ShapeAt.get());
+				bool bIsSimilarSwitch = MeshSwitch->IsCompatibleWith(ShapeSwitch);
+				if (!bIsSimilarSwitch)
+				{
+					break;
+				}
+
+				// Move the operation down all the paths
+				Ptr<ASTOpSwitch> NewSwitch = mu::Clone<ASTOpSwitch>(MeshAt);
+
+				if (NewSwitch->def)
+				{
+					Ptr<ASTOpMeshBindShape> NewBind = mu::Clone<ASTOpMeshBindShape>(this);
+					NewBind->Mesh = MeshSwitch->def.child();
+					NewBind->Shape = ShapeSwitch->def.child();
+					NewSwitch->def = NewBind;
+				}
+
+				for (int32 v = 0; v < NewSwitch->cases.Num(); ++v)
+				{
+					if (NewSwitch->cases[v].branch)
+					{
+						Ptr<ASTOpMeshBindShape> NewBind = mu::Clone<ASTOpMeshBindShape>(this);
+						NewBind->Mesh = MeshSwitch->cases[v].branch.child();
+						NewBind->Shape = ShapeSwitch->FindBranch(MeshSwitch->cases[v].condition);
+						NewSwitch->cases[v].branch = NewBind;
+					}
+				}
+
+				NewOp = NewSwitch;
+				break;
+			}
+
+
+			case OP_TYPE::ME_CONDITIONAL:
+			{
+				const ASTOpConditional* MeshConditional = static_cast<const ASTOpConditional*>(MeshAt.get());
+				const ASTOpConditional* ShapeConditional = static_cast<const ASTOpConditional*>(ShapeAt.get());
+				bool bIsSimilar = MeshConditional->condition == ShapeConditional->condition;
+				if (!bIsSimilar)
+				{
+					break;
+				}
+
+				Ptr<ASTOpConditional> NewConditional = mu::Clone<ASTOpConditional>(MeshAt);
+
+				if (NewConditional->yes)
+				{
+					Ptr<ASTOpMeshBindShape> NewBind = mu::Clone<ASTOpMeshBindShape>(this);
+					NewBind->Mesh = MeshConditional->yes.child();
+					NewBind->Shape = ShapeConditional->yes.child();
+					NewConditional->yes = NewBind;
+				}
+
+				if (NewConditional->no)
+				{
+					Ptr<ASTOpMeshBindShape> NewBind = mu::Clone<ASTOpMeshBindShape>(this);
+					NewBind->Mesh = MeshConditional->no.child();
+					NewBind->Shape = ShapeConditional->no.child();
+					NewConditional->no = NewBind;
+				}
+
+				NewOp = NewConditional;
+				break;
+			}
+
+
+			default:
+				break;
+
+			}
+		}
+
+
+		// If not already optimized
+		if (!NewOp)
+		{
+			// Optimize only the mesh parameter
+			switch (MeshType)
+			{
+
+			case OP_TYPE::ME_SWITCH:
+			{
+				// Move the operation down all the paths
+				Ptr<ASTOpSwitch> NewSwitch = mu::Clone<ASTOpSwitch>(MeshAt);
+
+				if (NewSwitch->def)
+				{
+					Ptr<ASTOpMeshBindShape> NewBind = mu::Clone<ASTOpMeshBindShape>(this);
+					NewBind->Mesh = NewSwitch->def.child();
+					NewSwitch->def = NewBind;
+				}
+
+				for (int32 v = 0; v < NewSwitch->cases.Num(); ++v)
+				{
+					if (NewSwitch->cases[v].branch)
+					{
+						Ptr<ASTOpMeshBindShape> NewBind = mu::Clone<ASTOpMeshBindShape>(this);
+						NewBind->Mesh = NewSwitch->cases[v].branch.child();
+						NewSwitch->cases[v].branch = NewBind;
+					}
+				}
+
+				NewOp = NewSwitch;
+				break;
+			}
+
+			case OP_TYPE::ME_CONDITIONAL:
+			{
+				// Move the operation down all the paths
+				Ptr<ASTOpConditional> NewConditional = mu::Clone<ASTOpConditional>(MeshAt);
+
+				if (NewConditional->yes)
+				{
+					Ptr<ASTOpMeshBindShape> NewBind = mu::Clone<ASTOpMeshBindShape>(this);
+					NewBind->Mesh = NewConditional->yes.child();
+					NewConditional->yes = NewBind;
+				}
+
+				if (NewConditional->no)
+				{
+					Ptr<ASTOpMeshBindShape> NewBind = mu::Clone<ASTOpMeshBindShape>(this);
+					NewBind->Mesh = NewConditional->no.child();
+					NewConditional->no = NewBind;
+				}
+
+				NewOp = NewConditional;
+				break;
+			}
+
+			case OP_TYPE::ME_REMOVEMASK:
+			{
+				// We bind something that could have a part removed: we can reorder to bind the entire mesh
+				// and apply remove later at runtime.
+
+				Ptr<ASTOpMeshRemoveMask> NewRemove = mu::Clone<ASTOpMeshRemoveMask>(MeshAt);
+				if (NewRemove->source)
+				{
+					Ptr<ASTOpMeshBindShape> NewBind = mu::Clone<ASTOpMeshBindShape>(this);
+					NewBind->Mesh = NewRemove->source.child();
+					NewRemove->source = NewBind;
+				}
+
+				NewOp = NewRemove;
+				break;
+			}
+
+			case OP_TYPE::ME_ADDTAGS:
+			{
+				Ptr<ASTOpMeshAddTags> New = mu::Clone<ASTOpMeshAddTags>(MeshAt);
+				if (New->Source)
+				{
+					Ptr<ASTOpMeshBindShape> NewBind = mu::Clone<ASTOpMeshBindShape>(this);
+					NewBind->Mesh = New->Source.child();
+					New->Source = NewBind;
+				}
+
+				NewOp = New;
+				break;
+			}
+
+			default:
+				break;
+
+			}
+		}
+
+		// If not already optimized
+		if (!NewOp)
+		{
+			// Optimize only the shape parameter
+			switch (ShapeType)
+			{
+
+			case OP_TYPE::ME_SWITCH:
+			{
+				// Move the operation down all the paths
+				Ptr<ASTOpSwitch> NewSwitch = mu::Clone<ASTOpSwitch>(ShapeAt);
+
+				if (NewSwitch->def)
+				{
+					Ptr<ASTOpMeshBindShape> NewBind = mu::Clone<ASTOpMeshBindShape>(this);
+					NewBind->Shape = NewSwitch->def.child();
+					NewSwitch->def = NewBind;
+				}
+
+				for (int32 v = 0; v < NewSwitch->cases.Num(); ++v)
+				{
+					if (NewSwitch->cases[v].branch)
+					{
+						Ptr<ASTOpMeshBindShape> NewBind = mu::Clone<ASTOpMeshBindShape>(this);
+						NewBind->Shape = NewSwitch->cases[v].branch.child();
+						NewSwitch->cases[v].branch = NewBind;
+					}
+				}
+
+				NewOp = NewSwitch;
+				break;
+			}
+
+			case OP_TYPE::ME_CONDITIONAL:
+			{
+				// Move the operation down all the paths
+				Ptr<ASTOpConditional> NewConditional = mu::Clone<ASTOpConditional>(ShapeAt);
+
+				if (NewConditional->yes)
+				{
+					Ptr<ASTOpMeshBindShape> NewBind = mu::Clone<ASTOpMeshBindShape>(this);
+					NewBind->Shape = NewConditional->yes.child();
+					NewConditional->yes = NewBind;
+				}
+
+				if (NewConditional->no)
+				{
+					Ptr<ASTOpMeshBindShape> NewBind = mu::Clone<ASTOpMeshBindShape>(this);
+					NewBind->Shape = NewConditional->no.child();
+					NewConditional->no = NewBind;
+				}
+
+				NewOp = NewConditional;
+				break;
+			}
+
+			case OP_TYPE::ME_ADDTAGS:
+			{
+				// Ignore the tags in the shape
+				Ptr<ASTOpMeshBindShape> NewBind = mu::Clone<ASTOpMeshBindShape>(this);
+				const ASTOpMeshAddTags* New = static_cast<const ASTOpMeshAddTags*>(ShapeAt.get());
+				NewBind->Shape = New->Source.child();
+				NewOp = NewBind;
+				break;
+			}
+
+			default:
+				break;
+
+			}
+		}
+
+
+		return NewOp;
+	}
+
+}
diff --git a/Engine/Plugins/Mutable/Source/MutableTools/Private/MuT/ASTOpMeshBindShape.h b/Engine/Plugins/Mutable/Source/MutableTools/Private/MuT/ASTOpMeshBindShape.h
new file mode 100644
--- /dev/null
+++ b/Engine/Plugins/Mutable/Source/MutableTools/Private/MuT/ASTOpMeshBindShape.h
@@ -0,0 +1,59 @@
+// Copyright Epic Games, Inc. All Rights Reserved.
+
+#pragma once
+
+#include "Containers/Array.h"
+#include "HAL/Platform.h"
+#include "MuR/Operations.h"
+#include "MuR/Ptr.h"
+#include "MuR/Skeleton.h"
+#include "MuT/AST.h"
+
+
+namespace mu
+{
+	struct FProgram;
+
+	//---------------------------------------------------------------------------------------------
+	//!
+	//---------------------------------------------------------------------------------------------
+	class ASTOpMeshBindShape final : public ASTOp
+	{
+	public:
+
+		ASTChild Mesh;
+		ASTChild Shape;
+
+		TArray<FBoneName> BonesToDeform;
+		TArray<FBoneName> PhysicsToDeform;
+
+		uint32 BindingMethod = 0;
+		
+		uint32 bRecomputeNormals	  : 1;
+		uint32 bReshapeSkeleton	      : 1;
+		uint32 bReshapePhysicsVolumes : 1;
+		uint32 bReshapeVertices       : 1;
+		uint32 bApplyLaplacian        : 1;
+
+		EVertexColorUsage RChannelUsage = EVertexColorUsage::None;
+		EVertexColorUsage GChannelUsage = EVertexColorUsage::None;
+		EVertexColorUsage BChannelUsage = EVertexColorUsage::None;
+		EVertexColorUsage AChannelUsage = EVertexColorUsage::None;
+	public:
+
+		ASTOpMeshBindShape();
+		ASTOpMeshBindShape(const ASTOpMeshBindShape&) = delete;
+		~ASTOpMeshBindShape();
+
+		OP_TYPE GetOpType() const override { return OP_TYPE::ME_BINDSHAPE; }
+		uint64 Hash() const override;
+		bool IsEqual(const ASTOp& otherUntyped) const override;
+		Ptr<ASTOp> Clone(MapChildFuncRef mapChild) const override;
+		void ForEachChild(const TFunctionRef<void(ASTChild&)>) override;
+		void Link(FProgram& program, FLinkerOptions* Options) override;
+		Ptr<ASTOp> OptimiseSink(const FModelOptimizationOptions&, FOptimizeSinkContext&) const override;
+
+	};
+
+}
+
diff --git a/Engine/Plugins/Mutable/Source/MutableTools/Private/MuT/ASTOpMeshClipDeform.cpp b/Engine/Plugins/Mutable/Source/MutableTools/Private/MuT/ASTOpMeshClipDeform.cpp
new file mode 100644
--- /dev/null
+++ b/Engine/Plugins/Mutable/Source/MutableTools/Private/MuT/ASTOpMeshClipDeform.cpp
@@ -0,0 +1,91 @@
+// Copyright Epic Games, Inc. All Rights Reserved.
+
+#include "MuT/ASTOpMeshClipDeform.h"
+
+#include "HAL/PlatformMath.h"
+#include "MuR/ModelPrivate.h"
+#include "MuR/RefCounted.h"
+#include "MuR/Types.h"
+
+
+namespace mu
+{
+
+
+	ASTOpMeshClipDeform::ASTOpMeshClipDeform()
+		: Mesh(this)
+		, ClipShape(this)
+	{
+	}
+
+
+	ASTOpMeshClipDeform::~ASTOpMeshClipDeform()
+	{
+		// Explicit call needed to avoid recursive destruction
+		ASTOp::RemoveChildren();
+	}
+
+
+	bool ASTOpMeshClipDeform::IsEqual(const ASTOp& OtherUntyped) const
+	{
+		if (OtherUntyped.GetOpType()==GetOpType())
+		{
+			const ASTOpMeshClipDeform* Other = static_cast<const ASTOpMeshClipDeform*>(&OtherUntyped);
+			return Mesh == Other->Mesh && ClipShape == Other->ClipShape;
+		}
+
+		return false;
+	}
+
+
+	uint64 ASTOpMeshClipDeform::Hash() const
+	{
+		uint64 res = std::hash<void*>()(Mesh.child().get());
+		hash_combine(res, ClipShape.child().get());
+		return res;
+	}
+
+
+	mu::Ptr<ASTOp> ASTOpMeshClipDeform::Clone(MapChildFuncRef mapChild) const
+	{
+		Ptr<ASTOpMeshClipDeform> n = new ASTOpMeshClipDeform();
+		n->Mesh = mapChild(Mesh.child());
+		n->ClipShape = mapChild(ClipShape.child());
+		return n;
+	}
+
+
+	void ASTOpMeshClipDeform::ForEachChild(const TFunctionRef<void(ASTChild&)> f)
+	{
+		f(Mesh);
+		f(ClipShape);
+	}
+
+
+	void ASTOpMeshClipDeform::Link(FProgram& program, FLinkerOptions*)
+	{
+		// Already linked?
+		if (!linkedAddress)
+		{
+			OP::MeshClipDeformArgs args;
+			memset(&args, 0, sizeof(args));
+
+			if (Mesh)
+			{
+				args.mesh = Mesh->linkedAddress;
+			}
+
+			if (ClipShape)
+			{
+				args.clipShape = ClipShape->linkedAddress;
+			}
+
+			linkedAddress = (OP::ADDRESS)program.m_opAddress.Num();
+			program.m_opAddress.Add((uint32_t)program.m_byteCode.Num());
+			AppendCode(program.m_byteCode, OP_TYPE::ME_CLIPDEFORM);
+			AppendCode(program.m_byteCode, args);
+		}
+
+	}
+
+}
diff --git a/Engine/Plugins/Mutable/Source/MutableTools/Private/MuT/ASTOpMeshClipMorphPlane.cpp b/Engine/Plugins/Mutable/Source/MutableTools/Private/MuT/ASTOpMeshClipMorphPlane.cpp
new file mode 100644
--- /dev/null
+++ b/Engine/Plugins/Mutable/Source/MutableTools/Private/MuT/ASTOpMeshClipMorphPlane.cpp
@@ -0,0 +1,154 @@
+// Copyright Epic Games, Inc. All Rights Reserved.
+
+#include "MuT/ASTOpMeshClipMorphPlane.h"
+
+#include "MuT/ASTOpMeshAddTags.h"
+#include "MuR/ModelPrivate.h"
+#include "MuR/RefCounted.h"
+#include "MuR/Types.h"
+#include "HAL/PlatformMath.h"
+
+
+namespace mu
+{
+
+
+	//-------------------------------------------------------------------------------------------------
+	//-------------------------------------------------------------------------------------------------
+	//-------------------------------------------------------------------------------------------------
+	ASTOpMeshClipMorphPlane::ASTOpMeshClipMorphPlane()
+		: source(this)
+	{
+	}
+
+
+	ASTOpMeshClipMorphPlane::~ASTOpMeshClipMorphPlane()
+	{
+		// Explicit call needed to avoid recursive destruction
+		ASTOp::RemoveChildren();
+	}
+
+
+	bool ASTOpMeshClipMorphPlane::IsEqual(const ASTOp& otherUntyped) const
+	{
+		if (otherUntyped.GetOpType()==GetOpType())
+		{
+			const ASTOpMeshClipMorphPlane* other = static_cast<const ASTOpMeshClipMorphPlane*>(&otherUntyped);
+
+			return source == other->source &&
+				morphShape == other->morphShape &&
+				selectionShape == other->selectionShape &&
+				vertexSelectionBone == other->vertexSelectionBone &&
+				vertexSelectionType == other->vertexSelectionType &&
+				dist == other->dist &&
+				factor == other->factor;
+		}
+		return false;
+	}
+
+
+	uint64 ASTOpMeshClipMorphPlane::Hash() const
+	{
+		uint64 res = std::hash<void*>()(source.child().get());
+		hash_combine(res, factor);
+		hash_combine(res, dist);
+		return res;
+	}
+
+
+	mu::Ptr<ASTOp> ASTOpMeshClipMorphPlane::Clone(MapChildFuncRef mapChild) const
+	{
+		Ptr<ASTOpMeshClipMorphPlane> n = new ASTOpMeshClipMorphPlane();
+		n->source = mapChild(source.child());
+		n->morphShape = morphShape;
+		n->selectionShape = selectionShape;
+		n->vertexSelectionBone = vertexSelectionBone;
+		n->vertexSelectionType = vertexSelectionType;
+		n->dist = dist;
+		n->factor = factor;
+		return n;
+	}
+
+
+	void ASTOpMeshClipMorphPlane::ForEachChild(const TFunctionRef<void(ASTChild&)> f)
+	{
+		f(source);
+	}
+
+
+	void ASTOpMeshClipMorphPlane::Link(FProgram& program, FLinkerOptions*)
+	{
+		// Already linked?
+		if (!linkedAddress)
+		{
+			OP::MeshClipMorphPlaneArgs args;
+			memset(&args, 0, sizeof(args));
+
+			if (source) args.source = source->linkedAddress;
+
+			args.morphShape = (OP::ADDRESS)program.m_constantShapes.Num();
+			program.m_constantShapes.Add(morphShape);
+
+			args.vertexSelectionType = (uint8_t)vertexSelectionType;
+			if (vertexSelectionType == OP::MeshClipMorphPlaneArgs::VS_BONE_HIERARCHY)
+			{
+				check(vertexSelectionBone.Id < MAX_uint32);
+				args.vertexSelectionShapeOrBone = vertexSelectionBone.Id;
+			}
+			else if (vertexSelectionType == OP::MeshClipMorphPlaneArgs::VS_SHAPE)
+			{
+				args.vertexSelectionShapeOrBone = program.AddConstant(selectionShape);
+			}
+
+			args.dist = dist;
+			args.factor = factor;
+			args.maxBoneRadius = vertexSelectionBoneMaxRadius;
+
+			linkedAddress = (OP::ADDRESS)program.m_opAddress.Num();
+			//program.m_code.push_back(op);
+			program.m_opAddress.Add((uint32_t)program.m_byteCode.Num());
+			AppendCode(program.m_byteCode, OP_TYPE::ME_CLIPMORPHPLANE);
+			AppendCode(program.m_byteCode, args);
+		}
+
+	}
+
+	mu::Ptr<ASTOp> ASTOpMeshClipMorphPlane::OptimiseSink(const FModelOptimizationOptions&, FOptimizeSinkContext&) const
+	{
+		Ptr<ASTOp> NewOp;
+
+		if (!source.child())
+		{
+			return nullptr;
+		}
+
+		OP_TYPE SourceType = source.child()->GetOpType();
+
+		// Optimize only the mesh parameter
+		switch (SourceType)
+		{
+
+		case OP_TYPE::ME_ADDTAGS:
+		{
+			Ptr<ASTOpMeshAddTags> NewAddTags = mu::Clone<ASTOpMeshAddTags>(source.child());
+
+			if (NewAddTags->Source)
+			{
+				Ptr<ASTOpMeshClipMorphPlane> New = mu::Clone<ASTOpMeshClipMorphPlane>(this);
+				New->source = NewAddTags->Source.child();
+				NewAddTags->Source = New;
+			}
+
+			NewOp = NewAddTags;
+			break;
+		}
+
+		default:
+			break;
+
+		}
+
+		return NewOp;
+	}
+
+}
diff --git a/Engine/Plugins/Mutable/Source/MutableTools/Private/MuT/ASTOpMeshClipMorphPlane.h b/Engine/Plugins/Mutable/Source/MutableTools/Private/MuT/ASTOpMeshClipMorphPlane.h
new file mode 100644
--- /dev/null
+++ b/Engine/Plugins/Mutable/Source/MutableTools/Private/MuT/ASTOpMeshClipMorphPlane.h
@@ -0,0 +1,50 @@
+// Copyright Epic Games, Inc. All Rights Reserved.
+
+#pragma once
+
+#include "HAL/Platform.h"
+#include "MuR/Operations.h"
+#include "MuR/ParametersPrivate.h"
+#include "MuR/Ptr.h"
+#include "MuT/AST.h"
+
+
+namespace mu
+{
+struct FProgram;
+
+	//---------------------------------------------------------------------------------------------
+	//!
+	//---------------------------------------------------------------------------------------------
+	class ASTOpMeshClipMorphPlane final : public ASTOp
+	{
+	public:
+
+		ASTChild source;
+
+		FShape morphShape;
+		FShape selectionShape;
+		FBoneName vertexSelectionBone;
+
+		OP::MeshClipMorphPlaneArgs::VERTEX_SELECTION_TYPE vertexSelectionType = OP::MeshClipMorphPlaneArgs::VS_NONE;
+
+		float dist = 0.f, factor = 0.f, vertexSelectionBoneMaxRadius = -1.f;
+
+	public:
+
+		ASTOpMeshClipMorphPlane();
+		ASTOpMeshClipMorphPlane(const ASTOpMeshClipMorphPlane&) = delete;
+		virtual ~ASTOpMeshClipMorphPlane();
+
+		virtual OP_TYPE GetOpType() const override { return OP_TYPE::ME_CLIPMORPHPLANE; }
+		virtual uint64 Hash() const override;
+		virtual bool IsEqual(const ASTOp& otherUntyped) const override;
+		virtual Ptr<ASTOp> Clone(MapChildFuncRef mapChild) const override;
+		virtual void ForEachChild(const TFunctionRef<void(ASTChild&)>) override;
+		virtual void Link(FProgram& program, FLinkerOptions* Options) override;
+		virtual mu::Ptr<ASTOp> OptimiseSink(const FModelOptimizationOptions&, FOptimizeSinkContext&) const override;
+	};
+
+
+}
+
diff --git a/Engine/Plugins/Mutable/Source/MutableTools/Private/MuT/ASTOpMeshExtractLayoutBlocks.cpp b/Engine/Plugins/Mutable/Source/MutableTools/Private/MuT/ASTOpMeshExtractLayoutBlocks.cpp
new file mode 100644
--- /dev/null
+++ b/Engine/Plugins/Mutable/Source/MutableTools/Private/MuT/ASTOpMeshExtractLayoutBlocks.cpp
@@ -0,0 +1,263 @@
+// Copyright Epic Games, Inc. All Rights Reserved.
+
+#include "MuT/ASTOpMeshExtractLayoutBlocks.h"
+
+#include "MuT/ASTOpSwitch.h"
+#include "MuT/ASTOpConditional.h"
+#include "MuT/ASTOpMeshAddTags.h"
+#include "MuT/ASTOpMeshMorph.h"
+#include "MuT/ASTOpMeshRemoveMask.h"
+#include "MuT/ASTOpMeshClipMorphPlane.h"
+#include "MuT/ASTOpMeshApplyPose.h"
+#include "MuR/ModelPrivate.h"
+#include "MuR/RefCounted.h"
+#include "MuR/Types.h"
+#include "Misc/AssertionMacros.h"
+
+namespace mu
+{
+
+	//---------------------------------------------------------------------------------------------
+	//---------------------------------------------------------------------------------------------
+	//---------------------------------------------------------------------------------------------
+	ASTOpMeshExtractLayoutBlocks::ASTOpMeshExtractLayoutBlocks()
+		: Source(this)
+	{
+	}
+
+
+	ASTOpMeshExtractLayoutBlocks::~ASTOpMeshExtractLayoutBlocks()
+	{
+		// Explicit call needed to avoid recursive destruction
+		ASTOp::RemoveChildren();
+	}
+
+
+	bool ASTOpMeshExtractLayoutBlocks::IsEqual(const ASTOp& otherUntyped) const
+	{
+		if (otherUntyped.GetOpType() == GetOpType())
+		{
+			const ASTOpMeshExtractLayoutBlocks* other = static_cast<const ASTOpMeshExtractLayoutBlocks*>(&otherUntyped);
+			return Source == other->Source && LayoutIndex == other->LayoutIndex && Blocks == other->Blocks;
+		}
+		return false;
+	}
+
+
+	mu::Ptr<ASTOp> ASTOpMeshExtractLayoutBlocks::Clone(MapChildFuncRef mapChild) const
+	{
+		Ptr<ASTOpMeshExtractLayoutBlocks> n = new ASTOpMeshExtractLayoutBlocks();
+		n->Source = mapChild(Source.child());
+		n->LayoutIndex = LayoutIndex;
+		n->Blocks = Blocks;
+		return n;
+	}
+
+
+	void ASTOpMeshExtractLayoutBlocks::Assert()
+	{
+		check(Blocks.Num() < std::numeric_limits<uint16>::max());
+		ASTOp::Assert();
+	}
+
+
+	void ASTOpMeshExtractLayoutBlocks::ForEachChild(const TFunctionRef<void(ASTChild&)> f)
+	{
+		f(Source);
+	}
+
+
+	uint64 ASTOpMeshExtractLayoutBlocks::Hash() const
+	{
+		uint64 res = std::hash<size_t>()(size_t(Source.child().get()));
+		return res;
+	}
+
+
+	void ASTOpMeshExtractLayoutBlocks::Link(FProgram& program, FLinkerOptions*)
+	{
+		// Already linked?
+		if (!linkedAddress)
+		{
+			linkedAddress = (OP::ADDRESS)program.m_opAddress.Num();
+
+			program.m_opAddress.Add((uint32)program.m_byteCode.Num());
+			AppendCode(program.m_byteCode, OP_TYPE::ME_EXTRACTLAYOUTBLOCK);
+			OP::ADDRESS sourceAt = Source ? Source->linkedAddress : 0;
+			AppendCode(program.m_byteCode, sourceAt);
+			AppendCode(program.m_byteCode, (uint16)LayoutIndex);
+			AppendCode(program.m_byteCode, (uint16)Blocks.Num());
+
+			for (uint64 Id : Blocks)
+			{
+				AppendCode(program.m_byteCode, Id);
+			}
+		}
+	}
+
+	 
+	mu::Ptr<ASTOp> ASTOpMeshExtractLayoutBlocks::OptimiseSink(const FModelOptimizationOptions&, FOptimizeSinkContext& Context) const
+	{
+		mu::Ptr<ASTOp> NewOp = Context.MeshExtractLayoutBlocksSinker.Apply(this);
+		return NewOp;
+	}
+
+
+	//---------------------------------------------------------------------------------------------
+	mu::Ptr<ASTOp> Sink_MeshExtractLayoutBlocksAST::Apply(const ASTOpMeshExtractLayoutBlocks* root)
+	{
+		m_root = root;
+
+		OldToNew.Reset();
+
+		m_initialSource = m_root->Source.child();
+		mu::Ptr<ASTOp> newSource = Visit(m_initialSource, m_root);
+
+		// If there is any change, it is the new root.
+		if (newSource != m_initialSource)
+		{
+			return newSource;
+		}
+
+		return nullptr;
+	}
+
+	mu::Ptr<ASTOp> Sink_MeshExtractLayoutBlocksAST::Visit(const mu::Ptr<ASTOp>& at, const ASTOpMeshExtractLayoutBlocks* currentSinkOp)
+	{
+		if (!at) return nullptr;
+
+		// Already visited?
+		const Ptr<ASTOp>* Cached = OldToNew.Find({ at, currentSinkOp });
+		if (Cached)
+		{
+			return *Cached;
+		}
+
+		mu::Ptr<ASTOp> newAt = at;
+		switch (at->GetOpType())
+		{
+
+		case OP_TYPE::ME_APPLYLAYOUT:
+		{
+			Ptr<ASTOpFixed> newOp = mu::Clone<ASTOpFixed>(at);
+			newOp->SetChild(newOp->op.args.MeshApplyLayout.mesh, Visit(newOp->children[newOp->op.args.MeshApplyLayout.mesh].child(), currentSinkOp));
+			newAt = newOp;
+			break;
+		}
+
+		case OP_TYPE::ME_SETSKELETON:
+		{
+			Ptr<ASTOpFixed> newOp = mu::Clone<ASTOpFixed>(at);
+			newOp->SetChild(newOp->op.args.MeshSetSkeleton.source, Visit(newOp->children[newOp->op.args.MeshSetSkeleton.source].child(), currentSinkOp));
+			newAt = newOp;
+			break;
+		}
+
+		case OP_TYPE::ME_ADDTAGS:
+		{
+			Ptr<ASTOpMeshAddTags> newOp = mu::Clone<ASTOpMeshAddTags>(at);
+			newOp->Source = Visit(newOp->Source.child(), currentSinkOp);
+			newAt = newOp;
+			break;
+		}
+
+		case OP_TYPE::ME_CLIPMORPHPLANE:
+		{
+			Ptr<ASTOpMeshClipMorphPlane> newOp = mu::Clone<ASTOpMeshClipMorphPlane>(at);
+			newOp->source = Visit(newOp->source.child(), currentSinkOp);
+			newAt = newOp;
+			break;
+		}
+
+		case OP_TYPE::ME_MORPH:
+		{
+			Ptr<ASTOpMeshMorph> NewOp = mu::Clone<ASTOpMeshMorph>(at);
+			NewOp->Base = Visit(NewOp->Base.child(), currentSinkOp);
+			NewOp->Target = Visit(NewOp->Target.child(), currentSinkOp);
+			newAt = NewOp;
+			break;
+		}
+
+		case OP_TYPE::ME_MERGE:
+		{
+			Ptr<ASTOpFixed> newOp = mu::Clone<ASTOpFixed>(at);
+			newOp->SetChild(newOp->op.args.MeshMerge.base, Visit(newOp->children[newOp->op.args.MeshMerge.base].child(), currentSinkOp));
+			newOp->SetChild(newOp->op.args.MeshMerge.added, Visit(newOp->children[newOp->op.args.MeshMerge.added].child(), currentSinkOp));
+			newAt = newOp;
+			break;
+		}
+
+		case OP_TYPE::ME_APPLYPOSE:
+		{
+			Ptr<ASTOpMeshApplyPose> NewOp = mu::Clone<ASTOpMeshApplyPose>(at);
+			NewOp->base = Visit(NewOp->base.child(), currentSinkOp);
+			newAt = NewOp;
+			break;
+		}
+
+		case OP_TYPE::ME_INTERPOLATE:
+		{
+			Ptr<ASTOpFixed> newOp = mu::Clone<ASTOpFixed>(at);
+			newOp->SetChild(newOp->op.args.MeshInterpolate.base, Visit(newOp->children[newOp->op.args.MeshInterpolate.base].child(), currentSinkOp));
+
+			for (int32 t = 0; t < MUTABLE_OP_MAX_INTERPOLATE_COUNT - 1; ++t)
+			{
+				if (newOp->children[newOp->op.args.MeshInterpolate.targets[t]])
+				{
+					newOp->SetChild(newOp->op.args.MeshInterpolate.targets[t], Visit(newOp->children[newOp->op.args.MeshInterpolate.targets[t]].child(), currentSinkOp));
+				}
+			}
+
+			newAt = newOp;
+			break;
+		}
+
+		case OP_TYPE::ME_REMOVEMASK:
+		{
+			// TODO: Make mask smaller?
+			Ptr<ASTOpMeshRemoveMask> newOp = mu::Clone<ASTOpMeshRemoveMask>(at);
+			newOp->source = Visit(newOp->source.child(), currentSinkOp);
+			newAt = newOp;
+			break;
+		}
+
+		case OP_TYPE::ME_CONDITIONAL:
+		{
+			Ptr<ASTOpConditional> newOp = mu::Clone<ASTOpConditional>(at);
+			newOp->yes = Visit(newOp->yes.child(), currentSinkOp);
+			newOp->no = Visit(newOp->no.child(), currentSinkOp);
+			newAt = newOp;
+			break;
+		}
+
+		case OP_TYPE::ME_SWITCH:
+		{
+			Ptr<ASTOpSwitch> newOp = mu::Clone<ASTOpSwitch>(at);
+			newOp->def = Visit(newOp->def.child(), currentSinkOp);
+			for (ASTOpSwitch::FCase& c : newOp->cases)
+			{
+				c.branch = Visit(c.branch.child(), currentSinkOp);
+			}
+			newAt = newOp;
+			break;
+		}
+
+		// If we reach here it means the operation type has not bee optimized.
+		default:
+			if (at != m_initialSource)
+			{
+				mu::Ptr<ASTOpMeshExtractLayoutBlocks> newOp = mu::Clone<ASTOpMeshExtractLayoutBlocks>(currentSinkOp);
+				newOp->Source = at;
+				newAt = newOp;
+			}
+			break;
+
+		}
+
+		OldToNew.Add({ at, currentSinkOp }, newAt);
+
+		return newAt;
+	}
+
+
+}
diff --git a/Engine/Plugins/Mutable/Source/MutableTools/Private/MuT/ASTOpMeshFormat.cpp b/Engine/Plugins/Mutable/Source/MutableTools/Private/MuT/ASTOpMeshFormat.cpp
new file mode 100644
--- /dev/null
+++ b/Engine/Plugins/Mutable/Source/MutableTools/Private/MuT/ASTOpMeshFormat.cpp
@@ -0,0 +1,394 @@
+// Copyright Epic Games, Inc. All Rights Reserved.
+
+#include "MuT/ASTOpMeshFormat.h"
+
+#include "HAL/UnrealMemory.h"
+#include "Misc/AssertionMacros.h"
+#include "MuR/Mesh.h"
+#include "MuR/MeshBufferSet.h"
+#include "MuR/ModelPrivate.h"
+#include "MuR/MutableTrace.h"
+#include "MuR/RefCounted.h"
+#include "MuR/System.h"
+#include "MuR/Types.h"
+#include "MuT/ASTOpConditional.h"
+#include "MuT/ASTOpConstantResource.h"
+#include "MuT/ASTOpMeshClipMorphPlane.h"
+#include "MuT/ASTOpMeshRemoveMask.h"
+#include "MuT/ASTOpMeshMorph.h"
+#include "MuT/ASTOpMeshAddTags.h"
+#include "MuT/ASTOpMeshApplyPose.h"
+#include "MuT/ASTOpSwitch.h"
+
+
+namespace mu
+{
+
+
+//-------------------------------------------------------------------------------------------------
+ASTOpMeshFormat::ASTOpMeshFormat()
+    : Source(this)
+    , Format(this)
+{
+}
+
+
+ASTOpMeshFormat::~ASTOpMeshFormat()
+{
+    // Explicit call needed to avoid recursive destruction
+    ASTOp::RemoveChildren();
+}
+
+
+bool ASTOpMeshFormat::IsEqual(const ASTOp& otherUntyped) const
+{
+	if (otherUntyped.GetOpType() == GetOpType())
+	{
+		const ASTOpMeshFormat* other = static_cast<const ASTOpMeshFormat*>(&otherUntyped);
+        return Source==other->Source && Format==other->Format && Flags ==other->Flags;
+    }
+    return false;
+}
+
+
+uint64 ASTOpMeshFormat::Hash() const
+{
+	uint64 res = std::hash<void*>()(Source.child().get() );
+    hash_combine( res, Format.child().get() );
+    return res;
+}
+
+
+mu::Ptr<ASTOp> ASTOpMeshFormat::Clone(MapChildFuncRef mapChild) const
+{
+	mu::Ptr<ASTOpMeshFormat> n = new ASTOpMeshFormat();
+    n->Source = mapChild(Source.child());
+    n->Format = mapChild(Format.child());
+	n->Flags = Flags;
+	n->bOptimizeBuffers = bOptimizeBuffers;
+    return n;
+}
+
+
+void ASTOpMeshFormat::ForEachChild(const TFunctionRef<void(ASTChild&)> f )
+{
+    f( Source );
+    f( Format );
+}
+
+
+void ASTOpMeshFormat::Link( FProgram& program, FLinkerOptions* )
+{
+    // Already linked?
+    if (!linkedAddress)
+    {
+        OP::MeshFormatArgs args;
+        FMemory::Memzero( &args, sizeof(args) );
+
+		args.Flags = Flags;
+		if (bOptimizeBuffers)
+		{
+			args.Flags = args.Flags | OP::MeshFormatArgs::OptimizeBuffers;
+		}
+
+		if (Source) args.source = Source->linkedAddress;
+		if (Format) args.format = Format->linkedAddress;
+
+        linkedAddress = (OP::ADDRESS)program.m_opAddress.Num();
+        program.m_opAddress.Add((uint32)program.m_byteCode.Num());
+        AppendCode(program.m_byteCode,OP_TYPE::ME_FORMAT);
+        AppendCode(program.m_byteCode,args);
+    }
+
+}
+
+
+mu::Ptr<ASTOp> ASTOpMeshFormat::OptimiseSink(const FModelOptimizationOptions& options, FOptimizeSinkContext& context) const
+{
+	mu::Ptr<ASTOp> at = context.MeshFormatSinker.Apply(this);
+	return at;
+}
+
+
+//---------------------------------------------------------------------------------------------
+//---------------------------------------------------------------------------------------------
+//---------------------------------------------------------------------------------------------
+mu::Ptr<ASTOp> Sink_MeshFormatAST::Apply(const ASTOpMeshFormat* root)
+{
+	m_root = root;
+
+	OldToNew.Reset();
+
+	m_initialSource = m_root->Source.child();
+	mu::Ptr<ASTOp> newSource = Visit(m_initialSource, m_root);
+
+	// If there is any change, it is the new root.
+	if (newSource != m_initialSource)
+	{
+		return newSource;
+	}
+
+	return nullptr;
+}
+
+
+namespace
+{
+	mu::Ptr<const Mesh> FindBaseMeshConstant(mu::Ptr<ASTOp> at)
+	{
+		mu::Ptr<const Mesh> res;
+
+		switch (at->GetOpType())
+		{
+		case OP_TYPE::ME_CONSTANT:
+		{
+			const ASTOpConstantResource* typed = static_cast<const ASTOpConstantResource*>(at.get());
+			res = static_cast<const Mesh*>(typed->GetValue().get());
+			break;
+		}
+
+		default:
+			check(false);
+		}
+
+		check(res);
+
+		return res;
+	}
+
+
+	// Make a mesh format suitable to morph a particular other format.
+	Ptr<Mesh> MakeMorphTargetFormat(Ptr<const Mesh> pTargetFormat)
+	{
+		MUTABLE_CPUPROFILER_SCOPE(MakeMorphTargetFormat);
+
+		// Make a morph format by adding all the vertex channels from the base into a single
+		// vertex buffer
+
+		int32 offset = 0;
+		int32 numChannels = 0;
+		TArray<EMeshBufferSemantic> semantics;
+		TArray<int32> semanticIndices;
+		TArray<EMeshBufferFormat> formats;
+		TArray<int32> components;
+		TArray<int32> offsets;
+
+		// Add the vertex channels from the new format
+		for (int32 vb = 0; vb < pTargetFormat->GetVertexBuffers().GetBufferCount(); ++vb)
+		{
+			for (int32 c = 0; c < pTargetFormat->GetVertexBuffers().GetBufferChannelCount(vb); ++c)
+			{
+				// Channel info
+				EMeshBufferSemantic semantic;
+				int semanticIndex;
+				EMeshBufferFormat format;
+				int component;
+				pTargetFormat->GetVertexBuffers().GetChannel(vb, c, &semantic, &semanticIndex, &format, &component, nullptr);
+
+				// TODO: Filter useless semantics for morphing.
+				// Maybe some formats like the ones with a packed tangent sign need to be tweaked here, to make sense of the whole buffer.
+				semantics.Add(semantic);
+				semanticIndices.Add(semanticIndex);
+				formats.Add(format);
+				components.Add(component);
+				offsets.Add(offset);
+				offset += components[numChannels] * GetMeshFormatData(formats[numChannels]).SizeInBytes;
+				numChannels++;
+			}
+		}
+
+
+		MeshPtr pTargetMorphFormat = new Mesh;
+		pTargetMorphFormat->GetVertexBuffers().SetBufferCount(1);
+
+		pTargetMorphFormat->GetVertexBuffers().SetBuffer(0, offset,
+			numChannels,
+			semantics.GetData(),
+			semanticIndices.GetData(),
+			formats.GetData(),
+			components.GetData(),
+			offsets.GetData());
+
+		return pTargetMorphFormat;
+	}
+
+}
+
+
+mu::Ptr<ASTOp> Sink_MeshFormatAST::Visit(const mu::Ptr<ASTOp>& at, const ASTOpMeshFormat* currentFormatOp)
+{
+	if (!at) return nullptr;
+
+	// Already visited?
+	const Ptr<ASTOp>* Cached = OldToNew.Find({ at,currentFormatOp });
+	if (Cached)
+	{
+		return *Cached;
+	}
+
+	mu::Ptr<ASTOp> newAt = at;
+	switch (at->GetOpType())
+	{
+
+	case OP_TYPE::ME_APPLYLAYOUT:
+	{
+		auto newOp = mu::Clone<ASTOpFixed>(at);
+		newOp->SetChild(newOp->op.args.MeshApplyLayout.mesh, Visit(newOp->children[newOp->op.args.MeshApplyLayout.mesh].child(), currentFormatOp));
+		newAt = newOp;
+		break;
+	}
+
+	case OP_TYPE::ME_SETSKELETON:
+	{
+		auto newOp = mu::Clone<ASTOpFixed>(at);
+		newOp->SetChild(newOp->op.args.MeshSetSkeleton.source, Visit(newOp->children[newOp->op.args.MeshSetSkeleton.source].child(), currentFormatOp));
+		newAt = newOp;
+		break;
+	}
+
+	case OP_TYPE::ME_ADDTAGS:
+	{
+		Ptr<ASTOpMeshAddTags> newOp = mu::Clone<ASTOpMeshAddTags>(at);
+		newOp->Source = Visit(newOp->Source.child(), currentFormatOp);
+		newAt = newOp;
+		break;
+	}
+
+	case OP_TYPE::ME_CLIPMORPHPLANE:
+	{
+		auto newOp = mu::Clone<ASTOpMeshClipMorphPlane>(at);
+		newOp->source = Visit(newOp->source.child(), currentFormatOp);
+		newAt = newOp;
+		break;
+	}
+
+	case OP_TYPE::ME_MORPH:
+	{
+		// Move the format down the base of the morph
+		Ptr<ASTOpMeshMorph> NewOp = mu::Clone<ASTOpMeshMorph>(at);
+		NewOp->Base = Visit(NewOp->Base.child(), currentFormatOp);
+
+		// Reformat the morph targets to match the new format.
+		Ptr<const Mesh> pTargetFormat = FindBaseMeshConstant(currentFormatOp->Format.child());
+		Ptr<const Mesh> pTargetMorphFormat = MakeMorphTargetFormat(pTargetFormat);
+
+		mu::Ptr<ASTOpConstantResource> NewFormatConstant = new ASTOpConstantResource();
+		NewFormatConstant->Type = OP_TYPE::ME_CONSTANT;
+		NewFormatConstant->SetValue(pTargetMorphFormat, nullptr);
+
+		if (NewOp->Target)
+		{
+			mu::Ptr<ASTOpMeshFormat> newFormat = mu::Clone<ASTOpMeshFormat>(currentFormatOp);
+			newFormat->Flags = OP::MeshFormatArgs::Vertex | OP::MeshFormatArgs::IgnoreMissing;
+			newFormat->Format = NewFormatConstant;
+
+			NewOp->Target = Visit(NewOp->Target.child(), newFormat.get());
+		}
+
+		newAt = NewOp;
+		break;
+	}
+
+	case OP_TYPE::ME_MERGE:
+	{
+		Ptr<ASTOpFixed> newOp = mu::Clone<ASTOpFixed>(at);
+		newOp->SetChild(newOp->op.args.MeshMerge.base, Visit(newOp->children[newOp->op.args.MeshMerge.base].child(), currentFormatOp));
+		newOp->SetChild(newOp->op.args.MeshMerge.added, Visit(newOp->children[newOp->op.args.MeshMerge.added].child(), currentFormatOp));
+		newAt = newOp;
+		break;
+	}
+
+	case OP_TYPE::ME_APPLYPOSE:
+	{
+		Ptr<ASTOpMeshApplyPose> NewOp = mu::Clone<ASTOpMeshApplyPose>(at);
+		NewOp->base = Visit(NewOp->base.child(), currentFormatOp);
+		newAt = NewOp;
+		break;
+	}
+
+	case OP_TYPE::ME_INTERPOLATE:
+	{
+		// Move the format down the base of the morph
+		Ptr<ASTOpFixed> newOp = mu::Clone<ASTOpFixed>(at);
+		newOp->SetChild(newOp->op.args.MeshInterpolate.base, Visit(newOp->children[newOp->op.args.MeshInterpolate.base].child(), currentFormatOp));
+
+		// Reformat the morph targets to match the new format.
+		Ptr<const Mesh> pTargetFormat = FindBaseMeshConstant(currentFormatOp->Format.child());
+		Ptr<const Mesh> pTargetMorphFormat = MakeMorphTargetFormat(pTargetFormat);
+
+		mu::Ptr<ASTOpConstantResource> TargetMorphFormatOp = new ASTOpConstantResource();
+		TargetMorphFormatOp->Type = OP_TYPE::ME_CONSTANT;
+		TargetMorphFormatOp->SetValue(pTargetMorphFormat, nullptr );
+
+		for (int32 t = 0; t < MUTABLE_OP_MAX_INTERPOLATE_COUNT - 1; ++t)
+		{
+			if (newOp->children[newOp->op.args.MeshInterpolate.targets[t]])
+			{
+				mu::Ptr<ASTOpMeshFormat> newFormat = mu::Clone<ASTOpMeshFormat>(currentFormatOp);
+				newFormat->Flags = OP::MeshFormatArgs::Vertex | OP::MeshFormatArgs::IgnoreMissing;
+				newFormat->Format = TargetMorphFormatOp;
+
+				newOp->SetChild(newOp->op.args.MeshInterpolate.targets[t], Visit(newOp->children[newOp->op.args.MeshInterpolate.targets[t]].child(), newFormat.get()));
+			}
+		}
+
+		newAt = newOp;
+		break;
+	}
+
+	case OP_TYPE::ME_REMOVEMASK:
+	{
+		auto newOp = mu::Clone<ASTOpMeshRemoveMask>(at);
+		newOp->source = Visit(newOp->source.child(), currentFormatOp);
+		newAt = newOp;
+		break;
+	}
+
+	case OP_TYPE::ME_CONDITIONAL:
+	{
+		auto newOp = mu::Clone<ASTOpConditional>(at);
+		newOp->yes = Visit(newOp->yes.child(), currentFormatOp);
+		newOp->no = Visit(newOp->no.child(), currentFormatOp);
+		newAt = newOp;
+		break;
+	}
+
+	case OP_TYPE::ME_SWITCH:
+	{
+		auto newOp = mu::Clone<ASTOpSwitch>(at);
+		newOp->def = Visit(newOp->def.child(), currentFormatOp);
+		for (auto& c : newOp->cases)
+		{
+			c.branch = Visit(c.branch.child(), currentFormatOp);
+		}
+		newAt = newOp;
+		break;
+	}
+
+	case OP_TYPE::ME_FORMAT:
+		// TODO: The child format can be removed. 
+		// Unless channels are removed and re-added, which would change their content?
+		break;
+
+
+	// This operation should not be optimized.
+	case OP_TYPE::ME_DIFFERENCE:
+
+	// If we reach here it means the operation type has not bee optimized.
+	default:
+		if (at != m_initialSource)
+		{
+			mu::Ptr<ASTOpMeshFormat> newOp = mu::Clone<ASTOpMeshFormat>(currentFormatOp);
+			newOp->Source = at;
+			newAt = newOp;
+		}
+		break;
+
+	}
+
+	OldToNew.Add({ at,currentFormatOp }, newAt);
+
+	return newAt;
+}
+
+}
diff --git a/Engine/Plugins/Mutable/Source/MutableTools/Private/MuT/ASTOpMeshGeometryOperation.cpp b/Engine/Plugins/Mutable/Source/MutableTools/Private/MuT/ASTOpMeshGeometryOperation.cpp
new file mode 100644
--- /dev/null
+++ b/Engine/Plugins/Mutable/Source/MutableTools/Private/MuT/ASTOpMeshGeometryOperation.cpp
@@ -0,0 +1,93 @@
+// Copyright Epic Games, Inc. All Rights Reserved.
+
+#include "MuT/ASTOpMeshGeometryOperation.h"
+
+#include "HAL/PlatformMath.h"
+#include "MuR/ModelPrivate.h"
+#include "MuR/RefCounted.h"
+
+
+namespace mu
+{
+
+	ASTOpMeshGeometryOperation::ASTOpMeshGeometryOperation()
+		: meshA(this)
+		, meshB(this)
+		, scalarA(this)
+		, scalarB(this)
+	{
+	}
+
+
+	ASTOpMeshGeometryOperation::~ASTOpMeshGeometryOperation()
+	{
+		// Explicit call needed to avoid recursive destruction
+		ASTOp::RemoveChildren();
+	}
+
+
+	bool ASTOpMeshGeometryOperation::IsEqual(const ASTOp& otherUntyped) const
+	{
+		if (otherUntyped.GetOpType() == GetOpType())
+		{
+			const ASTOpMeshGeometryOperation* other = static_cast<const ASTOpMeshGeometryOperation*>(&otherUntyped);
+			return meshA == other->meshA && meshB == other->meshB &&
+				scalarA == other->scalarA && scalarB == other->scalarB;
+		}
+		return false;
+	}
+
+
+	uint64 ASTOpMeshGeometryOperation::Hash() const
+	{
+		uint64 res = std::hash<void*>()(meshA.child().get());
+		hash_combine(res, meshB.child().get());
+		hash_combine(res, scalarA.child().get());
+		hash_combine(res, scalarB.child().get());
+		return res;
+	}
+
+
+	mu::Ptr<ASTOp> ASTOpMeshGeometryOperation::Clone(MapChildFuncRef mapChild) const
+	{
+		Ptr<ASTOpMeshGeometryOperation> n = new ASTOpMeshGeometryOperation();
+		n->meshA = mapChild(meshA.child());
+		n->meshB = mapChild(meshB.child());
+		n->scalarA = mapChild(scalarA.child());
+		n->scalarB = mapChild(scalarB.child());
+		return n;
+	}
+
+
+	void ASTOpMeshGeometryOperation::ForEachChild(const TFunctionRef<void(ASTChild&)> f)
+	{
+		f(meshA);
+		f(meshB);
+		f(scalarA);
+		f(scalarB);
+	}
+
+
+	void ASTOpMeshGeometryOperation::Link(FProgram& program, FLinkerOptions*)
+	{
+		// Already linked?
+		if (!linkedAddress)
+		{
+			OP::MeshGeometryOperationArgs args;
+			memset(&args, 0, sizeof(args));
+
+			if (meshA) args.meshA = meshA->linkedAddress;
+			if (meshB) args.meshB = meshB->linkedAddress;
+			if (scalarA) args.scalarA = scalarA->linkedAddress;
+			if (scalarB) args.scalarB = scalarB->linkedAddress;
+
+			linkedAddress = (OP::ADDRESS)program.m_opAddress.Num();
+			//program.m_code.push_back(op);
+			program.m_opAddress.Add((uint32_t)program.m_byteCode.Num());
+			AppendCode(program.m_byteCode, OP_TYPE::ME_GEOMETRYOPERATION);
+			AppendCode(program.m_byteCode, args);
+		}
+
+	}
+
+}
diff --git a/Engine/Plugins/Mutable/Source/MutableTools/Private/MuT/ASTOpMeshGeometryOperation.h b/Engine/Plugins/Mutable/Source/MutableTools/Private/MuT/ASTOpMeshGeometryOperation.h
new file mode 100644
--- /dev/null
+++ b/Engine/Plugins/Mutable/Source/MutableTools/Private/MuT/ASTOpMeshGeometryOperation.h
@@ -0,0 +1,44 @@
+// Copyright Epic Games, Inc. All Rights Reserved.
+
+#pragma once
+
+#include "HAL/Platform.h"
+#include "MuR/Operations.h"
+#include "MuR/Ptr.h"
+#include "MuT/AST.h"
+
+
+namespace mu
+{
+	struct FProgram;
+
+
+	//---------------------------------------------------------------------------------------------
+	//!
+	//---------------------------------------------------------------------------------------------
+	class ASTOpMeshGeometryOperation final : public ASTOp
+	{
+	public:
+
+		ASTChild meshA;
+		ASTChild meshB;
+		ASTChild scalarA;
+		ASTChild scalarB;
+
+	public:
+
+		ASTOpMeshGeometryOperation();
+		ASTOpMeshGeometryOperation(const ASTOpMeshGeometryOperation&) = delete;
+		~ASTOpMeshGeometryOperation();
+
+		OP_TYPE GetOpType() const override { return OP_TYPE::ME_GEOMETRYOPERATION; }
+		uint64 Hash() const override;
+		bool IsEqual(const ASTOp& otherUntyped) const override;
+		Ptr<ASTOp> Clone(MapChildFuncRef mapChild) const override;
+		void ForEachChild(const TFunctionRef<void(ASTChild&)>) override;
+		void Link(FProgram& program, FLinkerOptions* Options) override;
+	};
+
+
+}
+
diff --git a/Engine/Plugins/Mutable/Source/MutableTools/Private/MuT/ASTOpMeshMaskClipMesh.cpp b/Engine/Plugins/Mutable/Source/MutableTools/Private/MuT/ASTOpMeshMaskClipMesh.cpp
new file mode 100644
--- /dev/null
+++ b/Engine/Plugins/Mutable/Source/MutableTools/Private/MuT/ASTOpMeshMaskClipMesh.cpp
@@ -0,0 +1,368 @@
+// Copyright Epic Games, Inc. All Rights Reserved.
+
+#include "MuT/ASTOpMeshMaskClipMesh.h"
+
+#include "HAL/PlatformMath.h"
+#include "MuR/ModelPrivate.h"
+#include "MuR/RefCounted.h"
+#include "MuR/Types.h"
+#include "MuT/ASTOpConditional.h"
+#include "MuT/ASTOpMeshRemoveMask.h"
+#include "MuT/ASTOpMeshAddTags.h"
+#include "MuT/ASTOpSwitch.h"
+
+
+namespace mu
+{
+
+
+	//-------------------------------------------------------------------------------------------------
+	//-------------------------------------------------------------------------------------------------
+	//-------------------------------------------------------------------------------------------------
+	ASTOpMeshMaskClipMesh::ASTOpMeshMaskClipMesh()
+		: source(this)
+		, clip(this)
+	{
+	}
+
+
+	//-------------------------------------------------------------------------------------------------
+	ASTOpMeshMaskClipMesh::~ASTOpMeshMaskClipMesh()
+	{
+		// Explicit call needed to avoid recursive destruction
+		ASTOp::RemoveChildren();
+	}
+
+
+	//-------------------------------------------------------------------------------------------------
+	bool ASTOpMeshMaskClipMesh::IsEqual(const ASTOp& otherUntyped) const
+	{
+		if (otherUntyped.GetOpType() == GetOpType())
+		{
+			const ASTOpMeshMaskClipMesh* other = static_cast<const ASTOpMeshMaskClipMesh*>(&otherUntyped);
+			return source == other->source && clip == other->clip;
+		}
+		return false;
+	}
+
+
+	//-------------------------------------------------------------------------------------------------
+	uint64 ASTOpMeshMaskClipMesh::Hash() const
+	{
+		uint64 res = std::hash<void*>()(source.child().get());
+		hash_combine(res, clip.child());
+		return res;
+	}
+
+
+	//-------------------------------------------------------------------------------------------------
+	mu::Ptr<ASTOp> ASTOpMeshMaskClipMesh::Clone(MapChildFuncRef mapChild) const
+	{
+		Ptr<ASTOpMeshMaskClipMesh> n = new ASTOpMeshMaskClipMesh();
+		n->source = mapChild(source.child());
+		n->clip = mapChild(clip.child());
+		return n;
+	}
+
+
+	//-------------------------------------------------------------------------------------------------
+	void ASTOpMeshMaskClipMesh::ForEachChild(const TFunctionRef<void(ASTChild&)> f)
+	{
+		f(source);
+		f(clip);
+	}
+
+
+	//-------------------------------------------------------------------------------------------------
+	void ASTOpMeshMaskClipMesh::Link(FProgram& program, FLinkerOptions*)
+	{
+		// Already linked?
+		if (!linkedAddress)
+		{
+			OP::MeshMaskClipMeshArgs args;
+			memset(&args, 0, sizeof(args));
+
+			if (source) args.source = source->linkedAddress;
+			if (clip) args.clip = clip->linkedAddress;
+
+			linkedAddress = (OP::ADDRESS)program.m_opAddress.Num();
+			//program.m_code.push_back(op);
+			program.m_opAddress.Add((uint32_t)program.m_byteCode.Num());
+			AppendCode(program.m_byteCode, OP_TYPE::ME_MASKCLIPMESH);
+			AppendCode(program.m_byteCode, args);
+		}
+	}
+
+
+	//-------------------------------------------------------------------------------------------------
+	//-------------------------------------------------------------------------------------------------
+	//-------------------------------------------------------------------------------------------------
+	namespace
+	{
+		class Sink_MeshMaskClipMeshSource
+		{
+		public:
+
+			// \TODO This is recursive and may cause stack overflows in big models.
+
+			mu::Ptr<ASTOp> Apply(const ASTOpMeshMaskClipMesh* root)
+			{
+				m_root = root;
+				m_oldToNew.Empty();
+
+				m_initialSource = root->source.child();
+				mu::Ptr<ASTOp> NewSource = Visit(m_initialSource);
+
+				// If there is any change, it is the new root.
+				if (NewSource != m_initialSource)
+				{
+					return NewSource;
+				}
+
+				return nullptr;
+			}
+
+		protected:
+
+			const ASTOpMeshMaskClipMesh* m_root;
+			mu::Ptr<ASTOp> m_initialSource;
+			TMap<mu::Ptr<ASTOp>, mu::Ptr<ASTOp>> m_oldToNew;
+			TArray<mu::Ptr<ASTOp>> m_newOps;
+
+			mu::Ptr<ASTOp> Visit(const mu::Ptr<ASTOp>& at)
+			{
+				if (!at) return nullptr;
+
+				// Newly created?
+				if (m_newOps.Find(at) != INDEX_NONE)
+				{
+					return at;
+				}
+
+				// Already visited?
+				mu::Ptr<ASTOp>* cacheIt = m_oldToNew.Find(at);
+				if (cacheIt)
+				{
+					return *cacheIt;
+				}
+
+				mu::Ptr<ASTOp> newAt = at;
+				switch (at->GetOpType())
+				{
+
+				// This cannot be sunk since the result is different. Since the clipping is now correctly
+				// generated at the end of the chain when really necessary, this wrong optimisation is no 
+				// longer needed.
+				//case OP_TYPE::ME_MORPH:
+		        //{
+		        //    break;
+		        //}
+
+				case OP_TYPE::ME_REMOVEMASK:
+				{
+					// Remove this op:
+					// This may lead to the mask being bigger than needed since it will include
+					// faces removed by the ignored removemask, but it is ok
+
+					// TODO: Swap instead of ignore, and implement removemask on a mask?
+					const ASTOpMeshRemoveMask* typedAt = static_cast<const ASTOpMeshRemoveMask*>(at.get());
+					newAt = Visit(typedAt->source.child());
+					break;
+				}
+
+				case OP_TYPE::ME_ADDTAGS:
+				{
+					Ptr<ASTOpMeshAddTags> newOp = mu::Clone<ASTOpMeshAddTags>(at);
+					newOp->Source = Visit(newOp->Source.child());
+					newAt = newOp;
+					break;
+				}
+
+				case OP_TYPE::ME_CONDITIONAL:
+				{
+					// We move the mask creation down the two paths
+					// It always needs to be a clone because otherwise we could be modifying an
+					// instruction that shouldn't if the parent was a ME_REMOVEMASK above and we
+					// skipped the cloning for the parent.
+					Ptr<ASTOpConditional> newOp = mu::Clone<ASTOpConditional>(at);
+					newOp->yes = Visit(newOp->yes.child());
+					newOp->no = Visit(newOp->no.child());
+					newAt = newOp;
+					break;
+				}
+
+				case OP_TYPE::ME_SWITCH:
+				{
+					// We move the mask creation down all the paths
+					Ptr<ASTOpSwitch> newOp = mu::Clone<ASTOpSwitch>(at);
+					newOp->def = Visit(newOp->def.child());
+					for (ASTOpSwitch::FCase& c : newOp->cases)
+					{
+						c.branch = Visit(c.branch.child());
+					}
+					newAt = newOp;
+					break;
+				}
+
+				// This cannot be sunk since the result is different. Since the clipping is now correctly
+				// generated at the end of the chain when really necessary, this wrong optimisation is no 
+				// longer needed.
+				//case OP_TYPE::ME_CLIPMORPHPLANE:
+	   //         {
+	   //             // We move the mask creation down the source
+	   //             auto typedAt = dynamic_cast<const ASTOpMeshClipMorphPlane*>(at.get());
+	   //             newAt = Visit(typedAt->source.child());
+	   //             break;
+	   //         }
+
+				default:
+				{
+					//
+					if (at != m_initialSource)
+					{
+						Ptr<ASTOpMeshMaskClipMesh> newOp = mu::Clone<ASTOpMeshMaskClipMesh>(m_root);
+						newOp->source = at;
+						newAt = newOp;
+					}
+					break;
+				}
+
+				}
+
+				m_oldToNew.Add(at, newAt);
+
+				return newAt;
+			}
+		};
+
+
+		class Sink_MeshMaskClipMeshClip
+		{
+		public:
+
+			// \TODO This is recursive and may cause stack overflows in big models.
+
+			mu::Ptr<ASTOp> Apply(const ASTOpMeshMaskClipMesh* root)
+			{
+				m_root = root;
+				m_oldToNew.Empty();
+
+				m_initialClip = root->clip.child();
+				mu::Ptr<ASTOp> NewClip = Visit(m_initialClip);
+
+				// If there is any change, it is the new root.
+				if (NewClip != m_initialClip)
+				{
+					return NewClip;
+				}
+
+				return nullptr;
+			}
+
+		protected:
+
+			const ASTOpMeshMaskClipMesh* m_root;
+			mu::Ptr<ASTOp> m_initialClip;
+			TMap<mu::Ptr<ASTOp>, mu::Ptr<ASTOp>> m_oldToNew;
+			TArray<mu::Ptr<ASTOp>> m_newOps;
+
+			mu::Ptr<ASTOp> Visit(const mu::Ptr<ASTOp>& at)
+			{
+				if (!at) return nullptr;
+
+				// Newly created?
+				if (m_newOps.Contains(at))
+				{
+					return at;
+				}
+
+				// Already visited?
+				mu::Ptr<ASTOp>* cacheIt = m_oldToNew.Find(at);
+				if (cacheIt)
+				{
+					return *cacheIt;
+				}
+
+				mu::Ptr<ASTOp> newAt = at;
+				switch (at->GetOpType())
+				{
+
+				case OP_TYPE::ME_CONDITIONAL:
+				{
+					// We move the mask creation down the two paths
+					// It always needs to be a clone because otherwise we could be modifying an
+					// instruction that shouldn't if the parent was a ME_REMOVEMASK above and we
+					// skipped the cloning for the parent.
+					Ptr<ASTOpConditional> newOp = mu::Clone<ASTOpConditional>(at);
+					newOp->yes = Visit(newOp->yes.child());
+					newOp->no = Visit(newOp->no.child());
+					newAt = newOp;
+					break;
+				}
+
+				case OP_TYPE::ME_SWITCH:
+				{
+					// We move the mask creation down all the paths
+					Ptr<ASTOpSwitch> newOp = mu::Clone<ASTOpSwitch>(at);
+					newOp->def = Visit(newOp->def.child());
+					for (ASTOpSwitch::FCase& c : newOp->cases)
+					{
+						c.branch = Visit(c.branch.child());
+					}
+					newAt = newOp;
+					break;
+				}
+
+				case OP_TYPE::ME_ADDTAGS:
+				{
+					// Ignore tags in the clip mesh branch
+					const ASTOpMeshAddTags* AddOp = static_cast<const ASTOpMeshAddTags*>(at.get());
+					newAt = Visit(AddOp->Source.child());
+					break;
+				}
+
+				default:
+				{
+					//
+					if (at != m_initialClip)
+					{
+						Ptr<ASTOpMeshMaskClipMesh> newOp = mu::Clone<ASTOpMeshMaskClipMesh>(m_root);
+						newOp->clip = at;
+						newAt = newOp;
+					}
+					break;
+				}
+
+				}
+
+				m_oldToNew.Add(at, newAt);
+
+				return newAt;
+			}
+		};
+
+	}
+
+
+	//-------------------------------------------------------------------------------------------------
+	mu::Ptr<ASTOp> ASTOpMeshMaskClipMesh::OptimiseSink(const FModelOptimizationOptions&, FOptimizeSinkContext&) const
+	{
+		// \TODO: Add logic state to the sinkers to avoid explosion with switches in both branches and similar cases.
+
+		// This will sink the operation down the source
+		Sink_MeshMaskClipMeshSource SinkerSource;
+		mu::Ptr<ASTOp> at = SinkerSource.Apply(this);
+
+		// If we didn't sink it.
+		if (!at || at == this)
+		{
+			// This will sink the operation down the clip child
+			Sink_MeshMaskClipMeshClip SinkerClip;
+			at = SinkerClip.Apply(this);
+		}
+
+		return at;
+	}
+
+}
+ 
\ No newline at end of file
diff --git a/Engine/Plugins/Mutable/Source/MutableTools/Private/MuT/ASTOpMeshMorphReshape.cpp b/Engine/Plugins/Mutable/Source/MutableTools/Private/MuT/ASTOpMeshMorphReshape.cpp
new file mode 100644
--- /dev/null
+++ b/Engine/Plugins/Mutable/Source/MutableTools/Private/MuT/ASTOpMeshMorphReshape.cpp
@@ -0,0 +1,92 @@
+// Copyright Epic Games, Inc. All Rights Reserved.
+
+#include "MuT/ASTOpMeshMorphReshape.h"
+
+#include "HAL/PlatformMath.h"
+#include "MuR/ModelPrivate.h"
+#include "MuR/RefCounted.h"
+#include "MuR/Types.h"
+
+
+namespace mu
+{
+
+
+	ASTOpMeshMorphReshape::ASTOpMeshMorphReshape()
+		: Morph(this)
+		, Reshape(this)
+	{
+	}
+
+
+	ASTOpMeshMorphReshape::~ASTOpMeshMorphReshape()
+	{
+		// Explicit call needed to avoid recursive destruction
+		ASTOp::RemoveChildren();
+	}
+
+
+	bool ASTOpMeshMorphReshape::IsEqual(const ASTOp& otherUntyped) const
+	{
+		if (otherUntyped.GetOpType() == GetOpType())
+		{
+			const ASTOpMeshMorphReshape* other = static_cast<const ASTOpMeshMorphReshape*>(&otherUntyped);
+			return Morph == other->Morph && Reshape == other->Reshape;
+		}
+
+		return false;
+	}
+
+
+	uint64 ASTOpMeshMorphReshape::Hash() const
+	{
+		uint64 res = std::hash<void*>()(Morph.child().get());
+		hash_combine(res, Reshape.child().get());
+
+		return res;
+	}
+
+
+	mu::Ptr<ASTOp> ASTOpMeshMorphReshape::Clone(MapChildFuncRef mapChild) const
+	{
+		Ptr<ASTOpMeshMorphReshape> n = new ASTOpMeshMorphReshape();
+		n->Morph = mapChild(Morph.child());
+		n->Reshape = mapChild(Reshape.child());
+		return n;
+	}
+
+
+	void ASTOpMeshMorphReshape::ForEachChild(const TFunctionRef<void(ASTChild&)> f)
+	{
+		f(Morph);
+		f(Reshape);
+	}
+
+
+	void ASTOpMeshMorphReshape::Link(FProgram& program, FLinkerOptions*)
+	{
+		// Already linked?
+		if (!linkedAddress)
+		{
+			OP::MeshMorphReshapeArgs Args;
+			memset(&Args, 0, sizeof(Args));
+
+			if (Morph)
+			{
+				Args.Morph = Morph->linkedAddress;
+			}
+
+			if (Reshape)
+			{
+				Args.Reshape = Reshape->linkedAddress;
+			}
+
+			linkedAddress = (OP::ADDRESS)program.m_opAddress.Num();
+			program.m_opAddress.Add((uint32_t)program.m_byteCode.Num());
+			AppendCode(program.m_byteCode, OP_TYPE::ME_MORPHRESHAPE);
+			AppendCode(program.m_byteCode, Args);
+		}
+
+	}
+
+}
diff --git a/Engine/Plugins/Mutable/Source/MutableTools/Private/MuT/ASTOpMeshMorphReshape.h b/Engine/Plugins/Mutable/Source/MutableTools/Private/MuT/ASTOpMeshMorphReshape.h
new file mode 100644
--- /dev/null
+++ b/Engine/Plugins/Mutable/Source/MutableTools/Private/MuT/ASTOpMeshMorphReshape.h
@@ -0,0 +1,41 @@
+// Copyright Epic Games, Inc. All Rights Reserved.
+
+#pragma once
+
+#include "HAL/Platform.h"
+#include "MuR/Operations.h"
+#include "MuR/Ptr.h"
+#include "MuT/AST.h"
+
+
+namespace mu
+{
+struct FProgram;
+
+	//---------------------------------------------------------------------------------------------
+	//!
+	//---------------------------------------------------------------------------------------------
+	class ASTOpMeshMorphReshape final : public ASTOp
+	{
+	public:
+
+		ASTChild Morph;
+		ASTChild Reshape;
+
+	public:
+
+		ASTOpMeshMorphReshape();
+		ASTOpMeshMorphReshape(const ASTOpMeshMorphReshape&) = delete;
+		~ASTOpMeshMorphReshape();
+
+		OP_TYPE GetOpType() const override { return OP_TYPE::ME_MORPHRESHAPE; }
+		uint64 Hash() const override;
+		bool IsEqual(const ASTOp& otherUntyped) const override;
+		Ptr<ASTOp> Clone(MapChildFuncRef mapChild) const override;
+		void ForEachChild(const TFunctionRef<void(ASTChild&)>) override;
+		void Link(FProgram& program, FLinkerOptions* Options) override;
+	};
+
+
+}
+
diff --git a/Engine/Plugins/Mutable/Source/MutableTools/Private/MuT/ASTOpMeshRemoveMask.cpp b/Engine/Plugins/Mutable/Source/MutableTools/Private/MuT/ASTOpMeshRemoveMask.cpp
new file mode 100644
--- /dev/null
+++ b/Engine/Plugins/Mutable/Source/MutableTools/Private/MuT/ASTOpMeshRemoveMask.cpp
@@ -0,0 +1,276 @@
+// Copyright Epic Games, Inc. All Rights Reserved.
+
+#include "MuT/ASTOpMeshRemoveMask.h"
+
+#include "MuT/ASTOpMeshMorph.h"
+#include "MuT/ASTOpMeshAddTags.h"
+#include "HAL/PlatformMath.h"
+#include "MuR/ModelPrivate.h"
+#include "MuR/RefCounted.h"
+#include "MuR/Types.h"
+
+
+namespace mu
+{
+
+
+	//---------------------------------------------------------------------------------------------
+	ASTOpMeshRemoveMask::ASTOpMeshRemoveMask()
+		: source(this)
+	{
+	}
+
+
+	//---------------------------------------------------------------------------------------------
+	ASTOpMeshRemoveMask::~ASTOpMeshRemoveMask()
+	{
+		// Explicit call needed to avoid recursive destruction
+		ASTOp::RemoveChildren();
+	}
+
+
+	//---------------------------------------------------------------------------------------------
+	void ASTOpMeshRemoveMask::AddRemove(const Ptr<ASTOp>& condition, const Ptr<ASTOp>& mask)
+	{
+		removes.Add({ ASTChild(this,condition), ASTChild(this,mask) });
+	}
+
+
+	//---------------------------------------------------------------------------------------------
+	bool ASTOpMeshRemoveMask::IsEqual(const ASTOp& otherUntyped) const
+	{
+		if (otherUntyped.GetOpType() == GetOpType())
+		{
+			const ASTOpMeshRemoveMask* other = static_cast<const ASTOpMeshRemoveMask*>(&otherUntyped);
+			return source == other->source && removes == other->removes;
+		}
+		return false;
+	}
+
+
+	//---------------------------------------------------------------------------------------------
+	mu::Ptr<ASTOp> ASTOpMeshRemoveMask::Clone(MapChildFuncRef mapChild) const
+	{
+		Ptr<ASTOpMeshRemoveMask> n = new ASTOpMeshRemoveMask();
+		n->source = mapChild(source.child());
+		for (const TPair<ASTChild, ASTChild>& r : removes)
+		{
+			n->removes.Add({ ASTChild(n,mapChild(r.Key.child())), ASTChild(n,mapChild(r.Value.child())) });
+		}
+		return n;
+	}
+
+
+	//---------------------------------------------------------------------------------------------
+	void ASTOpMeshRemoveMask::ForEachChild(const TFunctionRef<void(ASTChild&)> f)
+	{
+		f(source);
+		for (TPair<ASTChild, ASTChild>& r : removes)
+		{
+			f(r.Key);
+			f(r.Value);
+		}
+	}
+
+
+	//---------------------------------------------------------------------------------------------
+	uint64 ASTOpMeshRemoveMask::Hash() const
+	{
+		uint64 res = std::hash<ASTOp*>()(source.child().get());
+		for (const TPair<ASTChild, ASTChild>& r : removes)
+		{
+			hash_combine(res, r.Key.child().get());
+			hash_combine(res, r.Value.child().get());
+		}
+		return res;
+	}
+
+
+	//---------------------------------------------------------------------------------------------
+	void ASTOpMeshRemoveMask::Link(FProgram& program, FLinkerOptions*)
+	{
+		// Already linked?
+		if (!linkedAddress)
+		{
+			linkedAddress = (OP::ADDRESS)program.m_opAddress.Num();
+
+			program.m_opAddress.Add((uint32_t)program.m_byteCode.Num());
+			AppendCode(program.m_byteCode, OP_TYPE::ME_REMOVEMASK);
+			OP::ADDRESS sourceAt = source ? source->linkedAddress : 0;
+			AppendCode(program.m_byteCode, sourceAt);
+			AppendCode(program.m_byteCode, (uint16)removes.Num());
+			for (const TPair<ASTChild, ASTChild>& b : removes)
+			{
+				OP::ADDRESS condition = b.Key ? b.Key->linkedAddress : 0;
+				AppendCode(program.m_byteCode, condition);
+
+				OP::ADDRESS remove = b.Value ? b.Value->linkedAddress : 0;
+				AppendCode(program.m_byteCode, remove);
+			}
+		}
+	}
+
+
+	//-------------------------------------------------------------------------------------------------
+	//-------------------------------------------------------------------------------------------------
+	//-------------------------------------------------------------------------------------------------
+	namespace
+	{
+		class Sink_MeshRemoveMaskAST
+		{
+		public:
+
+			// \TODO This is recursive and may cause stack overflows in big models.
+
+			mu::Ptr<ASTOp> Apply(const ASTOpMeshRemoveMask* root)
+			{
+				m_root = root;
+				m_oldToNew.Empty();
+
+				m_initialSource = root->source.child();
+				mu::Ptr<ASTOp> newSource = Visit(m_initialSource);
+
+				// If there is any change, it is the new root.
+				if (newSource != m_initialSource)
+				{
+					return newSource;
+				}
+
+				return nullptr;
+			}
+
+		protected:
+
+			const ASTOpMeshRemoveMask* m_root;
+			mu::Ptr<ASTOp> m_initialSource;
+			TMap<mu::Ptr<ASTOp>, mu::Ptr<ASTOp>> m_oldToNew;
+			TArray<mu::Ptr<ASTOp>> m_newOps;
+
+			mu::Ptr<ASTOp> Visit(const mu::Ptr<ASTOp>& at)
+			{
+				if (!at) return nullptr;
+
+				// Newly created?
+				if (m_newOps.Contains(at))
+				{
+					return at;
+				}
+
+				// Already visited?
+				Ptr<ASTOp>* cacheIt = m_oldToNew.Find(at);
+				if (cacheIt)
+				{
+					return *cacheIt;
+				}
+
+				mu::Ptr<ASTOp> newAt = at;
+				switch (at->GetOpType())
+				{
+
+				case OP_TYPE::ME_MORPH:
+				{
+					mu::Ptr<ASTOpMeshMorph> newOp = mu::Clone<ASTOpMeshMorph>(at);
+					newOp->Base = Visit(newOp->Base.child());
+					newAt = newOp;
+					break;
+				}
+
+				case OP_TYPE::ME_ADDTAGS:
+				{
+					mu::Ptr<ASTOpMeshAddTags> newOp = mu::Clone<ASTOpMeshAddTags>(at);
+					newOp->Source = Visit(newOp->Source.child());
+					newAt = newOp;
+					break;
+				}
+
+				// disabled to avoid code explosion (or bug?) TODO
+//            case OP_TYPE::ME_CONDITIONAL:
+//            {
+//                Ptr<ASTOpConditional> newOp = mu::Clone<ASTOpConditional>(at);
+//                newOp->yes = Visit(newOp->yes.child());
+//                newOp->no = Visit(newOp->no.child());
+//                newAt = newOp;
+//                break;
+//            }
+
+//            case OP_TYPE::ME_SWITCH:
+//            {
+//                auto newOp = mu::Clone<ASTOpSwitch>(at);
+//                newOp->def = Visit(newOp->def.child());
+//                for( auto& c:newOp->cases )
+//                {
+//                    c.branch = Visit(c.branch.child());
+//                }
+//                newAt = newOp;
+//                break;
+//            }
+
+				default:
+				{
+					//
+					if (at != m_initialSource)
+					{
+						Ptr<ASTOpMeshRemoveMask> newOp = mu::Clone<ASTOpMeshRemoveMask>(m_root);
+						newOp->source = at;
+						newAt = newOp;
+					}
+					break;
+				}
+
+				}
+
+				m_oldToNew.Add(at, newAt);
+
+				return newAt;
+			}
+		};
+	}
+
+
+	//-------------------------------------------------------------------------------------------------
+	mu::Ptr<ASTOp> ASTOpMeshRemoveMask::OptimiseSink(const FModelOptimizationOptions&, FOptimizeSinkContext&) const
+	{
+		Sink_MeshRemoveMaskAST sinker;
+		mu::Ptr<ASTOp> at = sinker.Apply(this);
+
+		// If not optimized already, see if we can optimize the "remove" branches
+		if (!at)
+		{
+			Ptr<ASTOpMeshRemoveMask> NewOp;
+
+			for (int32 RemoveIndex=0; RemoveIndex<removes.Num(); ++RemoveIndex)
+			{
+				if (!removes[RemoveIndex].Value)
+				{
+					continue;
+				}
+
+				OP_TYPE RemoveType = removes[RemoveIndex].Value->GetOpType();
+				switch (RemoveType)
+				{
+				case OP_TYPE::ME_ADDTAGS:
+				{
+					// It can be ignored.
+					if (!NewOp)
+					{
+						NewOp = mu::Clone<ASTOpMeshRemoveMask>(this);
+					}
+
+					const ASTOpMeshAddTags* Add = static_cast<const ASTOpMeshAddTags*>(removes[RemoveIndex].Value.child().get());
+					NewOp->removes[RemoveIndex].Value = Add->Source.child();
+
+					break;
+				}
+
+				default:
+					break;
+				}
+			}
+
+			at = NewOp;
+		}
+
+		return at;
+	}
+
+}
diff --git a/Engine/Plugins/Mutable/Source/MutableTools/Private/MuT/ASTOpMeshTransform.cpp b/Engine/Plugins/Mutable/Source/MutableTools/Private/MuT/ASTOpMeshTransform.cpp
new file mode 100644
--- /dev/null
+++ b/Engine/Plugins/Mutable/Source/MutableTools/Private/MuT/ASTOpMeshTransform.cpp
@@ -0,0 +1,83 @@
+// Copyright Epic Games, Inc. All Rights Reserved.
+
+#include "MuT/ASTOpMeshTransform.h"
+
+#include "HAL/PlatformMath.h"
+#include "MuR/ModelPrivate.h"
+#include "MuR/RefCounted.h"
+#include "MuR/Types.h"
+
+
+namespace mu
+{
+
+	ASTOpMeshTransform::ASTOpMeshTransform()
+		: source(this)
+	{
+	}
+
+
+	ASTOpMeshTransform::~ASTOpMeshTransform()
+	{
+		// Explicit call needed to avoid recursive destruction
+		ASTOp::RemoveChildren();
+	}
+
+
+	bool ASTOpMeshTransform::IsEqual(const ASTOp& otherUntyped) const
+	{
+		if (otherUntyped.GetOpType() == GetOpType())
+		{
+			const ASTOpMeshTransform* other = static_cast<const ASTOpMeshTransform*>(&otherUntyped);
+			return source == other->source &&
+				matrix == other->matrix;
+		}
+		return false;
+	}
+
+
+	uint64 ASTOpMeshTransform::Hash() const
+	{
+		uint64 res = std::hash<OP_TYPE>()(OP_TYPE::ME_TRANSFORM);
+		hash_combine(res, source.child().get());
+		return res;
+	}
+
+
+	mu::Ptr<ASTOp> ASTOpMeshTransform::Clone(MapChildFuncRef mapChild) const
+	{
+		Ptr<ASTOpMeshTransform> n = new ASTOpMeshTransform();
+		n->matrix = matrix;
+		n->source = mapChild(source.child());
+		return n;
+	}
+
+
+	void ASTOpMeshTransform::ForEachChild(const TFunctionRef<void(ASTChild&)> f)
+	{
+		f(source);
+	}
+
+
+	void ASTOpMeshTransform::Link(FProgram& program, FLinkerOptions*)
+	{
+		// Already linked?
+		if (!linkedAddress)
+		{
+			OP::MeshTransformArgs args;
+			memset(&args, 0, sizeof(args));
+
+			if (source) args.source = source->linkedAddress;
+
+			args.matrix = program.AddConstant(matrix);
+
+			linkedAddress = (OP::ADDRESS)program.m_opAddress.Num();
+			//program.m_code.push_back(op);
+			program.m_opAddress.Add((uint32_t)program.m_byteCode.Num());
+			AppendCode(program.m_byteCode, OP_TYPE::ME_TRANSFORM);
+			AppendCode(program.m_byteCode, args);
+		}
+
+	}
+
+}
diff --git a/Engine/Plugins/Mutable/Source/MutableTools/Private/MuT/ASTOpParameter.cpp b/Engine/Plugins/Mutable/Source/MutableTools/Private/MuT/ASTOpParameter.cpp
new file mode 100644
--- /dev/null
+++ b/Engine/Plugins/Mutable/Source/MutableTools/Private/MuT/ASTOpParameter.cpp
@@ -0,0 +1,149 @@
+// Copyright Epic Games, Inc. All Rights Reserved.
+
+#include "MuT/ASTOpParameter.h"
+
+#include "HAL/PlatformMath.h"
+#include "Misc/AssertionMacros.h"
+#include "MuR/ModelPrivate.h"
+#include "MuR/Parameters.h"
+#include "MuR/RefCounted.h"
+#include "MuR/Types.h"
+
+
+namespace mu
+{
+
+
+	ASTOpParameter::~ASTOpParameter()
+	{
+		// Explicit call needed to avoid recursive destruction
+		ASTOp::RemoveChildren();
+	}
+
+
+	void ASTOpParameter::ForEachChild(const TFunctionRef<void(ASTChild&)> f)
+	{
+		for (auto& c : ranges)
+		{
+			f(c.rangeSize);
+		}
+	}
+
+
+	bool ASTOpParameter::IsEqual(const ASTOp& otherUntyped) const
+	{
+		if (otherUntyped.GetOpType() == GetOpType())
+		{
+			const ASTOpParameter* other = static_cast<const ASTOpParameter*>(&otherUntyped);
+			return type == other->type &&
+				parameter == other->parameter &&
+				ranges == other->ranges;
+		}
+		return false;
+	}
+
+
+	mu::Ptr<ASTOp> ASTOpParameter::Clone(MapChildFuncRef mapChild) const
+	{
+		Ptr<ASTOpParameter> n = new ASTOpParameter();
+		n->type = type;
+		n->parameter = parameter;
+		for (const auto& c : ranges)
+		{
+			n->ranges.Emplace(n.get(), mapChild(c.rangeSize.child()), c.rangeName, c.rangeUID);
+		}
+		return n;
+	}
+
+
+	uint64 ASTOpParameter::Hash() const
+	{
+		uint64 res = std::hash<uint64>()(uint64(type));
+		hash_combine(res, parameter.m_type);
+		hash_combine(res, parameter.m_name.Len());
+		return res;
+	}
+
+
+	void ASTOpParameter::Assert()
+	{
+		ASTOp::Assert();
+	}
+
+
+	void ASTOpParameter::Link(FProgram& Program, FLinkerOptions*)
+	{
+		// Already linked?
+		if (!linkedAddress)
+		{
+			OP::ParameterArgs args;
+			FMemory::Memzero(&args, sizeof(args));
+
+			LinkedParameterIndex = Program.m_parameters.Find(parameter);
+
+			// If this fails, it means an ASTOpParameter was created at code generation time, but not registered into the
+			// parameters map in the CodeGenerator_FirstPass.
+			check(LinkedParameterIndex != INDEX_NONE);
+
+			args.variable = (OP::ADDRESS)LinkedParameterIndex;
+
+			for (const auto& d : ranges)
+			{
+				OP::ADDRESS sizeAt = 0;
+				uint16 rangeId = 0;
+				LinkRange(Program, d, sizeAt, rangeId);
+				Program.m_parameters[LinkedParameterIndex].m_ranges.Add(rangeId);
+			}
+
+			linkedAddress = (OP::ADDRESS)Program.m_opAddress.Num();
+
+			Program.m_opAddress.Add((uint32)Program.m_byteCode.Num());
+			AppendCode(Program.m_byteCode, type);
+			AppendCode(Program.m_byteCode, args);
+		}
+	}
+
+
+	int ASTOpParameter::EvaluateInt(ASTOpList& facts, bool& unknown) const
+	{
+		unknown = true;
+
+		// Check the facts, in case we have the value for our parameter.
+		for (const auto& f : facts)
+		{
+			if (f->GetOpType() == OP_TYPE::BO_EQUAL_INT_CONST)
+			{
+				const auto typedFact = static_cast<const ASTOpFixed*>(f.get());
+				auto value = typedFact->children[typedFact->op.args.BoolEqualScalarConst.value].child();
+				if (value.get() == this)
+				{
+					unknown = false;
+					return typedFact->op.args.BoolEqualScalarConst.constant;
+				}
+				else
+				{
+					// We could try something more if it was an expression and it had the parameter
+					// somewhere in it.
+				}
+			}
+		}
+
+		return 0;
+	}
+
+
+	ASTOp::FBoolEvalResult ASTOpParameter::EvaluateBool(ASTOpList& /*facts*/,
+		FEvaluateBoolCache*) const
+	{
+		check(type == OP_TYPE::BO_PARAMETER);
+		return BET_UNKNOWN;
+	}
+
+
+	FImageDesc ASTOpParameter::GetImageDesc(bool, FGetImageDescContext*) const
+	{
+		check(type == OP_TYPE::IM_PARAMETER);
+		return FImageDesc();
+	}
+
+}
diff --git a/Engine/Plugins/Mutable/Source/MutableTools/Private/MuT/ASTOpParameter.h b/Engine/Plugins/Mutable/Source/MutableTools/Private/MuT/ASTOpParameter.h
new file mode 100644
--- /dev/null
+++ b/Engine/Plugins/Mutable/Source/MutableTools/Private/MuT/ASTOpParameter.h
@@ -0,0 +1,55 @@
+// Copyright Epic Games, Inc. All Rights Reserved.
+
+#pragma once
+
+#include "HAL/Platform.h"
+#include "MuR/Image.h"
+#include "MuR/Operations.h"
+#include "MuR/ParametersPrivate.h"
+#include "MuR/Ptr.h"
+#include "MuT/AST.h"
+
+
+namespace mu
+{
+	struct FProgram;
+
+
+	//---------------------------------------------------------------------------------------------
+	//! Parameter operation.
+	//---------------------------------------------------------------------------------------------
+	class ASTOpParameter final : public ASTOp
+	{
+	public:
+
+		//! Type of parameter
+		OP_TYPE type;
+
+		//!
+		FParameterDesc parameter;
+
+		//** Ranges adding dimensions to this parameter. */
+		TArray<FRangeData> ranges;
+
+		/** Index of the parameter in the program parameter list. Generated ar link time. */
+		int32 LinkedParameterIndex = -1;
+
+	public:
+
+		~ASTOpParameter() override;
+
+		OP_TYPE GetOpType() const override { return type; }
+		uint64 Hash() const override;
+		void ForEachChild(const TFunctionRef<void(ASTChild&)>) override;
+		bool IsEqual(const ASTOp& otherUntyped) const override;
+		Ptr<ASTOp> Clone(MapChildFuncRef mapChild) const override;
+		void Assert() override;
+		void Link(FProgram& program, FLinkerOptions*) override;
+		int EvaluateInt(ASTOpList& facts, bool& unknown) const override;
+		FBoolEvalResult EvaluateBool(ASTOpList& /*facts*/, FEvaluateBoolCache* = nullptr) const override;
+		FImageDesc GetImageDesc(bool, FGetImageDescContext*) const override;
+
+	};
+
+}
+
diff --git a/Engine/Plugins/Mutable/Source/MutableTools/Private/MuT/ASTOpScalarCurve.cpp b/Engine/Plugins/Mutable/Source/MutableTools/Private/MuT/ASTOpScalarCurve.cpp
new file mode 100644
--- /dev/null
+++ b/Engine/Plugins/Mutable/Source/MutableTools/Private/MuT/ASTOpScalarCurve.cpp
@@ -0,0 +1,81 @@
+// Copyright Epic Games, Inc. All Rights Reserved.
+
+#include "MuT/ASTOpScalarCurve.h"
+
+#include "HAL/PlatformMath.h"
+#include "MuR/ModelPrivate.h"
+#include "MuR/RefCounted.h"
+#include "MuR/Types.h"
+
+
+namespace mu
+{
+
+
+	ASTOpScalarCurve::ASTOpScalarCurve()
+		: time(this)
+	{
+	}
+
+
+	ASTOpScalarCurve::~ASTOpScalarCurve()
+	{
+		// Explicit call needed to avoid recursive destruction
+		ASTOp::RemoveChildren();
+	}
+
+
+	void ASTOpScalarCurve::ForEachChild(const TFunctionRef<void(ASTChild&)> f)
+	{
+		f(time);
+	}
+
+
+	uint64 ASTOpScalarCurve::Hash() const
+	{
+		uint64 res = std::hash<uint64>()(size_t(OP_TYPE::SC_CURVE));
+		hash_combine(res, Curve.Keys.Num());
+		return res;
+	}
+
+
+	bool ASTOpScalarCurve::IsEqual(const ASTOp& otherUntyped) const
+	{
+		if (otherUntyped.GetOpType() == GetOpType())
+		{
+			const ASTOpScalarCurve* other = static_cast<const ASTOpScalarCurve*>(&otherUntyped);
+			return time == other->time && Curve == other->Curve;
+		}
+		return false;
+	}
+
+
+	mu::Ptr<ASTOp> ASTOpScalarCurve::Clone(MapChildFuncRef mapChild) const
+	{
+		Ptr<ASTOpScalarCurve> n = new ASTOpScalarCurve();
+		n->Curve = Curve;
+		n->time = mapChild(time.child());
+		return n;
+	}
+
+
+	void ASTOpScalarCurve::Link(FProgram& program, FLinkerOptions*)
+	{
+		if (!linkedAddress)
+		{
+			OP::ScalarCurveArgs args;
+			memset(&args, 0, sizeof(args));
+			args.time = time ? time->linkedAddress : 0;
+			args.curve = program.AddConstant(Curve);
+
+			linkedAddress = (OP::ADDRESS)program.m_opAddress.Num();
+			//program.m_code.push_back(op);
+			program.m_opAddress.Add((uint32_t)program.m_byteCode.Num());
+			AppendCode(program.m_byteCode, OP_TYPE::SC_CURVE);
+			AppendCode(program.m_byteCode, args);
+		}
+
+	}
+
+}
+
