[mutable] Moved the Mutable plugin out of Experimental status into Beta.

diff --git a/Engine/Plugins/Mutable/Source/MutableTools/Private/MuT/ASTOpSwitch.cpp b/Engine/Plugins/Mutable/Source/MutableTools/Private/MuT/ASTOpSwitch.cpp
new file mode 100644
--- /dev/null
+++ b/Engine/Plugins/Mutable/Source/MutableTools/Private/MuT/ASTOpSwitch.cpp
@@ -0,0 +1,658 @@
+// Copyright Epic Games, Inc. All Rights Reserved.
+
+#include "MuT/ASTOpSwitch.h"
+
+#include "MuT/ASTOpParameter.h"
+#include "Containers/Map.h"
+#include "Misc/AssertionMacros.h"
+#include "MuR/ImagePrivate.h"
+#include "MuR/ModelPrivate.h"
+#include "MuR/MutableMath.h"
+#include "MuR/RefCounted.h"
+#include "MuR/Types.h"
+
+
+namespace mu
+{
+
+
+	//---------------------------------------------------------------------------------------------
+	ASTOpSwitch::ASTOpSwitch()
+		: variable(this)
+		, def(this)
+	{
+	}
+
+
+	//-------------------------------------------------------------------------------------------------
+	ASTOpSwitch::~ASTOpSwitch()
+	{
+		// Explicit call needed to avoid recursive destruction
+		ASTOp::RemoveChildren();
+	}
+
+
+	//-------------------------------------------------------------------------------------------------
+	bool ASTOpSwitch::IsEqual(const ASTOp& otherUntyped) const
+	{
+		if (GetOpType()==otherUntyped.GetOpType())
+		{
+			const ASTOpSwitch* other = static_cast<const ASTOpSwitch*>(&otherUntyped);
+			return type == other->type && variable == other->variable &&
+				cases == other->cases && def == other->def;
+		}
+		return false;
+	}
+
+
+	//-------------------------------------------------------------------------------------------------
+	mu::Ptr<ASTOp> ASTOpSwitch::Clone(MapChildFuncRef mapChild) const
+	{
+		Ptr<ASTOpSwitch> n = new ASTOpSwitch();
+		n->type = type;
+		n->variable = mapChild(variable.child());
+		n->def = mapChild(def.child());
+		for (const auto& c : cases)
+		{
+			n->cases.Emplace(c.condition, n, mapChild(c.branch.child()));
+		}
+		return n;
+	}
+
+
+	//-------------------------------------------------------------------------------------------------
+	void ASTOpSwitch::Assert()
+	{
+		switch (type)
+		{
+		case OP_TYPE::NU_SWITCH:
+		case OP_TYPE::SC_SWITCH:
+		case OP_TYPE::CO_SWITCH:
+		case OP_TYPE::IM_SWITCH:
+		case OP_TYPE::ME_SWITCH:
+		case OP_TYPE::LA_SWITCH:
+		case OP_TYPE::IN_SWITCH:
+		case OP_TYPE::ED_SWITCH:
+			break;
+		default:
+			// Unexpected type
+			check(false);
+			break;
+		}
+
+		ASTOp::Assert();
+	}
+
+
+	//-------------------------------------------------------------------------------------------------
+	uint64 ASTOpSwitch::Hash() const
+	{
+		uint64 res = std::hash<uint64>()(uint64(type));
+		for (const auto& c : cases)
+		{
+			hash_combine(res, c.condition);
+			hash_combine(res, c.branch.child().get());
+		}
+		return res;
+	}
+
+
+	//-------------------------------------------------------------------------------------------------
+	mu::Ptr<ASTOp> ASTOpSwitch::GetFirstValidValue()
+	{
+		for (int32 i = 0; i < cases.Num(); ++i)
+		{
+			if (cases[i].branch)
+			{
+				return cases[i].branch.child();
+			}
+		}
+		return nullptr;
+	}
+
+
+	//-------------------------------------------------------------------------------------------------
+	bool ASTOpSwitch::IsCompatibleWith(const ASTOpSwitch* other) const
+	{
+		if (!other) return false;
+		if (variable.child() != other->variable.child()) return false;
+		if (cases.Num() != other->cases.Num()) return false;
+		for (const auto& c : cases)
+		{
+			bool found = false;
+			for (const auto& o : other->cases)
+			{
+				if (c.condition == o.condition)
+				{
+					found = true;
+					break;
+				}
+			}
+			if (!found)
+			{
+				return false;
+			}
+		}
+
+		return true;
+	}
+
+
+	//-------------------------------------------------------------------------------------------------
+	mu::Ptr<ASTOp> ASTOpSwitch::FindBranch(int32 condition) const
+	{
+		for (const auto& c : cases)
+		{
+			if (c.condition == condition)
+			{
+				return c.branch.child();
+			}
+		}
+
+		return def.child();
+	}
+
+
+	//-------------------------------------------------------------------------------------------------
+	void ASTOpSwitch::ForEachChild(const TFunctionRef<void(ASTChild&)> f)
+	{
+		f(variable);
+		f(def);
+		for (auto& cas : cases)
+		{
+			f(cas.branch);
+		}
+	}
+
+
+	//-------------------------------------------------------------------------------------------------
+	void ASTOpSwitch::Link(FProgram& program, FLinkerOptions*)
+	{
+		// Already linked?
+		if (!linkedAddress)
+		{
+			linkedAddress = (OP::ADDRESS)program.m_opAddress.Num();
+			program.m_opAddress.Add((uint32_t)program.m_byteCode.Num());
+
+			OP::ADDRESS VarAddress = variable ? variable->linkedAddress : 0;
+			OP::ADDRESS DefAddress = def ? def->linkedAddress : 0;
+
+			AppendCode(program.m_byteCode, type);
+			AppendCode(program.m_byteCode, VarAddress);
+			AppendCode(program.m_byteCode, DefAddress);
+			AppendCode(program.m_byteCode, (uint32_t)cases.Num());
+
+			for (const FCase& Case : cases)
+			{
+				OP::ADDRESS CaseBranchAddress = Case.branch ? Case.branch->linkedAddress : 0;
+				AppendCode(program.m_byteCode, Case.condition);
+				AppendCode(program.m_byteCode, CaseBranchAddress);
+			}
+		}
+	}
+
+	//-------------------------------------------------------------------------------------------------
+	FImageDesc ASTOpSwitch::GetImageDesc(bool bReturnBestOption, class FGetImageDescContext* Context) const
+	{
+		FImageDesc Result;
+
+		// Local context in case it is necessary
+		FGetImageDescContext LocalContext;
+		if (!Context)
+		{
+			Context = &LocalContext;
+		}
+		else
+		{
+			// Cached result?
+			FImageDesc* PtrValue = Context->m_results.Find(this);
+			if (PtrValue)
+			{
+				return *PtrValue;
+			}
+		}
+
+		// In a switch we cannot guarantee the size and format.
+		// We check all the options, and if they are the same we return that.
+		// Otherwise, we return a descriptor with empty fields in the conflicting ones, size or format.
+		// In some places this will force re-formatting of the image.
+		// The code optimiser will take care then of moving the format operations down to each
+		// branch and remove the unnecessary ones.
+		FImageDesc Candidate;
+
+		bool bSameSize = true;
+		bool bSameFormat = true;
+		bool bSameLods = true;
+		bool bFirst = true;
+
+		if (def)
+		{
+			FImageDesc ChildDesc = def->GetImageDesc(bReturnBestOption, Context);
+			Candidate = ChildDesc;
+			bFirst = false;
+		}
+
+		for (int32 CaseIndex = 0; CaseIndex < cases.Num(); ++CaseIndex)
+		{
+			if (cases[CaseIndex].branch)
+			{
+				FImageDesc ChildDesc = cases[CaseIndex].branch->GetImageDesc(bReturnBestOption, Context);
+				if (bFirst)
+				{
+					Candidate = ChildDesc;
+					bFirst = false;
+				}
+				else
+				{
+					bSameSize = bSameSize && (Candidate.m_size == ChildDesc.m_size);
+					bSameFormat = bSameFormat && (Candidate.m_format == ChildDesc.m_format);
+					bSameLods = bSameLods && (Candidate.m_lods == ChildDesc.m_lods);
+
+					if (bReturnBestOption)
+					{
+						Candidate.m_format = GetMostGenericFormat(Candidate.m_format, ChildDesc.m_format);
+					}
+				}
+			}
+		}
+
+		Result = Candidate;
+
+		// In case of ReturnBestOption the first valid case will be used to determine size and lods.
+		// Format will be the most generic from all cases.
+		if (!bSameFormat && !bReturnBestOption)
+		{
+			Result.m_format = EImageFormat::IF_NONE;
+		}
+
+		if (!bSameSize && !bReturnBestOption)
+		{
+			Result.m_size = FImageSize(0, 0);
+		}
+
+		if (!bSameLods && !bReturnBestOption)
+		{
+			Result.m_lods = 0;
+		}
+
+		// Cache the result
+		if (Context)
+		{
+			Context->m_results.Add(this, Result);
+		}
+
+		return Result;
+	}
+
+
+	//-------------------------------------------------------------------------------------------------
+	void ASTOpSwitch::GetBlockLayoutSize(uint64 BlockId, int32* pBlockX, int32* pBlockY, FBlockLayoutSizeCache* cache)
+	{
+		switch (type)
+		{
+		case OP_TYPE::LA_SWITCH:
+		{
+			Ptr<ASTOp> child = GetFirstValidValue();
+			if (!child)
+			{
+				child = def.child();
+			}
+
+			if (child)
+			{
+				child->GetBlockLayoutSizeCached(BlockId, pBlockX, pBlockY, cache);
+			}
+			else
+			{
+				*pBlockX = 0;
+				*pBlockY = 0;
+			}
+			break;
+		}
+
+		default:
+			check(false);
+			break;
+		}
+	}
+
+
+	//-------------------------------------------------------------------------------------------------
+	void ASTOpSwitch::GetLayoutBlockSize(int32* pBlockX, int32* pBlockY)
+	{
+		switch (type)
+		{
+		case OP_TYPE::IM_SWITCH:
+		{
+			Ptr<ASTOp> child = GetFirstValidValue();
+			if (!child)
+			{
+				child = def.child();
+			}
+
+			if (child)
+			{
+				child->GetLayoutBlockSize(pBlockX, pBlockY);
+			}
+			else
+			{
+				checkf(false, TEXT("Image switch had no options."));
+			}
+			break;
+		}
+
+		default:
+			checkf(false, TEXT("Instruction not supported"));
+		}
+	}
+
+
+	//-------------------------------------------------------------------------------------------------
+	bool ASTOpSwitch::GetNonBlackRect(FImageRect& maskUsage) const
+	{
+		if (type == OP_TYPE::IM_SWITCH)
+		{
+			FImageRect local;
+			bool localValid = false;
+			if (def)
+			{
+				localValid = def->GetNonBlackRect(local);
+				if (!localValid)
+				{
+					return false;
+				}
+			}
+
+			for (const auto& c : cases)
+			{
+				if (c.branch)
+				{
+					FImageRect branchRect;
+					bool validBranch = c.branch->GetNonBlackRect(branchRect);
+					if (validBranch)
+					{
+						if (localValid)
+						{
+							local.Bound(branchRect);
+						}
+						else
+						{
+							local = branchRect;
+							localValid = true;
+						}
+					}
+					else
+					{
+						return false;
+					}
+				}
+			}
+
+			if (localValid)
+			{
+				maskUsage = local;
+				return true;
+			}
+		}
+
+		return false;
+	}
+
+
+	//-------------------------------------------------------------------------------------------------
+	bool ASTOpSwitch::IsImagePlainConstant(FVector4f&) const
+	{
+		// We could check if every option is plain and exactly the same colour, but probably it is
+		// not worth.
+		return false;
+	}
+
+
+	//-------------------------------------------------------------------------------------------------
+	mu::Ptr<ASTOp> ASTOpSwitch::OptimiseSemantic(const FModelOptimizationOptions&, int32 Pass) const
+	{
+		// Constant condition?
+		if (variable->GetOpType() == OP_TYPE::NU_CONSTANT)
+		{
+			Ptr<ASTOp> Branch = def.child();
+
+			auto typedCondition = static_cast<const ASTOpFixed*>(variable.child().get());
+			for (int32 o = 0; o < cases.Num(); ++o)
+			{
+				if (cases[o].branch &&
+					typedCondition->op.args.IntConstant.value
+					==
+					(int)cases[o].condition)
+				{
+					Branch = cases[o].branch.child();
+					break;
+				}
+			}
+
+			return Branch;
+		}
+
+		else if (variable->GetOpType() == OP_TYPE::NU_PARAMETER)
+		{
+			// If all the branches for the possible values are the same op remove the instruction
+			const ASTOpParameter* ParamOp = static_cast<const ASTOpParameter*>(variable.child().get());
+			check(ParamOp);
+			check(!ParamOp->parameter.m_possibleValues.IsEmpty());
+
+			bool bFirstValue = true;
+			bool bAllSame = true;
+			Ptr<ASTOp> SameBranch = nullptr;
+			for (const FParameterDesc::FIntValueDesc& Value : ParamOp->parameter.m_possibleValues)
+			{
+				// Look for the switch branch it would take
+				Ptr<ASTOp> Branch = def.child();
+				for (const FCase& Case : cases)
+				{
+					if (Case.condition == Value.m_value)
+					{
+						Branch = Case.branch.child();
+						break;
+					}
+				}
+
+				if (bFirstValue)
+				{
+					bFirstValue = false;
+					SameBranch = Branch;
+				}
+				else
+				{
+					if (SameBranch != Branch)
+					{
+						bAllSame = false;
+						SameBranch = nullptr;
+						break;
+					}
+				}
+			}
+
+	        if (bAllSame)
+	        {
+				return SameBranch;
+	        }
+		}
+
+		// Ad-hoc logic optimization: check if all code paths leading to this operation have a switch with the same variable
+		// and the option on those switches for the path that connects to this one is always the same. In that case, we can 
+		// remove this switch and replace it by the value it has for that option. 
+		// This is something the generic logic optimizer should do whan re-enabled.
+		{
+			// List of parent operations that we have visited, and the child we have visited them from.
+			TSet<TTuple<const ASTOp*, const ASTOp*>> Visited;
+			Visited.Reserve(64);
+
+			// First is parent, second is what child we are reaching the parent from. This is necessary to find out what 
+			// switch branch we reach the parent from, if it is a switch.
+			TArray< TTuple<const ASTOp*, const ASTOp*>, TInlineAllocator<16>> Pending;
+			ForEachParent([this,&Pending](ASTOp* Parent)
+				{
+					Pending.Add({ Parent,this});
+				});
+
+			bool bAllPathsHaveMatchingSwitch = true;
+
+			// Switch option value of all parent compatible switches (if any)
+			int32 MatchingSwitchOption = -1;
+
+			while (!Pending.IsEmpty() && bAllPathsHaveMatchingSwitch)
+			{
+				TTuple<const ASTOp*, const ASTOp*> ParentPair = Pending.Pop();
+				bool bAlreadyVisited = false;
+				Visited.Add(ParentPair, &bAlreadyVisited);
+
+				if (!bAlreadyVisited)
+				{
+					const ASTOp* Parent = ParentPair.Get<0>();
+					const ASTOp* ParentChild = ParentPair.Get<1>();
+
+					bool bIsMatchingSwitch = false;
+
+					// TODO: Probably it could be a any switch, it doesn't need to be of the same type.
+					if (Parent->GetOpType() == GetOpType())
+					{
+						const ASTOpSwitch* ParentSwitch = static_cast<const ASTOpSwitch*>(Parent);
+						check(ParentSwitch);
+
+						// To be compatible the switch must be on the same variable
+						if (ParentSwitch->variable==variable)
+						{
+							bIsMatchingSwitch = true;
+							
+							// Find what switch option we are reaching it from
+							bool bIsSingleOption = true;
+							int OptionIndex = -1;
+							for (int32 CaseIndex = 0; CaseIndex < ParentSwitch->cases.Num(); ++CaseIndex)
+							{
+								if (ParentSwitch->cases[CaseIndex].branch.child().get() == ParentChild)
+								{
+									if (OptionIndex != -1)
+									{
+										// This means the same child is connected to more than one switch options
+										// so we cannot optimize.
+										// \TODO: We could if we track a "set of options" for all switches instead of just one.
+										bIsSingleOption = false;
+										break;
+									}
+									else
+									{
+										OptionIndex = CaseIndex;
+									}
+								}
+							}
+
+							// If we did reach it from one single option
+							if (bIsSingleOption && OptionIndex!=-1)
+							{
+								if (MatchingSwitchOption<0)
+								{
+									MatchingSwitchOption = ParentSwitch->cases[OptionIndex].condition;
+								}
+								else if (MatchingSwitchOption!= ParentSwitch->cases[OptionIndex].condition)
+								{
+									bAllPathsHaveMatchingSwitch = false;
+								}
+							}
+						}
+					}
+					
+					if (!bIsMatchingSwitch)
+					{
+						// If it has no parents, then the optimization cannot be applied
+						bool bHasParent = false;
+						Parent->ForEachParent([&bHasParent,this,&Pending,Parent](ASTOp* ParentParent)
+							{
+								Pending.Add({ ParentParent,Parent });
+								bHasParent = true;
+							});
+
+						if (!bHasParent)
+						{
+							// We reached a root without a matching switch along the path.
+							bAllPathsHaveMatchingSwitch = false;
+						}
+					}
+				}
+			}
+
+			if (bAllPathsHaveMatchingSwitch && MatchingSwitchOption>=0)
+			{
+				// We can remove this switch, all paths leading to it have the same condition for this switches variable.
+				return FindBranch(MatchingSwitchOption);
+			}
+
+		}
+
+		return nullptr;
+	}
+
+
+	//-------------------------------------------------------------------------------------------------
+	mu::Ptr<ImageSizeExpression> ASTOpSwitch::GetImageSizeExpression() const
+	{
+		Ptr<ImageSizeExpression> pRes = new ImageSizeExpression;
+
+		bool first = true;
+		for (const FCase& c : cases)
+		{
+			if (c.branch)
+			{
+				if (first)
+				{
+					pRes = c.branch->GetImageSizeExpression();
+				}
+				else
+				{
+					Ptr<ImageSizeExpression> pOther = c.branch->GetImageSizeExpression();
+					if (!(*pOther == *pRes))
+					{
+						pRes->type = ImageSizeExpression::ISET_UNKNOWN;
+						break;
+					}
+				}
+			}
+		}
+
+		return pRes;
+	}
+
+
+	FSourceDataDescriptor ASTOpSwitch::GetSourceDataDescriptor(FGetSourceDataDescriptorContext* Context) const
+	{
+		// Cache management
+		TUniquePtr<FGetSourceDataDescriptorContext> LocalContext;
+		if (!Context)
+		{
+			LocalContext.Reset(new FGetSourceDataDescriptorContext);
+			Context = LocalContext.Get();
+		}
+
+		FSourceDataDescriptor* Found = Context->Cache.Find(this);
+		if (Found)
+		{
+			return *Found;
+		}
+
+		// Not cached: calculate
+		FSourceDataDescriptor Result;
+
+		for (const FCase& Case : cases)
+		{
+			if (Case.branch)
+			{
+				FSourceDataDescriptor SourceDesc = Case.branch->GetSourceDataDescriptor(Context);
+				Result.CombineWith(SourceDesc);
+			}
+		}
+
+		Context->Cache.Add(this, Result);
+
+		return Result;
+	}
+
+
+}
