[mutable] Moved the Mutable plugin out of Experimental status into Beta.

diff --git a/Engine/Plugins/Mutable/Source/MutableTools/Private/MuT/CodeGenerator.cpp b/Engine/Plugins/Mutable/Source/MutableTools/Private/MuT/CodeGenerator.cpp
new file mode 100644
--- /dev/null
+++ b/Engine/Plugins/Mutable/Source/MutableTools/Private/MuT/CodeGenerator.cpp
@@ -0,0 +1,2643 @@
+// Copyright Epic Games, Inc. All Rights Reserved.
+
+#include "MuT/CodeGenerator.h"
+
+#include "Containers/Array.h"
+#include "Logging/LogCategory.h"
+#include "Logging/LogMacros.h"
+#include "Math/IntPoint.h"
+#include "Math/UnrealMathUtility.h"
+#include "MuR/ImagePrivate.h"
+#include "MuR/Layout.h"
+#include "MuR/MutableTrace.h"
+#include "MuR/Operations.h"
+#include "MuR/ParametersPrivate.h"
+#include "MuR/Platform.h"
+#include "MuT/ASTOpAddLOD.h"
+#include "MuT/ASTOpAddExtensionData.h"
+#include "MuT/ASTOpConditional.h"
+#include "MuT/ASTOpConstantBool.h"
+#include "MuT/ASTOpConstantResource.h"
+#include "MuT/ASTOpImageCompose.h"
+#include "MuT/ASTOpImageMipmap.h"
+#include "MuT/ASTOpImagePixelFormat.h"
+#include "MuT/ASTOpImageSwizzle.h"
+#include "MuT/ASTOpImageLayer.h"
+#include "MuT/ASTOpImageLayerColor.h"
+#include "MuT/ASTOpImagePatch.h"
+#include "MuT/ASTOpImageCrop.h"
+#include "MuT/ASTOpInstanceAdd.h"
+#include "MuT/ASTOpMeshBindShape.h"
+#include "MuT/ASTOpMeshClipDeform.h"
+#include "MuT/ASTOpMeshClipMorphPlane.h"
+#include "MuT/ASTOpMeshMaskClipMesh.h"
+#include "MuT/ASTOpMeshMaskClipUVMask.h"
+#include "MuT/ASTOpMeshRemoveMask.h"
+#include "MuT/ASTOpMeshDifference.h"
+#include "MuT/ASTOpMeshMorph.h"
+#include "MuT/ASTOpMeshOptimizeSkinning.h"
+#include "MuT/ASTOpMeshExtractLayoutBlocks.h"
+#include "MuT/ASTOpParameter.h"
+#include "MuT/ASTOpLayoutRemoveBlocks.h"
+#include "MuT/ASTOpLayoutFromMesh.h"
+#include "MuT/ASTOpLayoutMerge.h"
+#include "MuT/ASTOpLayoutPack.h"
+#include "MuT/CodeGenerator_SecondPass.h"
+#include "MuT/CodeOptimiser.h"
+#include "MuT/CompilerPrivate.h"
+#include "MuT/ErrorLogPrivate.h"
+#include "MuT/NodeColour.h"
+#include "MuT/NodeColourConstant.h"
+#include "MuT/NodeComponent.h"
+#include "MuT/NodeComponentSwitch.h"
+#include "MuT/NodeComponentVariation.h"
+#include "MuT/NodeImage.h"
+#include "MuT/NodeImageFormat.h"
+#include "MuT/NodeImageFormatPrivate.h"
+#include "MuT/NodeImageMipmap.h"
+#include "MuT/NodeImageMipmapPrivate.h"
+#include "MuT/NodeImageSwizzlePrivate.h"
+#include "MuT/NodeMesh.h"
+#include "MuT/NodeMeshClipMorphPlane.h"
+#include "MuT/NodeMeshClipWithMesh.h"
+#include "MuT/NodeMeshConstantPrivate.h"
+#include "MuT/NodeMeshFormat.h"
+#include "MuT/NodeMeshFragment.h"
+#include "MuT/NodeMeshGeometryOperation.h"
+#include "MuT/NodeMeshInterpolate.h"
+#include "MuT/NodeMeshMorph.h"
+#include "MuT/NodeMeshReshape.h"
+#include "MuT/NodeModifier.h"
+#include "MuT/NodeModifierMeshClipDeform.h"
+#include "MuT/NodeModifierMeshClipMorphPlane.h"
+#include "MuT/NodeModifierMeshClipWithMesh.h"
+#include "MuT/NodeModifierMeshClipWithUVMask.h"
+#include "MuT/NodeModifierSurfaceEdit.h"
+#include "MuT/NodeObject.h"
+#include "MuT/NodeObjectGroupPrivate.h"
+#include "MuT/NodeObjectNew.h"
+#include "MuT/NodePrivate.h"
+#include "MuT/NodeRange.h"
+#include "MuT/NodeRangeFromScalar.h"
+#include "MuT/NodeScalar.h"
+#include "MuT/NodeScalarConstant.h"
+#include "MuT/NodeSurface.h"
+#include "MuT/NodeSurfaceNew.h"
+#include "MuT/NodeSurfaceVariation.h"
+#include "MuT/NodeSurfaceSwitch.h"
+#include "MuT/TablePrivate.h"
+#include "Trace/Detail/Channel.h"
+
+
+namespace mu
+{
+
+	//---------------------------------------------------------------------------------------------
+	//---------------------------------------------------------------------------------------------
+	//---------------------------------------------------------------------------------------------
+	CodeGenerator::CodeGenerator(CompilerOptions::Private* options)
+	{
+		CompilerOptions = options;
+
+		// Create the message log
+		ErrorLog = new mu::ErrorLog;
+
+		// Add the parent at the top of the hierarchy
+		CurrentParents.Add(FParentKey());
+	}
+
+	//---------------------------------------------------------------------------------------------
+	void CodeGenerator::GenerateRoot(const Ptr<const Node> pNode)
+	{
+		MUTABLE_CPUPROFILER_SCOPE(Generate);
+
+		// First pass
+		FirstPass.Generate(ErrorLog, pNode.get(), CompilerOptions->bIgnoreStates, this);
+
+		// Second pass
+		SecondPassGenerator SecondPass(&FirstPass, CompilerOptions);
+		bool bSuccess = SecondPass.Generate(ErrorLog, pNode.get());
+		if (!bSuccess)
+		{
+			return;
+		}
+
+		// Main pass for each state
+		{
+			MUTABLE_CPUPROFILER_SCOPE(MainPass);
+
+            int32 CurrentStateIndex = 0;
+			for (const TPair<FObjectState, const Node*>& s : FirstPass.States)
+			{
+				MUTABLE_CPUPROFILER_SCOPE(MainPassState);
+
+				FGenericGenerationOptions Options;
+				Options.State = CurrentStateIndex;
+
+				Ptr<ASTOp> stateRoot = Generate_Generic(pNode, Options);
+				States.Emplace(s.Key, stateRoot);
+
+				AdditionalComponents.Empty();
+
+				++CurrentStateIndex;
+			}
+		}
+	}
+
+
+	Ptr<ASTOp> CodeGenerator::Generate_Generic(const Ptr<const Node> pNode, const FGenericGenerationOptions& Options )
+	{
+		if (!pNode)
+		{
+			return nullptr;
+		}
+
+        // Type-specific generation
+		if (pNode->GetType()->IsA(NodeScalar::GetStaticType()))
+		{
+			const NodeScalar* ScalarNode = static_cast<const NodeScalar*>(pNode.get());
+			FScalarGenerationResult ScalarResult;
+			GenerateScalar(ScalarResult, Options, ScalarNode);
+			return ScalarResult.op;
+		}
+
+		else if (pNode->GetType()->IsA(NodeColour::GetStaticType()))
+		{
+			const NodeColour* ColorNode = static_cast<const NodeColour*>(pNode.get());
+			FColorGenerationResult Result;
+			GenerateColor(Result, Options, ColorNode);
+			return Result.op;
+		}
+
+		else if (pNode->GetType()->IsA(NodeProjector::GetStaticType()))
+		{
+			const NodeProjector* projNode = static_cast<const NodeProjector*>(pNode.get());
+			FProjectorGenerationResult ProjResult;
+			GenerateProjector(ProjResult, Options, projNode);
+			return ProjResult.op;
+		}
+
+		else if (pNode->GetType()->IsA(NodeSurfaceNew::GetStaticType()))
+		{
+			const NodeSurfaceNew* surfNode = static_cast<const NodeSurfaceNew*>(pNode.get());
+
+			// This happens only if we generate a node graph that has a NodeSurfaceNew at the root.
+			FSurfaceGenerationResult surfResult;
+			FSurfaceGenerationOptions SurfaceOptions(Options);
+			GenerateSurface(surfResult, SurfaceOptions, surfNode);
+			return surfResult.surfaceOp;
+		}
+
+		else if (pNode->GetType()->IsA(NodeSurfaceVariation::GetStaticType()))
+		{
+			// This happens only if we generate a node graph that has a NodeSurfaceVariation at the root.
+			return nullptr;
+		}
+
+		else if (pNode->GetType()->IsA(NodeSurfaceSwitch::GetStaticType()))
+		{
+			// This happens only if we generate a node graph that has a NodeSurfaceSwitch at the root.
+			return nullptr;
+		}
+
+		else if (pNode->GetType()->IsA(NodeModifier::GetStaticType()))
+		{
+			// This happens only if we generate a node graph that has a modifier at the root.
+			return nullptr;
+		}
+
+		else if (pNode->GetType()->IsA(NodeComponent::GetStaticType()))
+		{
+			const NodeComponent* ComponentNode = static_cast<const NodeComponent*>(pNode.get());
+			FComponentGenerationOptions ComponentOptions( Options, nullptr );
+			FGenericGenerationResult Result;
+			GenerateComponent(ComponentOptions, Result, ComponentNode);
+			return Result.op;
+		}
+
+
+		Ptr<ASTOp> ResultOp;
+
+		// See if it was already generated
+
+		FGeneratedCacheKey Key;
+		Key.Node = pNode;
+		Key.Options = Options;
+		FGeneratedGenericNodesMap::ValueType* it = GeneratedGenericNodes.Find(Key);
+		if (it)
+		{
+			ResultOp = it->op;
+		}
+		else
+		{
+			FGenericGenerationResult Result;
+
+			// Generate for each different type of node
+			if (pNode->GetType() == NodeObjectNew::GetStaticType())
+			{
+				Generate_ObjectNew(Options, Result, static_cast<const NodeObjectNew*>(pNode.get()));
+			}
+			else if (pNode->GetType()==NodeObjectGroup::GetStaticType())
+			{
+				Generate_ObjectGroup(Options, Result, static_cast<const NodeObjectGroup*>(pNode.get()));
+			}
+			else
+			{
+				check(false);
+			}
+
+			ResultOp = Result.op;
+			GeneratedGenericNodes.Add(Key, Result);
+		}
+
+		// debug: expensive check of all code generation
+//        if (ResultOp)
+//        {
+//            ASTOpList roots;
+//            roots.push_back(ResultOp);
+//            ASTOp::FullAssert(roots);
+//        }
+
+		return ResultOp;
+	}
+
+
+	void CodeGenerator::GenerateRange(FRangeGenerationResult& Result, const FGenericGenerationOptions& Options, Ptr<const NodeRange> Untyped)
+	{
+		if (!Untyped)
+		{
+			Result = FRangeGenerationResult();
+			return;
+		}
+
+		// See if it was already generated
+		FGeneratedCacheKey Key;
+		Key.Node = Untyped;
+		Key.Options = Options;
+		FGeneratedRangeMap::ValueType* it = GeneratedRanges.Find(Key);
+		if (it)
+		{
+			Result = *it;
+			return;
+		}
+
+		// Generate for each different type of node
+		if (Untyped->GetType()==NodeRangeFromScalar::GetStaticType())
+		{
+			const NodeRangeFromScalar* FromScalar = static_cast<const NodeRangeFromScalar*>(Untyped.get());
+
+			Result = FRangeGenerationResult();
+			Result.rangeName = FromScalar->GetName();
+
+			FScalarGenerationResult ChildResult;
+			GenerateScalar(ChildResult, Options, FromScalar->GetSize());
+			Result.sizeOp = ChildResult.op;
+		}
+		else
+		{
+			check(false);
+		}
+
+
+		// Cache the result
+		GeneratedRanges.Add(Key, Result);
+	}
+
+
+	Ptr<ASTOp> CodeGenerator::GenerateTableVariable( Ptr<const Node> InNode, const FTableCacheKey& CacheKey, bool bAddNoneOption, const FString& DefaultRowName)
+	{
+		Ptr<ASTOp> result;
+
+        FParameterDesc param;
+        param.m_name = CacheKey.ParameterName;
+
+        if (param.m_name.Len() == 0)
+        {
+            param.m_name = CacheKey.Table->GetName();
+        }
+        
+		param.m_type = PARAMETER_TYPE::T_INT;
+        param.m_defaultValue.Set<ParamIntType>(0);
+
+		// Add the possible values
+		{
+			// See if there is a string column. If there is one, we will use it as names for the
+			// options. Only the first string column will be used.
+			int32 nameCol = -1;
+			int32 cols = CacheKey.Table->GetPrivate()->Columns.Num();
+			for (int32 c = 0; c < cols && nameCol < 0; ++c)
+			{
+				if (CacheKey.Table->GetPrivate()->Columns[c].Type == ETableColumnType::String)
+				{
+					nameCol = c;
+				}
+			}
+
+			if (bAddNoneOption)
+			{
+				FParameterDesc::FIntValueDesc nullValue;
+				nullValue.m_value = -1;
+				nullValue.m_name = "None";
+				param.m_possibleValues.Add(nullValue);
+				param.m_defaultValue.Set<ParamIntType>(nullValue.m_value);
+			}
+
+			// Add every row
+			int32 RowCount = CacheKey.Table->GetPrivate()->Rows.Num();
+			for (int32 i = 0; i < RowCount; ++i)
+			{
+				FParameterDesc::FIntValueDesc value;
+				value.m_value = (int16)CacheKey.Table->GetPrivate()->Rows[i].Id;
+
+				if (nameCol > -1)
+				{
+					value.m_name = CacheKey.Table->GetPrivate()->Rows[i].Values[nameCol].String;
+				}
+
+				param.m_possibleValues.Add(value);
+
+				// Set the first row as the default one (if there is none option)
+				if (i == 0 && !bAddNoneOption)
+				{
+					param.m_defaultValue.Set<ParamIntType>(value.m_value);
+				}
+				
+				// Set the selected row as default (if exists)
+				if (value.m_name == DefaultRowName)
+				{
+					param.m_defaultValue.Set<ParamIntType>(value.m_value);
+				}
+            }
+        }
+
+		Ptr<ASTOpParameter> op = new ASTOpParameter();
+		op->type = OP_TYPE::NU_PARAMETER;
+		op->parameter = param;
+
+		FirstPass.ParameterNodes.Add(InNode, op);
+
+		return op;
+	}
+
+
+	Ptr<const Layout> CodeGenerator::GenerateLayout(Ptr<const NodeLayout> SourceLayout, uint32 MeshIDPrefix)
+	{
+		Ptr<const Layout>* it = GeneratedLayouts.Find({ SourceLayout,MeshIDPrefix });
+
+		if (it)
+		{
+			return *it;
+		}
+
+		Ptr<Layout> GeneratedLayout = new Layout;
+		GeneratedLayout->Size = SourceLayout->Size;
+		GeneratedLayout->MaxSize = SourceLayout->MaxSize;
+		GeneratedLayout->Strategy = SourceLayout->Strategy;
+		GeneratedLayout->ReductionMethod = SourceLayout->ReductionMethod;
+
+		const int32 BlockCount = SourceLayout->Blocks.Num();
+		GeneratedLayout->Blocks.SetNum(BlockCount);
+		for (int32 BlockIndex = 0; BlockIndex < BlockCount; ++BlockIndex)
+		{
+			const FSourceLayoutBlock& From = SourceLayout->Blocks[BlockIndex];
+			FLayoutBlock& To = GeneratedLayout->Blocks[BlockIndex];
+			To.Min = From.Min;
+			To.Size = From.Size;
+			To.Priority = From.Priority;
+			To.bReduceBothAxes = From.bReduceBothAxes;
+			To.bReduceByTwo = From.bReduceByTwo;
+
+			// Assign unique ids to each layout block
+			uint64 Id = uint64(MeshIDPrefix) << 32 | uint64(BlockIndex);
+			To.Id = Id;
+		}
+
+		check(GeneratedLayout->Blocks.IsEmpty() || GeneratedLayout->Blocks[0].Id != FLayoutBlock::InvalidBlockId);
+		GeneratedLayouts.Add({ SourceLayout,MeshIDPrefix }, GeneratedLayout);
+
+		return GeneratedLayout;
+	}
+
+
+	Ptr<ASTOp> CodeGenerator::GenerateImageBlockPatch(Ptr<ASTOp> InBlockOp,
+		const NodeModifierSurfaceEdit::FTexture& Patch,
+		Ptr<Image> PatchMask,
+		Ptr<ASTOp> conditionAd,
+		const FImageGenerationOptions& ImageOptions )
+	{
+		// Blend operation
+		Ptr<ASTOp> FinalOp;
+		{
+			MUTABLE_CPUPROFILER_SCOPE(PatchBlend);
+
+			Ptr<ASTOpImageLayer> LayerOp = new ASTOpImageLayer();
+			LayerOp->blendType = Patch.PatchBlendType;
+			LayerOp->base = InBlockOp;
+
+			// When we patch from edit nodes, we want to apply it to all the channels.
+			// \todo: since we can choose the patch function, maybe we want to be able to select this as well.
+			LayerOp->Flags = Patch.bPatchApplyToAlpha ? OP::ImageLayerArgs::F_APPLY_TO_ALPHA : 0;
+
+			NodeImage* ImageNode = Patch.PatchImage.get();
+			Ptr<ASTOp> BlendOp;
+			if (ImageNode)
+			{
+				FImageGenerationResult BlendResult;
+				GenerateImage(ImageOptions, BlendResult, ImageNode);
+				BlendOp = BlendResult.op;
+			}
+			else
+			{
+				BlendOp = GenerateMissingImageCode(TEXT("Patch top image"), EImageFormat::IF_RGB_UBYTE, nullptr, ImageOptions);
+			}
+			BlendOp = GenerateImageFormat(BlendOp, InBlockOp->GetImageDesc().m_format);
+			BlendOp = GenerateImageSize(BlendOp, ImageOptions.RectSize);
+			LayerOp->blend = BlendOp;
+
+			// Create the rect mask constant
+			Ptr<ASTOp> RectConstantOp;
+			{
+				Ptr<NodeImageConstant> pNode = new NodeImageConstant();
+				pNode->SetValue(PatchMask.get());
+
+				FImageGenerationOptions ConstantOptions;
+				FImageGenerationResult ConstantResult;
+				GenerateImage(ConstantOptions, ConstantResult, pNode);
+				RectConstantOp = ConstantResult.op;
+			}
+
+			NodeImage* MaskNode = Patch.PatchMask.get();
+			Ptr<ASTOp> MaskOp;
+			if (MaskNode)
+			{
+				// Combine the block rect mask with the user provided mask.
+
+				FImageGenerationResult MaskResult;
+				GenerateImage(ImageOptions, MaskResult, MaskNode);
+				MaskOp = MaskResult.op;
+
+				Ptr<ASTOpImageLayer> PatchCombineOp = new ASTOpImageLayer;
+				PatchCombineOp->base = MaskOp;
+				PatchCombineOp->blend = RectConstantOp;
+				PatchCombineOp->blendType = EBlendType::BT_MULTIPLY;
+				MaskOp = PatchCombineOp;
+			}
+			else
+			{
+				MaskOp = RectConstantOp;
+			}
+			MaskOp = GenerateImageFormat(MaskOp, EImageFormat::IF_L_UBYTE);
+			MaskOp = GenerateImageSize(MaskOp, ImageOptions.RectSize);
+			LayerOp->mask = MaskOp;
+
+			FinalOp = LayerOp;
+		}
+
+		// Condition to enable this patch
+		if (conditionAd)
+		{
+			Ptr<ASTOp> conditionalAd;
+			{
+				Ptr<ASTOpConditional> op = new ASTOpConditional();
+				op->type = OP_TYPE::IM_CONDITIONAL;
+				op->no = InBlockOp;
+				op->yes = FinalOp;
+				op->condition = conditionAd;
+				conditionalAd = op;
+			}
+
+			FinalOp = conditionalAd;
+		}
+
+		return FinalOp;
+	}
+
+
+
+	//---------------------------------------------------------------------------------------------
+	void CodeGenerator::GenerateComponent(const FComponentGenerationOptions& InOptions, FGenericGenerationResult& OutResult, const NodeComponent* InUntypedNode)
+	{
+		if (!InUntypedNode)
+		{
+			OutResult = FGenericGenerationResult();
+			return;
+		}
+
+		// See if it was already generated
+		FGeneratedComponentCacheKey Key;
+		Key.Node = InUntypedNode;
+		Key.Options = InOptions;
+		GeneratedComponentMap::ValueType* it = GeneratedComponents.Find(Key);
+		if (it)
+		{
+			OutResult = *it;
+			return;
+		}
+
+		// Generate for each different type of node
+		const FNodeType* Type = InUntypedNode->GetType();
+		if (Type == NodeComponentNew::GetStaticType())
+		{
+			GenerateComponent_New(InOptions, OutResult, static_cast<const NodeComponentNew*>(InUntypedNode));
+		}
+		else if (Type == NodeComponentEdit::GetStaticType())
+		{
+			// Nothing to do because it is all preprocessed in the first code generator stage
+			//GenerateComponent_Edit(InOptions, OutResult, static_cast<const NodeComponentEdit*>(InUntypedNode));
+			OutResult.op = InOptions.BaseInstance;
+		}
+		else if (Type == NodeComponentSwitch::GetStaticType())
+		{
+			GenerateComponent_Switch(InOptions, OutResult, static_cast<const NodeComponentSwitch*>(InUntypedNode));
+		}
+		else if (Type == NodeComponentVariation::GetStaticType())
+		{
+			GenerateComponent_Variation(InOptions, OutResult, static_cast<const NodeComponentVariation*>(InUntypedNode));
+		}
+		else
+		{
+			check(false);
+		}
+
+		// Cache the result
+		GeneratedComponents.Add(Key, OutResult);
+	}
+
+
+	void CodeGenerator::GenerateComponent_New(const FComponentGenerationOptions& Options, FGenericGenerationResult& Result, const NodeComponentNew* InNode)
+	{
+		const NodeComponentNew& node = *InNode;
+
+		// Create the expression for each component in this object
+		Ptr<ASTOpAddLOD> LODsOp = new ASTOpAddLOD();
+
+		for (int32 LODIndex = 0; LODIndex < node.LODs.Num(); ++LODIndex)
+		{
+			if (const NodeLOD* LODNode = node.LODs[LODIndex].get())
+			{
+				CurrentParents.Last().Lod = LODIndex;
+
+				FLODGenerationOptions LODOptions(Options, LODIndex, InNode );
+				FGenericGenerationResult LODResult;
+				Generate_LOD(LODOptions, LODResult, LODNode);
+
+				LODsOp->lods.Emplace(LODsOp, LODResult.op);
+			}
+		}
+
+		Ptr<ASTOpInstanceAdd> InstanceOp = new ASTOpInstanceAdd();
+		InstanceOp->type = OP_TYPE::IN_ADDCOMPONENT;
+		InstanceOp->instance = Options.BaseInstance;
+		InstanceOp->value = LODsOp;
+		InstanceOp->id = node.Id;
+
+		Result.op = InstanceOp;
+
+		// Add a conditional if this component has conditions
+		for (const FirstPassGenerator::FComponent& Component: FirstPass.Components)
+		{
+			if (Component.Component != InNode)
+			{
+				continue;
+			}
+
+			if (Component.ComponentCondition || Component.ObjectCondition)
+			{
+				// TODO: This could be done earlier?
+				Ptr<ASTOpFixed> ConditionOp = new ASTOpFixed();
+				ConditionOp->op.type = OP_TYPE::BO_AND;
+				ConditionOp->SetChild(ConditionOp->op.args.BoolBinary.a, Component.ObjectCondition);
+				ConditionOp->SetChild(ConditionOp->op.args.BoolBinary.b, Component.ComponentCondition);
+
+				Ptr<ASTOpConditional> IfOp = new ASTOpConditional();
+				IfOp->type = OP_TYPE::IN_CONDITIONAL;
+				IfOp->no = Options.BaseInstance;
+				IfOp->yes = Result.op;
+				IfOp->condition = ConditionOp;
+
+				Result.op = IfOp;
+			}
+		}
+	}
+
+
+	void CodeGenerator::GenerateComponent_Switch(const FComponentGenerationOptions& Options, FGenericGenerationResult& Result, const NodeComponentSwitch* Node)
+	{
+		MUTABLE_CPUPROFILER_SCOPE(NodeComponentSwitch);
+
+		if (Node->Options.Num() == 0)
+		{
+			// No options in the switch!
+			Result.op = Options.BaseInstance;
+			return;
+		}
+
+		Ptr<ASTOpSwitch> Op = new ASTOpSwitch();
+		Op->type = OP_TYPE::IN_SWITCH;
+
+		// Variable value
+		if (Node->Parameter)
+		{
+			Op->variable = Generate_Generic(Node->Parameter.get(), Options);
+		}
+		else
+		{
+			// This argument is required
+			Op->variable = GenerateMissingScalarCode(TEXT("Switch variable"), 0.0f, Node->GetMessageContext());
+		}
+
+		// Options
+		for (int32 OptionIndex = 0; OptionIndex < Node->Options.Num(); ++OptionIndex)
+		{
+			Ptr<ASTOp> Branch;
+
+			if (Node->Options[OptionIndex])
+			{
+				FGenericGenerationResult BaseResult;
+				GenerateComponent(Options, BaseResult, Node->Options[OptionIndex].get());
+				Branch = BaseResult.op;
+			}
+			else
+			{
+				// This argument is not required
+				Branch = Options.BaseInstance;
+			}
+
+			Op->cases.Emplace(OptionIndex, Op, Branch);
+		}
+
+		Result.op = Op;
+	}
+
+
+	void CodeGenerator::GenerateComponent_Variation(const FComponentGenerationOptions& Options, FGenericGenerationResult& Result, const NodeComponentVariation* Node)
+	{
+		Ptr<ASTOp> CurrentMeshOp = Options.BaseInstance;
+
+		// Default case
+		if (Node->DefaultComponent)
+		{
+			FGenericGenerationResult BranchResults;
+
+			GenerateComponent(Options, BranchResults, Node->DefaultComponent.get());
+			CurrentMeshOp = BranchResults.op;
+		}
+
+		// Process variations in reverse order, since conditionals are built bottom-up.
+		for (int32 VariationIndex = Node->Variations.Num() - 1; VariationIndex >= 0; --VariationIndex)
+		{
+			int32 TagIndex = -1;
+			const FString& Tag = Node->Variations[VariationIndex].Tag;
+			for (int32 i = 0; i < FirstPass.Tags.Num(); ++i)
+			{
+				if (FirstPass.Tags[i].Tag == Tag)
+				{
+					TagIndex = i;
+				}
+			}
+
+			if (TagIndex < 0)
+			{
+				ErrorLog->GetPrivate()->Add(
+					FString::Printf(TEXT("Unknown tag found in component variation [%s]."), *Tag),
+					ELMT_WARNING,
+					Node->GetMessageContext(),
+					ELMSB_UNKNOWN_TAG
+				);
+				continue;
+			}
+
+			Ptr<ASTOp> VariationMeshOp = Options.BaseInstance;
+			if (Node->Variations[VariationIndex].Component)
+			{
+				FGenericGenerationResult BranchResults;
+				GenerateComponent(Options, BranchResults, Node->Variations[VariationIndex].Component.get());
+
+				VariationMeshOp = BranchResults.op;
+			}
+
+			Ptr<ASTOpConditional> Conditional = new ASTOpConditional;
+			Conditional->type = OP_TYPE::IN_CONDITIONAL;
+			Conditional->no = CurrentMeshOp;
+			Conditional->yes = VariationMeshOp;
+			Conditional->condition = FirstPass.Tags[TagIndex].GenericCondition;
+
+			CurrentMeshOp = Conditional;
+		}
+
+		Result.op = CurrentMeshOp;
+	}
+
+
+	Ptr<ASTOp> CodeGenerator::ApplyTiling(Ptr<ASTOp> Source, UE::Math::TIntVector2<int32> Size, EImageFormat Format)
+	{
+		// For now always apply tiling
+		if (CompilerOptions->ImageTiling==0)
+		{
+			return Source;
+		}
+
+		int32 TileSize = CompilerOptions->ImageTiling;
+
+		int32 TilesX = FMath::DivideAndRoundUp<int32>(Size[0], TileSize);
+		int32 TilesY = FMath::DivideAndRoundUp<int32>(Size[1], TileSize);
+		if (TilesX * TilesY <= 2)
+		{
+			return Source;
+		}
+
+		Ptr<ASTOpFixed> BaseImage = new ASTOpFixed;
+		BaseImage->op.type = OP_TYPE::IM_PLAINCOLOUR;
+		BaseImage->op.args.ImagePlainColour.size[0] = Size[0];
+		BaseImage->op.args.ImagePlainColour.size[1] = Size[1];
+		BaseImage->op.args.ImagePlainColour.format = Format;
+		BaseImage->op.args.ImagePlainColour.LODs = 1;
+
+		Ptr<ASTOp> CurrentImage = BaseImage;
+
+		for (int32 Y = 0; Y < TilesY; ++Y)
+		{
+			for (int32 X = 0; X < TilesX; ++X)
+			{
+				int32 MinX = X * TileSize;
+				int32 MinY = Y * TileSize;
+				int32 TileSizeX = FMath::Min(TileSize, Size[0] - MinX);
+				int32 TileSizeY = FMath::Min(TileSize, Size[1] - MinY);
+
+				Ptr<ASTOpImageCrop> TileImage = new ASTOpImageCrop();
+				TileImage->Source = Source;
+				TileImage->Min[0] = MinX;
+				TileImage->Min[1] = MinY;
+				TileImage->Size[0] = TileSizeX;
+				TileImage->Size[1] = TileSizeY;
+
+				Ptr<ASTOpImagePatch> PatchedImage = new ASTOpImagePatch();
+				PatchedImage->base = CurrentImage;
+				PatchedImage->patch = TileImage;
+				PatchedImage->location[0] = MinX;
+				PatchedImage->location[1] = MinY;
+
+				CurrentImage = PatchedImage;
+			}
+		}
+
+		return CurrentImage;
+	}
+
+
+	Ptr<Image> CodeGenerator::GenerateImageBlockPatchMask(const NodeModifierSurfaceEdit::FTexture& Patch, FIntPoint GridSize, int32 BlockPixelsX, int32 BlockPixelsY, box<FIntVector2> RectInCells )
+	{
+		// Create a patching mask for the block
+		Ptr<Image> PatchMask;
+
+		FIntVector2 SourceTextureSize = { GridSize[0] * BlockPixelsX, GridSize[1] * BlockPixelsY };
+
+		FInt32Rect BlockRectInPixels;
+		BlockRectInPixels.Min = { RectInCells.min[0] * BlockPixelsX, RectInCells.min[1] * BlockPixelsY };
+		BlockRectInPixels.Max = { (RectInCells.min[0] + RectInCells.size[0]) * BlockPixelsX, (RectInCells.min[1] + RectInCells.size[1]) * BlockPixelsY };
+
+		for (const FBox2f& PatchRect : Patch.PatchBlocks)
+		{
+			// Does the patch rect intersects the current block at all?
+			FInt32Rect PatchRectInPixels;
+			PatchRectInPixels.Min = { int32(PatchRect.Min[0] * SourceTextureSize[0]), int32(PatchRect.Min[1] * SourceTextureSize[1]) };
+			PatchRectInPixels.Max = { int32(PatchRect.Max[0] * SourceTextureSize[0]), int32(PatchRect.Max[1] * SourceTextureSize[1]) };
+
+			FInt32Rect BlockPatchRect = PatchRectInPixels;
+			BlockPatchRect.Clip(BlockRectInPixels);
+
+			if (BlockPatchRect.Area() > 0)
+			{
+				FInt32Point BlockSize = BlockRectInPixels.Size();
+				if (!PatchMask)
+				{
+					PatchMask = new mu::Image(BlockSize[0], BlockSize[1], 1, mu::EImageFormat::IF_L_UBYTE, mu::EInitializationType::Black);
+				}
+
+				uint8* Pixels = PatchMask->GetMipData(0);
+				FInt32Point BlockPatchOffset = BlockPatchRect.Min - BlockRectInPixels.Min;
+				FInt32Point BlockPatchSize = BlockPatchRect.Size();
+				for (int32 RowIndex = BlockPatchOffset[1]; RowIndex < BlockPatchOffset[1]+BlockPatchSize[1]; ++RowIndex)
+				{
+					uint8* RowPixels = Pixels + RowIndex * BlockSize[0] + BlockPatchOffset[0];
+					FMemory::Memset(RowPixels, 255, BlockPatchSize[0]);
+				}
+			}
+		}
+
+		return PatchMask;
+	}
+
+
+    //---------------------------------------------------------------------------------------------
+    void CodeGenerator::GenerateSurface( FSurfaceGenerationResult& result, 
+										 const FSurfaceGenerationOptions& Options,
+                                         Ptr<const NodeSurfaceNew> surfaceNode )
+    {
+        MUTABLE_CPUPROFILER_SCOPE(GenerateSurface);
+
+        const NodeSurfaceNew& node = *surfaceNode;
+
+        // Build a series of operations to assemble the surface
+        Ptr<ASTOp> lastSurfOp;
+
+        // Generate the mesh
+        //------------------------------------------------------------------------
+        FMeshGenerationResult MeshResults;
+
+		// We don't add the mesh here, since it will be added directly at the top of the
+		// component expression in the NodeComponentNew generator with the right merges
+		// and conditions.
+		// But we store it to be used then.
+
+		// Do we need to generate the mesh? Or was it already generated for state conditions 
+		// accepting the current state?
+		FirstPassGenerator::FSurface* targetSurface = nullptr;
+		for (FirstPassGenerator::FSurface& its : FirstPass.Surfaces)
+		{
+            if (its.Node == &node)
+            {
+                // Check state conditions
+                bool surfaceValidForThisState =
+                    Options.State >= its.StateCondition.Num()
+                    ||
+                    its.StateCondition[Options.State];
+
+                if (surfaceValidForThisState)
+                {
+                    if (its.ResultSurfaceOp)
+                    {
+                        // Reuse the entire surface
+                        result.surfaceOp = its.ResultSurfaceOp;
+                        return;
+                    }
+                    else
+                    {
+                        // Not already generated, we will generate this
+                        targetSurface = &its;
+                    }
+                }
+            }
+		}
+
+        if (!targetSurface)
+        {
+            return;
+        }
+
+		// This assumes that the lods are processed in order. It checks it this way because some platforms may have empty lods at the top.
+		const bool bIsBaseForSharedSurface = node.SharedSurfaceId != INDEX_NONE && !SharedMeshOptionsMap.Contains(node.SharedSurfaceId);
+
+		// If this is true, we will reuse the surface properties from a higher LOD, se we can skip the generation of material properties and images.
+		const bool bShareSurface = node.SharedSurfaceId != INDEX_NONE && !bIsBaseForSharedSurface;
+
+        if (node.Mesh)
+        {
+            MUTABLE_CPUPROFILER_SCOPE(SurfaceMesh);
+
+            Ptr<ASTOp> lastMeshOp;
+
+			// Store the data necessary to apply modifiers for the pre-normal operations stage.
+			ActiveTags.Add(node.Tags);
+
+            // Generate the mesh
+			FMeshGenerationOptions MeshOptions;
+			MeshOptions.bLayouts = true;
+			MeshOptions.State = Options.State;
+			MeshOptions.ActiveTags = node.Tags;
+
+			const FMeshGenerationResult* SharedMeshResults = nullptr;
+			if (bShareSurface)
+			{
+				// Do we have the surface we need to share it with?
+				SharedMeshResults = SharedMeshOptionsMap.Find(node.SharedSurfaceId);
+				check(SharedMeshResults);
+
+				// Override the layouts with the ones from the surface we share
+				MeshOptions.OverrideLayouts = SharedMeshResults->GeneratedLayouts;
+			}
+
+			// Normalize UVs if we're going to work with images and layouts.
+			// TODO: This should come from per-layout settings!
+			const bool bNormalizeUVs = false; // !node.Images.IsEmpty();
+			MeshOptions.bNormalizeUVs = bNormalizeUVs;
+
+			// Ensure UV islands remain within their main layout block on lower LODs to avoid unexpected reordering 
+			// of the layout blocks when reusing a surface between LODs. Used to fix small displacements on vertices
+			// that may cause them to fall on a different block.
+			MeshOptions.bClampUVIslands = bShareSurface && bNormalizeUVs;
+
+            GenerateMesh(MeshOptions, MeshResults, node.Mesh);
+
+			ActiveTags.Pop();
+
+			// Apply the modifier for the post-normal operations stage.
+			bool bModifiersForBeforeOperations = false;
+			FMeshGenerationOptions ModifierOptions(MeshOptions);
+			ModifierOptions.ActiveTags = node.Tags;
+			lastMeshOp = ApplyMeshModifiers(MeshOptions, MeshResults, SharedMeshResults, bModifiersForBeforeOperations, surfaceNode->GetMessageContext(), nullptr);
+
+            // Layouts
+            for ( int32 LayoutIndex=0; LayoutIndex < MeshResults.GeneratedLayouts.Num(); ++LayoutIndex)
+            {
+                Ptr<ASTOp> layoutOp;
+
+				Ptr<const Layout> pLayout = MeshResults.GeneratedLayouts[LayoutIndex].Layout;
+                if ( pLayout )
+                {
+					if (SharedMeshResults)
+					{
+						check(SharedMeshResults->LayoutOps.Num()>LayoutIndex);
+						layoutOp = SharedMeshResults->LayoutOps[LayoutIndex];
+					}
+					else
+					{
+						// Create a new layout expression
+
+						// Constant layout to start with
+						{
+							Ptr<ASTOpConstantResource> op = new ASTOpConstantResource();
+							op->Type = OP_TYPE::LA_CONSTANT;
+
+							op->SetValue(pLayout, CompilerOptions->OptimisationOptions.DiskCacheContext);
+							layoutOp = op;
+						}
+
+						// Add children merged meshes layouts
+						for (const FMeshGenerationResult::FExtraLayouts& Data : MeshResults.ExtraMeshLayouts)
+						{
+							if (!Data.MeshFragment)
+							{
+								// No mesh to add, we assume there are no layouts to add either.
+								check(Data.GeneratedLayouts.IsEmpty());
+								continue;
+							}
+
+							if (Data.GeneratedLayouts.Num() < MeshResults.GeneratedLayouts.Num())
+							{
+								ErrorLog->GetPrivate()->Add(TEXT("Merged layout has been ignored because the number of layouts is different."), ELMT_ERROR, surfaceNode->GetMessageContext());
+							}
+							else
+							{
+								// Constant layout to start with
+								Ptr<ASTOp> layoutFragmentAd;
+								{
+									Ptr<ASTOpConstantResource> op = new ASTOpConstantResource();
+									op->Type = OP_TYPE::LA_CONSTANT;
+
+									Ptr<const Layout> pCloned = Data.GeneratedLayouts[LayoutIndex].Layout;
+									op->SetValue(pCloned, CompilerOptions->OptimisationOptions.DiskCacheContext);
+
+									layoutFragmentAd = op;
+								}
+
+								// Merge operation
+								Ptr<ASTOpLayoutMerge> mergeAd = new ASTOpLayoutMerge();
+								mergeAd->Base = layoutOp;
+								mergeAd->Added = layoutFragmentAd;
+
+								// Condition to apply
+								if (Data.Condition)
+								{
+									Ptr<ASTOpConditional> conditionalAd = new ASTOpConditional();
+									conditionalAd->type = OP_TYPE::LA_CONDITIONAL;
+									conditionalAd->no = layoutOp;
+									conditionalAd->yes = mergeAd;
+									conditionalAd->condition = Data.Condition;
+									layoutOp = conditionalAd;
+								}
+								else
+								{
+									layoutOp = mergeAd;
+								}
+							}
+						}
+					}
+
+					bool bIsOverlayLayout = pLayout->GetLayoutPackingStrategy() == mu::EPackStrategy::Overlay;
+
+                    if (!bIsOverlayLayout && layoutOp)
+                    {
+                        // Add layout packing instructions
+						if (!SharedMeshResults)
+						{
+                            // Make sure we removed unnecessary blocks
+                            Ptr<ASTOpLayoutFromMesh> ExtractOp = new ASTOpLayoutFromMesh();
+							ExtractOp->Mesh = lastMeshOp;
+							check(LayoutIndex<256);
+							ExtractOp->LayoutIndex = uint8(LayoutIndex);
+
+							Ptr<ASTOpLayoutRemoveBlocks> RemoveOp = new ASTOpLayoutRemoveBlocks();
+							RemoveOp->Source = layoutOp;
+							RemoveOp->ReferenceLayout = ExtractOp;
+							layoutOp = RemoveOp;
+
+                            // Pack uv blocks
+                            Ptr<ASTOpLayoutPack> op = new ASTOpLayoutPack();
+                            op->Source = layoutOp;
+                            layoutOp = op;
+                        }
+
+                        // Create the expression to apply the layout to the mesh
+                        {
+                            Ptr<ASTOpFixed> op = new ASTOpFixed();
+                            op->op.type = OP_TYPE::ME_APPLYLAYOUT;
+                            op->SetChild(op->op.args.MeshApplyLayout.mesh, lastMeshOp );
+                            op->SetChild(op->op.args.MeshApplyLayout.layout, layoutOp );
+                            op->op.args.MeshApplyLayout.channel = (uint16)LayoutIndex;
+                            lastMeshOp = op;
+                        }
+                    }
+
+                }
+
+                MeshResults.LayoutOps.Add( layoutOp );
+            }
+
+            // Store in the surface for later use.
+            targetSurface->ResultMeshOp = lastMeshOp;
+        }
+
+
+        // Create the expression for each texture, if we are not reusing the surface from another LOD.
+        //------------------------------------------------------------------------
+		if (!bShareSurface)
+		{
+			for (int32 ImageIndex = 0; ImageIndex < node.Images.Num(); ++ImageIndex)
+			{
+				MUTABLE_CPUPROFILER_SCOPE(SurfaceTexture);
+
+				// Any image-specific format or mipmapping needs to be applied at the end
+				Ptr<NodeImageMipmap> mipmapNode;
+				Ptr<NodeImageFormat> formatNode;
+				Ptr<NodeImageSwizzle> swizzleNode;
+
+				bool bFound = false;
+				Ptr<NodeImage> pImageNode = node.Images[ImageIndex].Image;
+
+				while (!bFound && pImageNode)
+				{
+					if (pImageNode->GetType()==NodeImageMipmap::GetStaticType())
+					{
+						NodeImageMipmap* tm = static_cast<NodeImageMipmap*>(pImageNode.get());
+						if (!mipmapNode) mipmapNode = tm;
+						pImageNode = tm->GetSource();
+					}
+					else if (pImageNode->GetType() == NodeImageFormat::GetStaticType())
+					{
+						NodeImageFormat* tf = static_cast<NodeImageFormat*>(pImageNode.get());
+						if (!formatNode) formatNode = tf;
+						pImageNode = tf->GetSource();
+					}
+					else if (pImageNode->GetType() == NodeImageSwizzle::GetStaticType())
+					{
+						NodeImageSwizzle* ts = static_cast<NodeImageSwizzle*>(pImageNode.get());
+						NodeImage* Source = ts->GetSource(0).get();
+						if (!swizzleNode && Source ==ts->GetSource(1) && Source==ts->GetSource(2) && Source==ts->GetSource(3))
+						{
+							swizzleNode = ts;
+							pImageNode = Source;
+						}
+						else
+						{
+							bFound = true;
+						}
+					}
+					else
+					{
+						bFound = true;
+					}
+				}
+
+				if (bFound)
+				{
+
+					const int32 LayoutIndex = node.Images[ImageIndex].LayoutIndex;
+
+					// If the layout index has been set to negative, it means we should ignore the layout for this image.
+					CompilerOptions::TextureLayoutStrategy ImageLayoutStrategy = (LayoutIndex < 0)
+						? CompilerOptions::TextureLayoutStrategy::None
+						: CompilerOptions::TextureLayoutStrategy::Pack
+						;
+
+					if (ImageLayoutStrategy == CompilerOptions::TextureLayoutStrategy::None)
+					{
+						//check(desc.m_format != EImageFormat::IF_NONE);
+
+						// Generate the image
+						FImageGenerationOptions ImageOptions;
+						ImageOptions.State = Options.State;
+						ImageOptions.ImageLayoutStrategy = ImageLayoutStrategy;
+						ImageOptions.ActiveTags = node.Tags;
+						ImageOptions.RectSize = { 0, 0 };
+						FImageGenerationResult Result;
+						GenerateImage(ImageOptions, Result, pImageNode);
+						Ptr<ASTOp> imageAd = Result.op;
+
+						// Placeholder block. Ideally this should be the actual image size
+						constexpr int32 FakeLayoutSize = 256;
+						FIntPoint GridSize(FakeLayoutSize, FakeLayoutSize);
+						FLayoutBlockDesc LayoutBlockDesc;
+						LayoutBlockDesc.BlockPixelsX = 1;
+						LayoutBlockDesc.BlockPixelsY = 1;
+						box< FIntVector2 > RectInCells;
+						RectInCells.min = { 0,0 };
+						RectInCells.size = { FakeLayoutSize ,FakeLayoutSize };
+						
+						bool bModifiersForBeforeOperations = false;
+						imageAd = ApplyImageBlockModifiers(ImageOptions, imageAd, ImageIndex, GridSize, LayoutBlockDesc, RectInCells, bModifiersForBeforeOperations, surfaceNode->GetMessageContext());
+						
+						check(imageAd);
+
+						if (swizzleNode)
+						{
+							Ptr<ASTOpImageSwizzle> fop = new ASTOpImageSwizzle();
+							fop->Format = swizzleNode->GetPrivate()->m_format;
+							fop->Sources[0] = imageAd;
+							fop->Sources[1] = imageAd;
+							fop->Sources[2] = imageAd;
+							fop->Sources[3] = imageAd;
+							fop->SourceChannels[0] = swizzleNode->GetPrivate()->m_sourceChannels[0];
+							fop->SourceChannels[1] = swizzleNode->GetPrivate()->m_sourceChannels[1];
+							fop->SourceChannels[2] = swizzleNode->GetPrivate()->m_sourceChannels[2];
+							fop->SourceChannels[3] = swizzleNode->GetPrivate()->m_sourceChannels[3];
+							check(fop->Format != EImageFormat::IF_NONE);
+							imageAd = fop;
+						}
+
+						if (mipmapNode)
+						{
+							Ptr<ASTOpImageMipmap> op = new ASTOpImageMipmap();
+							op->Levels = 0;
+							op->Source = imageAd;
+							op->BlockLevels = 0;
+
+							op->AddressMode = mipmapNode->GetPrivate()->m_settings.AddressMode;
+							op->FilterType = mipmapNode->GetPrivate()->m_settings.FilterType;
+							imageAd = op;
+						}
+
+						if (formatNode)
+						{
+							Ptr<ASTOpImagePixelFormat> fop = new ASTOpImagePixelFormat();
+							fop->Format = formatNode->GetPrivate()->m_format;
+							fop->FormatIfAlpha = formatNode->GetPrivate()->m_formatIfAlpha;
+							fop->Source = imageAd;
+							check(fop->Format != EImageFormat::IF_NONE);
+							imageAd = fop;
+						}
+
+						Ptr<ASTOpInstanceAdd> op = new ASTOpInstanceAdd();
+						op->type = OP_TYPE::IN_ADDIMAGE;
+						op->instance = lastSurfOp;
+						op->value = imageAd;
+						op->name = node.Images[ImageIndex].Name;
+
+						lastSurfOp = op;
+					}
+
+					else if (ImageLayoutStrategy == CompilerOptions::TextureLayoutStrategy::Pack) //-V547
+					{						
+						if (LayoutIndex >= MeshResults.GeneratedLayouts.Num() ||
+							LayoutIndex >= MeshResults.LayoutOps.Num())
+						{
+							ErrorLog->GetPrivate()->Add("Missing layout in object, or its parent.", ELMT_ERROR, surfaceNode->GetMessageContext());
+						}
+						else
+						{
+							const Layout* pLayout = MeshResults.GeneratedLayouts[LayoutIndex].Layout.get();
+							check(pLayout);
+
+							Ptr<ASTOpInstanceAdd> op = new ASTOpInstanceAdd();
+							op->type = OP_TYPE::IN_ADDIMAGE;
+							op->instance = lastSurfOp;
+
+							// Image
+							//-------------------------------------
+
+							// Size of a layout block in pixels
+							FIntPoint GridSize = pLayout->GetGridSize();
+	
+							// Try to guess the layout block description from the first valid block that is generated.
+							FLayoutBlockDesc LayoutBlockDesc;
+							if (formatNode)
+							{
+								LayoutBlockDesc.FinalFormat = formatNode->GetPrivate()->m_formatIfAlpha;
+								if (LayoutBlockDesc.FinalFormat == EImageFormat::IF_NONE)
+								{
+									LayoutBlockDesc.FinalFormat = formatNode->GetPrivate()->m_format;
+								}
+							}
+
+							bool bImageSizeWarning = false;
+
+							// Start with a blank image. It will be completed later with the blockSize, format and mips information
+							Ptr<ASTOpFixed> BlankImageOp;
+							Ptr<ASTOp> imageAd;
+							{
+								BlankImageOp = new ASTOpFixed();
+								BlankImageOp->op.type = OP_TYPE::IM_BLANKLAYOUT;
+								BlankImageOp->SetChild(BlankImageOp->op.args.ImageBlankLayout.layout, MeshResults.LayoutOps[LayoutIndex]);
+								// The rest ok the op will be completed below
+								BlankImageOp->op.args.ImageBlankLayout.mipmapCount = 0;
+								imageAd = BlankImageOp;
+							}
+
+							for (int32 BlockIndex = 0; BlockIndex < pLayout->GetBlockCount(); ++BlockIndex)
+							{
+								// Generate the image
+								FImageGenerationOptions ImageOptions;
+								ImageOptions.State = Options.State;
+								ImageOptions.ImageLayoutStrategy = ImageLayoutStrategy;
+								ImageOptions.RectSize = { 0,0 };
+								ImageOptions.ActiveTags = node.Tags;
+								ImageOptions.LayoutToApply = pLayout;
+								ImageOptions.LayoutBlockId = pLayout->Blocks[BlockIndex].Id;
+								FImageGenerationResult Result;
+								GenerateImage(ImageOptions, Result, pImageNode);
+								Ptr<ASTOp> blockAd = Result.op;
+
+								if (!blockAd)
+								{
+									// The GenerateImage(...) above has failed, skip this block
+									result.surfaceOp = nullptr;
+									continue;
+								}
+
+								// Calculate the desc of the generated block.
+								constexpr bool bReturnBestOption = true;
+								FImageDesc BlockDesc = blockAd->GetImageDesc(bReturnBestOption, nullptr);
+
+								// Block in layout grid units (cells)
+								box< FIntVector2 > RectInCells;
+								RectInCells.min = pLayout->Blocks[BlockIndex].Min;
+								RectInCells.size = pLayout->Blocks[BlockIndex].Size;
+
+								// Try to update the layout block desc if we don't know it yet.
+								UpdateLayoutBlockDesc(LayoutBlockDesc, BlockDesc, RectInCells.size);
+
+								// Even if we force the size afterwards, we need some size hint in some cases, like image projections.
+								ImageOptions.RectSize = UE::Math::TIntVector2<int32>(BlockDesc.m_size);
+
+								bool bModifiersForBeforeOperations = false;
+								blockAd = ApplyImageBlockModifiers(ImageOptions, blockAd, ImageIndex, GridSize, LayoutBlockDesc, RectInCells, bModifiersForBeforeOperations, surfaceNode->GetMessageContext());
+
+								// Enforce block size and optimizations
+								blockAd = GenerateImageSize(blockAd, FIntVector2(BlockDesc.m_size));
+
+								EImageFormat baseFormat = imageAd->GetImageDesc().m_format;
+								// Actually don't do it, it will be propagated from the top format operation.
+								//Ptr<ASTOp> blockAd = GenerateImageFormat(blockAd, baseFormat);
+
+								// Apply tiling to avoid generating chunks of image that are too big.
+								blockAd = ApplyTiling(blockAd, ImageOptions.RectSize, LayoutBlockDesc.FinalFormat);
+
+								// Compose layout operation
+								Ptr<ASTOpImageCompose> composeOp = new ASTOpImageCompose();
+								composeOp->Layout = MeshResults.LayoutOps[LayoutIndex];
+								composeOp->Base = imageAd;
+								composeOp->BlockImage = blockAd;
+
+								// Set the absolute block index.
+								check(pLayout->Blocks[BlockIndex].Id != FLayoutBlock::InvalidBlockId);
+								composeOp->BlockId = pLayout->Blocks[BlockIndex].Id;
+
+								imageAd = composeOp;
+							}
+							check(imageAd);
+
+							FMeshGenerationOptions ModifierOptions;
+							ModifierOptions.State = Options.State;
+							ModifierOptions.ActiveTags = node.Tags;
+							bool bModifiersForBeforeOperations = false;
+							imageAd = ApplyImageExtendModifiers(ModifierOptions, MeshResults, imageAd, ImageLayoutStrategy, LayoutIndex, ImageIndex, GridSize, LayoutBlockDesc, bModifiersForBeforeOperations, surfaceNode->GetMessageContext());
+
+							// Complete the base op
+							BlankImageOp->op.args.ImageBlankLayout.blockSize[0] = uint16(LayoutBlockDesc.BlockPixelsX);
+							BlankImageOp->op.args.ImageBlankLayout.blockSize[1] = uint16(LayoutBlockDesc.BlockPixelsY);
+							BlankImageOp->op.args.ImageBlankLayout.format = GetUncompressedFormat(LayoutBlockDesc.FinalFormat);
+							BlankImageOp->op.args.ImageBlankLayout.generateMipmaps = LayoutBlockDesc.bBlocksHaveMips;
+							BlankImageOp->op.args.ImageBlankLayout.mipmapCount = 0;
+
+							if (swizzleNode)
+							{
+								Ptr<ASTOpImageSwizzle> fop = new ASTOpImageSwizzle();
+								fop->Format = swizzleNode->GetPrivate()->m_format;
+								fop->Sources[0] = imageAd;
+								fop->Sources[1] = imageAd;
+								fop->Sources[2] = imageAd;
+								fop->Sources[3] = imageAd;
+								fop->SourceChannels[0] = swizzleNode->GetPrivate()->m_sourceChannels[0];
+								fop->SourceChannels[1] = swizzleNode->GetPrivate()->m_sourceChannels[1];
+								fop->SourceChannels[2] = swizzleNode->GetPrivate()->m_sourceChannels[2];
+								fop->SourceChannels[3] = swizzleNode->GetPrivate()->m_sourceChannels[3];
+								check(fop->Format != EImageFormat::IF_NONE);
+								imageAd = fop;
+							}
+
+							// Apply mipmap and format if necessary, skip if format is IF_NONE (possibly because a block was skipped above)
+							bool bNeedsMips =
+								(mipmapNode && LayoutBlockDesc.FinalFormat != EImageFormat::IF_NONE)
+								||
+								LayoutBlockDesc.bBlocksHaveMips;
+
+							if (bNeedsMips)
+							{
+								Ptr<ASTOpImageMipmap> mop = new ASTOpImageMipmap();
+
+								// At the end of the day, we want all the mipmaps. Maybe the code
+								// optimiser will split the process later.
+								mop->Levels = 0;
+								mop->bOnlyTail = false;
+								mop->Source = imageAd;
+
+								// We have to avoid mips smaller than the image format block size, so
+								// we will devide the layout block by the format block
+								const FImageFormatData& PixelFormatInfo = GetImageFormatData(LayoutBlockDesc.FinalFormat);
+
+								int32 mipsX = FMath::CeilLogTwo(LayoutBlockDesc.BlockPixelsX / PixelFormatInfo.PixelsPerBlockX);
+								int32 mipsY = FMath::CeilLogTwo(LayoutBlockDesc.BlockPixelsY / PixelFormatInfo.PixelsPerBlockY);
+								mop->BlockLevels = (uint8)FMath::Max(mipsX, mipsY);
+								
+								if (LayoutBlockDesc.BlockPixelsX < PixelFormatInfo.PixelsPerBlockX || LayoutBlockDesc.BlockPixelsY < PixelFormatInfo.PixelsPerBlockY)
+								{
+									// In this case, the mipmap will never be useful for blocks, so we indicate that
+									// it should make the mips at the root of the expression.
+									mop->bOnlyTail = true;
+								}
+
+								mop->AddressMode = EAddressMode::ClampToEdge;
+								mop->FilterType = EMipmapFilterType::SimpleAverage;
+
+								if (mipmapNode)
+								{
+									mop->AddressMode = mipmapNode->GetPrivate()->m_settings.AddressMode;
+									mop->FilterType = mipmapNode->GetPrivate()->m_settings.FilterType;
+								}
+
+								imageAd = mop;
+							}
+
+							if (formatNode)
+							{
+								Ptr<ASTOpImagePixelFormat> fop = new ASTOpImagePixelFormat();
+								fop->Format = formatNode->GetPrivate()->m_format;
+								fop->FormatIfAlpha = formatNode->GetPrivate()->m_formatIfAlpha;
+								fop->Source = imageAd;
+								check(fop->Format != EImageFormat::IF_NONE);
+								imageAd = fop;
+							}
+
+							op->value = imageAd;
+
+							// Name
+							op->name = node.Images[ImageIndex].Name;
+
+							lastSurfOp = op;
+						}
+					}
+
+					else
+					{
+						// Unimplemented texture layout strategy
+						check(false);
+					}
+				}
+			}
+
+			// Create the expression for each vector
+			//------------------------------------------------------------------------
+			for (int32 t = 0; t < node.Vectors.Num(); ++t)
+			{
+				//MUTABLE_CPUPROFILER_SCOPE(SurfaceVector);
+
+				if (Ptr<NodeColour> pVectorNode = node.Vectors[t].Vector)
+				{
+					Ptr<ASTOpInstanceAdd> op = new ASTOpInstanceAdd();
+					op->type = OP_TYPE::IN_ADDVECTOR;
+					op->instance = lastSurfOp;
+
+					// Vector
+					FColorGenerationResult VectorResult;
+					GenerateColor(VectorResult, Options, pVectorNode);
+					op->value = VectorResult.op;
+
+					// Name
+					op->name = node.Vectors[t].Name;
+
+					lastSurfOp = op;
+				}
+			}
+
+			// Create the expression for each scalar
+			//------------------------------------------------------------------------
+			for (int32 t = 0; t < node.Scalars.Num(); ++t)
+			{
+				// MUTABLE_CPUPROFILER_SCOPE(SurfaceScalar);
+
+				if (NodeScalarPtr pScalarNode = node.Scalars[t].Scalar)
+				{
+					Ptr<ASTOpInstanceAdd> op = new ASTOpInstanceAdd();
+					op->type = OP_TYPE::IN_ADDSCALAR;
+					op->instance = lastSurfOp;
+
+					// Scalar
+					FScalarGenerationResult ScalarResult;
+					GenerateScalar(ScalarResult, Options, pScalarNode);
+					op->value = ScalarResult.op;
+
+					// Name
+					op->name = node.Scalars[t].Name;
+
+					lastSurfOp = op;
+				}
+			}
+
+			// Create the expression for each string
+			//------------------------------------------------------------------------
+			for (int32 t = 0; t < node.Strings.Num(); ++t)
+			{
+				if (NodeStringPtr pStringNode = node.Strings[t].String)
+				{
+					Ptr<ASTOpInstanceAdd> op = new ASTOpInstanceAdd();
+					op->type = OP_TYPE::IN_ADDSTRING;
+					op->instance = lastSurfOp;
+
+					FStringGenerationResult StringResult;
+					GenerateString(StringResult, Options, pStringNode);
+					op->value = StringResult.op;
+
+					// Name
+					op->name = node.Strings[t].Name;
+
+					lastSurfOp = op;
+				}
+			}
+		}
+
+        result.surfaceOp = lastSurfOp;
+        targetSurface->ResultSurfaceOp = lastSurfOp;
+
+		// If we are going to share this surface properties, remember it.
+		if (bIsBaseForSharedSurface)
+		{
+			check(!SharedMeshOptionsMap.Contains(node.SharedSurfaceId));
+			SharedMeshOptionsMap.Add( node.SharedSurfaceId, MeshResults );
+		}
+    }
+
+
+    //---------------------------------------------------------------------------------------------
+	void CodeGenerator::Generate_LOD(const FLODGenerationOptions& Options, FGenericGenerationResult& Result, const NodeLOD* InNode)
+	{
+		const NodeLOD& node = *InNode;
+
+		MUTABLE_CPUPROFILER_SCOPE(Generate_LOD);
+
+		// Build a series of operations to assemble the component
+        Ptr<ASTOp> lastCompOp;
+        Ptr<ASTOp> lastMeshOp;
+        FString lastMeshName;
+
+        // This generates a different ID for each surface. It can be used to match it to the
+        // mesh surface, or for debugging. It cannot be 0 because it is a special case for the
+        // merge operation.
+        int32 surfaceID=1;
+
+        // Look for all surfaces that belong to this component
+		for (int32 i = 0; i<FirstPass.Surfaces.Num(); ++i, ++surfaceID)
+		{
+			const FirstPassGenerator::FSurface& its = FirstPass.Surfaces[i];
+			if (its.Component==Options.Component
+				&&
+				its.LOD==Options.LODIndex)
+			{
+                // Apply state conditions: only generate it if it enabled in this state
+                {
+                    bool enabledInThisState = true;
+                    if (its.StateCondition.Num() && Options.State >= 0)
+                    {
+                        enabledInThisState =
+                                (Options.State < its.StateCondition.Num())
+                                &&
+                                ( its.StateCondition[Options.State] );
+                    }
+                    if (!enabledInThisState)
+                    {
+                        continue;
+                    }
+                }
+
+                Ptr<ASTOpInstanceAdd> sop = new ASTOpInstanceAdd();
+                sop->type = OP_TYPE::IN_ADDSURFACE;
+                sop->name = its.Node->Name;
+                sop->instance = lastCompOp;
+
+				FSurfaceGenerationOptions SurfaceOptions(Options);
+				FSurfaceGenerationResult surfaceGenerationResult;
+                GenerateSurface( surfaceGenerationResult, SurfaceOptions, its.Node );
+                sop->value = surfaceGenerationResult.surfaceOp;
+
+                sop->id = surfaceID;
+                sop->ExternalId = its.Node->ExternalId;
+                sop->SharedSurfaceId = its.Node->SharedSurfaceId;
+                Ptr<ASTOp> surfaceAt = sop;
+
+                Ptr<ASTOp> SurfaceConditionOp = its.FinalCondition;
+
+                {
+                    Ptr<ASTOpConditional> op = new ASTOpConditional();
+                    op->type = OP_TYPE::IN_CONDITIONAL;
+                    op->no = lastCompOp;
+                    op->yes = surfaceAt;
+                    op->condition = SurfaceConditionOp;
+                    lastCompOp = op;
+                }
+
+                // Add the mesh with its condition
+ 
+                // We add the merge op even for the first mesh, so that we set the surface id.
+				Ptr<ASTOp> mergeAd;
+				{
+                    Ptr<ASTOp> added = its.ResultMeshOp;
+ 
+                    Ptr<ASTOpFixed> mop = new ASTOpFixed();
+                    mop->op.type = OP_TYPE::ME_MERGE;
+                    mop->SetChild(mop->op.args.MeshMerge.base, lastMeshOp );
+                    mop->SetChild(mop->op.args.MeshMerge.added, added );
+                    mop->op.args.MeshMerge.newSurfaceID = surfaceID;
+                    mergeAd = mop;
+                }
+
+                if (SurfaceConditionOp)
+                {
+                    Ptr<ASTOpConditional> op = new ASTOpConditional();
+                    op->type = OP_TYPE::ME_CONDITIONAL;
+                    op->no = lastMeshOp;
+                    op->yes = mergeAd;
+                    op->condition = SurfaceConditionOp;
+                    lastMeshOp = op;
+                }
+                else
+                {
+                    lastMeshOp = mergeAd;
+                }
+			}
+		}
+
+		// Add op to optimize the skinning of the resulting mesh
+		{
+			Ptr<ASTOpMeshOptimizeSkinning> mop = new ASTOpMeshOptimizeSkinning();
+			mop->source = lastMeshOp;
+			lastMeshOp = mop;
+		}
+
+        // Add the component mesh
+        {
+            Ptr<ASTOpInstanceAdd> iop = new ASTOpInstanceAdd();
+            iop->type = OP_TYPE::IN_ADDMESH;
+            iop->instance = lastCompOp;
+            iop->value = lastMeshOp;
+
+            lastCompOp = iop;
+        }
+
+        Result.op = lastCompOp;
+    }
+
+
+    //---------------------------------------------------------------------------------------------
+   // Ptr<ASTOp> CodeGenerator::Visit( const NodeObjectNew::Private& node )
+	void CodeGenerator::Generate_ObjectNew(const FGenericGenerationOptions& Options, FGenericGenerationResult& Result, const NodeObjectNew* InNode)
+	{
+		MUTABLE_CPUPROFILER_SCOPE(NodeObjectNew);
+
+		// There is always at least a null parent
+		bool bIsChildObject = CurrentParents.Num() > 1;
+
+		// Add this object as current parent
+        CurrentParents.Add( FParentKey() );
+        CurrentParents.Last().ObjectNode = InNode;
+
+        // Parse the child objects first, which will accumulate operations in the patching lists
+        for ( int32 t=0; t< InNode->Children.Num(); ++t )
+        {
+            if ( const NodeObject* pChildNode = InNode->Children[t].get() )
+            {
+                Ptr<ASTOp> paramOp;
+
+                // If there are parent objects, the condition of this object depends on the
+                // condition of the parent object
+                if ( CurrentObject.Num() )
+                {
+                    paramOp = CurrentObject.Last().Condition;
+                }
+                else
+                {
+                    // In case there is no group node, we generate a constant true condition
+                    // This condition will be overwritten by the group nodes.
+                    Ptr<ASTOpConstantBool> op = new ASTOpConstantBool();
+                    op->value = true;
+                    paramOp = op;
+                }
+
+				FObjectGenerationData data;
+                data.Condition = paramOp;
+                CurrentObject.Add( data );
+
+                // This op is ignored: everything is stored as patches to apply to the parent when it is compiled.
+                Generate_Generic( pChildNode, Options );
+
+                CurrentObject.Pop();
+            }
+        }
+
+        // Create the expression adding all the components
+		Ptr<ASTOp> LastCompOp;
+		Ptr<ASTOp> PlaceholderOp;
+		if (bIsChildObject)
+		{
+			PlaceholderOp = new ASTOpInstanceAdd;
+			LastCompOp = PlaceholderOp;
+		}
+
+		// Add the components in this node
+        for ( int32 t=0; t< InNode->Components.Num(); ++t )
+        {
+			const NodeComponent* ComponentNode = InNode->Components[t].get();
+            if (ComponentNode)
+            {
+				FComponentGenerationOptions ComponentOptions( Options, LastCompOp );
+				FGenericGenerationResult ComponentResult;
+				GenerateComponent(ComponentOptions, ComponentResult, ComponentNode);
+				LastCompOp = ComponentResult.op;
+            }
+        }
+
+		// If we didn't generate anything, make sure we don't use the placeholder.
+		if (LastCompOp == PlaceholderOp)
+		{
+			LastCompOp = nullptr;
+			PlaceholderOp = nullptr;
+		}
+
+		// Add the components from child objects
+		FAdditionalComponentKey ThisKey;
+		ThisKey.ObjectNode = CurrentParents.Last().ObjectNode;
+		TArray<FAdditionalComponentData>* ThisAdditionalComponents = AdditionalComponents.Find(ThisKey);
+		if (LastCompOp && ThisAdditionalComponents)
+		{
+			for (const FAdditionalComponentData& Additional : *ThisAdditionalComponents)
+			{
+				check(Additional.PlaceholderOp);
+				ASTOp::Replace(Additional.PlaceholderOp, LastCompOp);
+				LastCompOp = Additional.ComponentOp;
+			}
+		}
+
+		// Store this chain of components for use in parent objects if necessary
+		// 2 is because there must be a parent and there is always a null element as well.
+		if (LastCompOp && bIsChildObject)
+		{
+			// TODO: Directly to the root object?
+			const FParentKey& ParentObjectKey = CurrentParents[CurrentParents.Num() - 2];
+			FAdditionalComponentKey ParentKey;
+			ParentKey.ObjectNode = ParentObjectKey.ObjectNode;
+
+			FAdditionalComponentData Data;
+			Data.ComponentOp = LastCompOp;
+			Data.PlaceholderOp = PlaceholderOp;
+			AdditionalComponents.FindOrAdd(ParentKey).Add(Data);
+		}
+
+
+        Ptr<ASTOp> RootOp = LastCompOp;
+
+		// Add an ASTOpAddExtensionData for each connected ExtensionData node
+		for (const NodeObjectNew::FNamedExtensionDataNode& NamedNode : InNode->ExtensionDataNodes)
+		{
+			if (!NamedNode.Node.get())
+			{
+				// No node connected
+				continue;
+			}
+
+			// Name must be valid
+			check(NamedNode.Name.Len() > 0);
+
+			FExtensionDataGenerationResult ChildResult;
+			GenerateExtensionData(ChildResult, Options, NamedNode.Node);
+
+			if (!ChildResult.Op.get())
+			{
+				// Failed to generate anything for this node
+				continue;
+			}
+
+			FConditionalExtensionDataOp& SavedOp = ConditionalExtensionDataOps.AddDefaulted_GetRef();
+			if (CurrentObject.Num() > 0)
+			{
+				SavedOp.Condition = CurrentObject.Last().Condition;
+			}
+			SavedOp.ExtensionDataOp = ChildResult.Op;
+			SavedOp.ExtensionDataName = NamedNode.Name;
+		}
+
+		if (CurrentObject.Num() == 0)
+		{
+			for (const FConditionalExtensionDataOp& SavedOp : ConditionalExtensionDataOps)
+			{
+				Ptr<ASTOpAddExtensionData> ExtensionPinOp = new ASTOpAddExtensionData();
+				ExtensionPinOp->Instance = ASTChild(ExtensionPinOp, RootOp);
+				ExtensionPinOp->ExtensionData = ASTChild(ExtensionPinOp, SavedOp.ExtensionDataOp);
+				ExtensionPinOp->ExtensionDataName = SavedOp.ExtensionDataName;
+
+				if (SavedOp.Condition.get())
+				{
+					Ptr<ASTOpConditional> ConditionOp = new ASTOpConditional();
+					ConditionOp->type = OP_TYPE::IN_CONDITIONAL;
+					ConditionOp->no = RootOp;
+					ConditionOp->yes = ExtensionPinOp;
+					ConditionOp->condition = ASTChild(ConditionOp, SavedOp.Condition);
+					
+					RootOp = ConditionOp;
+				}
+				else
+				{
+					RootOp = ExtensionPinOp;
+				}
+			}
+		}
+
+        CurrentParents.Pop();
+
+        Result.op = RootOp;
+    }
+
+
+    //---------------------------------------------------------------------------------------------
+    //Ptr<ASTOp> CodeGenerator::Visit( const NodeObjectGroup::Private& node )
+	void CodeGenerator::Generate_ObjectGroup(const FGenericGenerationOptions& Options, FGenericGenerationResult& Result, const NodeObjectGroup* InNode)
+	{
+		const NodeObjectGroup::Private& node = *InNode->GetPrivate();
+
+		TArray<FString> usedNames;
+
+        // Parse the child objects first, which will accumulate operations in the patching lists
+        for ( int32 t=0; t<node.m_children.Num(); ++t )
+        {
+            if ( const NodeObject* pChildNode = node.m_children[t].get() )
+            {
+				// Look for the child condition in the first pass
+                Ptr<ASTOp> conditionOp;
+				bool found = false;
+                for( int32 i = 0; !found && i != FirstPass.Objects.Num(); i++ )
+				{
+					FirstPassGenerator::FObject& it = FirstPass.Objects[i];
+					if (it.Node == pChildNode)
+					{
+						found = true;
+						conditionOp = it.Condition;
+					}
+				}
+
+                // \todo
+                // TrunkStraight_02_BranchTop crash
+                // It may happen with partial compilations?
+                // check(found);
+
+				FObjectGenerationData data;
+                data.Condition = conditionOp;
+                CurrentObject.Add( data );
+
+                // This op is ignored: everything is stored as patches to apply to the parent when
+                // it is compiled.
+                Generate_Generic( pChildNode, Options );
+
+                CurrentObject.Pop();
+
+				// Check for duplicated child names
+				FString strChildName = pChildNode->GetName();
+				if (usedNames.Contains(strChildName))
+				{
+					FString Msg = FString::Printf(TEXT("Object group has more than one children with the same name [%s]."), *strChildName );
+					ErrorLog->GetPrivate()->Add(Msg, ELMT_WARNING, InNode->GetMessageContext());
+				}
+				else
+				{
+					usedNames.Add(strChildName);
+				}
+            }
+        }
+    }
+
+
+    //---------------------------------------------------------------------------------------------
+    Ptr<ASTOp> CodeGenerator::GenerateMissingBoolCode(const TCHAR* Where, bool Value, const void* ErrorContext )
+    {
+        // Log a warning
+		FString Msg = FString::Printf(TEXT("Required connection not found: %s"), Where);
+        ErrorLog->GetPrivate()->Add( Msg, ELMT_ERROR, ErrorContext);
+
+        // Create a constant node
+        Ptr<NodeBoolConstant> pNode = new NodeBoolConstant();
+        pNode->SetValue(Value);
+
+		FBoolGenerationResult ChildResult;
+		FGenericGenerationOptions Options;
+        GenerateBool(ChildResult,Options, pNode );
+		return ChildResult.op;
+    }
+
+
+	//---------------------------------------------------------------------------------------------
+	void CodeGenerator::GetModifiersFor(
+		const TArray<FString>& SurfaceTags,
+		int32 LOD, bool bModifiersForBeforeOperations,
+		TArray<FirstPassGenerator::FModifier>& OutModifiers)
+	{
+        MUTABLE_CPUPROFILER_SCOPE(GetModifiersFor);
+
+		if (SurfaceTags.IsEmpty())
+		{
+			return;
+		}
+
+		for (const FirstPassGenerator::FModifier& m: FirstPass.Modifiers)
+		{
+			// Correct LOD?
+			// TODO: LOD restrictions inside the modifiers data.
+			//if (m.LOD != LOD)
+			//{
+			//	continue;
+			//}
+
+			// Correct stage?
+			if (m.Node->bApplyBeforeNormalOperations != bModifiersForBeforeOperations)
+			{
+				continue;
+			}
+
+			// Already there?
+			bool bAlreadyAdded = 
+				OutModifiers.FindByPredicate( [&m](const FirstPassGenerator::FModifier& c) {return c.Node == m.Node; })
+				!= 
+				nullptr;
+
+			if (bAlreadyAdded)
+			{
+				continue;
+			}
+
+			// Matching tags?
+			bool bApply = false;
+
+			switch (m.Node->MultipleTagsPolicy)
+			{
+			case EMutableMultipleTagPolicy::OnlyOneRequired:
+			{
+				for (const FString& Tag: m.Node->RequiredTags)
+				{
+					if (SurfaceTags.Contains(Tag))
+					{
+						bApply = true;
+						break;
+					}
+				}
+				break;
+			}
+
+			case EMutableMultipleTagPolicy::AllRequired:
+			{
+				bApply = true;
+				for (const FString& Tag : m.Node->RequiredTags)
+				{
+					if (!SurfaceTags.Contains(Tag))
+					{
+						bApply = false;
+						break;
+					}
+				}
+			}
+			}
+
+			if (bApply)
+			{
+				OutModifiers.Add(m);
+			}
+		}
+	}
+
+	//---------------------------------------------------------------------------------------------
+	Ptr<ASTOp> CodeGenerator::ApplyMeshModifiers(const FMeshGenerationOptions& Options, 
+		FMeshGenerationResult& BaseMeshResult,
+		const FMeshGenerationResult* SharedMeshResults,
+		bool bModifiersForBeforeOperations,
+		const void* ErrorContext,
+		const NodeMeshConstant* OriginalMeshNode )
+	{
+		Ptr<ASTOp> LastMeshOp = BaseMeshResult.MeshOp;
+
+		// Apply mesh modifiers
+		TArray<FirstPassGenerator::FModifier> Modifiers;
+
+		int32 currentLOD = CurrentParents.Last().Lod;
+		GetModifiersFor(Options.ActiveTags, currentLOD, bModifiersForBeforeOperations, Modifiers);
+
+		Ptr<ASTOp> PreModifiersMesh = LastMeshOp;
+
+		ActiveTags.Push({});
+
+		// Process mesh extend modifiers (from edit modifiers)
+		int32 EditIndex = 0;
+		for (const FirstPassGenerator::FModifier& m : Modifiers)
+		{
+			if (m.Node->GetType() == NodeModifierSurfaceEdit::GetStaticType())
+			{
+				const NodeModifierSurfaceEdit* Edit = static_cast<const NodeModifierSurfaceEdit*>(m.Node);
+
+				BaseMeshResult.ExtraMeshLayouts.Emplace();
+
+				if (Ptr<NodeMesh> pAdd = Edit->MeshAdd)
+				{
+					// Store the data necessary to apply modifiers for the pre-normal operations stage.
+					ActiveTags.Push(Edit->EnableTags);
+
+					FMeshGenerationOptions MergedMeshOptions(Options);
+					MergedMeshOptions.ActiveTags = Edit->EnableTags; // TODO: Append to current?
+
+					if (SharedMeshResults)
+					{
+						check(SharedMeshResults->ExtraMeshLayouts.IsValidIndex(EditIndex));
+						MergedMeshOptions.OverrideLayouts = SharedMeshResults->ExtraMeshLayouts[EditIndex].GeneratedLayouts;
+					}
+
+					FMeshGenerationResult AddResults;
+					GenerateMesh(MergedMeshOptions, AddResults, pAdd);
+
+					ActiveTags.Pop();
+
+					// Apply the modifier for the post-normal operations stage to the added mesh
+					FMeshGenerationOptions ModifierOptions(Options);
+					ModifierOptions.ActiveTags = Edit->EnableTags;
+					Ptr<ASTOp> AddedMeshOp = ApplyMeshModifiers(ModifierOptions, AddResults, SharedMeshResults, bModifiersForBeforeOperations, ErrorContext, nullptr);
+
+					FMeshGenerationResult::FExtraLayouts data;
+					data.GeneratedLayouts = AddResults.GeneratedLayouts;
+					data.Condition = m.FinalCondition;
+					data.MeshFragment = AddedMeshOp;
+					BaseMeshResult.ExtraMeshLayouts[EditIndex] = data;
+
+					Ptr<ASTOpFixed> mop = new ASTOpFixed();
+					mop->op.type = OP_TYPE::ME_MERGE;
+					mop->SetChild(mop->op.args.MeshMerge.base, LastMeshOp);
+					mop->SetChild(mop->op.args.MeshMerge.added, AddedMeshOp);
+					// will merge the meshes under the same surface
+					mop->op.args.MeshMerge.newSurfaceID = 0;
+
+					// Condition to apply
+					if (m.FinalCondition)
+					{
+						Ptr<ASTOpConditional> conditionalAd = new ASTOpConditional();
+						conditionalAd->type = OP_TYPE::ME_CONDITIONAL;
+						conditionalAd->no = LastMeshOp;
+						conditionalAd->yes = mop;
+						conditionalAd->condition = m.FinalCondition;
+						LastMeshOp = conditionalAd;
+					}
+					else
+					{
+						LastMeshOp = mop;
+					}
+				}
+
+				++EditIndex;
+			}
+
+		}
+
+		// "remove" operation to group all the removes
+		Ptr<ASTOpMeshRemoveMask> RemoveOp;
+
+		// Process mesh remove modifiers (from edit modifiers)
+		for (const FirstPassGenerator::FModifier& m : Modifiers)
+		{
+			if (m.Node->GetType() == NodeModifierSurfaceEdit::GetStaticType())
+			{
+				const NodeModifierSurfaceEdit* Edit = static_cast<const NodeModifierSurfaceEdit*>(m.Node);
+
+				// Apply mesh removes from child objects "edit surface" nodes.
+				// "Removes" need to come after "Adds" because some removes may refer to added meshes,
+				// and not the base.
+				// \TODO: Apply base removes first, and then "added meshes" removes here. It may have lower memory footprint during generation.
+				if (Ptr<NodeMesh> pRemove = Edit->MeshRemove)
+				{
+					FMeshGenerationResult removeResults;
+					FMeshGenerationOptions RemoveMeshOptions;
+					RemoveMeshOptions.bLayouts = false;
+					RemoveMeshOptions.State = Options.State;
+					RemoveMeshOptions.ActiveTags = Edit->EnableTags;
+
+					GenerateMesh(RemoveMeshOptions, removeResults, pRemove);
+
+					Ptr<ASTOpFixed> maskOp = new ASTOpFixed();
+					maskOp->op.type = OP_TYPE::ME_MASKDIFF;
+
+					// By default, remove from the base
+					Ptr<ASTOp> removeFrom = BaseMeshResult.BaseMeshOp;
+					maskOp->SetChild(maskOp->op.args.MeshMaskDiff.source, removeFrom);
+					maskOp->SetChild(maskOp->op.args.MeshMaskDiff.fragment, removeResults.MeshOp);
+
+					if (!RemoveOp)
+					{
+						RemoveOp = new ASTOpMeshRemoveMask();
+						RemoveOp->source = LastMeshOp;
+						LastMeshOp = RemoveOp;
+					}
+
+					RemoveOp->AddRemove(m.FinalCondition, maskOp);
+				}
+			}
+		}
+
+
+		// Process mesh morph modifiers (from edit modifiers)
+		for (const FirstPassGenerator::FModifier& m : Modifiers)
+		{
+			if (m.Node->GetType() == NodeModifierSurfaceEdit::GetStaticType())
+			{
+				const NodeModifierSurfaceEdit* Edit = static_cast<const NodeModifierSurfaceEdit*>(m.Node);
+
+				if (Edit->MeshMorph.IsEmpty())
+				{
+					continue;
+				}
+
+				check(OriginalMeshNode);
+
+				Ptr<Mesh> TargetMesh = OriginalMeshNode->FindMorph(Edit->MeshMorph);
+				if (!TargetMesh)
+				{
+					continue;
+				}
+				
+				{
+					// Target mesh
+					Ptr<ASTOpConstantResource> TargetMeshOp = new ASTOpConstantResource;
+					TargetMeshOp->Type = OP_TYPE::ME_CONSTANT;
+					TargetMeshOp->SetValue(TargetMesh->Clone(), CompilerOptions->OptimisationOptions.DiskCacheContext);
+
+					// Morph generation through mesh diff
+					Ptr<ASTOpMeshDifference> diffAd;
+					{
+						Ptr<ASTOpMeshDifference> op = new ASTOpMeshDifference();
+						op->Base = BaseMeshResult.BaseMeshOp;
+						op->Target = TargetMeshOp;
+
+						// Morphing tex coords here is not supported:
+						// Generating the homogoneous UVs is difficult since we don't have the base
+						// layout yet.                       
+						op->bIgnoreTextureCoords = true;
+						diffAd = op;
+					}
+
+					// Morph operation
+					Ptr<ASTOp> morphAd;
+					{
+						Ptr<ASTOpMeshMorph> op = new ASTOpMeshMorph();
+
+						// Factor
+						if (Edit->MorphFactor)
+						{
+							FScalarGenerationResult ChildResult;
+							GenerateScalar(ChildResult, Options, Edit->MorphFactor);
+							op->Factor = ChildResult.op;
+						}
+						else
+						{
+							Ptr<NodeScalarConstant> auxNode = new NodeScalarConstant();
+							auxNode->SetValue(1.0f);
+
+							FScalarGenerationResult ChildResult;
+							GenerateScalar(ChildResult, Options, auxNode);
+							op->Factor = ChildResult.op;
+						}
+
+						// Base		
+						op->Base = LastMeshOp;
+
+						// Targets
+						op->Target = diffAd;
+						morphAd = op;
+					}
+
+					// Condition to apply the morph
+					if (m.FinalCondition)
+					{
+						Ptr<ASTOpConditional> conditionalAd = new ASTOpConditional();
+						conditionalAd->type = OP_TYPE::ME_CONDITIONAL;
+						conditionalAd->no = LastMeshOp;
+						conditionalAd->yes = morphAd;
+						conditionalAd->condition = m.FinalCondition;
+						LastMeshOp = conditionalAd;
+					}
+					else
+					{
+						LastMeshOp = morphAd;
+					}
+				}
+			}
+		}
+
+
+
+		// Process clip-with-mesh modifiers
+		RemoveOp = nullptr;
+		for (const FirstPassGenerator::FModifier& m : Modifiers)
+		{
+			if (m.Node->GetType()== NodeModifierMeshClipWithMesh::GetStaticType())
+			{
+				const NodeModifierMeshClipWithMesh* TypedClipNode = static_cast<const NodeModifierMeshClipWithMesh*>(m.Node);
+				Ptr<ASTOpMeshMaskClipMesh> op = new ASTOpMeshMaskClipMesh();
+				op->source = PreModifiersMesh;
+
+				// Parameters
+				FMeshGenerationOptions ClipOptions;
+				ClipOptions.bLayouts = false;
+				ClipOptions.State = Options.State;
+
+				FMeshGenerationResult clipResult;
+				GenerateMesh(ClipOptions, clipResult, TypedClipNode->ClipMesh);
+				op->clip = clipResult.MeshOp;
+
+				if (!op->clip)
+				{
+					ErrorLog->GetPrivate()->Add("Clip mesh has not been generated", ELMT_ERROR, ErrorContext);
+					continue;
+				}
+
+				Ptr<ASTOp> maskAt = op;
+
+				if (!RemoveOp)
+				{
+					RemoveOp = new ASTOpMeshRemoveMask();
+					RemoveOp->source = LastMeshOp;
+					LastMeshOp = RemoveOp;
+				}
+
+				Ptr<ASTOp> fullCondition = m.FinalCondition;
+
+				RemoveOp->AddRemove(fullCondition, maskAt);
+			}
+		}
+
+		// Process clip-with-mask modifiers
+		for (const FirstPassGenerator::FModifier& m : Modifiers)
+		{
+			if (m.Node->GetType() == NodeModifierMeshClipWithUVMask::GetStaticType())
+			{
+				// Create a constant mesh with the original UVs required by this modifier.
+				// TODO: Optimize, by caching.
+				// TODO: Optimize by formatting and keeping only UVs
+				check(OriginalMeshNode);
+				const Mesh* OriginalMesh = OriginalMeshNode->GetPrivate()->Value.get();
+				Ptr<ASTOpConstantResource> UVMeshOp = new ASTOpConstantResource();
+				UVMeshOp->Type = OP_TYPE::ME_CONSTANT;
+				UVMeshOp->SetValue(OriginalMesh->Clone(), CompilerOptions->OptimisationOptions.DiskCacheContext);
+
+				const NodeModifierMeshClipWithUVMask* TypedClipNode = static_cast<const NodeModifierMeshClipWithUVMask*>(m.Node);
+
+				Ptr<ASTOp> MeshMaskAt;
+
+				Ptr<ASTOpMeshMaskClipUVMask> op = new ASTOpMeshMaskClipUVMask();
+				MeshMaskAt = op;
+				op->Source = BaseMeshResult.BaseMeshOp; 
+				op->UVSource = UVMeshOp; 
+				op->LayoutIndex = TypedClipNode->LayoutIndex;
+
+				if (TypedClipNode->ClipMask)
+				{
+					// Parameters to generate the mask image
+					FImageGenerationOptions ClipOptions;
+					ClipOptions.ImageLayoutStrategy = CompilerOptions::TextureLayoutStrategy::None;
+					ClipOptions.LayoutBlockId = FLayoutBlock::InvalidBlockId;
+					ClipOptions.State = Options.State;
+
+					FImageGenerationResult ClipMaskResult;
+					GenerateImage(ClipOptions, ClipMaskResult, TypedClipNode->ClipMask);
+
+					// It could be IF_L_UBIT, but since this should be optimized out at compile time, leave the most cpu efficient.
+					op->MaskImage = GenerateImageFormat(ClipMaskResult.op, mu::EImageFormat::IF_L_UBYTE);
+
+					if (!op->MaskImage)
+					{
+						ErrorLog->GetPrivate()->Add("Clip UV mask has not been generated", ELMT_ERROR, ErrorContext);
+						continue;
+					}
+				}
+
+				else if (TypedClipNode->ClipLayout)
+				{
+					// Generate the layout with blocks to extract
+					Ptr<const Layout> Layout = GenerateLayout(TypedClipNode->ClipLayout, 0);
+
+					Ptr<ASTOpConstantResource> LayoutOp = new ASTOpConstantResource();
+					LayoutOp->Type = OP_TYPE::LA_CONSTANT;
+					LayoutOp->SetValue(Layout, CompilerOptions->OptimisationOptions.DiskCacheContext);
+					op->MaskLayout = LayoutOp;
+				}
+
+				else
+				{
+					// No mask or layout specified to clip. Don't clip anything.
+				}
+
+				if (MeshMaskAt)
+				{
+					if (!RemoveOp)
+					{
+						RemoveOp = new ASTOpMeshRemoveMask();
+						RemoveOp->source = LastMeshOp;
+						LastMeshOp = RemoveOp;
+					}
+
+					Ptr<ASTOp> fullCondition = m.FinalCondition;
+					RemoveOp->AddRemove(fullCondition, MeshMaskAt);
+				}
+			}
+		}
+
+		// Process clip-morph-plane modifiers
+		for (const FirstPassGenerator::FModifier& m : Modifiers)
+		{
+			Ptr<ASTOp> modifiedMeshOp;
+
+			if (m.Node->GetType() == NodeModifierMeshClipMorphPlane::GetStaticType())
+			{
+				const NodeModifierMeshClipMorphPlane* TypedNode = static_cast<const NodeModifierMeshClipMorphPlane*>(m.Node);
+				Ptr<ASTOpMeshClipMorphPlane> op = new ASTOpMeshClipMorphPlane();
+				op->source = LastMeshOp;
+
+				// Morph to an ellipse
+				{
+					FShape morphShape;
+					morphShape.type = (uint8_t)FShape::Type::Ellipse;
+					morphShape.position = TypedNode->Parameters.Origin;
+					morphShape.up = TypedNode->Parameters.Normal;
+					// TODO: Move rotation to ellipse rotation reference base instead of passing it directly
+					morphShape.size = FVector3f(TypedNode->Parameters.Radius1, TypedNode->Parameters.Radius2, TypedNode->Parameters.Rotation);
+
+					// Generate a "side" vector.
+					// \todo: make generic and move to the vector class
+					{
+						// Generate vector perpendicular to normal for ellipse rotation reference base
+						FVector3f aux_base(0.f, 1.f, 0.f);
+
+						if (FMath::Abs(FVector3f::DotProduct(TypedNode->Parameters.Normal, aux_base)) > 0.95f)
+						{
+							aux_base = FVector3f(0.f, 0.f, 1.f);
+						}
+
+						morphShape.side = FVector3f::CrossProduct(TypedNode->Parameters.Normal, aux_base);
+					}
+					op->morphShape = morphShape;
+				}
+
+				// Selection box
+				if (TypedNode->Parameters.VertexSelectionType == FClipMorphPlaneParameters::VS_SHAPE)
+				{
+					op->vertexSelectionType = OP::MeshClipMorphPlaneArgs::VS_SHAPE;
+					FShape selectionShape;
+					selectionShape.type = (uint8)FShape::Type::AABox;
+					selectionShape.position = TypedNode->Parameters.SelectionBoxOrigin;
+					selectionShape.size = TypedNode->Parameters.SelectionBoxRadius;
+					op->selectionShape = selectionShape;
+				}
+				else if (TypedNode->Parameters.VertexSelectionType == FClipMorphPlaneParameters::VS_BONE_HIERARCHY)
+				{
+					op->vertexSelectionType = OP::MeshClipMorphPlaneArgs::VS_BONE_HIERARCHY;
+					op->vertexSelectionBone = TypedNode->Parameters.VertexSelectionBone;
+					op->vertexSelectionBoneMaxRadius = TypedNode->Parameters.MaxEffectRadius;
+				}
+				else
+				{
+					op->vertexSelectionType = OP::MeshClipMorphPlaneArgs::VS_NONE;
+				}
+
+				// Parameters
+				op->dist = TypedNode->Parameters.DistanceToPlane;
+				op->factor = TypedNode->Parameters.LinearityFactor;
+
+				modifiedMeshOp = op;
+
+				Ptr<ASTOp> fullCondition = m.FinalCondition;
+
+				Ptr<ASTOpConditional> ConditionalOp = new ASTOpConditional();
+				ConditionalOp->type = OP_TYPE::ME_CONDITIONAL;
+				ConditionalOp->no = LastMeshOp;
+				ConditionalOp->yes = modifiedMeshOp;
+				ConditionalOp->condition = fullCondition;
+				LastMeshOp = ConditionalOp;
+			}
+		}
+
+    	// Process clip deform modifiers.
+		for (const FirstPassGenerator::FModifier& M : Modifiers)
+		{
+			Ptr<ASTOp> ModifiedMeshOp;
+
+			if (M.Node->GetType()==NodeModifierMeshClipDeform::GetStaticType())
+			{
+				const NodeModifierMeshClipDeform* TypedClipNode = static_cast<const NodeModifierMeshClipDeform*>(M.Node);
+				Ptr<ASTOpMeshBindShape>  BindOp = new ASTOpMeshBindShape();
+				Ptr<ASTOpMeshClipDeform> ClipOp = new ASTOpMeshClipDeform();
+
+				FMeshGenerationOptions ClipOptions;
+				ClipOptions.bLayouts = false;
+				ClipOptions.State = Options.State;
+
+				FMeshGenerationResult ClipShapeResult;
+				GenerateMesh(ClipOptions, ClipShapeResult, TypedClipNode->ClipMesh);
+				ClipOp->ClipShape = ClipShapeResult.MeshOp;
+				
+				BindOp->Mesh = LastMeshOp;
+				BindOp->Shape = ClipShapeResult.MeshOp; 
+				BindOp->BindingMethod = static_cast<uint32>(TypedClipNode->BindingMethod);
+	
+				ClipOp->Mesh = BindOp;
+
+				if (!ClipOp->ClipShape)
+				{
+					ErrorLog->GetPrivate()->Add
+					("Clip shape mesh has not been generated", ELMT_ERROR, ErrorContext);
+				}
+				else
+				{
+					ModifiedMeshOp = ClipOp;
+				}
+			}
+			
+			if (ModifiedMeshOp)
+			{
+				Ptr<ASTOp> FullCondition = M.FinalCondition;
+
+				Ptr<ASTOpConditional> Op = new ASTOpConditional();
+				Op->type = OP_TYPE::ME_CONDITIONAL;
+				Op->no = LastMeshOp;
+				Op->yes = ModifiedMeshOp;
+				Op->condition = FullCondition;
+				LastMeshOp = Op;
+			}
+		}
+			
+		ActiveTags.Pop();
+
+		return LastMeshOp;
+	}
+
+
+	Ptr<ASTOp> CodeGenerator::ApplyImageBlockModifiers(const FImageGenerationOptions& Options, Ptr<ASTOp> BaseImageOp, int32 ImageIndex, 
+		FIntPoint GridSize,
+		const FLayoutBlockDesc& LayoutBlockDesc,
+		box< FIntVector2 > RectInCells,
+		bool bModifiersForBeforeOperations, const void* ErrorContext)
+	{
+		Ptr<ASTOp> LastImageOp = BaseImageOp;
+
+		// Apply mesh modifiers
+		TArray<FirstPassGenerator::FModifier> Modifiers;
+
+		int32 currentLOD = CurrentParents.Last().Lod;
+		GetModifiersFor(Options.ActiveTags, currentLOD, bModifiersForBeforeOperations, Modifiers);
+
+		ActiveTags.Push({});
+
+		// Process patch image modifiers (from edit modifiers)
+		for (const FirstPassGenerator::FModifier& m : Modifiers)
+		{
+			if (m.Node->GetType() == NodeModifierSurfaceEdit::GetStaticType())
+			{
+				const NodeModifierSurfaceEdit* Edit = static_cast<const NodeModifierSurfaceEdit*>(m.Node);
+
+
+				if (ImageIndex >= Edit->Textures.Num())
+				{
+					continue;
+				}
+
+				const NodeModifierSurfaceEdit::FTexture& Patch = Edit->Textures[ImageIndex];
+				if (Patch.PatchImage.get())
+				{
+					// Does the current block need to be patched? Find out by building a mask.
+					Ptr<Image> PatchMask = GenerateImageBlockPatchMask(Patch, GridSize, LayoutBlockDesc.BlockPixelsX, LayoutBlockDesc.BlockPixelsY, RectInCells);
+
+					if (PatchMask)
+					{
+						LastImageOp = GenerateImageBlockPatch(LastImageOp, Patch, PatchMask, m.FinalCondition, Options);
+					}
+				}
+			}
+
+		}
+
+		ActiveTags.Pop();
+
+		return LastImageOp;
+	}
+
+
+	void CodeGenerator::UpdateLayoutBlockDesc(CodeGenerator::FLayoutBlockDesc& Out, FImageDesc BlockDesc, FIntVector2 LayoutCellSize)
+	{
+		if (Out.BlockPixelsX == 0 && LayoutCellSize.X > 0 && LayoutCellSize.Y > 0)
+		{
+			Out.BlockPixelsX = FMath::Max(1, BlockDesc.m_size[0] / LayoutCellSize[0]);
+			Out.BlockPixelsY = FMath::Max(1, BlockDesc.m_size[1] / LayoutCellSize[1]);
+			Out.bBlocksHaveMips = BlockDesc.m_lods > 1;
+
+			if (Out.FinalFormat==EImageFormat::IF_NONE)
+			{
+				Out.FinalFormat = BlockDesc.m_format;
+			}
+		}
+	};
+
+
+	Ptr<ASTOp> CodeGenerator::ApplyImageExtendModifiers(
+		const FGenericGenerationOptions& Options, 
+		const FMeshGenerationResult& BaseMeshResults,
+		Ptr<ASTOp> BaseImageOp, 
+		CompilerOptions::TextureLayoutStrategy ImageLayoutStrategy,
+		int32 LayoutIndex, int32 ImageIndex,
+		FIntPoint GridSize, 
+		CodeGenerator::FLayoutBlockDesc& InOutLayoutBlockDesc,
+		bool bModifiersForBeforeOperations, const void* ErrorContext)
+	{
+		Ptr<ASTOp> LastImageOp = BaseImageOp;
+
+		// Apply mesh modifiers
+		TArray<FirstPassGenerator::FModifier> Modifiers;
+
+		int32 currentLOD = CurrentParents.Last().Lod;
+		GetModifiersFor(Options.ActiveTags, currentLOD, bModifiersForBeforeOperations, Modifiers);
+
+		ActiveTags.Push({});
+
+		// Process mesh extend modifiers (from edit modifiers)
+		int32 EditIndex = 0;
+		for (const FirstPassGenerator::FModifier& m : Modifiers)
+		{
+			if (m.Node->GetType() == NodeModifierSurfaceEdit::GetStaticType())
+			{
+				const NodeModifierSurfaceEdit* Edit = static_cast<const NodeModifierSurfaceEdit*>(m.Node);
+
+				int32 ThisEditIndex = EditIndex;
+				++EditIndex;
+
+				if (ImageIndex >= Edit->Textures.Num())
+				{
+					continue;
+				}
+
+				Ptr<NodeImage> pExtend = Edit->Textures[ImageIndex].Extend;
+				if (!pExtend)
+				{
+					continue;
+				}
+
+				const TArray<FGeneratedLayout>& ExtraLayouts = BaseMeshResults.ExtraMeshLayouts[ThisEditIndex].GeneratedLayouts;
+
+				if (LayoutIndex >= ExtraLayouts.Num() || !ExtraLayouts[LayoutIndex].Layout)
+				{
+					FString Msg = FString::Printf(TEXT("Trying to extend a layout that doesn't exist in object [%s]."), *CurrentParents.Last().ObjectNode->Name );
+					ErrorLog->GetPrivate()->Add(Msg, ELMT_ERROR, ErrorContext);
+				}
+				else
+				{
+					Ptr<const Layout> pExtendLayout = ExtraLayouts[LayoutIndex].Layout;
+
+					Ptr<ASTOp> lastBase = LastImageOp;
+
+					for (int32 b = 0; b < pExtendLayout->GetBlockCount(); ++b)
+					{
+						// Generate the image block
+						FImageGenerationOptions ImageOptions;
+						ImageOptions.State = Options.State;
+						ImageOptions.ImageLayoutStrategy = ImageLayoutStrategy;
+						ImageOptions.ActiveTags = Edit->EnableTags; // TODO: Merge with current tags?
+						ImageOptions.RectSize = { 0,0 };
+						ImageOptions.LayoutToApply = pExtendLayout;
+						ImageOptions.LayoutBlockId = pExtendLayout->Blocks[b].Id;
+						FImageGenerationResult ExtendResult;
+						GenerateImage(ImageOptions, ExtendResult, pExtend);
+						Ptr<ASTOp> fragmentAd = ExtendResult.op;
+
+						// Block in layout grid units
+						box< FIntVector2 > rectInCells;
+						rectInCells.min = pExtendLayout->Blocks[b].Min;
+						rectInCells.size = pExtendLayout->Blocks[b].Size;
+
+						FImageDesc ExtendDesc = fragmentAd->GetImageDesc();
+
+						// If we don't know the size of a layout block in pixels, calculate it
+						UpdateLayoutBlockDesc(InOutLayoutBlockDesc, ExtendDesc, rectInCells.size);
+
+						// Adjust the format and size of the block to be added
+						// Actually don't do it, it will be propagated from the top format operation.
+						//fragmentAd = GenerateImageFormat(fragmentAd, FinalImageFormat);
+
+						UE::Math::TIntVector2<int32> expectedSize;
+						expectedSize[0] = InOutLayoutBlockDesc.BlockPixelsX * rectInCells.size[0];
+						expectedSize[1] = InOutLayoutBlockDesc.BlockPixelsY * rectInCells.size[1];
+						fragmentAd = GenerateImageSize(fragmentAd, expectedSize);
+
+						// Apply tiling to avoid generating chunks of image that are too big.
+						fragmentAd = ApplyTiling(fragmentAd, expectedSize, InOutLayoutBlockDesc.FinalFormat);
+
+						// Compose operation
+						Ptr<ASTOpImageCompose> composeOp = new ASTOpImageCompose();
+						composeOp->Layout = BaseMeshResults.LayoutOps[LayoutIndex];
+						composeOp->Base = lastBase;
+						composeOp->BlockImage = fragmentAd;
+
+						// Set the absolute block index.
+						check(pExtendLayout->Blocks[b].Id != FLayoutBlock::InvalidBlockId);
+						composeOp->BlockId = pExtendLayout->Blocks[b].Id;
+
+						lastBase = composeOp;
+					}
+
+					// Condition to enable this image extension
+					if (m.FinalCondition)
+					{
+						Ptr<ASTOp> conditionalAd;
+						Ptr<ASTOpConditional> cop = new ASTOpConditional();
+						cop->type = OP_TYPE::IM_CONDITIONAL;
+						cop->no = LastImageOp;
+						cop->yes = lastBase;
+						cop->condition = m.FinalCondition;
+						conditionalAd = cop;
+						LastImageOp = conditionalAd;
+					}
+					else
+					{
+						LastImageOp = lastBase;
+					}
+
+				}
+			}
+		}
+
+		ActiveTags.Pop();
+
+		return LastImageOp;
+	}
+
+
+	Ptr<ASTOp> CodeGenerator::GenerateDefaultTableValue(ETableColumnType NodeType)
+	{
+		switch (NodeType)
+		{
+		case mu::ETableColumnType::Scalar:
+		{
+			//TODO(Max):MTBL-1660
+			//Ptr<NodeScalarConstant> pNode = new NodeScalarConstant();
+			//pNode->SetValue(-UE_MAX_FLT);
+			//
+			//return Generate(pNode);
+			return nullptr;
+		}
+		case mu::ETableColumnType::Color:
+		{
+			mu::Ptr<mu::NodeColourConstant> pNode = new NodeColourConstant();
+			pNode->Value = mu::DefaultMutableColorValue;
+
+			FColorGenerationResult ChildResult;
+			FGenericGenerationOptions Options;
+			GenerateColor(ChildResult, Options, pNode);
+			return ChildResult.op;
+		}
+		case mu::ETableColumnType::Image:
+		{
+			//TODO(Max):MTBL-1660
+			//FImageGenerationOptions DummyOptions;
+			//FImageGenerationResult DefaultValue;
+			//
+			//mu::Ptr<mu::NodeImageReference> ImageNode = new mu::NodeImageReference();
+			//ImageNode->SetImageReference(-1);
+			//
+			//GenerateImage(DummyOptions, DefaultValue, ImageNode);
+			//
+			//return DefaultValue.op;
+			return nullptr;
+		}
+		case mu::ETableColumnType::Mesh:
+			/* The default mesh is always null */
+			return nullptr;
+		default:
+			break;
+		}
+
+		return nullptr;
+	}
+
+}
