[mutable] Moved the Mutable plugin out of Experimental status into Beta.

diff --git a/Engine/Plugins/Mutable/Source/MutableTools/Private/MuT/CodeGenerator.h b/Engine/Plugins/Mutable/Source/MutableTools/Private/MuT/CodeGenerator.h
new file mode 100644
--- /dev/null
+++ b/Engine/Plugins/Mutable/Source/MutableTools/Private/MuT/CodeGenerator.h
@@ -0,0 +1,907 @@
+// Copyright Epic Games, Inc. All Rights Reserved.
+
+#pragma once
+
+#include "Containers/Map.h"
+#include "HAL/PlatformMath.h"
+#include "Misc/AssertionMacros.h"
+#include "MuR/Image.h"
+#include "MuR/Mesh.h"
+#include "MuR/MutableMath.h"
+#include "MuR/MutableMemory.h"
+#include "MuR/Operations.h"
+#include "MuR/Parameters.h"
+#include "MuR/Ptr.h"
+#include "MuR/RefCounted.h"
+#include "MuR/System.h"
+#include "MuT/AST.h"
+#include "MuT/ASTOpSwitch.h"
+#include "MuT/CodeGenerator_FirstPass.h"
+#include "MuT/Compiler.h"
+#include "MuT/ErrorLog.h"
+#include "MuT/Node.h"
+#include "MuT/NodeBool.h"
+#include "MuT/NodeColourSampleImage.h"
+#include "MuT/NodeComponentEdit.h"
+#include "MuT/NodeComponentNew.h"
+#include "MuT/NodeComponentSwitch.h"
+#include "MuT/NodeExtensionData.h"
+#include "MuT/NodeImageProject.h"
+#include "MuT/NodeLOD.h"
+#include "MuT/NodeLayout.h"
+#include "MuT/NodeMeshGeometryOperation.h"
+#include "MuT/NodeMeshTable.h"
+#include "MuT/NodeObjectGroup.h"
+#include "MuT/NodeObjectNew.h"
+#include "MuT/NodeProjector.h"
+#include "MuT/NodeString.h"
+#include "MuT/NodeSurfaceNew.h"
+#include "MuT/NodeModifierSurfaceEdit.h"
+#include "MuT/Table.h"
+#include "MuT/TablePrivate.h"
+#include "Templates/TypeHash.h"
+
+
+namespace mu
+{
+	class ASTOpParameter;
+	class Layout;
+	class NodeColourArithmeticOperation;
+	class NodeColourConstant;
+	class NodeColourFromScalars;
+	class NodeColourParameter;
+	class NodeColourSwitch;
+	class NodeColourTable;
+	class NodeColourVariation;
+	class NodeImageBinarise;
+	class NodeImageColourMap;
+	class NodeImageConditional;
+	class NodeImageConstant;
+	class NodeImageFormat;
+	class NodeImageGradient;
+	class NodeImageInterpolate;
+	class NodeImageInvert;
+	class NodeImageLayer;
+	class NodeImageLayerColour;
+	class NodeImageLuminance;
+	class NodeImageMipmap;
+	class NodeImageMultiLayer;
+	class NodeImageNormalComposite;
+	class NodeImageParameter;
+	class NodeImagePlainColour;
+	class NodeImageResize;
+	class NodeImageSaturate;
+	class NodeImageSwitch;
+	class NodeImageSwizzle;
+	class NodeImageTable;
+	class NodeImageTransform;
+	class NodeImageVariation;
+	class NodeMeshApplyPose;
+	class NodeMeshClipDeform;
+	class NodeMeshClipMorphPlane;
+	class NodeMeshClipWithMesh;
+	class NodeMeshConstant;
+	class NodeMeshFormat;
+	class NodeMeshFragment;
+	class NodeMeshInterpolate;
+	class NodeMeshMakeMorph;
+	class NodeMeshMorph;
+	class NodeMeshReshape;
+	class NodeMeshSwitch;
+	class NodeMeshTransform;
+	class NodeMeshVariation;
+	class NodeRange;
+	class NodeScalarArithmeticOperation;
+	class NodeScalarConstant;
+	class NodeScalarCurve;
+	class NodeScalarEnumParameter;
+	class NodeScalarParameter;
+	class NodeScalarSwitch;
+	class NodeScalarTable;
+	class NodeScalarVariation;
+	class NodeStringConstant;
+	class NodeStringParameter;
+	struct FObjectState;
+	struct FProgram;
+
+
+
+    //---------------------------------------------------------------------------------------------
+    //! Code generator
+    //---------------------------------------------------------------------------------------------
+    class CodeGenerator
+    {
+		
+		friend class FirstPassGenerator;
+
+    public:
+
+        CodeGenerator( CompilerOptions::Private* options );
+
+        //! Data will be stored in States
+        void GenerateRoot( const Ptr<const Node> );
+
+	public:
+
+		// Generic top-level nodes
+		struct FGenericGenerationOptions
+		{
+			friend FORCEINLINE uint32 GetTypeHash(const FGenericGenerationOptions& InKey)
+			{
+				uint32 KeyHash = 0;
+				KeyHash = HashCombineFast(KeyHash, ::GetTypeHash(InKey.State));
+				KeyHash = HashCombineFast(KeyHash, ::GetTypeHash(InKey.ActiveTags.Num()));
+				return KeyHash;
+			}
+
+			bool operator==(const FGenericGenerationOptions& InKey) const = default;
+
+			int32 State = -1;
+			TArray<FString> ActiveTags;
+		};
+
+		struct FComponentGenerationOptions : public FGenericGenerationOptions
+		{
+			FComponentGenerationOptions(const FGenericGenerationOptions& BaseOptions, const Ptr<ASTOp>& InBaseInstance )
+			{
+				BaseInstance = InBaseInstance;
+				State = BaseOptions.State;
+				ActiveTags = BaseOptions.ActiveTags;
+			}
+
+			/** Instance to which the possibly generated components should be added. */
+			Ptr<ASTOp> BaseInstance;
+		};
+
+		struct FLODGenerationOptions : public FGenericGenerationOptions
+		{
+			FLODGenerationOptions(const FGenericGenerationOptions& BaseOptions, int32 InLODIndex, const NodeComponentNew* InComponent)
+			{
+				Component = InComponent;
+				LODIndex = InLODIndex;
+				State = BaseOptions.State;
+				ActiveTags = BaseOptions.ActiveTags;
+			}
+
+			const NodeComponentNew* Component = nullptr;
+			int32 LODIndex;
+		};
+
+		struct FSurfaceGenerationOptions : public FGenericGenerationOptions
+		{
+			explicit FSurfaceGenerationOptions(const FGenericGenerationOptions& BaseOptions)
+			{
+				Component = nullptr;
+				LODIndex = -1;
+				State = BaseOptions.State;
+				ActiveTags = BaseOptions.ActiveTags;
+			}
+
+			explicit FSurfaceGenerationOptions(const FLODGenerationOptions& BaseOptions)
+			{
+				Component = BaseOptions.Component;
+				LODIndex = BaseOptions.LODIndex;
+				State = BaseOptions.State;
+				ActiveTags = BaseOptions.ActiveTags;
+			}
+
+			const NodeComponentNew* Component = nullptr;
+			int32 LODIndex = -1;
+		};
+
+		struct FGenericGenerationResult
+		{
+			Ptr<ASTOp> op;
+		};
+
+		struct FGeneratedCacheKey
+		{
+			Ptr<const Node> Node;
+			FGenericGenerationOptions Options;
+
+			friend FORCEINLINE uint32 GetTypeHash(const FGeneratedCacheKey& InKey)
+			{
+				uint32 KeyHash = 0;
+				KeyHash = HashCombineFast(KeyHash, ::GetTypeHash(InKey.Node.get()));
+				KeyHash = HashCombineFast(KeyHash, GetTypeHash(InKey.Options));
+				return KeyHash;
+			}
+
+			FORCEINLINE bool operator==(const FGeneratedCacheKey& Other) const = default;
+		};
+
+		typedef TMap<FGeneratedCacheKey, FGenericGenerationResult> FGeneratedGenericNodesMap;
+		FGeneratedGenericNodesMap GeneratedGenericNodes;
+
+		Ptr<ASTOp> Generate_Generic(const Ptr<const Node>, const FGenericGenerationOptions& );
+		void Generate_LOD(const FLODGenerationOptions&, FGenericGenerationResult&, const NodeLOD*);
+		void Generate_ObjectNew(const FGenericGenerationOptions&, FGenericGenerationResult&, const NodeObjectNew*);
+		void Generate_ObjectGroup(const FGenericGenerationOptions&, FGenericGenerationResult&, const NodeObjectGroup*);
+
+    public:
+
+        //! Settings
+        CompilerOptions::Private* CompilerOptions = nullptr;
+
+		//!
+		FirstPassGenerator FirstPass;
+
+        //!
+        ErrorLogPtr ErrorLog;
+
+        //! After the entire code generation this contains the information about all the states
+        typedef TArray< TPair<FObjectState, Ptr<ASTOp>> > StateList;
+        StateList States;
+
+    private:
+
+        /** Container of meshes generated to be able to reuse them. 
+		* They are sorted by a cheap hash to speed up searches.
+		*/
+		struct FGeneratedConstantMesh
+		{
+			Ptr<Mesh> Mesh;
+			Ptr<ASTOp> LastMeshOp;
+		};
+		TMap<uint64,TArray<FGeneratedConstantMesh>> GeneratedConstantMeshes;
+
+        /** List of image resources for every image formats that have been generated so far as palceholders for missing images. */
+        Ptr<Image> MissingImage[size_t(EImageFormat::IF_COUNT)];
+
+
+        //! List of already used vertex ID groups that must be unique.
+        TSet<uint32> UniqueVertexIDGroups;
+
+		// (top-down) Tags that are active when generating nodes.
+		TArray< TArray<FString> > ActiveTags;
+
+        struct FParentKey
+        {
+            const NodeObjectNew* ObjectNode = nullptr;
+            int32 Lod = -1;
+         };
+
+		TArray< FParentKey > CurrentParents;
+
+        /** List of additional components to add to an object that come from child objects.
+        * The index is the object and lod that should receive the components.
+		*/
+        struct FAdditionalComponentKey
+        {
+			FAdditionalComponentKey()
+            {
+				ObjectNode = nullptr;
+            }
+
+            const NodeObjectNew* ObjectNode;
+
+			FORCEINLINE bool operator==(const FAdditionalComponentKey& Other) const
+			{
+				return ObjectNode == Other.ObjectNode;
+			}
+
+			friend FORCEINLINE uint32 GetTypeHash(const FAdditionalComponentKey& InKey)
+			{
+				uint32 KeyHash = 0;
+				KeyHash = HashCombineFast(KeyHash, ::GetTypeHash(InKey.ObjectNode));
+				return KeyHash;
+			}
+		};
+
+		struct FAdditionalComponentData
+		{
+			Ptr<ASTOp> ComponentOp;
+			Ptr<ASTOp> PlaceholderOp;
+		};
+
+        TMap< FAdditionalComponentKey, TArray<FAdditionalComponentData> > AdditionalComponents;
+
+
+        struct FObjectGenerationData
+        {
+            // Condition that enables a specific object
+            Ptr<ASTOp> Condition;
+        };
+		TArray<FObjectGenerationData> CurrentObject;
+
+		/** The key for generated tables is made of the source table and a parameter name. */
+		struct FTableCacheKey
+		{
+			Ptr<const Table> Table;
+			FString ParameterName;
+
+			friend FORCEINLINE uint32 GetTypeHash(const FTableCacheKey& InKey)
+			{
+				uint32 KeyHash = ::GetTypeHash(InKey.Table.get());
+				KeyHash = HashCombineFast(KeyHash, GetTypeHash(InKey.ParameterName));
+				return KeyHash;
+			}
+
+			FORCEINLINE bool operator==(const FTableCacheKey& InKey) const
+			{
+				if (Table != InKey.Table) return false;
+				if (ParameterName != InKey.ParameterName) return false;
+				return true;
+			}
+		};
+		TMap< FTableCacheKey, Ptr<ASTOp> > GeneratedTables;
+
+		struct FConditionalExtensionDataOp
+		{
+			Ptr<ASTOp> Condition;
+			Ptr<ASTOp> ExtensionDataOp;
+			FString ExtensionDataName;
+		};
+
+		TArray<FConditionalExtensionDataOp> ConditionalExtensionDataOps;
+
+		struct FGeneratedComponentCacheKey
+		{
+			Ptr<const Node> Node;
+			FGenericGenerationOptions Options;
+
+			friend FORCEINLINE uint32 GetTypeHash(const FGeneratedComponentCacheKey& InKey)
+			{
+				uint32 KeyHash = 0;
+				KeyHash = HashCombineFast(KeyHash, ::GetTypeHash(InKey.Node.get()));
+				KeyHash = HashCombineFast(KeyHash, GetTypeHash(InKey.Options));
+				return KeyHash;
+			}
+
+			FORCEINLINE bool operator==(const FGeneratedComponentCacheKey& Other) const
+			{
+				return Node == Other.Node && Options == Other.Options;
+			}
+		};
+
+		typedef TMap<FGeneratedComponentCacheKey, FGenericGenerationResult> GeneratedComponentMap;
+		GeneratedComponentMap GeneratedComponents;
+
+		void GenerateComponent(const FComponentGenerationOptions&, FGenericGenerationResult&, const NodeComponent*);
+		void GenerateComponent_New(const FComponentGenerationOptions&, FGenericGenerationResult&, const NodeComponentNew*);
+		void GenerateComponent_Switch(const FComponentGenerationOptions&, FGenericGenerationResult&, const NodeComponentSwitch*);
+		void GenerateComponent_Variation(const FComponentGenerationOptions&, FGenericGenerationResult&, const NodeComponentVariation*);
+
+        //-----------------------------------------------------------------------------------------
+        //!
+        Ptr<ASTOp> GenerateTableVariable(Ptr<const Node>, const FTableCacheKey&, bool bAddNoneOption, const FString& DefaultRowName);
+
+        //!
+        Ptr<ASTOp> GenerateMissingBoolCode(const TCHAR* strWhere, bool value, const void* errorContext );
+
+        //!
+		template<class NODE_TABLE, ETableColumnType TYPE, OP_TYPE OPTYPE, typename F>
+		Ptr<ASTOp> GenerateTableSwitch( const NODE_TABLE& node, F&& GenerateOption );
+
+
+		//-----------------------------------------------------------------------------------------
+		// Images
+		
+		/** Options that affect the generation of images. It is like list of what required data we want while parsing down the image node graph. */
+		struct FImageGenerationOptions : public FGenericGenerationOptions
+		{
+			/** */
+			CompilerOptions::TextureLayoutStrategy ImageLayoutStrategy = CompilerOptions::TextureLayoutStrategy::None;
+
+			/** If different than {0,0} this is the mandatory size of the image that needs to be generated. */
+			UE::Math::TIntVector2<int32> RectSize = {0, 0};
+
+			/** Layout block that we are trying to generate if any. */
+			uint64 LayoutBlockId = FLayoutBlock::InvalidBlockId;
+			Ptr<const Layout> LayoutToApply;
+
+			friend FORCEINLINE uint32 GetTypeHash(const FImageGenerationOptions& InKey)
+			{
+				uint32 KeyHash = 0;
+				KeyHash = HashCombineFast(KeyHash, ::GetTypeHash(InKey.ImageLayoutStrategy));
+				KeyHash = HashCombineFast(KeyHash, GetTypeHash(InKey.RectSize));
+				KeyHash = HashCombineFast(KeyHash, ::GetTypeHash(InKey.State));
+				KeyHash = HashCombineFast(KeyHash, ::GetTypeHash(InKey.LayoutBlockId));
+				KeyHash = HashCombineFast(KeyHash, ::GetTypeHash(InKey.LayoutToApply.get()));
+				return KeyHash;
+			}
+
+			FORCEINLINE bool operator==(const FImageGenerationOptions& Other) const
+			{
+				return ImageLayoutStrategy == Other.ImageLayoutStrategy
+					&&
+					RectSize == Other.RectSize
+					&&
+					State == Other.State
+					&&
+					LayoutBlockId == Other.LayoutBlockId
+					&&
+					LayoutToApply == Other.LayoutToApply
+					&&
+					ActiveTags == Other.ActiveTags;
+			}
+
+		};
+
+		/** */
+		struct FImageGenerationResult
+		{
+			Ptr<ASTOp> op;
+		};
+
+		/** */
+		struct FGeneratedImageCacheKey
+		{
+			NodePtrConst Node;
+			FImageGenerationOptions Options;
+
+			friend FORCEINLINE uint32 GetTypeHash(const FGeneratedImageCacheKey& InKey)
+			{
+				uint32 KeyHash = 0;
+				KeyHash = HashCombineFast(KeyHash, ::GetTypeHash(InKey.Node.get()));
+				KeyHash = HashCombineFast(KeyHash, GetTypeHash(InKey.Options));
+				return KeyHash;
+			}
+
+			FORCEINLINE bool operator==(const FGeneratedImageCacheKey& Other) const
+			{
+				return Node == Other.Node && Options == Other.Options;
+			}
+		};
+
+		typedef TMap<FGeneratedImageCacheKey, FImageGenerationResult> GeneratedImagesMap;
+		GeneratedImagesMap GeneratedImages;
+
+		void GenerateImage(const FImageGenerationOptions&, FImageGenerationResult& result, const NodeImagePtrConst& node);
+		void GenerateImage_Constant(const FImageGenerationOptions&, FImageGenerationResult&, const NodeImageConstant*);
+		void GenerateImage_Interpolate(const FImageGenerationOptions&, FImageGenerationResult&, const NodeImageInterpolate*);
+		void GenerateImage_Saturate(const FImageGenerationOptions&, FImageGenerationResult&, const NodeImageSaturate*);
+		void GenerateImage_Table(const FImageGenerationOptions&, FImageGenerationResult&, const NodeImageTable*);
+		void GenerateImage_Swizzle(const FImageGenerationOptions&, FImageGenerationResult&, const NodeImageSwizzle*);
+		void GenerateImage_ColourMap(const FImageGenerationOptions&, FImageGenerationResult&, const NodeImageColourMap*);
+		void GenerateImage_Gradient(const FImageGenerationOptions&, FImageGenerationResult&, const NodeImageGradient*);
+		void GenerateImage_Binarise(const FImageGenerationOptions&, FImageGenerationResult&, const NodeImageBinarise*);
+		void GenerateImage_Luminance(const FImageGenerationOptions&, FImageGenerationResult&, const NodeImageLuminance*);
+		void GenerateImage_Layer(const FImageGenerationOptions&, FImageGenerationResult&, const NodeImageLayer*);
+		void GenerateImage_LayerColour(const FImageGenerationOptions&, FImageGenerationResult&, const NodeImageLayerColour*);
+		void GenerateImage_Resize(const FImageGenerationOptions&, FImageGenerationResult&, const NodeImageResize*);
+		void GenerateImage_PlainColour(const FImageGenerationOptions&, FImageGenerationResult&, const NodeImagePlainColour*);
+		void GenerateImage_Project(const FImageGenerationOptions&, FImageGenerationResult&, const NodeImageProject*);
+		void GenerateImage_Mipmap(const FImageGenerationOptions&, FImageGenerationResult&, const NodeImageMipmap*);
+		void GenerateImage_Switch(const FImageGenerationOptions&, FImageGenerationResult&, const NodeImageSwitch*);
+		void GenerateImage_Conditional(const FImageGenerationOptions&, FImageGenerationResult&, const NodeImageConditional*);
+		void GenerateImage_Format(const FImageGenerationOptions&, FImageGenerationResult&, const NodeImageFormat*);
+		void GenerateImage_Parameter(const FImageGenerationOptions&, FImageGenerationResult&, const NodeImageParameter*);
+		void GenerateImage_MultiLayer(const FImageGenerationOptions&, FImageGenerationResult&, const NodeImageMultiLayer*);
+		void GenerateImage_Invert(const FImageGenerationOptions&, FImageGenerationResult&, const NodeImageInvert*);
+		void GenerateImage_Variation(const FImageGenerationOptions&, FImageGenerationResult&, const NodeImageVariation*);
+		void GenerateImage_NormalComposite(const FImageGenerationOptions&, FImageGenerationResult&, const NodeImageNormalComposite*);
+		void GenerateImage_Transform(const FImageGenerationOptions&, FImageGenerationResult&, const NodeImageTransform*);
+
+		//!
+		Ptr<Image> GenerateMissingImage(EImageFormat);
+
+		//!
+		Ptr<ASTOp> GenerateMissingImageCode(const TCHAR* strWhere, EImageFormat, const void* errorContext, const FImageGenerationOptions& Options);
+
+		//!
+		Ptr<ASTOp> GeneratePlainImageCode(const FVector4f& Color, const FImageGenerationOptions& Options);
+
+		//!
+		Ptr<ASTOp> GenerateImageFormat(Ptr<ASTOp>, EImageFormat);
+
+		//!
+		Ptr<ASTOp> GenerateImageUncompressed(Ptr<ASTOp>);
+
+		//!
+		Ptr<ASTOp> GenerateImageSize(Ptr<ASTOp>, UE::Math::TIntVector2<int32>);
+
+		/** Evaluate if the image to generate is big enough to be split in separate operations and tiled afterwards. */
+		Ptr<ASTOp> ApplyTiling(Ptr<ASTOp> Source, UE::Math::TIntVector2<int32> Size, EImageFormat Format);
+
+		/** Generate a layout block-sized image with a mask including all pixels in the blocks defined in the patch node. */
+		Ptr<Image> GenerateImageBlockPatchMask(const NodeModifierSurfaceEdit::FTexture&, FIntPoint GridSize, int32 BlockPixelsX, int32 BlockPixelsY, box<FIntVector2> RectInCells);
+
+		/** Generate all the operations to apply the block patching on top of the BlockOp, and masking with PatchMask. */
+		Ptr<ASTOp> GenerateImageBlockPatch(Ptr<ASTOp> BlockOp, const NodeModifierSurfaceEdit::FTexture&, Ptr<Image> PatchMask, Ptr<ASTOp> ConditionOp, const FImageGenerationOptions&);
+
+        //-----------------------------------------------------------------------------------------
+        // Meshes
+
+		/** */
+		struct FGeneratedLayout
+		{
+			Ptr<const Layout> Layout;
+			Ptr<const NodeLayout> Source;
+
+			FORCEINLINE bool operator==(const FGeneratedLayout& Other) const
+			{
+				return Layout == Other.Layout
+					&& Source == Other.Source;
+			}
+		};
+
+		/** Options that affect the generation of meshes. It is like list of what required data we want
+		* while parsing down the mesh node graph.
+		*/
+		struct FMeshGenerationOptions : public FGenericGenerationOptions
+		{
+			/** The meshes at the leaves will need their own layout block data. */
+			bool bLayouts = false;
+
+			/** If true, Ensure UV Islands are not split between two or more blocks. UVs shared between multiple 
+			* layout blocks will be clamped to fit the one with more vertices belonging to the UV island.
+			* Mainly used to keep consistent layouts when reusing textures between LODs. */
+			bool bClampUVIslands = false;
+
+			/** If true, UVs will be normalized. Normalize UVs should be done in cases where we operate with Images and Layouts */
+			bool bNormalizeUVs = false;
+
+			/** If true, assign vertices without layout to the first block. */
+			bool bEnsureAllVerticesHaveLayoutBlock = true;
+
+			/** If this has something the layouts in constant meshes will be ignored, because
+			* they are supposed to match some other set of layouts. If the array is empty, layouts
+			* are generated normally.
+			*/
+			TArray<FGeneratedLayout> OverrideLayouts;
+
+			/** Optional context to use instead of the node error context.
+			 * Be careful since it is not used everywhere. Check usages before assigning a value to it. */
+			TOptional<const void*> OverrideContext;
+			
+			friend FORCEINLINE uint32 GetTypeHash(const FMeshGenerationOptions& InKey)
+			{
+				uint32 KeyHash = 0;
+				KeyHash = HashCombineFast(KeyHash, ::GetTypeHash(InKey.bLayouts));
+				KeyHash = HashCombineFast(KeyHash, ::GetTypeHash(InKey.OverrideLayouts.Num()));
+				KeyHash = HashCombineFast(KeyHash, ::GetTypeHash(InKey.State));
+				KeyHash = HashCombineFast(KeyHash, ::GetTypeHash(InKey.ActiveTags.Num()));
+				return KeyHash;
+			}
+
+			FORCEINLINE bool operator==(const FMeshGenerationOptions& Other) const
+			{
+				return State==Other.State 
+					&& bLayouts==Other.bLayouts 
+					&& bClampUVIslands == Other.bClampUVIslands && bNormalizeUVs == Other.bNormalizeUVs
+					&& bEnsureAllVerticesHaveLayoutBlock == Other.bEnsureAllVerticesHaveLayoutBlock
+					&& ActiveTags==Other.ActiveTags
+					&& OverrideLayouts ==Other.OverrideLayouts;
+			}
+		};
+
+		//! Store the results of the code generation of a mesh.
+		struct FMeshGenerationResult
+		{
+			//! Mesh after all code tree is applied
+			Ptr<ASTOp> MeshOp;
+
+			//! Original base mesh before removes, morphs, etc.
+			Ptr<ASTOp> BaseMeshOp;
+
+			/** Generated node layouts with their own block ids. */
+			TArray<FGeneratedLayout> GeneratedLayouts;
+
+			TArray<Ptr<ASTOp>> LayoutOps;
+
+			struct FExtraLayouts
+			{
+				/** Source node layouts to use with these extra mesh. They don't have block ids. */
+				TArray<FGeneratedLayout> GeneratedLayouts;
+				Ptr<ASTOp> Condition;
+				Ptr<ASTOp> MeshFragment;
+			};
+			TArray< FExtraLayouts > ExtraMeshLayouts;
+		};
+		
+		struct FGeneratedMeshCacheKey
+		{
+			NodePtrConst Node;
+			FMeshGenerationOptions Options;
+
+			friend FORCEINLINE uint32 GetTypeHash(const FGeneratedMeshCacheKey& InKey)
+			{
+				uint32 KeyHash = 0;
+				KeyHash = HashCombineFast(KeyHash, ::GetTypeHash(InKey.Node.get()));
+				KeyHash = HashCombineFast(KeyHash, GetTypeHash(InKey.Options));
+				return KeyHash;
+			}
+
+			FORCEINLINE bool operator==(const FGeneratedMeshCacheKey& Other) const
+			{
+				return Node == Other.Node && Options == Other.Options;
+			}
+		};
+
+        typedef TMap<FGeneratedMeshCacheKey,FMeshGenerationResult> GeneratedMeshMap;
+        GeneratedMeshMap GeneratedMeshes;
+
+		/** Store the mesh generation data for surfaces that we intend to share across LODs. 
+		* The key is the SharedSurfaceId.
+		*/
+		TMap<int32, FMeshGenerationResult> SharedMeshOptionsMap;
+
+		//! Map of layouts found in the code already generated. The map is from the source layout
+		//! node to the generated layout.
+		struct FGeneratedLayoutKey
+		{
+			Ptr<const NodeLayout> SourceLayout;
+			uint32 MeshIdPrefix;
+
+			friend FORCEINLINE uint32 GetTypeHash(const FGeneratedLayoutKey& InKey)
+			{
+				uint32 KeyHash = ::GetTypeHash(InKey.SourceLayout.get());
+				KeyHash = HashCombineFast(KeyHash, GetTypeHash(InKey.MeshIdPrefix));
+				return KeyHash;
+			}
+
+			FORCEINLINE bool operator==(const FGeneratedLayoutKey& Other) const = default;
+		};
+		TMap<FGeneratedLayoutKey, Ptr<const Layout>> GeneratedLayouts;
+
+        void GenerateMesh(const FMeshGenerationOptions&, FMeshGenerationResult& result, const NodeMeshPtrConst&);
+        void GenerateMesh_Constant(const FMeshGenerationOptions&, FMeshGenerationResult&, const NodeMeshConstant* );
+        void GenerateMesh_Format(const FMeshGenerationOptions&, FMeshGenerationResult&, const NodeMeshFormat* );
+        void GenerateMesh_Morph(const FMeshGenerationOptions&, FMeshGenerationResult&, const NodeMeshMorph* );
+        void GenerateMesh_MakeMorph(const FMeshGenerationOptions&, FMeshGenerationResult&, const NodeMeshMakeMorph* );
+        void GenerateMesh_Fragment(const FMeshGenerationOptions&, FMeshGenerationResult&, const NodeMeshFragment* );
+        void GenerateMesh_Interpolate(const FMeshGenerationOptions&, FMeshGenerationResult&, const NodeMeshInterpolate* );
+        void GenerateMesh_Switch(const FMeshGenerationOptions&, FMeshGenerationResult&, const NodeMeshSwitch* );
+        void GenerateMesh_Transform(const FMeshGenerationOptions&, FMeshGenerationResult&, const NodeMeshTransform* );
+        void GenerateMesh_ClipMorphPlane(const FMeshGenerationOptions&, FMeshGenerationResult&, const NodeMeshClipMorphPlane* );
+        void GenerateMesh_ClipWithMesh(const FMeshGenerationOptions&, FMeshGenerationResult&, const NodeMeshClipWithMesh* );
+        void GenerateMesh_ApplyPose(const FMeshGenerationOptions&, FMeshGenerationResult&, const NodeMeshApplyPose* );
+        void GenerateMesh_Variation(const FMeshGenerationOptions&, FMeshGenerationResult&, const NodeMeshVariation* );
+		void GenerateMesh_Table(const FMeshGenerationOptions&, FMeshGenerationResult&, const NodeMeshTable*);
+		void GenerateMesh_GeometryOperation(const FMeshGenerationOptions&, FMeshGenerationResult&, const NodeMeshGeometryOperation*);
+		void GenerateMesh_Reshape(const FMeshGenerationOptions&, FMeshGenerationResult&, const NodeMeshReshape*);
+		void GenerateMesh_ClipDeform(const FMeshGenerationOptions&, FMeshGenerationResult&, const NodeMeshClipDeform*);
+
+		//-----------------------------------------------------------------------------------------
+		void PrepareMeshForLayout(const FGeneratedLayout&, Ptr<Mesh>,
+			int32 currentLayoutChannel,
+			const void* errorContext,
+			const FMeshGenerationOptions&,
+			bool bUseAbsoluteBlockIds);
+
+		//!
+		Ptr<const Layout> GenerateLayout(Ptr<const NodeLayout> SourceLayout, uint32 MeshIDPrefix);
+
+		struct FExtensionDataGenerationResult
+		{
+			Ptr<ASTOp> Op;
+		};
+
+		typedef const NodeExtensionData* FGeneratedExtensionDataCacheKey;
+		typedef TMap<FGeneratedExtensionDataCacheKey, FExtensionDataGenerationResult> FGeneratedExtensionDataMap;
+		FGeneratedExtensionDataMap GeneratedExtensionData;
+
+		void GenerateExtensionData(FExtensionDataGenerationResult& OutResult, const FGenericGenerationOptions&, const Ptr<const NodeExtensionData>&);
+		void GenerateExtensionData_Constant(FExtensionDataGenerationResult& OutResult, const FGenericGenerationOptions&, const class NodeExtensionDataConstant*);
+		void GenerateExtensionData_Switch(FExtensionDataGenerationResult& OutResult, const FGenericGenerationOptions&, const class NodeExtensionDataSwitch*);
+		void GenerateExtensionData_Variation(FExtensionDataGenerationResult& OutResult, const FGenericGenerationOptions&, const class NodeExtensionDataVariation*);
+		Ptr<ASTOp> GenerateMissingExtensionDataCode(const TCHAR* StrWhere, const void* ErrorContext);
+
+        //-----------------------------------------------------------------------------------------
+        // Projectors
+        struct FProjectorGenerationResult
+        {
+            Ptr<ASTOp> op;
+            PROJECTOR_TYPE type;
+        };
+
+        typedef TMap<FGeneratedCacheKey,FProjectorGenerationResult> FGeneratedProjectorsMap;
+        FGeneratedProjectorsMap GeneratedProjectors;
+
+        void GenerateProjector( FProjectorGenerationResult&, const FGenericGenerationOptions&, const Ptr<const NodeProjector>& );
+        void GenerateProjector_Constant( FProjectorGenerationResult&, const FGenericGenerationOptions&, const Ptr<const NodeProjectorConstant>& );
+        void GenerateProjector_Parameter( FProjectorGenerationResult&, const FGenericGenerationOptions&, const Ptr<const NodeProjectorParameter>& );
+        void GenerateMissingProjectorCode( FProjectorGenerationResult&, const void* errorContext );
+
+		//-----------------------------------------------------------------------------------------
+		// Bools
+		struct FBoolGenerationResult
+		{
+			Ptr<ASTOp> op;
+		};
+
+		typedef TMap<FGeneratedCacheKey, FBoolGenerationResult> FGeneratedBoolsMap;
+		FGeneratedBoolsMap GeneratedBools;
+
+		void GenerateBool(FBoolGenerationResult&, const FGenericGenerationOptions&, const Ptr<const NodeBool>&);
+		void GenerateBool_Constant(FBoolGenerationResult&, const FGenericGenerationOptions&, const Ptr<const NodeBoolConstant>&);
+		void GenerateBool_Parameter(FBoolGenerationResult&, const FGenericGenerationOptions&, const Ptr<const NodeBoolParameter>&);
+		void GenerateBool_Not(FBoolGenerationResult&, const FGenericGenerationOptions&, const Ptr<const NodeBoolNot>&);
+		void GenerateBool_And(FBoolGenerationResult&, const FGenericGenerationOptions&, const Ptr<const NodeBoolAnd>&);
+
+		//-----------------------------------------------------------------------------------------
+		// Scalars
+		struct FScalarGenerationResult
+		{
+			Ptr<ASTOp> op;
+		};
+
+		typedef TMap<FGeneratedCacheKey, FScalarGenerationResult> FGeneratedScalarsMap;
+		FGeneratedScalarsMap GeneratedScalars;
+
+		void GenerateScalar(FScalarGenerationResult&, const FGenericGenerationOptions&, const Ptr<const NodeScalar>&);
+		void GenerateScalar_Constant(FScalarGenerationResult&, const FGenericGenerationOptions&, const Ptr<const NodeScalarConstant>&);
+		void GenerateScalar_Parameter(FScalarGenerationResult&, const FGenericGenerationOptions&, const Ptr<const NodeScalarParameter>&);
+		void GenerateScalar_Switch(FScalarGenerationResult&, const FGenericGenerationOptions&, const Ptr<const NodeScalarSwitch>&);
+		void GenerateScalar_EnumParameter(FScalarGenerationResult&, const FGenericGenerationOptions&, const Ptr<const NodeScalarEnumParameter>&);
+		void GenerateScalar_Curve(FScalarGenerationResult&, const FGenericGenerationOptions&, const Ptr<const NodeScalarCurve>&);
+		void GenerateScalar_Arithmetic(FScalarGenerationResult&, const FGenericGenerationOptions&, const Ptr<const NodeScalarArithmeticOperation>&);
+		void GenerateScalar_Variation(FScalarGenerationResult&, const FGenericGenerationOptions&, const Ptr<const NodeScalarVariation>&);
+		void GenerateScalar_Table(FScalarGenerationResult&, const FGenericGenerationOptions&, const Ptr<const NodeScalarTable>&);
+		Ptr<ASTOp> GenerateMissingScalarCode(const TCHAR* strWhere, float value, const void* errorContext);
+
+		//-----------------------------------------------------------------------------------------
+		// Colors
+		struct FColorGenerationResult
+		{
+			Ptr<ASTOp> op;
+		};
+
+		typedef TMap<FGeneratedCacheKey, FColorGenerationResult> FGeneratedColorsMap;
+		FGeneratedColorsMap GeneratedColors;
+
+		void GenerateColor(FColorGenerationResult&, const FGenericGenerationOptions&, const Ptr<const NodeColour>&);
+		void GenerateColor_Constant(FColorGenerationResult&, const FGenericGenerationOptions&, const Ptr<const NodeColourConstant>&);
+		void GenerateColor_Parameter(FColorGenerationResult&, const FGenericGenerationOptions&, const Ptr<const NodeColourParameter>&);
+		void GenerateColor_Switch(FColorGenerationResult&, const FGenericGenerationOptions&, const Ptr<const NodeColourSwitch>&);
+		void GenerateColor_SampleImage(FColorGenerationResult&, const FGenericGenerationOptions&, const Ptr<const NodeColourSampleImage>&);
+		void GenerateColor_FromScalars(FColorGenerationResult&, const FGenericGenerationOptions&, const Ptr<const NodeColourFromScalars>&);
+		void GenerateColor_Arithmetic(FColorGenerationResult&, const FGenericGenerationOptions&, const Ptr<const NodeColourArithmeticOperation>&);
+		void GenerateColor_Variation(FColorGenerationResult&, const FGenericGenerationOptions&, const Ptr<const NodeColourVariation>&);
+		void GenerateColor_Table(FColorGenerationResult&, const FGenericGenerationOptions&, const Ptr<const NodeColourTable>&);
+		Ptr<ASTOp> GenerateMissingColourCode(const TCHAR* strWhere, const void* errorContext);
+
+		//-----------------------------------------------------------------------------------------
+		// Strings
+		struct FStringGenerationResult
+		{
+			Ptr<ASTOp> op;
+		};
+
+		typedef TMap<FGeneratedCacheKey, FStringGenerationResult> FGeneratedStringsMap;
+		FGeneratedStringsMap GeneratedStrings;
+
+		void GenerateString(FStringGenerationResult&, const FGenericGenerationOptions&, const Ptr<const NodeString>&);
+		void GenerateString_Constant(FStringGenerationResult&, const FGenericGenerationOptions& Options, const Ptr<const NodeStringConstant>&);
+		void GenerateString_Parameter(FStringGenerationResult&, const FGenericGenerationOptions& Options, const Ptr<const NodeStringParameter>&);
+
+        //-----------------------------------------------------------------------------------------
+        // Ranges
+        struct FRangeGenerationResult
+        {
+            //
+            Ptr<ASTOp> sizeOp;
+
+            //
+            FString rangeName;
+
+            //
+			FString rangeUID;
+        };
+
+        typedef TMap<FGeneratedCacheKey,FRangeGenerationResult> FGeneratedRangeMap;
+        FGeneratedRangeMap GeneratedRanges;
+
+        void GenerateRange(FRangeGenerationResult&, const FGenericGenerationOptions&, Ptr<const NodeRange>);
+
+
+        //-----------------------------------------------------------------------------------------
+        struct FSurfaceGenerationResult
+        {
+            Ptr<ASTOp> surfaceOp;
+        };
+
+        void GenerateSurface( FSurfaceGenerationResult&, const FSurfaceGenerationOptions&, Ptr<const NodeSurfaceNew> );
+
+		//-----------------------------------------------------------------------------------------
+		// Default Table Parameters
+		Ptr<ASTOp> GenerateDefaultTableValue(ETableColumnType NodeType);
+
+
+		//-----------------------------------------------------------------------------------------
+
+		struct FLayoutBlockDesc
+		{
+			EImageFormat FinalFormat = EImageFormat::IF_NONE;
+			int32 BlockPixelsX = 0;
+			int32 BlockPixelsY = 0;
+			bool bBlocksHaveMips = false;
+		};
+
+		void UpdateLayoutBlockDesc(FLayoutBlockDesc& Out, FImageDesc BlockDesc, FIntVector2 LayoutCellSize);
+
+		// Get the modifiers that have to be applied to elements with a specific tag.
+		void GetModifiersFor(const TArray<FString>& SurfaceTags, int32 LOD, bool bModifiersForBeforeOperations, TArray<FirstPassGenerator::FModifier>& OutModifiers);
+
+		// Apply the required mesh modifiers to the given operation.
+		Ptr<ASTOp> ApplyMeshModifiers(const FMeshGenerationOptions&, FMeshGenerationResult& BaseResults, 
+			const FMeshGenerationResult* SharedSurfaceResults, 
+			bool bModifiersForBeforeOperations, const void* ErrorContext,
+			const NodeMeshConstant* OriginalMeshNode);
+
+		Ptr<ASTOp> ApplyImageBlockModifiers(const FImageGenerationOptions&, Ptr<ASTOp> BaseImageOp, int32 ImageIndex, 
+			FIntPoint GridSize,
+			const FLayoutBlockDesc& LayoutBlockDesc,
+			box< FIntVector2 > RectInCells,
+			bool bModifiersForBeforeOperations, const void* ErrorContext);
+
+		Ptr<ASTOp> ApplyImageExtendModifiers(
+			const FGenericGenerationOptions& Options, 
+			const FMeshGenerationResult& BaseMeshResults,
+			Ptr<ASTOp> ImageAd, 
+			CompilerOptions::TextureLayoutStrategy ImageLayoutStrategy, 
+			int32 LayoutIndex, int32 ImageIndex,
+			FIntPoint GridSize,
+			CodeGenerator::FLayoutBlockDesc& InOutLayoutBlockDesc,
+			bool bModifiersForBeforeOperations, const void* ErrorContext);
+
+    };
+
+	
+    //---------------------------------------------------------------------------------------------
+    template<class NODE_TABLE, ETableColumnType TYPE, OP_TYPE OPTYPE, typename F>
+    Ptr<ASTOp> CodeGenerator::GenerateTableSwitch( const NODE_TABLE& node, F&& GenerateOption )
+    {
+        Ptr<const Table> NodeTable = node.Table;
+        Ptr<ASTOp> Variable;
+
+		FTableCacheKey CacheKey = FTableCacheKey{ node.Table, node.ParameterName };
+        Ptr<ASTOp>* it = GeneratedTables.Find( CacheKey );
+        if ( it )
+        {
+            Variable = *it;
+        }
+
+        if ( !Variable)
+        {
+            // Create the table variable expression
+            Variable = GenerateTableVariable( &node, CacheKey, node.bNoneOption, node.DefaultRowName);
+
+            GeneratedTables.Add(CacheKey, Variable );
+        }
+
+		int32 NumRows = NodeTable->GetPrivate()->Rows.Num();
+
+        // Verify that the table column is the right type
+        int32 ColIndex = NodeTable->FindColumn( node.ColumnName );
+
+		if (NumRows == 0)
+		{
+			ErrorLog->GetPrivate()->Add("The table has no rows.", ELMT_ERROR, node.GetMessageContext());
+			return nullptr;
+		}
+        else if (ColIndex < 0)
+        {
+            ErrorLog->GetPrivate()->Add("Table column not found.", ELMT_ERROR, node.GetMessageContext());
+            return nullptr;
+        }
+
+        if (NodeTable->GetPrivate()->Columns[ ColIndex ].Type != TYPE )
+        {
+            ErrorLog->GetPrivate()->Add("Table column type is not the right type.", ELMT_ERROR, node.GetMessageContext());
+            return nullptr;
+        }
+
+        // Create the switch to cover all the options
+        Ptr<ASTOp> lastSwitch;
+        Ptr<ASTOpSwitch> SwitchOp = new ASTOpSwitch();
+		SwitchOp->type = OPTYPE;
+		SwitchOp->variable = Variable;
+		SwitchOp->def = GenerateDefaultTableValue(TYPE);
+
+		for (int32 i = 0; i < NumRows; ++i)
+        {
+            check(NodeTable->GetPrivate()->Rows[i].Id <= 0xFFFF);
+            auto Condition = (uint16)NodeTable->GetPrivate()->Rows[i].Id;
+            Ptr<ASTOp> Branch = GenerateOption( node, ColIndex, (int)i, ErrorLog.get() );
+
+			if (Branch || TYPE != ETableColumnType::Mesh)
+			{
+				SwitchOp->cases.Add(ASTOpSwitch::FCase(Condition, SwitchOp, Branch));
+			}
+        }
+
+        return SwitchOp;
+    }
+}
diff --git a/Engine/Plugins/Mutable/Source/MutableTools/Private/MuT/CodeGenerator_Bool.cpp b/Engine/Plugins/Mutable/Source/MutableTools/Private/MuT/CodeGenerator_Bool.cpp
new file mode 100644
--- /dev/null
+++ b/Engine/Plugins/Mutable/Source/MutableTools/Private/MuT/CodeGenerator_Bool.cpp
@@ -0,0 +1,156 @@
+// Copyright Epic Games, Inc. All Rights Reserved.
+
+#include "HAL/PlatformCrt.h"
+#include "Misc/AssertionMacros.h"
+#include "MuR/Operations.h"
+#include "MuR/Parameters.h"
+#include "MuR/ParametersPrivate.h"
+#include "MuR/Ptr.h"
+#include "MuR/RefCounted.h"
+#include "MuR/Types.h"
+#include "MuT/AST.h"
+#include "MuT/ASTOpConstantBool.h"
+#include "MuT/ASTOpParameter.h"
+#include "MuT/CodeGenerator.h"
+#include "MuT/NodeBool.h"
+#include "MuT/NodeBoolPrivate.h"
+#include "MuT/NodeRange.h"
+
+
+namespace mu
+{
+
+	void CodeGenerator::GenerateBool(FBoolGenerationResult& Result, const FGenericGenerationOptions& Options, const Ptr<const NodeBool>& Untyped)
+	{
+		if (!Untyped)
+		{
+			Result = FBoolGenerationResult();
+			return;
+		}
+
+		// See if it was already generated
+		FGeneratedCacheKey Key;
+		Key.Node = Untyped;
+		Key.Options = Options;
+		FGeneratedBoolsMap::ValueType* it = GeneratedBools.Find(Key);
+		if (it)
+		{
+			Result = *it;
+			return;
+		}
+
+		// Generate for each different type of node
+		if (Untyped->GetType()==NodeBoolConstant::GetStaticType())
+		{
+			const NodeBoolConstant* Constant = static_cast<const NodeBoolConstant*>(Untyped.get());
+			GenerateBool_Constant(Result, Options, Constant);
+		}
+		else if (Untyped->GetType() == NodeBoolParameter::GetStaticType())
+		{
+			const NodeBoolParameter* Param = static_cast<const NodeBoolParameter*>(Untyped.get());
+			GenerateBool_Parameter(Result, Options, Param);
+		}
+		else if (Untyped->GetType() == NodeBoolNot::GetStaticType())
+		{
+			const NodeBoolNot* Sample = static_cast<const NodeBoolNot*>(Untyped.get());
+			GenerateBool_Not(Result, Options, Sample);
+		}
+		else if (Untyped->GetType() == NodeBoolAnd::GetStaticType())
+		{
+			const NodeBoolAnd* From = static_cast<const NodeBoolAnd*>(Untyped.get());
+			GenerateBool_And(Result, Options, From);
+		}
+		else
+		{
+			check(false);
+		}
+
+		// Cache the result
+		GeneratedBools.Add(Key, Result);
+	}
+
+
+	void CodeGenerator::GenerateBool_Constant(FBoolGenerationResult& result, const FGenericGenerationOptions& Options, const Ptr<const NodeBoolConstant>& Typed)
+	{
+		const NodeBoolConstant::Private& node = *Typed->GetPrivate();
+
+		Ptr<ASTOpConstantBool> op = new ASTOpConstantBool();
+		op->value = node.m_value;
+
+		result.op = op;
+	}
+
+
+	void CodeGenerator::GenerateBool_Parameter(FBoolGenerationResult& result, const FGenericGenerationOptions& Options, const Ptr<const NodeBoolParameter>& Typed)
+	{
+		const NodeBoolParameter::Private& node = *Typed->GetPrivate();
+
+		Ptr<ASTOpParameter> op;
+
+		Ptr<ASTOpParameter>* it = FirstPass.ParameterNodes.Find(node.m_pNode);
+		if (!it)
+		{
+			FParameterDesc param;
+			param.m_name = node.m_name;
+			const TCHAR* CStr = ToCStr(node.m_uid);
+			param.m_uid.ImportTextItem(CStr, 0, nullptr, nullptr);
+			param.m_type = PARAMETER_TYPE::T_BOOL;
+			param.m_defaultValue.Set<ParamBoolType>(node.m_defaultValue);
+
+			op = new ASTOpParameter();
+			op->type = OP_TYPE::BO_PARAMETER;
+			op->parameter = param;
+
+			// Generate the code for the ranges
+			for (int32 a = 0; a < node.m_ranges.Num(); ++a)
+			{
+				FRangeGenerationResult rangeResult;
+				GenerateRange(rangeResult, Options, node.m_ranges[a]);
+				op->ranges.Emplace(op.get(), rangeResult.sizeOp, rangeResult.rangeName, rangeResult.rangeUID);
+			}
+
+			FirstPass.ParameterNodes.Add(node.m_pNode,op);
+		}
+		else
+		{
+			op = *it;
+		}
+
+		result.op = op;
+	}
+
+
+	void CodeGenerator::GenerateBool_Not(FBoolGenerationResult& result, const FGenericGenerationOptions& Options, const Ptr<const NodeBoolNot>& Typed)
+	{
+		const NodeBoolNot::Private& node = *Typed->GetPrivate();
+
+		Ptr<ASTOpFixed> op = new ASTOpFixed();
+		op->op.type = OP_TYPE::BO_NOT;
+
+		FBoolGenerationResult ChildResult;
+		GenerateBool(ChildResult, Options, node.m_pSource);
+		op->SetChild(op->op.args.BoolNot.source, ChildResult.op);
+
+		result.op = op;
+	}
+
+
+	void CodeGenerator::GenerateBool_And(FBoolGenerationResult& result, const FGenericGenerationOptions& Options, const Ptr<const NodeBoolAnd>& Typed)
+	{
+		const NodeBoolAnd::Private& node = *Typed->GetPrivate();
+
+		Ptr<ASTOpFixed> op = new ASTOpFixed();
+		op->op.type = OP_TYPE::BO_AND;
+
+		FBoolGenerationResult ChildResult;
+		GenerateBool(ChildResult, Options, node.m_pA);
+		op->SetChild(op->op.args.BoolBinary.a, ChildResult.op);
+
+		GenerateBool(ChildResult, Options, node.m_pB);
+		op->SetChild(op->op.args.BoolBinary.b, ChildResult.op);
+
+		result.op = op;
+	}
+
+
+}
diff --git a/Engine/Plugins/Mutable/Source/MutableTools/Private/MuT/CodeGenerator_Color.cpp b/Engine/Plugins/Mutable/Source/MutableTools/Private/MuT/CodeGenerator_Color.cpp
new file mode 100644
--- /dev/null
+++ b/Engine/Plugins/Mutable/Source/MutableTools/Private/MuT/CodeGenerator_Color.cpp
@@ -0,0 +1,524 @@
+// Copyright Epic Games, Inc. All Rights Reserved.
+
+#include "HAL/PlatformMath.h"
+#include "Misc/AssertionMacros.h"
+#include "MuR/Image.h"
+#include "MuR/MutableMath.h"
+#include "MuR/MutableTrace.h"
+#include "MuR/Operations.h"
+#include "MuR/Parameters.h"
+#include "MuR/ParametersPrivate.h"
+#include "MuR/Platform.h"
+#include "MuR/Ptr.h"
+#include "MuR/RefCounted.h"
+#include "MuR/Types.h"
+#include "MuT/AST.h"
+#include "MuT/ASTOpConditional.h"
+#include "MuT/ASTOpParameter.h"
+#include "MuT/ASTOpSwitch.h"
+#include "MuT/CodeGenerator.h"
+#include "MuT/CodeGenerator_FirstPass.h"
+#include "MuT/ErrorLog.h"
+#include "MuT/ErrorLogPrivate.h"
+#include "MuT/Node.h"
+#include "MuT/NodeColour.h"
+#include "MuT/NodeColourArithmeticOperation.h"
+#include "MuT/NodeColourConstant.h"
+#include "MuT/NodeColourFromScalars.h"
+#include "MuT/NodeColourParameter.h"
+#include "MuT/NodeColourSampleImage.h"
+#include "MuT/NodeColourSwitch.h"
+#include "MuT/NodeColourTable.h"
+#include "MuT/NodeColourVariation.h"
+#include "MuT/NodeImage.h"
+#include "MuT/NodeImageProject.h"
+#include "MuT/NodeMeshGeometryOperation.h"
+#include "MuT/NodeMeshTable.h"
+#include "MuT/NodeRange.h"
+#include "MuT/NodeScalar.h"
+#include "MuT/NodeScalarConstant.h"
+#include "MuT/Table.h"
+#include "MuT/TablePrivate.h"
+
+
+namespace mu
+{
+
+
+	//-------------------------------------------------------------------------------------------------
+	void CodeGenerator::GenerateColor(FColorGenerationResult& Result, const FGenericGenerationOptions& Options, const Ptr<const NodeColour>& Untyped)
+	{
+		if (!Untyped)
+		{
+			Result = FColorGenerationResult();
+			return;
+		}
+
+		// See if it was already generated
+		FGeneratedCacheKey Key;
+		Key.Node = Untyped;
+		Key.Options = Options;
+		FGeneratedColorsMap::ValueType* it = GeneratedColors.Find(Key);
+		if (it)
+		{
+			Result = *it;
+			return;
+		}
+
+		// Generate for each different type of node
+		if (Untyped->GetType()==NodeColourConstant::GetStaticType())
+		{
+			const NodeColourConstant* Constant = static_cast<const NodeColourConstant*>(Untyped.get());
+			GenerateColor_Constant(Result, Options, Constant);
+		}
+		else if (Untyped->GetType() == NodeColourParameter::GetStaticType())
+		{
+			const NodeColourParameter* Param = static_cast<const NodeColourParameter*>(Untyped.get());
+			GenerateColor_Parameter(Result, Options, Param);
+		}
+		else if (Untyped->GetType() == NodeColourSwitch::GetStaticType())
+		{
+			const NodeColourSwitch* Switch = static_cast<const NodeColourSwitch*>(Untyped.get());
+			GenerateColor_Switch(Result, Options, Switch);
+		}
+		else if (Untyped->GetType() == NodeColourSampleImage::GetStaticType())
+		{
+			const NodeColourSampleImage* Sample = static_cast<const NodeColourSampleImage*>(Untyped.get());
+			GenerateColor_SampleImage(Result, Options, Sample);
+		}
+		else if (Untyped->GetType() == NodeColourFromScalars::GetStaticType())
+		{
+			const NodeColourFromScalars* From = static_cast<const NodeColourFromScalars*>(Untyped.get());
+			GenerateColor_FromScalars(Result, Options, From);
+		}
+		else if (Untyped->GetType() == NodeColourArithmeticOperation::GetStaticType())
+		{
+			const NodeColourArithmeticOperation* Arithmetic = static_cast<const NodeColourArithmeticOperation*>(Untyped.get());
+			GenerateColor_Arithmetic(Result, Options, Arithmetic);
+		}
+		else if (Untyped->GetType() == NodeColourVariation::GetStaticType())
+		{
+			const NodeColourVariation* Variation = static_cast<const NodeColourVariation*>(Untyped.get());
+			GenerateColor_Variation(Result, Options, Variation);
+		}
+		else if (Untyped->GetType() == NodeColourTable::GetStaticType())
+		{
+			const NodeColourTable* Table = static_cast<const NodeColourTable*>(Untyped.get());
+			GenerateColor_Table(Result, Options, Table);
+		}
+		else
+		{
+			check(false);
+		}
+
+		// Cache the result
+		GeneratedColors.Add(Key, Result);
+	}
+
+
+	//-------------------------------------------------------------------------------------------------
+	void CodeGenerator::GenerateColor_Constant(FColorGenerationResult& result, const FGenericGenerationOptions& Options, const Ptr<const NodeColourConstant>& Typed)
+	{
+		const NodeColourConstant& node = *Typed;
+
+		Ptr<ASTOpFixed> op = new ASTOpFixed();
+		op->op.type = OP_TYPE::CO_CONSTANT;
+		op->op.args.ColourConstant.value[0] = node.Value.X;
+		op->op.args.ColourConstant.value[1] = node.Value.Y;
+		op->op.args.ColourConstant.value[2] = node.Value.Z;
+		op->op.args.ColourConstant.value[3] = node.Value.W;
+
+		result.op = op;
+	}
+
+
+	//-------------------------------------------------------------------------------------------------
+	void CodeGenerator::GenerateColor_Parameter(FColorGenerationResult& result, const FGenericGenerationOptions& Options, const Ptr<const NodeColourParameter>& Typed)
+	{
+		const NodeColourParameter& node = *Typed;
+
+		Ptr<ASTOpParameter> op;
+
+		Ptr<ASTOpParameter>* it = FirstPass.ParameterNodes.Find(Typed);
+
+		if (!it)
+		{
+			FParameterDesc param;
+			param.m_name = node.Name;
+			const TCHAR* CStr = ToCStr(node.Uid);
+			param.m_uid.ImportTextItem(CStr, 0, nullptr, nullptr);
+			param.m_type = PARAMETER_TYPE::T_COLOUR;
+
+			ParamColorType Value;
+			Value[0] =  node.DefaultValue[0];
+			Value[1] = node.DefaultValue[1];
+			Value[2] = node.DefaultValue[2];
+			Value[3] = node.DefaultValue[3];
+
+			param.m_defaultValue.Set<ParamColorType>(Value);
+			
+			op = new ASTOpParameter();
+			op->type = OP_TYPE::CO_PARAMETER;
+			op->parameter = param;
+
+			// Generate the code for the ranges
+			for (int32 a = 0; a < node.Ranges.Num(); ++a)
+			{
+				FRangeGenerationResult rangeResult;
+				GenerateRange(rangeResult, Options, node.Ranges[a]);
+				op->ranges.Emplace(op.get(), rangeResult.sizeOp, rangeResult.rangeName, rangeResult.rangeUID);
+			}
+
+			FirstPass.ParameterNodes.Add(Typed, op);
+		}
+		else
+		{
+			op = *it;
+		}
+
+		result.op = op;
+	}
+
+
+	//-------------------------------------------------------------------------------------------------
+	void CodeGenerator::GenerateColor_Switch(FColorGenerationResult& result, const FGenericGenerationOptions& Options, const Ptr<const NodeColourSwitch>& Typed)
+	{
+		const NodeColourSwitch& node = *Typed;
+
+		MUTABLE_CPUPROFILER_SCOPE(NodeColourSwitch);
+
+		if (node.Options.Num() == 0)
+		{
+			// No options in the switch!
+			Ptr<ASTOp> missingOp = GenerateMissingColourCode(TEXT("Switch option"), Typed->GetMessageContext());
+			result.op = missingOp;
+			return;
+		}
+
+		Ptr<ASTOpSwitch> op = new ASTOpSwitch();
+		op->type = OP_TYPE::CO_SWITCH;
+
+		// Variable value
+		if (node.Parameter)
+		{
+			op->variable = Generate_Generic(node.Parameter.get(), Options);
+		}
+		else
+		{
+			// This argument is required
+			op->variable = GenerateMissingScalarCode(TEXT("Switch variable"), 0.0f, Typed->GetMessageContext());
+		}
+
+		// Options
+		for (int32 t = 0; t < node.Options.Num(); ++t)
+		{
+			Ptr<ASTOp> branch;
+			if (node.Options[t])
+			{
+				branch = Generate_Generic(node.Options[t].get(), Options);
+			}
+			else
+			{
+				// This argument is required
+				branch = GenerateMissingColourCode(TEXT("Switch option"), Typed->GetMessageContext());
+			}
+			op->cases.Emplace((int16)t, op, branch);
+		}
+
+		result.op = op;
+	}
+
+
+	//-------------------------------------------------------------------------------------------------
+	void CodeGenerator::GenerateColor_Variation(FColorGenerationResult& result, const FGenericGenerationOptions& Options, const Ptr<const NodeColourVariation>& Typed)
+	{
+		const NodeColourVariation& node = *Typed;
+
+		Ptr<ASTOp> currentOp;
+
+		// Default case
+		if (node.DefaultColour)
+		{
+			FColorGenerationResult BranchResults;
+			GenerateColor(BranchResults, Options, node.DefaultColour);
+			currentOp = BranchResults.op;
+		}
+
+		// Process variations in reverse order, since conditionals are built bottom-up.
+		for (int32 t = node.Variations.Num() - 1; t >= 0; --t)
+		{
+			int32 tagIndex = -1;
+			const FString& tag = node.Variations[t].Tag;
+			for (int32 i = 0; i < FirstPass.Tags.Num(); ++i)
+			{
+				if (FirstPass.Tags[i].Tag == tag)
+				{
+					tagIndex = i;
+				}
+			}
+
+			if (tagIndex < 0)
+			{
+				FString Msg = FString::Printf(TEXT("Unknown tag found in color variation [%s]."), *tag);
+				ErrorLog->GetPrivate()->Add(Msg, ELMT_WARNING, Typed->GetMessageContext());
+				continue;
+			}
+
+			Ptr<ASTOp> variationOp;
+			if (node.Variations[t].Colour)
+			{
+				FColorGenerationResult BranchResults;
+				GenerateColor(BranchResults,Options,node.Variations[t].Colour);
+				variationOp = BranchResults.op;
+			}
+			else
+			{
+				// This argument is required
+				variationOp = GenerateMissingColourCode(TEXT("Variation option"), Typed->GetMessageContext());
+			}
+
+
+			Ptr<ASTOpConditional> conditional = new ASTOpConditional;
+			conditional->type = OP_TYPE::CO_CONDITIONAL;
+			conditional->no = currentOp;
+			conditional->yes = variationOp;
+			conditional->condition = FirstPass.Tags[tagIndex].GenericCondition;
+
+			currentOp = conditional;
+		}
+
+		result.op = currentOp;
+	}
+
+
+	//-------------------------------------------------------------------------------------------------
+	void CodeGenerator::GenerateColor_SampleImage(FColorGenerationResult& result, const FGenericGenerationOptions& Options, const Ptr<const NodeColourSampleImage>& Typed)
+	{
+		const NodeColourSampleImage& node = *Typed;
+
+		// Generate the code
+		Ptr<ASTOpFixed> op = new ASTOpFixed();
+		op->op.type = OP_TYPE::CO_SAMPLEIMAGE;
+
+		// Source image
+		FImageGenerationOptions ImageOptions;
+		ImageOptions.State = Options.State;
+		if (!ActiveTags.IsEmpty())
+		{
+			ImageOptions.ActiveTags = ActiveTags.Last();
+		}
+
+		Ptr<ASTOp> base;
+		if (node.Image)
+		{
+			// Generate
+			FImageGenerationResult MapResult;
+			GenerateImage(ImageOptions, MapResult, node.Image);
+			base = MapResult.op;
+		}
+		else
+		{
+			// This argument is required
+			base = GenerateMissingImageCode(TEXT("Sample image"), EImageFormat::IF_RGB_UBYTE, Typed->GetMessageContext(), ImageOptions);
+		}
+		base = GenerateImageFormat(base, EImageFormat::IF_RGB_UBYTE);
+		op->SetChild(op->op.args.ColourSampleImage.image, base);
+
+		FScalarGenerationResult ChildResult;
+
+		// X
+		if (NodeScalar* pX = node.X.get())
+		{
+			GenerateScalar(ChildResult, Options, pX);
+			op->SetChild(op->op.args.ColourSampleImage.x, ChildResult.op);
+		}
+		else
+		{
+			// Set a constant 0.5 value
+			NodeScalarConstantPtr pNode = new NodeScalarConstant();
+			pNode->SetValue(0.5f);
+			GenerateScalar(ChildResult, Options, pNode);
+			op->SetChild(op->op.args.ColourSampleImage.x, ChildResult.op);
+		}
+
+
+		// Y
+		if (NodeScalar* pY = node.Y.get())
+		{
+			GenerateScalar(ChildResult, Options, pY);
+			op->SetChild(op->op.args.ColourSampleImage.y, ChildResult.op);
+		}
+		else
+		{
+			// Set a constant 0.5 value
+			NodeScalarConstantPtr pNode = new NodeScalarConstant();
+			pNode->SetValue(0.5f);
+			GenerateScalar(ChildResult, Options, pNode);
+			op->SetChild(op->op.args.ColourSampleImage.y, ChildResult.op);
+		}
+
+		// TODO
+		op->op.args.ColourSampleImage.filter = 0;
+
+		result.op = op;
+	}
+
+
+	//-------------------------------------------------------------------------------------------------
+	void CodeGenerator::GenerateColor_FromScalars(FColorGenerationResult& result, const FGenericGenerationOptions& Options, const Ptr<const NodeColourFromScalars>& Typed)
+	{
+		const NodeColourFromScalars& node = *Typed;
+
+		Ptr<ASTOpFixed> op = new ASTOpFixed();
+		op->op.type = OP_TYPE::CO_FROMSCALARS;
+
+		FScalarGenerationResult ChildResult;
+
+		// X
+		if (NodeScalar* pX = node.X.get())
+		{
+			GenerateScalar(ChildResult, Options, pX);
+			op->SetChild(op->op.args.ColourFromScalars.v[0], ChildResult.op );
+		}
+		else
+		{
+			NodeScalarConstantPtr pNode = new NodeScalarConstant();
+			pNode->SetValue(1.0f);
+			GenerateScalar(ChildResult, Options, pNode);
+			op->SetChild(op->op.args.ColourFromScalars.v[0], ChildResult.op);
+		}
+
+		// Y
+		if (NodeScalar* pY = node.Y.get())
+		{
+			GenerateScalar(ChildResult, Options, pY);
+			op->SetChild(op->op.args.ColourFromScalars.v[1], ChildResult.op);
+		}
+		else
+		{
+			NodeScalarConstantPtr pNode = new NodeScalarConstant();
+			pNode->SetValue(1.0f);
+			GenerateScalar(ChildResult, Options, pNode);
+			op->SetChild(op->op.args.ColourFromScalars.v[1], ChildResult.op);
+		}
+
+		// Z
+		if (NodeScalar* pZ = node.Z.get())
+		{
+			GenerateScalar(ChildResult, Options, pZ);
+			op->SetChild(op->op.args.ColourFromScalars.v[2], ChildResult.op);
+		}
+		else
+		{
+			NodeScalarConstantPtr pNode = new NodeScalarConstant();
+			pNode->SetValue(1.0f);
+			GenerateScalar(ChildResult, Options, pNode);
+			op->SetChild(op->op.args.ColourFromScalars.v[2], ChildResult.op);
+		}
+
+		// W
+		if (NodeScalar* pW = node.W.get())
+		{
+			GenerateScalar(ChildResult, Options, pW);
+			op->SetChild(op->op.args.ColourFromScalars.v[3], ChildResult.op);
+		}
+		else
+		{
+			NodeScalarConstantPtr pNode = new NodeScalarConstant();
+			pNode->SetValue(1.0f);
+			GenerateScalar(ChildResult, Options, pNode);
+			op->SetChild(op->op.args.ColourFromScalars.v[3], ChildResult.op);
+		}
+
+		result.op = op;
+	}
+
+
+	//-------------------------------------------------------------------------------------------------
+	void CodeGenerator::GenerateColor_Arithmetic(FColorGenerationResult& result, const FGenericGenerationOptions& Options, const Ptr<const NodeColourArithmeticOperation>& Typed)
+	{
+		const NodeColourArithmeticOperation& node = *Typed;
+
+		Ptr<ASTOpFixed> op = new ASTOpFixed();
+		op->op.type = OP_TYPE::CO_ARITHMETIC;
+
+		switch (node.Operation)
+		{
+		case NodeColourArithmeticOperation::EOperation::Add: op->op.args.ColourArithmetic.operation = OP::ArithmeticArgs::ADD; break;
+		case NodeColourArithmeticOperation::EOperation::Subtract: op->op.args.ColourArithmetic.operation = OP::ArithmeticArgs::SUBTRACT; break;
+		case NodeColourArithmeticOperation::EOperation::Multiply: op->op.args.ColourArithmetic.operation = OP::ArithmeticArgs::MULTIPLY; break;
+		case NodeColourArithmeticOperation::EOperation::Divide: op->op.args.ColourArithmetic.operation = OP::ArithmeticArgs::DIVIDE; break;
+		default:
+			checkf(false, TEXT("Unknown arithmetic operation."));
+			op->op.args.ColourArithmetic.operation = OP::ArithmeticArgs::NONE;
+			break;
+		}
+
+		FColorGenerationResult ChildResult;
+
+		// A
+		if (NodeColour* pA = node.A.get())
+		{
+			GenerateColor(ChildResult, Options, pA );
+			op->SetChild(op->op.args.ColourArithmetic.a, ChildResult.op);
+		}
+		else
+		{
+			op->SetChild(op->op.args.ColourArithmetic.a,
+				CodeGenerator::GenerateMissingColourCode(TEXT("ColourArithmetic A"), Typed->GetMessageContext()));
+		}
+
+		// B
+		if (NodeColour* pB = node.B.get())
+		{
+			GenerateColor(ChildResult, Options, pB);
+			op->SetChild(op->op.args.ColourArithmetic.b, ChildResult.op);
+		}
+		else
+		{
+			op->SetChild(op->op.args.ColourArithmetic.b,
+				CodeGenerator::GenerateMissingColourCode(TEXT("ColourArithmetic B"), Typed->GetMessageContext()));
+		}
+
+		result.op = op;
+	}
+
+
+	//-------------------------------------------------------------------------------------------------
+	void CodeGenerator::GenerateColor_Table(FColorGenerationResult& result, const FGenericGenerationOptions& Options, const Ptr<const NodeColourTable>& Typed)
+	{
+		const NodeColourTable& node = *Typed;
+
+		result.op = GenerateTableSwitch<NodeColourTable, ETableColumnType::Color, OP_TYPE::CO_SWITCH>(node,
+			[this, &Options](const NodeColourTable& node, int32 colIndex, int32 row, mu::ErrorLog* pErrorLog)
+			{
+				Ptr<NodeColourConstant> CellData = new NodeColourConstant();
+				FVector4f Colour = node.Table->GetPrivate()->Rows[row].Values[colIndex].Color;
+				CellData->Value = Colour;
+				FColorGenerationResult BranchResults;
+				GenerateColor(BranchResults, Options, CellData );
+				return BranchResults.op;
+			});
+	}
+
+
+	//---------------------------------------------------------------------------------------------
+	mu::Ptr<ASTOp> CodeGenerator::GenerateMissingColourCode(const TCHAR* strWhere, const void* errorContext)
+	{
+		// Log a warning
+		FString Msg = FString::Printf(TEXT("Required connection not found: %s"), strWhere);
+		ErrorLog->GetPrivate()->Add(Msg, ELMT_ERROR, errorContext);
+
+		// Create a constant colour node
+		Ptr<NodeColourConstant> pNode = new NodeColourConstant();
+		pNode->Value = FVector4f(1, 1, 0, 1);
+
+		FColorGenerationResult Result;
+		FGenericGenerationOptions Options;
+		GenerateColor(Result, Options, pNode);
+
+		return Result.op;
+	}
+
+
+}
diff --git a/Engine/Plugins/Mutable/Source/MutableTools/Private/MuT/CodeGenerator_FirstPass.cpp b/Engine/Plugins/Mutable/Source/MutableTools/Private/MuT/CodeGenerator_FirstPass.cpp
new file mode 100644
--- /dev/null
+++ b/Engine/Plugins/Mutable/Source/MutableTools/Private/MuT/CodeGenerator_FirstPass.cpp
@@ -0,0 +1,738 @@
+// Copyright Epic Games, Inc. All Rights Reserved.
+
+#include "MuT/CodeGenerator_FirstPass.h"
+
+#include "MuT/CodeGenerator.h"
+#include "HAL/PlatformMath.h"
+#include "Misc/AssertionMacros.h"
+#include "MuR/MutableTrace.h"
+#include "MuR/Operations.h"
+#include "MuR/Parameters.h"
+#include "MuR/ParametersPrivate.h"
+#include "MuR/Platform.h"
+#include "MuT/AST.h"
+#include "MuT/ASTOpConstantBool.h"
+#include "MuT/ASTOpParameter.h"
+#include "MuT/CompilerPrivate.h"
+#include "MuT/ErrorLog.h"
+#include "MuT/ErrorLogPrivate.h"
+#include "MuT/NodeComponent.h"
+#include "MuT/NodeModifierMeshClipDeform.h"
+#include "MuT/NodeModifierMeshClipMorphPlane.h"
+#include "MuT/NodeModifierMeshClipWithMesh.h"
+#include "MuT/NodeModifierMeshClipWithUVMask.h"
+#include "MuT/NodeModifierSurfaceEdit.h"
+#include "MuT/NodeObject.h"
+#include "MuT/NodeObjectGroupPrivate.h"
+#include "MuT/NodeObjectNew.h"
+#include "MuT/NodePrivate.h"
+#include "MuT/NodeSurface.h"
+#include "MuT/NodeSurfaceNew.h"
+#include "MuT/NodeSurfaceVariation.h"
+#include "MuT/NodeSurfaceSwitch.h"
+#include "MuT/NodeScalarEnumParameterPrivate.h"
+
+namespace mu
+{
+
+	//---------------------------------------------------------------------------------------------
+	FirstPassGenerator::FirstPassGenerator()
+	{
+		// Default conditions when there is no restriction accumulated.
+		FConditionContext noCondition;
+        CurrentCondition.Add(noCondition);
+        CurrentStateCondition.Add(StateCondition());
+	}
+
+
+	//---------------------------------------------------------------------------------------------
+    void FirstPassGenerator::Generate( Ptr<mu::ErrorLog> InErrorLog,
+                                       const Node* Root,
+									   bool bIgnoreStates,
+									   CodeGenerator* InGenerator )
+	{
+		MUTABLE_CPUPROFILER_SCOPE(FirstPassGenerate);
+
+		Generator = InGenerator;
+		ErrorLog = InErrorLog;
+
+		// Step 1: collect all objects, surfaces and object conditions
+        if (Root)
+		{
+ 			Generate_Generic(Root);
+ 		}
+
+		// Step 2: Collect all tags and a list of the surfaces and modifiers that activate them
+		for (int32 s=0; s<Surfaces.Num(); ++s)
+		{
+			// Collect the tags in new surfaces
+			for (int32 t=0; t< Surfaces[s].Node->Tags.Num(); ++t)
+			{
+				int32 tag = -1;
+                const FString& tagStr = Surfaces[s].Node->Tags[t];
+                for (int32 i = 0; i<Tags.Num() && tag<0; ++i)
+				{
+                    if (Tags[i].Tag == tagStr)
+					{
+						tag = i;
+					}
+				}
+
+				// New tag?
+				if (tag < 0)
+				{
+                    tag = Tags.Num();
+					FTag newTag;
+                    newTag.Tag = tagStr;
+                    Tags.Add(newTag);
+				}
+
+                if (Tags[tag].Surfaces.Find(s)==INDEX_NONE)
+				{
+                    Tags[tag].Surfaces.Add(s);
+				}
+			}
+		}
+
+		// TODO: Modifier's enabling tags?
+		for (int32 ModifierIndex = 0; ModifierIndex < Modifiers.Num(); ++ModifierIndex)
+		{
+			// Collect the tags in the modifiers
+			for (const FString& ModifierTag: Modifiers[ModifierIndex].Node->EnableTags)
+			{
+				int32 TagIndex = Tags.IndexOfByPredicate([&](const FTag& Candidate)
+					{ 
+						return Candidate.Tag == ModifierTag;
+					});
+
+				// New tag?
+				if (TagIndex < 0)
+				{
+					FTag newTag;
+					newTag.Tag = ModifierTag;
+					TagIndex = Tags.Add(newTag);
+				}
+
+				if (Tags[TagIndex].Modifiers.Find(ModifierIndex) == INDEX_NONE)
+				{
+					Tags[TagIndex].Modifiers.Add(ModifierIndex);
+				}
+			}
+		}
+
+		// Step 3: Create default state if necessary
+        if ( bIgnoreStates )
+        {
+            States.Empty();
+        }
+
+        if ( States.IsEmpty() )
+        {
+            FObjectState data;
+            data.Name = "Default";
+            States.Emplace( data, Root );
+        }
+	}
+
+
+	void FirstPassGenerator::Generate_Generic(const Node* Root)
+	{
+		if (!Root)
+		{
+			return;
+		}
+
+		if (Root->GetType()==NodeSurfaceNew::GetStaticType())
+		{
+			Generate_SurfaceNew(static_cast<const NodeSurfaceNew*>(Root));
+		}
+		else if (Root->GetType() == NodeSurfaceVariation::GetStaticType())
+		{
+			Generate_SurfaceVariation(static_cast<const NodeSurfaceVariation*>(Root));
+		}
+		else if (Root->GetType() == NodeSurfaceSwitch::GetStaticType())
+		{
+			Generate_SurfaceSwitch(static_cast<const NodeSurfaceSwitch*>(Root));
+		}
+		else if (Root->GetType() == NodeComponentNew::GetStaticType())
+		{
+			Generate_ComponentNew(static_cast<const NodeComponentNew*>(Root));
+		}
+		else if (Root->GetType() == NodeComponentEdit::GetStaticType())
+		{
+			Generate_ComponentEdit(static_cast<const NodeComponentEdit*>(Root));
+		}
+		else if (Root->GetType() == NodeComponentSwitch::GetStaticType())
+		{
+			Generate_ComponentSwitch(static_cast<const NodeComponentSwitch*>(Root));
+		}
+		else if (Root->GetType() == NodeComponentVariation::GetStaticType())
+		{
+			Generate_ComponentVariation(static_cast<const NodeComponentVariation*>(Root));
+		}
+		else if (Root->GetType() == NodeObjectNew::GetStaticType())
+		{
+			Generate_ObjectNew(static_cast<const NodeObjectNew*>(Root));
+		}
+		else if (Root->GetType() == NodeObjectGroup::GetStaticType())
+		{
+			Generate_ObjectGroup(static_cast<const NodeObjectGroup*>(Root));
+		}
+		else if (Root->GetType() == NodeLOD::GetStaticType())
+		{
+			Generate_LOD(static_cast<const NodeLOD*>(Root));
+		}
+		else if (Root->GetType() == NodeModifier::GetStaticType())
+		{
+			Generate_Modifier(static_cast<const NodeModifier*>(Root));
+		}
+		else
+		{
+			check(false);
+		}
+	}
+
+	//---------------------------------------------------------------------------------------------
+	void FirstPassGenerator::Generate_Modifier(const NodeModifier* InNode)
+	{
+		// Add the data about this modifier
+		FModifier thisData;
+		thisData.Node = InNode;
+		thisData.ObjectCondition = CurrentCondition.Last().ObjectCondition;
+		thisData.StateCondition = CurrentStateCondition.Last();
+		thisData.PositiveTags = CurrentPositiveTags;
+		thisData.NegativeTags = CurrentNegativeTags;
+		Modifiers.Add(thisData);
+	}
+
+
+	//---------------------------------------------------------------------------------------------
+	void FirstPassGenerator::Generate_SurfaceNew(const NodeSurfaceNew* InNode)
+	{
+		// Add the data about this surface
+		FSurface thisData;
+		thisData.Node = InNode;
+		thisData.Component = CurrentComponent;
+		thisData.LOD = CurrentLOD;
+		thisData.ObjectCondition = CurrentCondition.Last().ObjectCondition;
+		thisData.StateCondition = CurrentStateCondition.Last();
+		thisData.PositiveTags = CurrentPositiveTags;
+		thisData.NegativeTags = CurrentNegativeTags;
+		Surfaces.Add(thisData);
+	}
+
+
+	//---------------------------------------------------------------------------------------------
+	void FirstPassGenerator::Generate_SurfaceVariation(const NodeSurfaceVariation* InNode)
+	{
+        switch(InNode->Type)
+        {
+
+        case NodeSurfaceVariation::VariationType::Tag:
+        {
+            // Any of the tags in the variations would prevent the default surface
+			TArray<FString> OldNegativeTags = CurrentNegativeTags;
+            for (int32 v=0; v< InNode->Variations.Num(); ++v)
+            {
+                CurrentNegativeTags.Add(InNode->Variations[v].Tag);
+            }
+
+            for(const Ptr<NodeSurface>& n: InNode->DefaultSurfaces)
+            {
+				Generate_Generic(n.get());
+			}
+            for(const Ptr<NodeModifier>& n: InNode->DefaultModifiers)
+            {
+				Generate_Modifier(n.get());
+            }
+
+            CurrentNegativeTags = OldNegativeTags;
+
+            for (int32 v=0; v< InNode->Variations.Num(); ++v)
+            {
+                CurrentPositiveTags.Add(InNode->Variations[v].Tag);
+                for (const Ptr<NodeSurface>& s : InNode->Variations[v].Surfaces)
+                {
+					Generate_Generic(s.get());
+                }
+
+                for (const Ptr<NodeModifier>& s : InNode->Variations[v].Modifiers)
+                {
+					Generate_Modifier(s.get());
+				}
+
+                CurrentPositiveTags.Pop();
+
+                // Tags have an order in a variation node: the current tag should prevent any following
+                // variation surface
+                CurrentNegativeTags.Add(InNode->Variations[v].Tag);
+            }
+
+            CurrentNegativeTags = OldNegativeTags;
+
+            break;
+        }
+
+
+        case NodeSurfaceVariation::VariationType::State:
+        {
+            int32 stateCount = States.Num();
+
+            // Default
+            {
+                // Store the states for the default branch here
+                StateCondition defaultStates;
+                {
+					StateCondition AllTrue;
+					AllTrue.Init(true,stateCount);
+                    defaultStates = CurrentStateCondition.Last().IsEmpty()
+                            ? AllTrue
+                            : CurrentStateCondition.Last();
+
+                    for (const NodeSurfaceVariation::FVariation& v: InNode->Variations)
+                    {
+                        for( size_t s=0; s<stateCount; ++s )
+                        {
+                            if (States[s].Key.Name==v.Tag)
+                            {
+                                // Remove this state from the default options, since it has its own variation
+                                defaultStates[s] = false;
+                            }
+                        }
+                    }
+                }
+
+                CurrentStateCondition.Add(defaultStates);
+
+                for (const Ptr<NodeSurface>& n : InNode->DefaultSurfaces)
+                {
+					Generate_Generic(n.get());
+				}
+                for (const Ptr<NodeModifier>& n : InNode->DefaultModifiers)
+                {
+					Generate_Modifier(n.get());
+				}
+
+                CurrentStateCondition.Pop();
+            }
+
+            // Variation branches
+            for (const auto& v: InNode->Variations)
+            {
+                // Store the states for this variation here
+				StateCondition variationStates;
+				variationStates.Init(false,stateCount);
+
+                for( int32 StateIndex=0; StateIndex<stateCount; ++StateIndex)
+                {
+                    if (States[StateIndex].Key.Name==v.Tag)
+                    {
+                        variationStates[StateIndex] = true;
+                    }
+                }
+
+                CurrentStateCondition.Add(variationStates);
+
+                for (const Ptr<NodeSurface>& n : v.Surfaces)
+                {
+					Generate_Generic(n.get());
+				}
+                for (const Ptr<NodeModifier>& n : v.Modifiers)
+                {
+					Generate_Modifier(n.get());
+				}
+
+                CurrentStateCondition.Pop();
+            }
+
+            break;
+        }
+
+        default:
+            // Case not implemented.
+            check(false);
+            break;
+        }
+	}
+
+
+	//---------------------------------------------------------------------------------------------
+	void FirstPassGenerator::Generate_SurfaceSwitch(const NodeSurfaceSwitch* InNode)
+	{
+		if (InNode->Options.Num() == 0)
+		{
+			// No options in the switch!
+			return;
+		}
+
+		// Prepare the enumeration parameter
+		CodeGenerator::FGenericGenerationOptions Options;
+		CodeGenerator::FScalarGenerationResult ScalarResult;
+		if (InNode->Parameter)
+		{
+			Generator->GenerateScalar( ScalarResult, Options, InNode->Parameter );
+		}
+		else
+		{
+			// This argument is required
+			ScalarResult.op = Generator->GenerateMissingScalarCode(TEXT("Switch variable"), 0.0f, InNode->GetMessageContext());
+		}
+
+		// Parse the options
+		for (int32 t = 0; t < InNode->Options.Num(); ++t)
+		{
+			// Create a comparison operation as the boolean parameter for the child
+			Ptr<ASTOpFixed> ParamOp = new ASTOpFixed();
+			ParamOp->op.type = OP_TYPE::BO_EQUAL_INT_CONST;
+			ParamOp->SetChild(ParamOp->op.args.BoolEqualScalarConst.value, ScalarResult.op);
+			ParamOp->op.args.BoolEqualScalarConst.constant = (int16)t;
+
+			// Combine the new condition with previous conditions coming from parent objects
+			if (CurrentCondition.Last().ObjectCondition)
+			{
+				Ptr<ASTOpFixed> op = new ASTOpFixed();
+				op->op.type = OP_TYPE::BO_AND;
+				op->SetChild(op->op.args.BoolBinary.a, CurrentCondition.Last().ObjectCondition);
+				op->SetChild(op->op.args.BoolBinary.b, ParamOp);
+				ParamOp = op;
+			}
+
+			FConditionContext data;
+			data.ObjectCondition = ParamOp;
+			CurrentCondition.Push(data);
+
+			if (InNode->Options[t])
+			{
+				Generate_Generic(InNode->Options[t].get());
+			}
+
+			CurrentCondition.Pop();
+		}
+	}
+
+
+	//---------------------------------------------------------------------------------------------
+	void FirstPassGenerator::Generate_ComponentNew(const NodeComponentNew* InNode)
+	{
+		// Add the data about this surface
+		FComponent ThisData;
+		ThisData.Component = InNode;
+		ThisData.ObjectCondition = CurrentCondition.Last().ObjectCondition;
+		ThisData.PositiveTags = CurrentPositiveTags;
+		ThisData.NegativeTags = CurrentNegativeTags;
+		Components.Add(ThisData);
+
+        CurrentComponent = InNode;
+
+		CurrentLOD = 0;
+		for (const Ptr<NodeLOD>& c : InNode->LODs)
+		{
+			if (c)
+			{
+				Generate_LOD(c.get());
+			}
+			++CurrentLOD;
+		}
+		CurrentLOD = -1;
+
+		CurrentComponent = nullptr;
+	}
+
+
+	//---------------------------------------------------------------------------------------------
+	void FirstPassGenerator::Generate_ComponentEdit(const NodeComponentEdit* InNode)
+	{
+		CurrentComponent = InNode->GetParentComponentNew();
+
+		CurrentLOD = 0;
+		for (const Ptr<NodeLOD>& c : InNode->LODs)
+		{
+			if (c)
+			{
+				Generate_LOD(c.get());
+			}
+			++CurrentLOD;
+		}
+		CurrentLOD = -1;
+
+		CurrentComponent = nullptr;
+	}
+
+
+	//---------------------------------------------------------------------------------------------
+	void FirstPassGenerator::Generate_ComponentVariation(const NodeComponentVariation* InNode)
+	{
+		// Any of the tags in the variations would prevent the default surface
+		TArray<FString> OldNegativeTags = CurrentNegativeTags;
+		for (int32 v = 0; v < InNode->Variations.Num(); ++v)
+		{
+			CurrentNegativeTags.Add(InNode->Variations[v].Tag);
+		}
+
+		Generate_Generic(InNode->DefaultComponent.get());
+
+		CurrentNegativeTags = OldNegativeTags;
+
+		for (int32 v = 0; v < InNode->Variations.Num(); ++v)
+		{
+			CurrentPositiveTags.Add(InNode->Variations[v].Tag);
+			Generate_Generic(InNode->Variations[v].Component.get());
+
+			CurrentPositiveTags.Pop();
+
+			// Tags have an order in a variation node: the current tag should prevent any following variation
+			CurrentNegativeTags.Add(InNode->Variations[v].Tag);
+		}
+
+		CurrentNegativeTags = OldNegativeTags;
+	}
+
+
+	//---------------------------------------------------------------------------------------------
+	void FirstPassGenerator::Generate_ComponentSwitch(const NodeComponentSwitch* InNode)
+	{
+		if (InNode->Options.Num() == 0)
+		{
+			// No options in the switch!
+			return;
+		}
+
+		// Prepare the enumeration parameter
+		CodeGenerator::FGenericGenerationOptions Options;
+		CodeGenerator::FScalarGenerationResult ScalarResult;
+		if (InNode->Parameter)
+		{
+			Generator->GenerateScalar(ScalarResult, Options, InNode->Parameter);
+		}
+		else
+		{
+			// This argument is required
+			ScalarResult.op = Generator->GenerateMissingScalarCode(TEXT("Switch variable"), 0.0f, InNode->GetMessageContext());
+		}
+
+		// Parse the options
+		for (int32 t = 0; t < InNode->Options.Num(); ++t)
+		{
+			// Create a comparison operation as the boolean parameter for the child
+			Ptr<ASTOpFixed> ParamOp = new ASTOpFixed();
+			ParamOp->op.type = OP_TYPE::BO_EQUAL_INT_CONST;
+			ParamOp->SetChild(ParamOp->op.args.BoolEqualScalarConst.value, ScalarResult.op);
+			ParamOp->op.args.BoolEqualScalarConst.constant = (int16)t;
+
+			// Combine the new condition with previous conditions coming from parent objects
+			if (CurrentCondition.Last().ObjectCondition)
+			{
+				Ptr<ASTOpFixed> op = new ASTOpFixed();
+				op->op.type = OP_TYPE::BO_AND;
+				op->SetChild(op->op.args.BoolBinary.a, CurrentCondition.Last().ObjectCondition);
+				op->SetChild(op->op.args.BoolBinary.b, ParamOp);
+				ParamOp = op;
+			}
+
+			FConditionContext data;
+			data.ObjectCondition = ParamOp;
+			CurrentCondition.Push(data);
+
+			if (InNode->Options[t])
+			{
+				Generate_Generic(InNode->Options[t].get());
+			}
+
+			CurrentCondition.Pop();
+		}
+	}
+
+
+	//---------------------------------------------------------------------------------------------
+	void FirstPassGenerator::Generate_LOD(const NodeLOD* InNode)
+	{
+		for (const Ptr<NodeSurface>& c : InNode->Surfaces)
+		{
+			if (c)
+			{
+				Generate_Generic(c.get());
+			}
+		}
+	}
+
+
+	//---------------------------------------------------------------------------------------------
+	void FirstPassGenerator::Generate_ObjectNew(const NodeObjectNew* InNode)
+	{
+		// Add the data about this object
+		FObject thisData;
+		thisData.Node = InNode;
+        thisData.Condition = CurrentCondition.Last().ObjectCondition;
+		Objects.Add(thisData);
+
+        // Accumulate the model states
+        for ( const FObjectState& s: InNode->States )
+        {
+            States.Emplace( s, InNode );
+
+            if ( s.RuntimeParams.Num() > MUTABLE_MAX_RUNTIME_PARAMETERS_PER_STATE )
+            {
+                FString Msg = FString::Printf( TEXT("State [%s] has more than %d runtime parameters. Their update may fail."), 
+					*s.Name,
+                    MUTABLE_MAX_RUNTIME_PARAMETERS_PER_STATE);
+				ErrorLog->GetPrivate()->Add(Msg, ELMT_ERROR, InNode->GetMessageContext());
+            }
+        }
+
+		// Process the components
+		for (const Ptr<NodeComponent>& Component : InNode->Components)
+		{
+			if (Component)
+			{
+				Generate_Generic(Component.get());
+			}
+		}
+
+		// Process the modifiers
+		for (const Ptr<NodeModifier>& c : InNode->Modifiers)
+		{
+			if (c)
+			{
+				Generate_Modifier(c.get());
+			}
+		}
+
+		// Process the children
+		for (const Ptr<NodeObject>& Child : InNode->Children)
+		{
+			if (Child)
+			{
+				Generate_Generic(Child.get());
+			}
+		}
+	}
+
+
+	//---------------------------------------------------------------------------------------------
+	void FirstPassGenerator::Generate_ObjectGroup(const NodeObjectGroup* InNode)
+	{
+		const NodeObjectGroup::Private* Private = InNode->GetPrivate();
+
+		// Prepare the enumeration parameter if necessary
+        Ptr<ASTOpParameter> enumOp;
+        if (Private->m_type==NodeObjectGroup::CS_ALWAYS_ONE ||
+			Private->m_type==NodeObjectGroup::CS_ONE_OR_NONE )
+        {
+            Ptr<ASTOpParameter> op = new ASTOpParameter();
+            op->type = OP_TYPE::NU_PARAMETER;
+
+            op->parameter.m_name = Private->Name;
+			const TCHAR* CStr = ToCStr(Private->Uid);
+			op->parameter.m_uid.ImportTextItem(CStr, 0, nullptr, nullptr);
+            op->parameter.m_type = PARAMETER_TYPE::T_INT;
+            op->parameter.m_defaultValue.Set<ParamIntType>(Private->DefaultValue);
+
+            if (Private->m_type==NodeObjectGroup::CS_ONE_OR_NONE )
+            {
+                FParameterDesc::FIntValueDesc nullValue;
+                nullValue.m_value = -1;
+                nullValue.m_name = "None";
+                op->parameter.m_possibleValues.Add( nullValue );
+            }
+
+			ParameterNodes.Add(InNode, op);
+
+            enumOp = op;
+        }
+
+
+        // Parse the child objects
+		for ( int32 t=0; t< Private->m_children.Num(); ++t )
+        {
+            if ( mu::Ptr<const NodeObject> pChildNode = Private->m_children[t] )
+            {
+                // Overwrite the implicit condition
+                Ptr<ASTOp> paramOp = 0;
+                switch (Private->m_type )
+                {
+                    case NodeObjectGroup::CS_TOGGLE_EACH:
+                    {           
+						if (pChildNode->GetType() == NodeObjectGroup::GetStaticType())
+						{
+							FString Msg = FString::Printf(TEXT("The Group Node [%s] has type Toggle and its direct child is a Group node, which is not allowed. Change the type or add a Child Object node in between them."),
+								*Private->Name);
+							ErrorLog->GetPrivate()->Add(Msg, ELMT_ERROR, InNode->GetMessageContext());
+						}
+						else
+						{
+							// Create a new boolean parameter
+							Ptr<ASTOpParameter> op = new ASTOpParameter();
+							op->type = OP_TYPE::BO_PARAMETER;
+
+							op->parameter.m_name = pChildNode->GetName();
+							const TCHAR* CStr = ToCStr(pChildNode->GetUid());
+							op->parameter.m_uid.ImportTextItem(CStr, 0, nullptr, nullptr);
+							op->parameter.m_type = PARAMETER_TYPE::T_BOOL;
+							op->parameter.m_defaultValue.Set<ParamBoolType>(false);
+
+							ParameterNodes.Add(pChildNode, op);
+
+							paramOp = op;
+						}
+
+                        break;
+                    }
+
+                    case NodeObjectGroup::CS_ALWAYS_ALL:
+                    {
+                        // Create a constant true boolean that the optimiser will remove later.
+                        Ptr<ASTOpConstantBool> op = new ASTOpConstantBool();
+                        op->value = true;
+
+                        paramOp = op;
+                        break;
+                    }
+
+                    case NodeObjectGroup::CS_ONE_OR_NONE:
+                    case NodeObjectGroup::CS_ALWAYS_ONE:
+                    {
+                        // Add the option to the enumeration parameter
+                        FParameterDesc::FIntValueDesc value;
+                        value.m_value = (int16)t;
+                        value.m_name = pChildNode->GetName();
+                        enumOp->parameter.m_possibleValues.Add( value );
+
+                        check(enumOp);
+
+                        // Create a comparison operation as the boolean parameter for the child
+                        Ptr<ASTOpFixed> op = new ASTOpFixed();
+                        op->op.type = OP_TYPE::BO_EQUAL_INT_CONST;
+                        op->SetChild( op->op.args.BoolEqualScalarConst.value, enumOp);
+                        op->op.args.BoolEqualScalarConst.constant = (int16_t)t;
+
+                        paramOp = op;
+                        break;
+                    }
+
+                    default:
+                        check( false );
+                }
+
+                // Combine the new condition with previous conditions coming from parent objects
+                if (CurrentCondition.Last().ObjectCondition)
+                {
+                    Ptr<ASTOpFixed> op = new ASTOpFixed();
+                    op->op.type = OP_TYPE::BO_AND;
+                    op->SetChild( op->op.args.BoolBinary.a,CurrentCondition.Last().ObjectCondition);
+                    op->SetChild( op->op.args.BoolBinary.b,paramOp);
+                    paramOp = op;
+                }
+
+				FConditionContext data;
+                data.ObjectCondition = paramOp;
+                CurrentCondition.Add( data );
+
+				Generate_Generic(pChildNode.get());
+
+                CurrentCondition.Pop();
+            }
+        }
+ 	}
+
+}
+
