[mutable] Moved the Mutable plugin out of Experimental status into Beta.

diff --git a/Engine/Plugins/Mutable/Source/MutableTools/Private/MuT/CodeGenerator_Image.cpp b/Engine/Plugins/Mutable/Source/MutableTools/Private/MuT/CodeGenerator_Image.cpp
new file mode 100644
--- /dev/null
+++ b/Engine/Plugins/Mutable/Source/MutableTools/Private/MuT/CodeGenerator_Image.cpp
@@ -0,0 +1,2119 @@
+// Copyright Epic Games, Inc. All Rights Reserved.
+
+#include "HAL/PlatformMath.h"
+#include "Misc/AssertionMacros.h"
+#include "MuR/Image.h"
+#include "MuR/ImagePrivate.h"
+#include "MuR/Mesh.h"
+#include "MuR/MutableMath.h"
+#include "MuR/MutableTrace.h"
+#include "MuR/OpImageProject.h"
+#include "MuR/Operations.h"
+#include "MuR/Parameters.h"
+#include "MuR/ParametersPrivate.h"
+#include "MuR/Platform.h"
+#include "MuR/Ptr.h"
+#include "MuR/RefCounted.h"
+#include "MuR/Serialisation.h"
+#include "MuR/System.h"
+#include "MuT/AST.h"
+#include "MuT/ASTOpConditional.h"
+#include "MuT/ASTOpConstantResource.h"
+#include "MuT/ASTOpReferenceResource.h"
+#include "MuT/ASTOpImageMipmap.h"
+#include "MuT/ASTOpImageLayer.h"
+#include "MuT/ASTOpImageLayerColor.h"
+#include "MuT/ASTOpImageMultiLayer.h"
+#include "MuT/ASTOpImageNormalComposite.h"
+#include "MuT/ASTOpImagePixelFormat.h"
+#include "MuT/ASTOpImageTransform.h"
+#include "MuT/ASTOpImageMakeGrowMap.h"
+#include "MuT/ASTOpImageSwizzle.h"
+#include "MuT/ASTOpImageRasterMesh.h"
+#include "MuT/ASTOpImageCrop.h"
+#include "MuT/ASTOpMeshExtractLayoutBlocks.h"
+#include "MuT/ASTOpMeshFormat.h"
+#include "MuT/ASTOpParameter.h"
+#include "MuT/ASTOpSwitch.h"
+#include "MuT/CodeGenerator.h"
+#include "MuT/CodeGenerator_FirstPass.h"
+#include "MuT/Compiler.h"
+#include "MuT/CompilerPrivate.h"
+#include "MuT/ErrorLog.h"
+#include "MuT/ErrorLogPrivate.h"
+#include "MuT/Node.h"
+#include "MuT/NodeBool.h"
+#include "MuT/NodeColour.h"
+#include "MuT/NodeColourFromScalars.h"
+#include "MuT/NodeColourSampleImage.h"
+#include "MuT/NodeColourConstant.h"
+#include "MuT/NodeImage.h"
+#include "MuT/NodeImageBinarise.h"
+#include "MuT/NodeImageBinarisePrivate.h"
+#include "MuT/NodeImageColourMap.h"
+#include "MuT/NodeImageColourMapPrivate.h"
+#include "MuT/NodeImageConditional.h"
+#include "MuT/NodeImageConditionalPrivate.h"
+#include "MuT/NodeImageConstant.h"
+#include "MuT/NodeImageFormat.h"
+#include "MuT/NodeImageFormatPrivate.h"
+#include "MuT/NodeImageGradient.h"
+#include "MuT/NodeImageGradientPrivate.h"
+#include "MuT/NodeImageInterpolate.h"
+#include "MuT/NodeImageInterpolatePrivate.h"
+#include "MuT/NodeImageInvert.h"
+#include "MuT/NodeImageInvertPrivate.h"
+#include "MuT/NodeImageLayer.h"
+#include "MuT/NodeImageLayerColour.h"
+#include "MuT/NodeImageLayerColourPrivate.h"
+#include "MuT/NodeImageLayerPrivate.h"
+#include "MuT/NodeImageLuminance.h"
+#include "MuT/NodeImageLuminancePrivate.h"
+#include "MuT/NodeImageMipmap.h"
+#include "MuT/NodeImageMipmapPrivate.h"
+#include "MuT/NodeImageMultiLayer.h"
+#include "MuT/NodeImageMultiLayerPrivate.h"
+#include "MuT/NodeImageNormalComposite.h"
+#include "MuT/NodeImageNormalCompositePrivate.h"
+#include "MuT/NodeImageParameter.h"
+#include "MuT/NodeImageParameterPrivate.h"
+#include "MuT/NodeImagePlainColour.h"
+#include "MuT/NodeImagePlainColourPrivate.h"
+#include "MuT/NodeImageProject.h"
+#include "MuT/NodeImageProjectPrivate.h"
+#include "MuT/NodeImageResize.h"
+#include "MuT/NodeImageResizePrivate.h"
+#include "MuT/NodeImageSaturate.h"
+#include "MuT/NodeImageSaturatePrivate.h"
+#include "MuT/NodeImageSwitch.h"
+#include "MuT/NodeImageSwitchPrivate.h"
+#include "MuT/NodeImageSwizzle.h"
+#include "MuT/NodeImageSwizzlePrivate.h"
+#include "MuT/NodeImageTable.h"
+#include "MuT/NodeImageTransform.h"
+#include "MuT/NodeImageTransformPrivate.h"
+#include "MuT/NodeImageVariation.h"
+#include "MuT/NodeImageVariationPrivate.h"
+#include "MuT/NodeMeshGeometryOperation.h"
+#include "MuT/NodeMeshTable.h"
+#include "MuT/NodeProjector.h"
+#include "MuT/NodeRange.h"
+#include "MuT/NodeScalar.h"
+#include "MuT/NodeScalarConstant.h"
+#include "MuT/Table.h"
+#include "MuT/TablePrivate.h"
+
+
+namespace mu
+{
+
+
+	//-------------------------------------------------------------------------------------------------
+	void CodeGenerator::GenerateImage(const FImageGenerationOptions& Options, FImageGenerationResult& Result, const NodeImagePtrConst& Untyped)
+	{
+		if (!Untyped)
+		{
+			Result = FImageGenerationResult();
+			return;
+		}
+
+		// See if it was already generated
+		FGeneratedImageCacheKey key;
+		key.Options = Options;
+		key.Node = Untyped;
+		GeneratedImagesMap::ValueType* CachedPtr = GeneratedImages.Find(key);
+		if (CachedPtr)
+		{
+			Result = *CachedPtr;
+		}
+		else
+		{ 
+			const NodeImage* Node = Untyped.get();
+
+			// Generate for each different type of node
+			switch (Untyped->GetType()->Type)
+			{
+			case Node::EType::ImageConstant: GenerateImage_Constant(Options, Result, static_cast<const NodeImageConstant*>(Node)); break;
+			case Node::EType::ImageInterpolate: GenerateImage_Interpolate(Options, Result, static_cast<const NodeImageInterpolate*>(Node)); break;
+			case Node::EType::ImageSaturate: GenerateImage_Saturate(Options, Result, static_cast<const NodeImageSaturate*>(Node)); break;
+			case Node::EType::ImageTable: GenerateImage_Table(Options, Result, static_cast<const NodeImageTable*>(Node)); break;
+			case Node::EType::ImageSwizzle: GenerateImage_Swizzle(Options, Result, static_cast<const NodeImageSwizzle*>(Node)); break;
+			case Node::EType::ImageColorMap: GenerateImage_ColourMap(Options, Result, static_cast<const NodeImageColourMap*>(Node)); break;
+			case Node::EType::ImageGradient: GenerateImage_Gradient(Options, Result, static_cast<const NodeImageGradient*>(Node)); break;
+			case Node::EType::ImageBinarise: GenerateImage_Binarise(Options, Result, static_cast<const NodeImageBinarise*>(Node)); break;
+			case Node::EType::ImageLuminance: GenerateImage_Luminance(Options, Result, static_cast<const NodeImageLuminance*>(Node)); break;
+			case Node::EType::ImageLayer: GenerateImage_Layer(Options, Result, static_cast<const NodeImageLayer*>(Node)); break;
+			case Node::EType::ImageLayerColour: GenerateImage_LayerColour(Options, Result, static_cast<const NodeImageLayerColour*>(Node)); break;
+			case Node::EType::ImageResize: GenerateImage_Resize(Options, Result, static_cast<const NodeImageResize*>(Node)); break;
+			case Node::EType::ImagePlainColour: GenerateImage_PlainColour(Options, Result, static_cast<const NodeImagePlainColour*>(Node)); break;
+			case Node::EType::ImageProject: GenerateImage_Project(Options, Result, static_cast<const NodeImageProject*>(Node)); break;
+			case Node::EType::ImageMipmap: GenerateImage_Mipmap(Options, Result, static_cast<const NodeImageMipmap*>(Node)); break;
+			case Node::EType::ImageSwitch: GenerateImage_Switch(Options, Result, static_cast<const NodeImageSwitch*>(Node)); break;
+			case Node::EType::ImageConditional: GenerateImage_Conditional(Options, Result, static_cast<const NodeImageConditional*>(Node)); break;
+			case Node::EType::ImageFormat: GenerateImage_Format(Options, Result, static_cast<const NodeImageFormat*>(Node)); break;
+			case Node::EType::ImageParameter: GenerateImage_Parameter(Options, Result, static_cast<const NodeImageParameter*>(Node)); break;
+			case Node::EType::ImageMultiLayer: GenerateImage_MultiLayer(Options, Result, static_cast<const NodeImageMultiLayer*>(Node)); break;
+			case Node::EType::ImageInvert: GenerateImage_Invert(Options, Result, static_cast<const NodeImageInvert*>(Node)); break;
+			case Node::EType::ImageVariation: GenerateImage_Variation(Options, Result, static_cast<const NodeImageVariation*>(Node)); break;
+			case Node::EType::ImageNormalComposite: GenerateImage_NormalComposite(Options, Result, static_cast<const NodeImageNormalComposite*>(Node)); break;
+			case Node::EType::ImageTransform: GenerateImage_Transform(Options, Result, static_cast<const NodeImageTransform*>(Node)); break;
+			default: check(false);
+			}
+
+			// Cache the Result
+			GeneratedImages.Add(key, Result);
+		}
+	}
+
+
+    //---------------------------------------------------------------------------------------------
+    void CodeGenerator::GenerateImage_Constant(const FImageGenerationOptions& Options, FImageGenerationResult& Result, const NodeImageConstant* InNode)
+    {
+		const NodeImageConstant& node = *InNode;
+		
+        // TODO: check duplicates
+        Ptr<const Image> pImage;
+		if (node.Proxy)
+		{
+			pImage = node.Proxy->Get();
+		}
+
+        if (!pImage)
+        {
+            // This data is required
+            pImage = GenerateMissingImage(EImageFormat::IF_RGB_UBYTE );
+
+            // Log an error message
+            ErrorLog->GetPrivate()->Add( "Constant image not set.", ELMT_WARNING, InNode->GetMessageContext());
+        }
+
+
+		if (pImage->IsReference())
+		{
+			Ptr<ASTOpReferenceResource> ReferenceOp = new ASTOpReferenceResource();
+			ReferenceOp->type = OP_TYPE::IM_REFERENCE;
+			ReferenceOp->ID = pImage->GetReferencedTexture();
+			ReferenceOp->bForceLoad = pImage->IsForceLoad();
+			ReferenceOp->SourceDataDescriptor = InNode->SourceDataDescriptor;
+
+			// Don't store the format. Format can vary between loaded constant image and reference and cause
+			// code optimization bugs.
+			// As it is now, reference will always have alpha channel but constant resolution can remove the 
+			// channel if not used.
+			// TODO: review this, probably the reference descriptor generation needs to check for alpha channels 
+			// as well.
+			ReferenceOp->ImageDesc = FImageDesc(pImage->GetSize(), EImageFormat::IF_NONE, pImage->GetLODCount());
+			Result.op = ReferenceOp;
+		}
+		else
+		{
+			Ptr<ASTOpConstantResource> op = new ASTOpConstantResource();
+			op->Type = OP_TYPE::IM_CONSTANT;
+			op->SetValue(pImage, CompilerOptions->OptimisationOptions.DiskCacheContext);
+			op->SourceDataDescriptor = InNode->SourceDataDescriptor;
+			Result.op = op;
+		}
+
+		if (Options.ImageLayoutStrategy!= CompilerOptions::TextureLayoutStrategy::None && Options.LayoutToApply)
+		{
+			// We want to generate only a block from the image.
+
+			FIntVector2 SourceImageSize(pImage->GetSizeX(), pImage->GetSizeY());
+
+			int32 BlockIndex = Options.LayoutToApply->FindBlock(Options.LayoutBlockId);
+			check( BlockIndex>=0 );
+
+			// Block in layout grid units
+			box< FIntVector2 > RectInCells;
+			RectInCells.min = Options.LayoutToApply->Blocks[BlockIndex].Min;
+			RectInCells.size = Options.LayoutToApply->Blocks[BlockIndex].Size;
+
+			FIntPoint grid = Options.LayoutToApply->GetGridSize();
+			grid[0] = FMath::Max(1, grid[0]);
+			grid[1] = FMath::Max(1, grid[1]);
+
+			// Transform to pixels
+			box< FIntVector2 > RectInPixels;
+			RectInPixels.min[0]  = (RectInCells.min[0]  * SourceImageSize[0]) / grid[0];
+			RectInPixels.min[1]  = (RectInCells.min[1]  * SourceImageSize[1]) / grid[1];
+			RectInPixels.size[0] = (RectInCells.size[0] * SourceImageSize[0]) / grid[0];
+			RectInPixels.size[1] = (RectInCells.size[1] * SourceImageSize[1]) / grid[1];
+
+			// Do we need to crop?
+			if (RectInPixels.min[0]!=0 || RectInPixels.min[1]!=0 || pImage->GetSizeX()!=RectInPixels.size[0] || pImage->GetSizeY()!=RectInPixels.size[1])
+			{
+				// See if the rect belongs to a single texture tile
+				FIntVector2 TileMin(RectInPixels.min[0] / SourceImageSize[0], RectInPixels.min[1] / SourceImageSize[1]);
+				FIntVector2 TileMax((RectInPixels.min[0]+ RectInPixels.size[0]-1) / SourceImageSize[0], (RectInPixels.min[1]+ RectInPixels.size[1]-1) / SourceImageSize[1]);
+
+				if (TileMin != TileMax)
+				{
+					// Blocks spaning multiple texture tiles are not supported.
+					// To implement them, assemble a series of instructions to crop and compose the necessary rects from each tile into the final image.
+					ensure(false);
+
+					// Log an error message
+					ErrorLog->GetPrivate()->Add("A layout block goes across different texture tiles. This is not supported yet.", ELMT_ERROR, InNode->GetMessageContext());
+				}
+				else
+				{
+					Ptr<ASTOpImageCrop> CropOp = new ASTOpImageCrop();
+					CropOp->Source = Result.op;
+					
+					// Bring the crop rect to tile 0,0
+					CropOp->Min[0] = RectInPixels.min[0] - TileMin[0] * SourceImageSize[0];
+					CropOp->Min[1] = RectInPixels.min[1] - TileMin[1] * SourceImageSize[1];
+					CropOp->Size[0] = RectInPixels.size[0];
+					CropOp->Size[1] = RectInPixels.size[1];
+					Result.op = CropOp;
+				}
+			}
+		}
+    }
+
+
+    //---------------------------------------------------------------------------------------------
+	void CodeGenerator::GenerateImage_Parameter(const FImageGenerationOptions& Options, FImageGenerationResult& Result, const NodeImageParameter* InNode )
+    {
+		const NodeImageParameter::Private& node = *InNode->GetPrivate();
+
+        Ptr<ASTOpParameter> op;
+
+		Ptr<ASTOpParameter>* it = FirstPass.ParameterNodes.Find( node.m_pNode );
+        if ( !it )
+        {
+            op = new ASTOpParameter();
+            op->type = OP_TYPE::IM_PARAMETER;
+
+			op->parameter.m_name = node.m_name;
+        	const TCHAR* CStr = ToCStr(node.m_uid);
+        	op->parameter.m_uid.ImportTextItem(CStr, 0, nullptr, nullptr);
+        	op->parameter.m_type = PARAMETER_TYPE::T_IMAGE;
+        	op->parameter.m_defaultValue.Set<ParamImageType>(node.m_defaultValue);
+
+			// Generate the code for the ranges
+			for (int32 a = 0; a < node.m_ranges.Num(); ++a)
+			{
+				FRangeGenerationResult rangeResult;
+				GenerateRange(rangeResult, Options, node.m_ranges[a]);
+				op->ranges.Emplace(op.get(), rangeResult.sizeOp, rangeResult.rangeName, rangeResult.rangeUID);
+			}
+
+			FirstPass.ParameterNodes.Add(node.m_pNode, op);
+		}
+        else
+        {
+            op = *it;
+        }
+
+		Result.op = op;
+    }
+
+
+    //---------------------------------------------------------------------------------------------
+	void CodeGenerator::GenerateImage_Layer(const FImageGenerationOptions& Options, FImageGenerationResult& Result, const NodeImageLayer* InNode)
+	{
+		const NodeImageLayer::Private& node = *InNode->GetPrivate();
+		
+		MUTABLE_CPUPROFILER_SCOPE(NodeImageLayer);
+
+        Ptr<ASTOpImageLayer> op = new ASTOpImageLayer();
+
+        op->blendType = node.m_type;
+
+        // Base image
+        Ptr<ASTOp> base;
+        if ( node.m_pBase )
+        {
+			FImageGenerationResult BaseResult;
+			GenerateImage(Options, BaseResult, node.m_pBase);
+			base = BaseResult.op;
+        }
+        else
+        {
+            // This argument is required
+            base = GenerateMissingImageCode(TEXT("Image Layer base"), EImageFormat::IF_RGB_UBYTE, InNode->GetMessageContext(), Options);
+        }
+
+		UE::Math::TIntVector2<int32> TargetSize = Options.RectSize;
+
+		FImageDesc BaseDesc = base->GetImageDesc(true);
+
+		// If no target size was specified in the generation options, get the base size to apply it to the mask and blend.
+		if (TargetSize == UE::Math::TIntVector2<int32>(0, 0))
+		{
+			TargetSize = UE::Math::TIntVector2<int32>(BaseDesc.m_size);
+		}
+
+		EImageFormat baseFormat = BaseDesc.m_format;
+        base = GenerateImageSize( base, TargetSize);
+        op->base = base;
+
+        // Mask of the effect
+        Ptr<ASTOp> mask;
+        if ( node.m_pMask )
+        {
+			FImageGenerationResult MaskResult;
+			FImageGenerationOptions MaskOptions(Options);
+			MaskOptions.RectSize = TargetSize;
+			GenerateImage(MaskOptions, MaskResult, node.m_pMask);
+			mask = MaskResult.op;
+
+            mask = GenerateImageFormat( mask, EImageFormat::IF_L_UBYTE );
+            mask = GenerateImageSize( mask, TargetSize);
+        }
+        op->mask = mask;
+
+        // Image to apply
+        Ptr<ASTOp> blended = 0;
+        if ( node.m_pBlended )
+        {
+			FImageGenerationResult BlendedResult;
+			FImageGenerationOptions BlendOptions(Options);
+			BlendOptions.RectSize = TargetSize;
+			GenerateImage(BlendOptions, BlendedResult, node.m_pBlended);
+			blended = BlendedResult.op;
+        }
+        else
+        {
+            // This argument is required
+            blended = GeneratePlainImageCode(FVector4f( 1,1,0,1 ), Options );
+        }
+        blended = GenerateImageFormat( blended, baseFormat );
+        blended = GenerateImageSize( blended, TargetSize);
+        op->blend = blended;
+
+        Result.op = op;
+    }
+
+
+    //---------------------------------------------------------------------------------------------
+	void CodeGenerator::GenerateImage_LayerColour(const FImageGenerationOptions& Options, FImageGenerationResult& Result, const NodeImageLayerColour* InNode)
+	{
+		const NodeImageLayerColour::Private& node = *InNode->GetPrivate();
+		
+		MUTABLE_CPUPROFILER_SCOPE(NodeImageLayerColour);
+
+        Ptr<ASTOpImageLayerColor> op = new ASTOpImageLayerColor();
+		op->blendType = node.m_type;
+
+        // Base image
+        Ptr<ASTOp> base;
+        if ( node.m_pBase )
+        {
+			FImageGenerationResult BaseResult;
+			GenerateImage(Options, BaseResult, node.m_pBase);
+			base = BaseResult.op;
+        }
+        else
+        {
+            // This argument is required
+            base = GenerateMissingImageCode(TEXT("Layer base image"), EImageFormat::IF_RGB_UBYTE, InNode->GetMessageContext(), Options );
+        }
+
+		UE::Math::TIntVector2<int32> TargetSize = Options.RectSize;
+
+		// If no target size was specified in the generation options, get the base size to apply it to the mask and blend.
+		if (TargetSize == UE::Math::TIntVector2<int32>(0, 0))
+		{
+			FImageDesc BaseDesc = base->GetImageDesc(true);
+			TargetSize = UE::Math::TIntVector2<int32>(BaseDesc.m_size);
+		}
+
+        base = GenerateImageFormat( base, EImageFormat::IF_RGB_UBYTE );
+        base = GenerateImageSize( base, TargetSize);
+        op->base = base;
+
+        // Mask of the effect
+        Ptr<ASTOp> mask;
+        if ( node.m_pMask )
+        {
+			FImageGenerationResult MaskResult;
+			GenerateImage(Options, MaskResult, node.m_pMask);
+			mask = MaskResult.op;
+			
+			mask = GenerateImageFormat( mask, EImageFormat::IF_L_UBYTE );
+            mask = GenerateImageSize( mask, TargetSize);
+        }
+        op->mask = mask;
+
+        // Colour to apply
+        Ptr<ASTOp> colour = 0;
+        if ( node.m_pColour )
+        {
+			FColorGenerationResult ColorResult;
+            GenerateColor(ColorResult, Options, node.m_pColour);
+			colour = ColorResult.op;
+        }
+        else
+        {
+            // This argument is required
+            colour = GenerateMissingColourCode(TEXT("Layer colour"), InNode->GetMessageContext() );
+        }
+        op->color = colour;
+
+        Result.op = op;
+    }
+
+
+    //---------------------------------------------------------------------------------------------
+	void CodeGenerator::GenerateImage_MultiLayer(const FImageGenerationOptions& Options, FImageGenerationResult& Result, const NodeImageMultiLayer* InNode)
+	{
+		const NodeImageMultiLayer::Private& node = *InNode->GetPrivate();
+
+		MUTABLE_CPUPROFILER_SCOPE(NodeImageMultiLayer);
+
+        Ptr<ASTOpImageMultiLayer> op = new ASTOpImageMultiLayer();
+
+		op->blendType = node.m_type;
+
+        // Base image
+        Ptr<ASTOp> base;
+        if ( node.m_pBase )
+        {
+			FImageGenerationResult BaseResult;
+			GenerateImage(Options, BaseResult, node.m_pBase);
+			base = BaseResult.op;
+		}
+        else
+        {
+            // This argument is required
+            base = GenerateMissingImageCode(TEXT("Image MultiLayer base"), EImageFormat::IF_RGB_UBYTE, InNode->GetMessageContext(), Options);
+        }
+
+		FImageDesc BaseDesc = base->GetImageDesc(true);
+
+		UE::Math::TIntVector2<int32> TargetSize = Options.RectSize;
+
+		// If no target size was specified in the generation options, get the base size to apply it to the mask and blend.
+		if (TargetSize == UE::Math::TIntVector2<int32>(0, 0))
+		{
+			TargetSize = UE::Math::TIntVector2<int32>(BaseDesc.m_size);
+		}
+
+		EImageFormat baseFormat = BaseDesc.m_format;
+        base = GenerateImageSize( base, TargetSize);
+        op->base = base;
+
+        // Mask of the effect
+        Ptr<ASTOp> mask;
+        if ( node.m_pMask )
+        {
+			FImageGenerationResult MaskResult;
+			GenerateImage(Options, MaskResult, node.m_pMask);
+			mask = MaskResult.op;
+
+			mask = GenerateImageFormat( mask, EImageFormat::IF_L_UBYTE );
+            mask = GenerateImageSize( mask, TargetSize);
+        }
+        op->mask = mask;
+
+        // Image to apply
+        Ptr<ASTOp> blended;
+        if (node.m_pBlended)
+        {
+			FImageGenerationResult MaskResult;
+			GenerateImage(Options, MaskResult, node.m_pBlended);
+			blended = MaskResult.op;
+        }
+        else
+        {
+            // This argument is required
+            blended = GeneratePlainImageCode(FVector4f( 1,1,0,1 ), Options);
+        }
+        blended = GenerateImageFormat( blended, baseFormat );
+        blended = GenerateImageSize( blended, TargetSize);
+        op->blend = blended;
+
+        // Range of iteration
+        if ( node.m_pRange )
+        {
+            FRangeGenerationResult rangeResult;
+            GenerateRange( rangeResult, Options, node.m_pRange );
+
+            op->range.rangeSize = rangeResult.sizeOp;
+            op->range.rangeName = rangeResult.rangeName;
+            op->range.rangeUID = rangeResult.rangeUID;
+        }
+
+        Result.op = op;
+    }
+
+
+	//---------------------------------------------------------------------------------------------
+	void CodeGenerator::GenerateImage_NormalComposite(const FImageGenerationOptions& Options, FImageGenerationResult& Result, const NodeImageNormalComposite* InNode)
+	{
+		const NodeImageNormalComposite::Private& node = *InNode->GetPrivate();
+
+		MUTABLE_CPUPROFILER_SCOPE(NodeImageNormalComposite);
+
+        Ptr<ASTOpImageNormalComposite> op = new ASTOpImageNormalComposite();
+
+		op->Mode = node.m_mode; 
+		op->Power = node.m_power;
+
+        // Base image
+        Ptr<ASTOp> base;
+        if ( node.m_pBase )
+        {
+			FImageGenerationResult BaseResult;
+			GenerateImage(Options, BaseResult, node.m_pBase);
+			base = BaseResult.op;
+		}
+        else
+        {
+            // This argument is required
+            base = GenerateMissingImageCode(TEXT("Image Composite Base"), EImageFormat::IF_RGB_UBYTE, InNode->GetMessageContext(), Options);
+        }
+
+		FImageDesc BaseDesc = base->GetImageDesc(true);
+
+		UE::Math::TIntVector2<int32> TargetSize = Options.RectSize;
+
+		// If no target size was specified in the generation options, get the base size to apply it to the mask and blend.
+		if (TargetSize == UE::Math::TIntVector2<int32>(0, 0))
+		{
+			TargetSize = UE::Math::TIntVector2<int32>(BaseDesc.m_size);
+		}
+
+
+		EImageFormat baseFormat = BaseDesc.m_format;
+        base = GenerateImageSize( base, TargetSize);
+        op->Base = base;
+
+        Ptr<ASTOp> normal;
+        if ( node.m_pNormal )
+        {
+			FImageGenerationResult BaseResult;
+			GenerateImage(Options, BaseResult, node.m_pNormal);
+			normal = BaseResult.op;
+
+            normal = GenerateImageFormat( normal, EImageFormat::IF_RGB_UBYTE );
+        }
+		else
+		{
+            // This argument is required
+            normal = GenerateMissingImageCode(TEXT("Image Composite Normal"), EImageFormat::IF_RGB_UBYTE, InNode->GetMessageContext(), Options);
+		}
+
+		normal = GenerateImageSize(normal, TargetSize);
+
+        op->Normal = normal;
+        
+		Result.op = op;
+    }
+
+	void CodeGenerator::GenerateImage_Transform(const FImageGenerationOptions& Options, FImageGenerationResult& Result, const NodeImageTransform* InNode)
+    {
+		const NodeImageTransform::Private& Node = *InNode->GetPrivate();
+
+        MUTABLE_CPUPROFILER_SCOPE(NodeImageTransform);
+
+        Ptr<ASTOpImageTransform> Op = new ASTOpImageTransform();
+
+		Ptr<ASTOp> OffsetX;
+		if (Node.m_pOffsetX)
+		{
+			FScalarGenerationResult ScalarResult;
+			GenerateScalar(ScalarResult, Options, Node.m_pOffsetX);
+			OffsetX = ScalarResult.op;
+		}
+
+		Ptr<ASTOp> OffsetY;
+		if (Node.m_pOffsetY)
+		{
+			FScalarGenerationResult ScalarResult;
+			GenerateScalar(ScalarResult, Options, Node.m_pOffsetY);
+			OffsetY = ScalarResult.op;
+		}
+	
+		Ptr<ASTOp> ScaleX;
+		if (Node.m_pScaleX)
+		{
+			FScalarGenerationResult ScalarResult;
+			GenerateScalar(ScalarResult, Options, Node.m_pScaleX);
+			ScaleX = ScalarResult.op;
+		}
+	
+		Ptr<ASTOp> ScaleY;
+		if (Node.m_pScaleY)
+		{
+			FScalarGenerationResult ScalarResult;
+			GenerateScalar(ScalarResult, Options, Node.m_pScaleY);
+			ScaleY = ScalarResult.op;
+		}
+
+		Ptr<ASTOp> Rotation;
+		if (Node.m_pRotation)
+		{
+			FScalarGenerationResult ScalarResult;
+			GenerateScalar(ScalarResult, Options, Node.m_pRotation);
+			Rotation = ScalarResult.op;
+		}
+
+		// If one of the inputs (offset or scale) is missig assume unifrom translation/scaling 
+		Op->OffsetX = OffsetX ? OffsetX : OffsetY;
+		Op->OffsetY = OffsetY ? OffsetY : OffsetX;
+ 		Op->ScaleX = ScaleX ? ScaleX : ScaleY;
+		Op->ScaleY = ScaleY ? ScaleY : ScaleX;
+		Op->Rotation = Rotation; 
+		Op->AddressMode = Node.AddressMode;
+		Op->SizeX = Node.SizeX;
+		Op->SizeY = Node.SizeY;
+		Op->bKeepAspectRatio = Node.bKeepAspectRatio;
+
+		// Base image
+        Ptr<ASTOp> Base;
+		FImageGenerationOptions NewOptions = Options;
+		NewOptions.ImageLayoutStrategy = CompilerOptions::TextureLayoutStrategy::None;
+		NewOptions.LayoutToApply = nullptr;
+		NewOptions.LayoutBlockId = -1;
+		NewOptions.RectSize = {};
+
+        if (Node.m_pBase)
+        {
+			FImageGenerationResult BaseResult;
+			GenerateImage(NewOptions, BaseResult, Node.m_pBase);
+			Base = BaseResult.op;
+		}
+        else
+        {
+            // This argument is required
+            Base = GenerateMissingImageCode(TEXT("Image Transform Base"), EImageFormat::IF_RGB_UBYTE, InNode->GetMessageContext(), NewOptions);
+        }
+		
+		FImageDesc BaseDesc = Base->GetImageDesc();
+		
+        Op->Base = Base;
+		Op->SourceSizeX = BaseDesc.m_size.X;
+		Op->SourceSizeY = BaseDesc.m_size.Y;
+
+        Result.op = Op; 
+    }
+
+    //---------------------------------------------------------------------------------------------
+	void CodeGenerator::GenerateImage_Interpolate(const FImageGenerationOptions& Options, FImageGenerationResult& Result, const NodeImageInterpolate* InNode)
+	{
+		const NodeImageInterpolate::Private& node = *InNode->GetPrivate();
+
+		MUTABLE_CPUPROFILER_SCOPE(NodeImageInterpolate);
+
+        Ptr<ASTOpFixed> op = new ASTOpFixed();
+        op->op.type = OP_TYPE::IM_INTERPOLATE;
+
+        // Factor
+        if ( Node* pFactor = node.m_pFactor.get() )
+        {
+            op->SetChild( op->op.args.ImageInterpolate.factor, Generate_Generic( pFactor, Options));
+        }
+        else
+        {
+            // This argument is required
+            op->SetChild( op->op.args.ImageInterpolate.factor,
+                          GenerateMissingScalarCode(TEXT("Interpolation factor"), 0.5f, InNode->GetMessageContext() ));
+        }
+
+        // Target images
+        int32 numTargets = 0;
+
+		UE::Math::TIntVector2<int32> FinalRectSize = Options.RectSize;
+
+        for ( std::size_t t=0
+            ; t< node.m_targets.Num() && numTargets<MUTABLE_OP_MAX_INTERPOLATE_COUNT
+            ; ++t )
+        {
+            if ( node.m_targets[t] )
+            {
+				FImageGenerationOptions ChildOptions = Options;
+				ChildOptions.RectSize = FinalRectSize;
+				FImageGenerationResult BaseResult;
+				GenerateImage(ChildOptions, BaseResult, node.m_targets[t]);
+				Ptr<ASTOp> target = BaseResult.op;
+
+				if (FinalRectSize[0] == 0)
+				{
+					FImageDesc ChildDesc = target->GetImageDesc();
+					FinalRectSize = UE::Math::TIntVector2<int32>(ChildDesc.m_size);
+				}
+
+                // TODO: Support other formats
+                target = GenerateImageFormat( target, EImageFormat::IF_RGB_UBYTE );
+                target = GenerateImageSize( target, FinalRectSize);
+
+                op->SetChild( op->op.args.ImageInterpolate.targets[numTargets], target);
+                numTargets++;
+            }
+        }
+
+        // At least one target is required
+        if (!op->op.args.ImageInterpolate.targets[0])
+        {
+            Ptr<ASTOp> target = GenerateMissingImageCode(TEXT("First interpolation image"), EImageFormat::IF_RGB_UBYTE, InNode->GetMessageContext(), Options);
+            target = GenerateImageSize( target, Options.RectSize);
+            op->SetChild( op->op.args.ImageInterpolate.targets[0], target);
+        }
+
+        Result.op = op;
+    }
+
+
+	//---------------------------------------------------------------------------------------------
+	void CodeGenerator::GenerateImage_Swizzle(const FImageGenerationOptions& Options, FImageGenerationResult& Result, const NodeImageSwizzle* InNode)
+	{
+		const NodeImageSwizzle::Private& node = *InNode->GetPrivate();
+
+		//MUTABLE_CPUPROFILER_SCOPE(NodeImageSwizzle);
+
+        // This node always produces a swizzle operation and sometimes it may produce a pixelformat
+		// operation to compress the Result
+        Ptr<ASTOpImageSwizzle> SwizzleOp = new ASTOpImageSwizzle();
+
+		// Format
+		EImageFormat CompressedFormat = EImageFormat::IF_NONE;
+
+		switch (node.m_format)
+		{
+        case EImageFormat::IF_BC1:
+        case EImageFormat::IF_ASTC_4x4_RGB_LDR:
+            CompressedFormat = node.m_format;
+            SwizzleOp->Format = node.m_sources[3] ? EImageFormat::IF_RGBA_UBYTE : EImageFormat::IF_RGB_UBYTE;
+			break;
+
+		case EImageFormat::IF_BC2:
+		case EImageFormat::IF_BC3:
+		case EImageFormat::IF_BC6:
+        case EImageFormat::IF_BC7:
+        case EImageFormat::IF_ASTC_4x4_RGBA_LDR:
+            CompressedFormat = node.m_format;
+            SwizzleOp->Format = EImageFormat::IF_RGBA_UBYTE;
+			break;
+
+		case EImageFormat::IF_BC4:
+			CompressedFormat = node.m_format;
+            SwizzleOp->Format = EImageFormat::IF_L_UBYTE;
+			break;
+
+		case EImageFormat::IF_BC5:
+        case EImageFormat::IF_ASTC_4x4_RG_LDR:
+            CompressedFormat = node.m_format;
+			// TODO: Should be RG
+            SwizzleOp->Format = EImageFormat::IF_RGB_UBYTE;
+			break;
+
+		default:
+            SwizzleOp->Format = node.m_format;
+			break;
+
+		}
+
+		check(node.m_format != EImageFormat::IF_NONE);
+
+		// Source images and channels
+		check(node.m_sources.Num() == node.m_sourceChannels.Num());
+
+		// First source, for reference in the size
+        Ptr<ASTOp> FirstValid;
+		FImageDesc FirstValidDesc;
+		int32 FirstValidSourceIndex = -1;
+
+		check(MUTABLE_OP_MAX_SWIZZLE_CHANNELS >= node.m_sources.Num());
+		for (int32 SourceIndex = 0; SourceIndex < node.m_sources.Num(); ++SourceIndex)
+		{
+			if (node.m_sources[SourceIndex])
+			{
+				FImageGenerationResult BaseResult;
+				GenerateImage(Options, BaseResult, node.m_sources[SourceIndex]);
+                Ptr<ASTOp> Source = BaseResult.op;
+
+				Source = GenerateImageUncompressed(Source);
+
+				if (!Source)
+				{
+					// TODO: Warn?
+					Source = GenerateMissingImageCode(TEXT("Swizzle channel"), EImageFormat::IF_L_UBYTE, InNode->GetMessageContext(), Options);
+				}
+
+                Ptr<ASTOp> SizedSource;
+				if (FirstValid && FirstValidDesc.m_size[0])
+				{
+					SizedSource = GenerateImageSize(Source, FIntVector2(FirstValidDesc.m_size));
+				}
+				else
+				{
+					FirstValid = Source;
+					SizedSource = Source;
+					FirstValidDesc = FirstValid->GetImageDesc();
+					FirstValidSourceIndex = SourceIndex;
+				}
+
+                SwizzleOp->Sources[SourceIndex] = SizedSource;
+                SwizzleOp->SourceChannels[SourceIndex] = (uint8)node.m_sourceChannels[SourceIndex];
+			}
+		}
+
+        if (FirstValidSourceIndex < 0)
+		{
+            Ptr<ASTOp> Source = GenerateMissingImageCode(TEXT("First swizzle image"), EImageFormat::IF_RGBA_UBYTE, InNode->GetMessageContext(), Options);
+            SwizzleOp->Sources[0] = Source;
+
+		}
+
+        Ptr<ASTOp> ResultOp = SwizzleOp;
+
+		if (CompressedFormat != EImageFormat::IF_NONE)
+		{
+            Ptr<ASTOpImagePixelFormat> FormatOp = new ASTOpImagePixelFormat();
+            FormatOp->Source = ResultOp;
+            FormatOp->Format = CompressedFormat;
+			ResultOp = FormatOp;
+		}
+
+        Result.op = ResultOp;
+	}
+
+
+	//---------------------------------------------------------------------------------------------
+	void CodeGenerator::GenerateImage_Format(const FImageGenerationOptions& Options, FImageGenerationResult& Result, const NodeImageFormat* InNode)
+	{
+		const NodeImageFormat::Private& node = *InNode->GetPrivate();
+
+		MUTABLE_CPUPROFILER_SCOPE(NodeImageFormat);
+
+        check(node.m_format != EImageFormat::IF_NONE);
+
+        Ptr<ASTOpImagePixelFormat> fop = new ASTOpImagePixelFormat();
+        fop->Format = node.m_format;
+        fop->FormatIfAlpha = node.m_formatIfAlpha;
+
+		// Source is required
+		if (!node.m_source)
+		{
+            fop->Source = GenerateMissingImageCode(TEXT("Source image for format."), EImageFormat::IF_RGBA_UBYTE, InNode->GetMessageContext(), Options);
+		}
+		else
+		{
+			FImageGenerationResult BaseResult;
+			GenerateImage(Options, BaseResult, node.m_source);
+			fop->Source = BaseResult.op;
+		}
+
+        Result.op = fop;
+	}
+
+
+    //---------------------------------------------------------------------------------------------
+	void CodeGenerator::GenerateImage_Saturate(const FImageGenerationOptions& Options, FImageGenerationResult& Result, const NodeImageSaturate* InNode)
+	{
+		const NodeImageSaturate::Private& node = *InNode->GetPrivate();
+
+		Ptr<ASTOpFixed> op = new ASTOpFixed();
+        op->op.type = OP_TYPE::IM_SATURATE;
+
+
+        // Source image
+        Ptr<ASTOp> base = 0;
+        if ( node.m_pSource )
+        {
+			FImageGenerationResult BaseResult;
+			GenerateImage(Options, BaseResult, node.m_pSource);
+			base = BaseResult.op;
+		}
+        else
+        {
+            // This argument is required
+            base = GenerateMissingImageCode(TEXT("Saturate image"), EImageFormat::IF_RGB_UBYTE, InNode->GetMessageContext(), Options);
+        }
+		
+        base = GenerateImageFormat(base, GetRGBOrRGBAFormat(base->GetImageDesc().m_format));
+        base = GenerateImageSize( base, Options.RectSize);
+        op->SetChild( op->op.args.ImageSaturate.base, base);
+
+
+        // Factor
+        if ( Node* pFactor = node.m_pFactor.get() )
+        {
+            op->SetChild( op->op.args.ImageSaturate.factor, Generate_Generic( pFactor, Options));
+        }
+        else
+        {
+            // This argument is required
+            op->SetChild( op->op.args.ImageSaturate.factor, GenerateMissingScalarCode(TEXT("Saturation factor"), 0.5f, InNode->GetMessageContext() ) );
+        }
+
+        Result.op = op;
+    }
+
+
+    //---------------------------------------------------------------------------------------------
+	void CodeGenerator::GenerateImage_Luminance(const FImageGenerationOptions& Options, FImageGenerationResult& Result, const NodeImageLuminance* InNode)
+	{
+		const NodeImageLuminance::Private& node = *InNode->GetPrivate();
+
+		Ptr<ASTOpFixed> op = new ASTOpFixed();
+        op->op.type = OP_TYPE::IM_LUMINANCE;
+
+        // Source image
+        Ptr<ASTOp> base;
+        if ( node.m_pSource )
+        {
+			FImageGenerationResult BaseResult;
+			GenerateImage(Options, BaseResult, node.m_pSource);
+			base = BaseResult.op;
+        }
+        else
+        {
+            // This argument is required
+            base = GenerateMissingImageCode(TEXT("Image luminance"), EImageFormat::IF_RGB_UBYTE, InNode->GetMessageContext(), Options);
+        }
+        base = GenerateImageFormat( base, EImageFormat::IF_RGB_UBYTE );
+        op->SetChild( op->op.args.ImageLuminance.base, base);
+
+        Result.op = op;
+    }
+
+
+    //---------------------------------------------------------------------------------------------
+	void CodeGenerator::GenerateImage_ColourMap(const FImageGenerationOptions& Options, FImageGenerationResult& Result, const NodeImageColourMap* InNode)
+	{
+		const NodeImageColourMap::Private& node = *InNode->GetPrivate();
+
+		Ptr<ASTOpFixed> op = new ASTOpFixed();
+        op->op.type = OP_TYPE::IM_COLOURMAP;
+
+        // Base image
+        Ptr<ASTOp> base ;
+        if ( node.m_pBase )
+        {
+			FImageGenerationResult BaseResult;
+			GenerateImage(Options, BaseResult, node.m_pBase);
+			base = BaseResult.op;
+		}
+        else
+        {
+            // This argument is required
+            base = GenerateMissingImageCode(TEXT("Colourmap base image"), EImageFormat::IF_RGB_UBYTE, InNode->GetMessageContext(), Options);
+        }
+        base = GenerateImageSize( base, Options.RectSize);
+        op->SetChild( op->op.args.ImageColourMap.base, base);
+
+        // Mask of the effect
+        Ptr<ASTOp> mask;
+        if ( node.m_pMask )
+        {
+			FImageGenerationResult BaseResult;
+			GenerateImage(Options, BaseResult, node.m_pMask);
+			mask = BaseResult.op;
+        }
+        else
+        {
+            // Set the argument default value: affect all pixels.
+            // TODO: Special operation code without mask
+            mask = GeneratePlainImageCode(FVector4f( 1,1,1,1 ), Options);
+        }
+        mask = GenerateImageFormat( mask, EImageFormat::IF_L_UBYTE );
+        mask = GenerateImageSize( mask, Options.RectSize);
+        op->SetChild( op->op.args.ImageColourMap.mask, mask);
+
+        // Map image
+        Ptr<ASTOp> MapImageOp;
+        if ( node.m_pMap )
+        {
+			FImageGenerationOptions MapOptions = Options;
+			MapOptions.ImageLayoutStrategy = CompilerOptions::TextureLayoutStrategy::None;
+			MapOptions.LayoutToApply = nullptr;
+			MapOptions.LayoutBlockId = -1;
+			MapOptions.RectSize = {};
+
+			FImageGenerationResult BaseResult;
+			GenerateImage(MapOptions, BaseResult, node.m_pMap);
+			MapImageOp = BaseResult.op;
+			MapImageOp = GenerateImageFormat(MapImageOp, EImageFormat::IF_RGB_UBYTE);
+		}
+        else
+        {
+            // This argument is required
+			MapImageOp = GenerateMissingImageCode(TEXT("Map image"), EImageFormat::IF_RGB_UBYTE, InNode->GetMessageContext(), Options);
+        }
+        op->SetChild( op->op.args.ImageColourMap.map, MapImageOp);
+
+        Result.op = op;
+    }
+
+
+    //---------------------------------------------------------------------------------------------
+	void CodeGenerator::GenerateImage_Gradient(const FImageGenerationOptions& Options, FImageGenerationResult& Result, const NodeImageGradient* InNode)
+	{
+		const NodeImageGradient::Private& node = *InNode->GetPrivate();
+
+		Ptr<ASTOpFixed> op = new ASTOpFixed();
+        op->op.type = OP_TYPE::IM_GRADIENT;
+
+        // First colour
+        Ptr<ASTOp> colour0 = 0;
+        if ( Node* pColour0 = node.m_pColour0.get() )
+        {
+            colour0 = Generate_Generic( pColour0, Options);
+        }
+        else
+        {
+            // This argument is required
+            colour0 = GenerateMissingColourCode(TEXT("Gradient colour 0"), InNode->GetMessageContext() );
+        }
+        op->SetChild( op->op.args.ImageGradient.colour0, colour0);
+
+        // Second colour
+        Ptr<ASTOp> colour1 = 0;
+        if ( Node* pColour1 = node.m_pColour1.get() )
+        {
+            colour1 = Generate_Generic( pColour1, Options);
+        }
+        else
+        {
+            // This argument is required
+            colour1 = GenerateMissingColourCode(TEXT("Gradient colour 1"), InNode->GetMessageContext() );
+        }
+        op->SetChild( op->op.args.ImageGradient.colour1, colour1);
+
+        op->op.args.ImageGradient.size[0] = (uint16)FMath::Max( 2, FMath::Min( 1024, node.m_size[0] ) );
+        op->op.args.ImageGradient.size[1] = (uint16)FMath::Max( 1, FMath::Min( 1024, node.m_size[1] ) );
+
+        Result.op = op;
+    }
+
+
+    //---------------------------------------------------------------------------------------------
+	void CodeGenerator::GenerateImage_Binarise(const FImageGenerationOptions& Options, FImageGenerationResult& Result, const NodeImageBinarise* InNode)
+	{
+		const NodeImageBinarise::Private& node = *InNode->GetPrivate();
+
+		Ptr<ASTOpFixed> op = new ASTOpFixed();
+        op->op.type = OP_TYPE::IM_BINARISE;
+
+        // A image
+        Ptr<ASTOp> a;
+        if ( node.m_pBase )
+        {
+			FImageGenerationResult BaseResult;
+			GenerateImage(Options, BaseResult, node.m_pBase);
+			a = BaseResult.op;
+		}
+        else
+        {
+            // This argument is required
+            a = GenerateMissingImageCode(TEXT("Image Binarise Base"), EImageFormat::IF_RGB_UBYTE, InNode->GetMessageContext(), Options);
+        }
+        a = GenerateImageFormat( a, EImageFormat::IF_RGB_UBYTE );
+        a = GenerateImageSize( a, Options.RectSize);
+        op->SetChild( op->op.args.ImageBinarise.base, a );
+
+        // Threshold
+        Ptr<ASTOp> b = 0;
+        if ( Node* pScalar = node.m_pThreshold.get() )
+        {
+            b = Generate_Generic( pScalar, Options);
+        }
+        else
+        {
+            // This argument is required
+            b = GenerateMissingScalarCode(TEXT("Image Binarise Threshold"), 0.5f, InNode->GetMessageContext() );
+        }
+        op->SetChild( op->op.args.ImageBinarise.threshold, b );
+
+        Result.op = op;
+    }
+
+
+    //---------------------------------------------------------------------------------------------
+	void CodeGenerator::GenerateImage_Resize(const FImageGenerationOptions& Options, FImageGenerationResult& Result, const NodeImageResize* InNode)
+	{
+		const NodeImageResize::Private& node = *InNode->GetPrivate();
+
+		MUTABLE_CPUPROFILER_SCOPE(NodeImageResize);
+
+        Ptr<ASTOp> at = 0;
+
+        // Source image
+        Ptr<ASTOp> base;
+        if ( node.m_pBase )
+        {
+			FImageGenerationOptions NewOptions = Options;
+			
+			if (node.m_relative)
+			{
+				NewOptions.RectSize[0] = FMath::RoundToInt32(NewOptions.RectSize[0] / node.m_sizeX);
+				NewOptions.RectSize[1] = FMath::RoundToInt32(NewOptions.RectSize[1] / node.m_sizeY);
+			}
+
+			FImageGenerationResult BaseResult;
+			GenerateImage(NewOptions, BaseResult, node.m_pBase);
+			base = BaseResult.op;
+        }
+        else
+        {
+            // This argument is required
+            base = GenerateMissingImageCode(TEXT("Image resize base"), EImageFormat::IF_RGB_UBYTE, InNode->GetMessageContext(), Options);
+        }
+
+        // Size
+        if ( node.m_relative )
+        {
+            Ptr<ASTOpFixed> op = new ASTOpFixed();
+            op->op.type = OP_TYPE::IM_RESIZEREL;
+            op->op.args.ImageResizeRel.factor[0] = node.m_sizeX;
+            op->op.args.ImageResizeRel.factor[1] = node.m_sizeY;
+            op->SetChild( op->op.args.ImageResizeRel.source, base);
+            at = op;
+        }
+        else
+        {
+			// Apply the layout block to the rect size
+			UE::Math::TIntVector2<int32> FinalImageSize = { int32(node.m_sizeX), int32(node.m_sizeY) };
+			if (Options.LayoutToApply)
+			{
+				int32 BlockIndex = Options.LayoutToApply->FindBlock(Options.LayoutBlockId);
+				check(BlockIndex >= 0);
+
+				// Block in layout grid units
+				box< FIntVector2 > RectInCells;
+				RectInCells.min = Options.LayoutToApply->Blocks[BlockIndex].Min;
+				RectInCells.size = Options.LayoutToApply->Blocks[BlockIndex].Size;
+
+				FIntPoint grid = Options.LayoutToApply->GetGridSize();
+				grid[0] = FMath::Max(1, grid[0]);
+				grid[1] = FMath::Max(1, grid[1]);
+
+				// Transform to pixels
+				FinalImageSize[0] = (RectInCells.size[0] * FinalImageSize[0]) / grid[0];
+				FinalImageSize[1] = (RectInCells.size[1] * FinalImageSize[1]) / grid[1];
+			}
+
+            Ptr<ASTOpFixed> op = new ASTOpFixed();
+            op->op.type = OP_TYPE::IM_RESIZE;
+            op->op.args.ImageResize.size[0] = (uint16)FinalImageSize[0];
+            op->op.args.ImageResize.size[1] = (uint16)FinalImageSize[1];
+            op->SetChild( op->op.args.ImageResize.source, base);
+            at = op;
+        }
+
+        Result.op = at;
+    }
+
+
+    //---------------------------------------------------------------------------------------------
+	void CodeGenerator::GenerateImage_PlainColour(const FImageGenerationOptions& Options, FImageGenerationResult& Result, const NodeImagePlainColour* InNode)
+	{
+		const NodeImagePlainColour::Private& node = *InNode->GetPrivate();
+
+		// Source colour
+        Ptr<ASTOp> base = 0;
+        if ( node.m_pColour )
+        {
+            base = Generate_Generic( node.m_pColour.get(), Options);
+        }
+        else
+        {
+            // This argument is required
+            base = GenerateMissingColourCode(TEXT("Image plain colour base"), InNode->GetMessageContext() );
+        }
+
+		UE::Math::TIntVector2<int32> FinalImageSize = { 0, 0 };
+
+		if (Options.RectSize.X > 0)
+		{
+			FinalImageSize = Options.RectSize;
+		}
+		else
+		{
+			FinalImageSize = { node.m_sizeX, node.m_sizeY };
+
+			// Apply the layout block to the rect size
+			if (Options.LayoutToApply)
+			{
+				int32 BlockIndex = Options.LayoutToApply->FindBlock(Options.LayoutBlockId);
+				check(BlockIndex >= 0);
+
+				// Block in layout grid units
+				box< FIntVector2 > RectInCells;
+				RectInCells.min = Options.LayoutToApply->Blocks[BlockIndex].Min;
+				RectInCells.size = Options.LayoutToApply->Blocks[BlockIndex].Size;
+
+				FIntPoint grid = Options.LayoutToApply->GetGridSize();
+				grid[0] = FMath::Max(1, grid[0]);
+				grid[1] = FMath::Max(1, grid[1]);
+
+				// Transform to pixels
+				FinalImageSize[0] = (RectInCells.size[0] * FinalImageSize[0]) / grid[0];
+				FinalImageSize[1] = (RectInCells.size[1] * FinalImageSize[1]) / grid[1];
+			}
+		}
+
+
+        Ptr<ASTOpFixed> op = new ASTOpFixed();
+        op->op.type = OP_TYPE::IM_PLAINCOLOUR;
+        op->SetChild( op->op.args.ImagePlainColour.colour, base);
+		op->op.args.ImagePlainColour.format = node.Format;
+		op->op.args.ImagePlainColour.size[0] = FinalImageSize[0];
+		op->op.args.ImagePlainColour.size[1] = FinalImageSize[1];
+		op->op.args.ImagePlainColour.LODs = 1;
+
+        Result.op = op;
+    }
+
+
+    //---------------------------------------------------------------------------------------------
+	void CodeGenerator::GenerateImage_Switch(const FImageGenerationOptions& Options, FImageGenerationResult& Result, const NodeImageSwitch* InNode)
+	{
+		const NodeImageSwitch::Private& node = *InNode->GetPrivate();
+
+		MUTABLE_CPUPROFILER_SCOPE(NodeImageSwitch);
+
+        if (node.m_options.Num() == 0)
+		{
+			// No options in the switch!
+            Ptr<ASTOp> missingOp = GenerateMissingImageCode(TEXT("Switch option"), EImageFormat::IF_RGBA_UBYTE, InNode->GetMessageContext(), Options);
+			Result.op = missingOp;
+			return;
+		}
+
+        Ptr<ASTOpSwitch> op = new ASTOpSwitch();
+        op->type = OP_TYPE::IM_SWITCH;
+
+		// Variable value
+		if ( node.m_pParameter )
+		{
+            op->variable = Generate_Generic( node.m_pParameter.get(), Options);
+		}
+		else
+		{
+			// This argument is required
+            op->variable = GenerateMissingScalarCode(TEXT("Switch variable"), 0.0f, InNode->GetMessageContext() );
+		}
+
+		// Options
+        for ( std::size_t t=0; t< node.m_options.Num(); ++t )
+        {
+            Ptr<ASTOp> branch;
+
+            if (node.m_options[t])
+            {
+				FImageGenerationResult BaseResult;
+				GenerateImage(Options, BaseResult, node.m_options[t]);
+				branch = BaseResult.op;
+			}
+            else
+            {
+                // This argument is required
+                branch = GenerateMissingImageCode(TEXT("Switch option"), EImageFormat::IF_RGBA_UBYTE, InNode->GetMessageContext(), Options);
+            }
+
+            op->cases.Emplace((int16_t)t,op,branch);
+        }
+
+        Ptr<ASTOp> switchAt = op;
+
+        // Make sure all options are the same format and size
+		// Disabled: This is not always desirable. For example if the image is going to be used in a 
+		// projector, the size doesn't need to be constrained.
+        //auto desc = switchAt->GetImageDesc( true );
+        //if ( desc.m_format == EImageFormat::IF_NONE )
+        //{
+        //    // TODO: Look for the most generic of the options?
+        //    // For now force a decently generic one
+        //    desc.m_format = EImageFormat::IF_RGBA_UBYTE;
+        //}
+
+        //if (desc.m_size[0]!=0 && desc.m_size[1]!=0)
+        //{
+        //    Ptr<ASTOpFixed> sop = new ASTOpFixed();
+        //    sop->op.type = OP_TYPE::IM_RESIZE;
+        //    sop->op.args.ImageResize.size[0] = desc.m_size[0];
+        //    sop->op.args.ImageResize.size[1] = desc.m_size[1];
+        //    sop->SetChild( sop->op.args.ImageResize.source, switchAt );
+        //    switchAt = sop;
+        //}
+
+        //{
+        //    Ptr<ASTOpFixed> fop = new ASTOpFixed();
+        //    fop->op.type = OP_TYPE::IM_PIXELFORMAT;
+        //    fop->op.args.ImagePixelFormat.format = desc.m_format;
+        //    fop->SetChild( fop->op.args.ImagePixelFormat.source, switchAt );
+        //    switchAt = fop;
+        //}
+
+        Result.op = switchAt;
+    }
+
+
+    //---------------------------------------------------------------------------------------------
+	void CodeGenerator::GenerateImage_Conditional(const FImageGenerationOptions& Options, FImageGenerationResult& Result, const NodeImageConditional* InNode)
+	{
+		const NodeImageConditional::Private& node = *InNode->GetPrivate();
+
+		Ptr<ASTOpConditional> op = new ASTOpConditional();
+        op->type = OP_TYPE::IM_CONDITIONAL;
+
+        // Condition
+        if ( node.m_parameter )
+        {
+            op->condition = Generate_Generic( node.m_parameter.get(), Options);
+        }
+        else
+        {
+            // This argument is required
+            op->condition = GenerateMissingBoolCode(TEXT("Conditional condition"), true, InNode->GetMessageContext() );
+        }
+
+        // Options
+		FImageGenerationResult YesResult;
+		GenerateImage(Options, YesResult, node.m_true);
+		op->yes = YesResult.op;
+
+		FImageGenerationResult NoResult;
+		GenerateImage(Options, NoResult, node.m_false);
+		op->no = NoResult.op;
+		
+        Result.op = op;
+    }
+
+
+    //---------------------------------------------------------------------------------------------
+	void CodeGenerator::GenerateImage_Project(const FImageGenerationOptions& Options, FImageGenerationResult& Result, const NodeImageProject* InNode)
+	{
+		const NodeImageProject::Private& node = *InNode->GetPrivate();
+
+		MUTABLE_CPUPROFILER_SCOPE(NodeImageProject);
+
+        // Mesh project operation
+        //------------------------------
+        Ptr<ASTOpFixed> ProjectOp = new ASTOpFixed();
+		ProjectOp->op.type = OP_TYPE::ME_PROJECT;
+
+        Ptr<ASTOp> LastMeshOp = ProjectOp;
+
+        // Projector
+        FProjectorGenerationResult projectorResult;
+        if ( node.m_pProjector )
+        {
+            GenerateProjector( projectorResult, Options, node.m_pProjector );
+            //projectorAt = Generate( node.m_pProjector.get() );
+        }
+        else
+        {
+            // This argument is required
+            GenerateMissingProjectorCode( projectorResult, InNode->GetMessageContext() );
+        }
+
+		ProjectOp->SetChild(ProjectOp->op.args.MeshProject.projector, projectorResult.op );
+
+		int32 LayoutBlockIndex = -1;
+		if (Options.LayoutToApply)
+		{
+			LayoutBlockIndex = Options.LayoutToApply->Blocks.IndexOfByPredicate([&](const FLayoutBlock& Block) { return Block.Id == Options.LayoutBlockId; });
+		}
+		uint64 GeneratedLayoutBlockId = FLayoutBlock::InvalidBlockId;
+
+        // Mesh
+        if ( node.m_pMesh )
+        {
+			// TODO: This will probably Result in a duplicated mesh subgraph, with the original mesh but new layout block ids.
+			// See if it can be optimized and try to reuse the existing layout block ids instead of generating new ones.
+			FMeshGenerationOptions MeshOptions;
+			MeshOptions.State = Options.State;
+			MeshOptions.ActiveTags = Options.ActiveTags;
+			MeshOptions.bLayouts = true;			// We need the layout that we will use to render
+			MeshOptions.bNormalizeUVs = true;		// We need normalized UVs for the projection
+			
+			FMeshGenerationResult MeshResult;
+			GenerateMesh( MeshOptions, MeshResult, node.m_pMesh );
+
+			// Match the block id of the block we are generating with the id that resulted in the generated mesh			
+			Ptr<const Layout> Layout = MeshResult.GeneratedLayouts.IsValidIndex(node.m_layout) ? MeshResult.GeneratedLayouts[node.m_layout].Layout : nullptr;
+			if (Layout && Layout->Blocks.IsValidIndex(LayoutBlockIndex))
+			{
+				GeneratedLayoutBlockId = Layout->Blocks[LayoutBlockIndex].Id;
+			}
+			else if (Layout && Layout->Blocks.Num() == 1)
+			{
+				// Layout management disabled, use the only block available
+				GeneratedLayoutBlockId = Layout->Blocks[0].Id;
+			}
+			else
+			{
+				ErrorLog->GetPrivate()->Add("Layout or block index error.", ELMT_ERROR, InNode->GetMessageContext());
+			}
+
+			// TODO: 
+			// MeshResult.MeshOp has some modifiers applied already: the ones applied before other operations directly in the mesh constant generation. 
+			// This is not what was happening before the refactor  so use MeshResult.BaseMeshOp. This is another case of ambiguity of order of modifiers 
+			// that whould be fixed with the general ordering design.
+			//Ptr<ASTOp> CurrentMeshToProjectOp = MeshResult.MeshOp;
+			Ptr<ASTOp> CurrentMeshToProjectOp = MeshResult.BaseMeshOp;
+
+            if (projectorResult.type == PROJECTOR_TYPE::WRAPPING)
+            {
+                // For wrapping projector we need the entire mesh. The actual project operation
+                // will remove the faces that are not in the layout block we are generating.
+                Ptr<ASTOpConstantResource> cop = new ASTOpConstantResource();
+                cop->Type = OP_TYPE::ME_CONSTANT;
+				Ptr<Mesh> FormatMeshResult = new Mesh();
+				CreateMeshOptimisedForWrappingProjection(FormatMeshResult.get(), node.m_layout);
+
+                cop->SetValue(FormatMeshResult, CompilerOptions->OptimisationOptions.DiskCacheContext);
+
+                Ptr<ASTOpMeshFormat> FormatOp = new ASTOpMeshFormat();
+				FormatOp->Flags = OP::MeshFormatArgs::Vertex
+                        | OP::MeshFormatArgs::Index
+                        | OP::MeshFormatArgs::ResetBufferIndices;
+				FormatOp->Format = cop;
+				FormatOp->Source = CurrentMeshToProjectOp;
+				CurrentMeshToProjectOp = FormatOp;
+            }
+            else
+            {
+                // Extract the mesh layout block
+                if ( GeneratedLayoutBlockId!= FLayoutBlock::InvalidBlockId)
+                {
+                    Ptr<ASTOpMeshExtractLayoutBlocks> eop = new ASTOpMeshExtractLayoutBlocks();
+                    eop->Source = CurrentMeshToProjectOp;
+                    eop->LayoutIndex = node.m_layout;
+
+                    eop->Blocks.Add(GeneratedLayoutBlockId);
+
+					CurrentMeshToProjectOp = eop;
+                }
+
+                // Reformat the mesh to a more efficient format for this operation
+                Ptr<ASTOpConstantResource> cop = new ASTOpConstantResource();
+                cop->Type = OP_TYPE::ME_CONSTANT;
+
+				Ptr<Mesh> FormatMeshResult = new Mesh();
+                CreateMeshOptimisedForProjection(FormatMeshResult.get(), node.m_layout);
+
+                cop->SetValue(FormatMeshResult, CompilerOptions->OptimisationOptions.DiskCacheContext);
+
+                Ptr<ASTOpMeshFormat> FormatOp = new ASTOpMeshFormat();
+				FormatOp->Flags = OP::MeshFormatArgs::Vertex
+					| OP::MeshFormatArgs::Index
+					| OP::MeshFormatArgs::ResetBufferIndices;
+				FormatOp->Format = cop;
+				FormatOp->Source = CurrentMeshToProjectOp;
+				CurrentMeshToProjectOp = FormatOp;
+            }
+
+			ProjectOp->SetChild(ProjectOp->op.args.MeshProject.mesh, CurrentMeshToProjectOp);
+
+        }
+        else
+        {
+            // This argument is required
+            Ptr<const Mesh> TempMesh = new Mesh();
+            Ptr<ASTOpConstantResource> cop = new ASTOpConstantResource();
+            cop->Type = OP_TYPE::ME_CONSTANT;
+            cop->SetValue(TempMesh, CompilerOptions->OptimisationOptions.DiskCacheContext);
+			ProjectOp->SetChild(ProjectOp->op.args.MeshProject.mesh, cop );
+            ErrorLog->GetPrivate()->Add( "Projector mesh not set.", ELMT_ERROR, InNode->GetMessageContext() );
+        }
+
+
+        // Image raster operation
+        //------------------------------
+        Ptr<ASTOpImageRasterMesh> ImageRasterOp = new ASTOpImageRasterMesh();
+		ImageRasterOp->mesh = LastMeshOp;
+		ImageRasterOp->projector = projectorResult.op;
+
+		// Calculate size of image to raster:
+		// The full image is:
+		// 0) The hint value in the image options passed down.
+		// 1) whatever is specified in the node attributes.
+		// 2) if that is 0, the size of the mask
+		// 3) if still 0, take the size of the image to project (which is not necessarily related, but often)
+		// 4) if still 0, a default value bigger than 0
+		// then if we are applying a layout a layout block rect need to be calculated of that size, like in image constants.
+		UE::Math::TIntVector2<int32> RasterImageSize = Options.RectSize;
+		bool bApplyLayoutToSize = false;
+
+		if (RasterImageSize.X == 0)
+		{
+			RasterImageSize = UE::Math::TIntVector2<int32>(node.m_imageSize);
+			bApplyLayoutToSize = true;
+		}
+
+		// Target mask
+		if (node.m_pMask)
+		{
+			FImageGenerationResult MaskResult;
+			GenerateImage(Options, MaskResult, node.m_pMask);
+			Ptr<ASTOp> mask = MaskResult.op;
+
+			mask = GenerateImageFormat(mask, EImageFormat::IF_L_UBYTE);
+
+			ImageRasterOp->mask = GenerateImageSize(mask, RasterImageSize);
+
+			if (RasterImageSize.X == 0)
+			{
+				FImageDesc MaskDesc = ImageRasterOp->mask->GetImageDesc();
+				RasterImageSize = UE::Math::TIntVector2<int32>(MaskDesc.m_size);
+				bApplyLayoutToSize = true;
+			}
+		}
+
+        // Image
+        if ( node.m_pImage )
+        {
+            // Generate
+			FImageGenerationOptions NewOptions = Options;
+			NewOptions.ImageLayoutStrategy = CompilerOptions::TextureLayoutStrategy::None;
+			NewOptions.LayoutToApply = nullptr;
+			NewOptions.LayoutBlockId = FLayoutBlock::InvalidBlockId;
+			NewOptions.RectSize = { 0,0 };
+
+			FImageGenerationResult ImageResult;
+			GenerateImage(NewOptions, ImageResult, node.m_pImage);
+			ImageRasterOp->image = ImageResult.op;
+
+			FImageDesc desc = ImageRasterOp->image->GetImageDesc();
+			ImageRasterOp->SourceSizeX = desc.m_size[0];
+			ImageRasterOp->SourceSizeY = desc.m_size[1];
+
+			if (RasterImageSize.X == 0)
+			{
+				RasterImageSize = UE::Math::TIntVector2<int32>(desc.m_size);
+				bApplyLayoutToSize = true;
+			}
+        }
+        else
+        {
+            // This argument is required
+			ImageRasterOp->image = GenerateMissingImageCode(TEXT("Projector image"), EImageFormat::IF_RGB_UBYTE, InNode->GetMessageContext(), Options);
+        }
+
+		if (RasterImageSize.X == 0)
+		{
+			// Last resort
+			RasterImageSize = { 256, 256 };
+		}
+
+		// Apply the layout block to the rect size
+		if (bApplyLayoutToSize && Options.LayoutToApply)
+		{
+			int32 BlockIndex = Options.LayoutToApply->FindBlock(Options.LayoutBlockId);
+			check(BlockIndex >= 0);
+
+			// Block in layout grid units
+			box< FIntVector2 > RectInCells;
+			RectInCells.min = Options.LayoutToApply->Blocks[BlockIndex].Min;
+			RectInCells.size = Options.LayoutToApply->Blocks[BlockIndex].Size;
+
+			FIntPoint grid = Options.LayoutToApply->GetGridSize();
+			grid[0] = FMath::Max(1, grid[0]);
+			grid[1] = FMath::Max(1, grid[1]);
+
+			// Transform to pixels
+			RasterImageSize[0] = (RectInCells.size[0] * RasterImageSize[0]) / grid[0];
+			RasterImageSize[1] = (RectInCells.size[1] * RasterImageSize[1]) / grid[1];
+		}
+
+        // Image size, from the current block being generated
+		ImageRasterOp->SizeX = RasterImageSize[0];
+		ImageRasterOp->SizeY = RasterImageSize[1];
+		ImageRasterOp->BlockId = GeneratedLayoutBlockId;
+		ImageRasterOp->LayoutIndex = node.m_layout;
+
+		ImageRasterOp->bIsRGBFadingEnabled = node.bIsRGBFadingEnabled;
+		ImageRasterOp->bIsAlphaFadingEnabled = node.bIsAlphaFadingEnabled;
+		ImageRasterOp->SamplingMethod = node.SamplingMethod;
+		ImageRasterOp->MinFilterMethod = node.MinFilterMethod;
+
+		// Fading angles are optional, and stored in a colour. If one exists, we generate both.
+		if (node.m_pAngleFadeStart || node.m_pAngleFadeEnd)
+		{
+			Ptr<NodeScalarConstant> pDefaultFade = new NodeScalarConstant();
+			pDefaultFade->SetValue(180.0f);
+
+			Ptr<NodeColourFromScalars> pPropsNode = new NodeColourFromScalars();
+
+			if (node.m_pAngleFadeStart) pPropsNode->X = node.m_pAngleFadeStart;
+			else pPropsNode->X = pDefaultFade;
+
+			if (node.m_pAngleFadeEnd) pPropsNode->Y = node.m_pAngleFadeEnd;
+			else pPropsNode->Y = pDefaultFade;
+
+			ImageRasterOp->angleFadeProperties = Generate_Generic(pPropsNode, Options);
+		}
+
+        // Seam correction operations
+        //------------------------------
+		if (node.bEnableTextureSeamCorrection)
+		{
+			Ptr<ASTOpImageRasterMesh> MaskRasterOp = new ASTOpImageRasterMesh();
+			MaskRasterOp->mesh = ImageRasterOp->mesh.child();
+			MaskRasterOp->image = 0;
+			MaskRasterOp->mask = 0;
+			MaskRasterOp->BlockId = ImageRasterOp->BlockId;
+			MaskRasterOp->LayoutIndex = ImageRasterOp->LayoutIndex;
+			MaskRasterOp->SizeX = ImageRasterOp->SizeX;
+			MaskRasterOp->SizeY = ImageRasterOp->SizeY;
+			MaskRasterOp->UncroppedSizeX = ImageRasterOp->UncroppedSizeX;
+			MaskRasterOp->UncroppedSizeY = ImageRasterOp->UncroppedSizeY;
+			MaskRasterOp->CropMinX = ImageRasterOp->CropMinX;
+			MaskRasterOp->CropMinY = ImageRasterOp->CropMinY;
+			MaskRasterOp->SamplingMethod = ESamplingMethod::Point;
+			MaskRasterOp->MinFilterMethod = EMinFilterMethod::None;
+
+			Ptr<ASTOpImageMakeGrowMap> MakeGrowMapOp = new ASTOpImageMakeGrowMap();
+			MakeGrowMapOp->Mask = MaskRasterOp;
+			MakeGrowMapOp->Border = MUTABLE_GROW_BORDER_VALUE;
+
+			// If we want to be able to generate progressive mips efficiently, we need mipmaps for the "displacement map".
+			if (CompilerOptions->OptimisationOptions.bEnableProgressiveImages)
+			{
+				Ptr<ASTOpImageMipmap> MipMask = new ASTOpImageMipmap;
+				MipMask->Source = MakeGrowMapOp->Mask.child();
+				MipMask->bPreventSplitTail = true;
+				MakeGrowMapOp->Mask = MipMask;
+			}
+
+			Ptr<ASTOpFixed> DisplaceOp = new ASTOpFixed();
+			DisplaceOp->op.type = OP_TYPE::IM_DISPLACE;
+			DisplaceOp->SetChild(DisplaceOp->op.args.ImageDisplace.displacementMap, MakeGrowMapOp);
+			DisplaceOp->SetChild(DisplaceOp->op.args.ImageDisplace.source, ImageRasterOp);
+
+			Result.op = DisplaceOp;
+		}
+		else
+		{
+			Result.op = ImageRasterOp;
+		}		
+    }
+
+
+    //---------------------------------------------------------------------------------------------
+	void CodeGenerator::GenerateImage_Mipmap(const FImageGenerationOptions& Options, FImageGenerationResult& Result, const NodeImageMipmap* InNode)
+	{
+		const NodeImageMipmap::Private& node = *InNode->GetPrivate();
+
+		MUTABLE_CPUPROFILER_SCOPE(NodeImageMipmap);
+
+        Ptr<ASTOp> res;
+
+        Ptr<ASTOpImageMipmap> op = new ASTOpImageMipmap();
+
+        // At the end of the day, we want all the mipmaps. Maybe the code optimiser will split the process later.
+        op->Levels = 0;
+
+        // Source image
+        Ptr<ASTOp> base;
+        if ( node.m_pSource )
+        {
+            MUTABLE_CPUPROFILER_SCOPE(Base);
+			FImageGenerationResult BaseResult;
+			GenerateImage(Options, BaseResult, node.m_pSource);
+			base = BaseResult.op;
+		}
+        else
+        {
+            // This argument is required
+            base = GenerateMissingImageCode(TEXT("Mipmap image"), EImageFormat::IF_RGB_UBYTE, InNode->GetMessageContext(), Options);
+        }
+
+        op->Source = base;
+
+        // The number of tail mipmaps depends on the cell size. We need to know it for some code
+        // optimisation operations. Scan the source image code looking for this info
+        int32 blockX = 0;
+        int32 blockY = 0;
+        if ( Options.ImageLayoutStrategy
+             !=
+             CompilerOptions::TextureLayoutStrategy::None )
+        {
+            MUTABLE_CPUPROFILER_SCOPE(GetLayoutBlockSize);
+            op->Source->GetLayoutBlockSize( &blockX, &blockY );
+        }
+
+        if ( blockX && blockY )
+        {
+            int32 mipsX = (int32)ceilf( logf( (float)blockX )/logf(2.0f) );
+            int32 mipsY = (int32)ceilf( logf( (float)blockY )/logf(2.0f) );
+            op->BlockLevels = (uint8)FMath::Max( mipsX, mipsY );
+        }
+        else
+        {
+            // No layout. Mipmap all the way down.
+            op->BlockLevels = 0;
+        }
+
+		op->AddressMode = node.m_settings.AddressMode;
+		op->FilterType = node.m_settings.FilterType;
+
+        res = op;
+
+        Result.op = res;
+    }
+
+
+	//---------------------------------------------------------------------------------------------
+	void CodeGenerator::GenerateImage_Invert(const FImageGenerationOptions& Options, FImageGenerationResult& Result, const NodeImageInvert* InNode)
+	{
+		const NodeImageInvert::Private& node = *InNode->GetPrivate();
+
+		Ptr<ASTOpFixed> op = new ASTOpFixed();
+		op->op.type = OP_TYPE::IM_INVERT;
+
+		// A image
+		Ptr<ASTOp> a;
+		if (node.m_pBase)
+		{
+			FImageGenerationResult BaseResult;
+			GenerateImage(Options, BaseResult, node.m_pBase);
+			a = BaseResult.op;
+		}
+		else
+		{
+			// This argument is required
+			a = GenerateMissingImageCode(TEXT("Image Invert Color"), EImageFormat::IF_RGB_UBYTE, InNode->GetMessageContext(), Options);
+		}
+		a = GenerateImageFormat(a, EImageFormat::IF_RGB_UBYTE);
+		a = GenerateImageSize(a, Options.RectSize);
+		op->SetChild(op->op.args.ImageInvert.base, a);
+
+		Result.op = op;
+	}
+
+
+    //---------------------------------------------------------------------------------------------
+	void CodeGenerator::GenerateImage_Variation(const FImageGenerationOptions& Options, FImageGenerationResult& Result, const NodeImageVariation* InNode)
+	{
+		const NodeImageVariation::Private& node = *InNode->GetPrivate();
+
+		Ptr<ASTOp> currentOp;
+
+        // Default case
+        if ( node.m_defaultImage )
+        {
+            FImageGenerationResult BranchResults;
+			GenerateImage(Options, BranchResults, node.m_defaultImage);
+			currentOp = BranchResults.op;
+        }
+        
+        // Process variations in reverse order, since conditionals are built bottom-up.
+        for ( int32 t = int32( node.m_variations.Num() ) - 1; t >= 0; --t )
+        {
+            int32 tagIndex = -1;
+            const FString& tag = node.m_variations[t].m_tag;
+            for ( int32 i = 0; i < int32( FirstPass.Tags.Num() ); ++i )
+            {
+                if ( FirstPass.Tags[i].Tag == tag )
+                {
+                    tagIndex = i;
+                }
+            }
+
+            if ( tagIndex < 0 )
+            {
+				FString Msg = FString::Printf(TEXT("Unknown tag found in image variation [%s]."), *tag );
+
+                ErrorLog->GetPrivate()->Add( Msg, ELMT_WARNING, InNode->GetMessageContext() );
+                continue;
+            }
+
+            Ptr<ASTOp> variationOp;
+            if ( node.m_variations[t].m_image )
+            {
+				FImageGenerationResult VariationResult;
+				GenerateImage(Options, VariationResult, node.m_variations[t].m_image);
+				variationOp = VariationResult.op;
+            }
+            else
+            {
+                // This argument is required
+                variationOp = GenerateMissingImageCode(TEXT("Variation option"), EImageFormat::IF_RGBA_UBYTE, InNode->GetMessageContext(), Options);
+            }
+
+
+            Ptr<ASTOpConditional> conditional = new ASTOpConditional;
+            conditional->type = OP_TYPE::IM_CONDITIONAL;
+            conditional->no = currentOp;
+            conditional->yes = variationOp;
+            conditional->condition = FirstPass.Tags[tagIndex].GenericCondition;
+
+            currentOp = conditional;
+        }
+
+        // Make sure all options are the same format and size
+        FImageDesc desc = currentOp->GetImageDesc( true );
+        if ( desc.m_format == EImageFormat::IF_NONE )
+        {
+            // TODO: Look for the most generic of the options?
+            // For now force a decently generic one
+            desc.m_format = EImageFormat::IF_RGBA_UBYTE;
+        }
+
+        if ( desc.m_size[0] != 0 && desc.m_size[1] != 0 )
+        {
+            Ptr<ASTOpFixed> sop = new ASTOpFixed();
+            sop->op.type = OP_TYPE::IM_RESIZE;
+            sop->op.args.ImageResize.size[0] = desc.m_size[0];
+            sop->op.args.ImageResize.size[1] = desc.m_size[1];
+            sop->SetChild( sop->op.args.ImageResize.source, currentOp );
+            currentOp = sop;
+        }
+
+        {
+            Ptr<ASTOpImagePixelFormat> fop = new ASTOpImagePixelFormat();
+            fop->Format = desc.m_format;
+            fop->Source = currentOp;
+            currentOp = fop;
+        }
+
+        Result.op = currentOp;
+    }
+
+	
+	//---------------------------------------------------------------------------------------------
+	void CodeGenerator::GenerateImage_Table(const FImageGenerationOptions& Options, FImageGenerationResult& Result, const NodeImageTable* InNode)
+	{
+		const NodeImageTable& node = *InNode;
+
+		Result.op = GenerateTableSwitch<NodeImageTable, ETableColumnType::Image, OP_TYPE::IM_SWITCH>(node,
+			[this, InNode, Options](const NodeImageTable& node, int32 colIndex, int32 row, mu::ErrorLog* pErrorLog)
+			{
+				const FTableValue& CellData = node.Table->GetPrivate()->Rows[row].Values[colIndex];
+				Ptr<const Image> pImage = nullptr;
+
+				if (Ptr<ResourceProxy<Image>> pProxyImage = CellData.ProxyImage)
+				{
+					pImage = pProxyImage->Get();
+				}
+
+				Ptr<ASTOp> ImageOp;
+
+				if (!pImage)
+				{
+					FString Msg = FString::Printf(TEXT("Table has a missing image in column %d, row %d."), colIndex, row);
+					pErrorLog->GetPrivate()->Add(Msg, ELMT_ERROR, InNode->GetMessageContext());
+
+					return ImageOp;
+				}
+				else
+				{
+					Ptr<NodeImageConstant> ImageConst = new NodeImageConstant();
+					ImageConst->SetValue(pImage.get());
+
+					// TODO: We probably want to get the data tags from the table row.
+					ImageConst->SourceDataDescriptor = InNode->SourceDataDescriptor;
+
+					FImageGenerationResult Result;
+					GenerateImage(Options, Result, ImageConst);
+					ImageOp = Result.op;
+
+					int32 MaxTextureSize = FMath::Max(node.ReferenceImageDesc.m_size[0], node.ReferenceImageDesc.m_size[1]);
+
+					if (MaxTextureSize > 0 && (MaxTextureSize < pImage->GetSizeX() || MaxTextureSize < pImage->GetSizeY()))
+					{
+						float Factor = FMath::Min(MaxTextureSize / (float)(pImage->GetSizeX()), MaxTextureSize / (float)(pImage->GetSizeY()));
+						Ptr<ASTOpFixed> op = new ASTOpFixed();
+						op->op.type = OP_TYPE::IM_RESIZE;
+						op->op.args.ImageResize.size[0] = (uint16)pImage->GetSizeX() * Factor;
+						op->op.args.ImageResize.size[1] = (uint16)pImage->GetSizeY() * Factor;
+						op->SetChild(op->op.args.ImageResize.source, ImageOp);
+						ImageOp = op;
+					}
+				}
+
+				return ImageOp;
+			});
+	}
+
+
+	Ptr<Image> CodeGenerator::GenerateMissingImage(EImageFormat Format)
+	{
+		// Create the image node if it hasn't been created yet.
+		if (!MissingImage[SIZE_T(Format)])
+		{
+			// Make a checkered debug image
+			const FImageSize Size(16, 16);
+
+			Ptr<Image> GeneratedImage = new Image(Size[0], Size[1], 1, Format, EInitializationType::NotInitialized);
+
+			switch (Format)
+			{
+			case EImageFormat::IF_L_UBYTE:
+			{
+				uint8* DataPtr = GeneratedImage->GetLODData(0);
+				for (int32 P = 0; P < Size[0]*Size[1]; ++P)
+				{
+					if ((P + P/Size[0]) % 2)
+					{
+						DataPtr[0] = 255;
+					}
+					else
+					{
+						DataPtr[0] = 64;
+					}
+
+					DataPtr++;
+				}
+				break;
+			}
+			case EImageFormat::IF_RGB_UBYTE:
+			{
+				uint8* DataPtr = GeneratedImage->GetLODData(0);
+				for (int32 P = 0; P < Size[0]*Size[1]; ++P)
+				{
+					if ((P + P/Size[0]) % 2)
+					{
+						DataPtr[0] = 255;
+						DataPtr[1] = 255;
+						DataPtr[2] = 64;
+					}
+					else
+					{
+						DataPtr[0] = 64;
+						DataPtr[1] = 64;
+						DataPtr[2] = 255;
+					}
+
+					DataPtr += 3;
+				}
+				break;
+			}
+			case EImageFormat::IF_BGRA_UBYTE:
+			case EImageFormat::IF_RGBA_UBYTE:
+			{
+				uint8* DataPtr = GeneratedImage->GetLODData(0);
+				for (int32 P = 0; P < Size[0]*Size[1]; ++P)
+				{
+					if ((P + P/Size[0]) % 2)
+					{
+						DataPtr[0] = 255;
+						DataPtr[1] = 255;
+						DataPtr[2] = 64;
+						DataPtr[3] = 255;
+					}
+					else
+					{
+						DataPtr[0] = 64;
+						DataPtr[1] = 64;
+						DataPtr[2] = 255;
+						DataPtr[3] = 128;
+					}
+
+					DataPtr += 4;
+				}
+				break;
+			}
+
+			default:
+				check( false );
+				break;
+
+			}
+
+			MissingImage[(SIZE_T)Format] = GeneratedImage;
+		}
+
+		return MissingImage[(SIZE_T)Format].get();
+	}
+
+
+	Ptr<ASTOp> CodeGenerator::GenerateMissingImageCode(const TCHAR* strWhere, EImageFormat format, const void* errorContext, const FImageGenerationOptions& Options )
+	{
+		// Log an error message
+		FString Msg = FString::Printf(TEXT("Required connection not found: %s"), strWhere );
+		ErrorLog->GetPrivate()->Add( Msg, ELMT_ERROR, errorContext );
+
+		// Make a checkered debug image
+		Ptr<Image> GeneratedImage = GenerateMissingImage( format );
+
+		Ptr<NodeImageConstant> pNode = new NodeImageConstant();
+		pNode->SetValue(GeneratedImage.get());
+
+		FImageGenerationResult Result;
+		GenerateImage(Options, Result, pNode);
+
+		return Result.op;
+	}
+
+
+    //---------------------------------------------------------------------------------------------
+    Ptr<ASTOp> CodeGenerator::GeneratePlainImageCode( const FVector4f& InColor, const FImageGenerationOptions& Options )
+    {
+		Ptr<NodeColourConstant> ConstantColor = new NodeColourConstant();
+		ConstantColor->Value = InColor;
+
+        Ptr<NodeImagePlainColour> PlainNode = new NodeImagePlainColour();
+		PlainNode->SetColour(ConstantColor);
+
+		FImageGenerationResult TempResult;
+		GenerateImage(Options, TempResult, PlainNode);
+        Ptr<ASTOp> Result = TempResult.op;
+
+        return Result;
+    }
+
+
+    //---------------------------------------------------------------------------------------------
+    Ptr<ASTOp> CodeGenerator::GenerateImageFormat( Ptr<ASTOp> Op, EImageFormat InFormat)
+    {
+        Ptr<ASTOp> Result = Op;
+
+        if (InFormat!=EImageFormat::IF_NONE && Op && Op->GetImageDesc().m_format!=InFormat)
+        {
+            // Generate the format change code
+            Ptr<ASTOpImagePixelFormat> op = new ASTOpImagePixelFormat();
+            op->Source = Op;
+            op->Format = InFormat;
+			Result = op;
+        }
+
+        return Result;
+    }
+
+
+    //---------------------------------------------------------------------------------------------
+    Ptr<ASTOp> CodeGenerator::GenerateImageUncompressed( Ptr<ASTOp> at )
+    {
+        Ptr<ASTOp> Result = at;
+
+        if (at)
+        {
+            EImageFormat sourceFormat = at->GetImageDesc().m_format;
+            EImageFormat targetFormat = GetUncompressedFormat( sourceFormat );
+
+            if ( targetFormat != sourceFormat )
+            {
+                // Generate the format change code
+                Ptr<ASTOpImagePixelFormat> op = new ASTOpImagePixelFormat();
+                op->Source = at;
+                op->Format = targetFormat;
+                Result = op;
+            }
+        }
+
+        return Result;
+    }
+
+
+    //---------------------------------------------------------------------------------------------
+    Ptr<ASTOp> CodeGenerator::GenerateImageSize( Ptr<ASTOp> at, UE::Math::TIntVector2<int32> size )
+    {
+        Ptr<ASTOp> Result = at;
+
+		if (size[0] > 0 && size[1] > 0)
+		{
+			if (UE::Math::TIntVector2<int32>(at->GetImageDesc().m_size) != size)
+			{
+				Ptr<ASTOpFixed> op = new ASTOpFixed();
+				op->op.type = OP_TYPE::IM_RESIZE;
+				op->SetChild(op->op.args.ImageResize.source, at);
+				op->op.args.ImageResize.size[0] = size[0];
+				op->op.args.ImageResize.size[1] = size[1];
+				Result = op;
+			}
+		}
+
+        return Result;
+    }
+
+ 
+}
