[mutable] Moved the Mutable plugin out of Experimental status into Beta.

diff --git a/Engine/Plugins/Mutable/Source/MutableTools/Private/MuT/CodeGenerator_Mesh.cpp b/Engine/Plugins/Mutable/Source/MutableTools/Private/MuT/CodeGenerator_Mesh.cpp
new file mode 100644
--- /dev/null
+++ b/Engine/Plugins/Mutable/Source/MutableTools/Private/MuT/CodeGenerator_Mesh.cpp
@@ -0,0 +1,1700 @@
+// Copyright Epic Games, Inc. All Rights Reserved.
+
+#include "Containers/Array.h"
+#include "Math/IntPoint.h"
+#include "Math/UnrealMathSSE.h"
+#include "Misc/AssertionMacros.h"
+#include "MuR/Layout.h"
+#include "MuR/Mesh.h"
+#include "MuR/MeshPrivate.h"
+#include "MuR/MeshBufferSet.h"
+#include "MuR/MutableMath.h"
+#include "MuR/Operations.h"
+#include "MuR/ParametersPrivate.h"
+#include "MuR/Platform.h"
+#include "MuR/Ptr.h"
+#include "MuR/RefCounted.h"
+#include "MuR/System.h"
+#include "MuT/AST.h"
+#include "MuT/ASTOpConditional.h"
+#include "MuT/ASTOpConstantResource.h"
+#include "MuT/ASTOpMeshApplyPose.h"
+#include "MuT/ASTOpMeshApplyShape.h"
+#include "MuT/ASTOpMeshBindShape.h"
+#include "MuT/ASTOpMeshClipDeform.h"
+#include "MuT/ASTOpMeshClipMorphPlane.h"
+#include "MuT/ASTOpMeshExtractLayoutBlocks.h"
+#include "MuT/ASTOpMeshFormat.h"
+#include "MuT/ASTOpMeshGeometryOperation.h"
+#include "MuT/ASTOpMeshMorphReshape.h"
+#include "MuT/ASTOpMeshTransform.h"
+#include "MuT/ASTOpMeshDifference.h"
+#include "MuT/ASTOpMeshMorph.h"
+#include "MuT/ASTOpMeshAddTags.h"
+#include "MuT/ASTOpSwitch.h"
+#include "MuT/ASTOpReferenceResource.h"
+#include "MuT/CodeGenerator.h"
+#include "MuT/CodeGenerator_FirstPass.h"
+#include "MuT/CompilerPrivate.h"
+#include "MuT/ErrorLog.h"
+#include "MuT/ErrorLogPrivate.h"
+#include "MuT/NodeImageProject.h"
+#include "MuT/NodeLayout.h"
+#include "MuT/NodeMesh.h"
+#include "MuT/NodeMeshApplyPose.h"
+#include "MuT/NodeMeshApplyPosePrivate.h"
+#include "MuT/NodeMeshClipDeform.h"
+#include "MuT/NodeMeshClipMorphPlane.h"
+#include "MuT/NodeMeshClipWithMesh.h"
+#include "MuT/NodeMeshConstant.h"
+#include "MuT/NodeMeshConstantPrivate.h"
+#include "MuT/NodeMeshFormat.h"
+#include "MuT/NodeMeshFormatPrivate.h"
+#include "MuT/NodeMeshFragment.h"
+#include "MuT/NodeMeshGeometryOperation.h"
+#include "MuT/NodeMeshGeometryOperationPrivate.h"
+#include "MuT/NodeMeshInterpolate.h"
+#include "MuT/NodeMeshInterpolatePrivate.h"
+#include "MuT/NodeMeshMakeMorph.h"
+#include "MuT/NodeMeshMakeMorphPrivate.h"
+#include "MuT/NodeMeshMorph.h"
+#include "MuT/NodeMeshMorphPrivate.h"
+#include "MuT/NodeMeshReshape.h"
+#include "MuT/NodeMeshReshapePrivate.h"
+#include "MuT/NodeMeshSwitch.h"
+#include "MuT/NodeMeshSwitchPrivate.h"
+#include "MuT/NodeMeshTable.h"
+#include "MuT/NodeMeshTransform.h"
+#include "MuT/NodeMeshTransformPrivate.h"
+#include "MuT/NodeMeshVariation.h"
+#include "MuT/NodeMeshVariationPrivate.h"
+#include "MuT/NodeScalar.h"
+#include "MuT/Table.h"
+#include "MuT/TablePrivate.h"
+
+namespace mu
+{
+	class Node;
+
+	template<typename T>
+	struct TArray2D
+	{
+		int32 SizeX = 0;
+		int32 SizeY = 0;
+		TArray<T> Data;
+
+		void Init(const T& Value, int32 InSizeX, int32 InSizeY)
+		{
+			SizeX = InSizeX;
+			SizeY = InSizeY;
+			Data.Init(Value, SizeX*SizeY );
+		}
+
+		inline const T& Get(int32 X, int32 Y) const
+		{
+			check(X >= 0 && X < SizeX);
+			check(Y >= 0 && Y < SizeY);
+			return Data[SizeX * Y + X];
+		}
+
+		inline void Set(int32 X, int32 Y, const T& Value)
+		{
+			check(X >= 0 && X < SizeX);
+			check(Y >= 0 && Y < SizeY);
+			Data[SizeX * Y + X] = Value;
+		}
+
+	};
+
+    //---------------------------------------------------------------------------------------------
+	void CodeGenerator::PrepareMeshForLayout( const FGeneratedLayout& GeneratedLayout,
+		Ptr<Mesh> Mesh,
+		int32 LayoutChannel,
+		const void* errorContext,
+		const FMeshGenerationOptions& MeshOptions,
+		bool bUseAbsoluteBlockIds
+		)
+	{
+		MUTABLE_CPUPROFILER_SCOPE(PrepareMeshForLayout);
+
+		if (Mesh->GetVertexCount() == 0)
+		{
+			return;
+		}
+
+		// The layout must have block ids.
+		check(GeneratedLayout.Layout->Blocks.IsEmpty() || GeneratedLayout.Layout->Blocks[0].Id != FLayoutBlock::InvalidBlockId);
+
+		// 
+		Ptr<const Layout> Layout = GeneratedLayout.Layout;
+		Mesh->AddLayout(Layout);
+
+		const int32 NumVertices = Mesh->GetVertexCount();
+		const int32 NumBlocks = Layout->GetBlockCount();
+
+		bool bIsSingleFullBlock = (NumBlocks == 1) && (Layout->Blocks[0].Min == FIntVector2(0, 0) && Layout->Blocks[0].Size == Layout->Size);
+
+		// Find block ids for each block in the grid. Calculate a grid size that contains all blocks
+		FIntPoint LayoutGrid = Layout->GetGridSize();
+		FIntPoint WorkingGrid = LayoutGrid;
+		for (const FSourceLayoutBlock& Block: GeneratedLayout.Source->Blocks)
+		{
+			WorkingGrid.X = FMath::Max(WorkingGrid.X, Block.Min.X + Block.Size.X );
+			WorkingGrid.Y = FMath::Max(WorkingGrid.Y, Block.Min.Y + Block.Size.Y );
+		}
+
+
+		TArray2D<int32> GridBlockBlockId;
+		GridBlockBlockId.Init(MAX_uint16, WorkingGrid.X, WorkingGrid.Y);
+
+		// 
+		TArray<box<FVector2f>> BlockRects;
+		BlockRects.SetNumUninitialized(NumBlocks);
+
+		// Create an array of block index per cell
+		TArray<int32> OverlappingBlocks;
+		for (int32 BlockIndex = 0; BlockIndex < NumBlocks; ++BlockIndex)
+		{
+			bool bBlockHasMask = GeneratedLayout.Source->Blocks[BlockIndex].Mask.get() != nullptr;
+
+			// Fill the block rect
+			FIntVector2 Min = Layout->Blocks[BlockIndex].Min;
+			FIntVector2 Size = Layout->Blocks[BlockIndex].Size;
+
+			box<FVector2f>& BlockRect = BlockRects[BlockIndex];
+			BlockRect.min[0] = float(Min.X) / float(LayoutGrid.X);
+			BlockRect.min[1] = float(Min.Y) / float(LayoutGrid.Y);
+			BlockRect.size[0] = float(Size.X) / float(LayoutGrid.X);
+			BlockRect.size[1] = float(Size.Y) / float(LayoutGrid.Y);
+
+			// Fill the block index per cell array
+			// Ignore the block in this stage if it has a mask, because blocks with masks will very likely overlap other blocks
+			if (!bBlockHasMask)
+			{
+				for (uint16 Y = Min.Y; Y < Min.Y + Size.Y; ++Y)
+				{
+					for (uint16 X = Min.X; X < Min.X + Size.X; ++X)
+					{
+						if (GridBlockBlockId.Get(X,Y) == MAX_uint16)
+						{
+							GridBlockBlockId.Set(X,Y, BlockIndex);
+						}
+						else
+						{
+							OverlappingBlocks.AddUnique(BlockIndex);
+						}
+					}
+				}
+			}
+		}
+
+
+		// Notify Overlapping layout blocks
+		if (!OverlappingBlocks.IsEmpty())
+		{
+			FString Msg = FString::Printf(TEXT("Source mesh has %d layout block overlapping in LOD %d"),
+				OverlappingBlocks.Num() + 1, CurrentParents.Last().Lod
+			);
+			ErrorLog->GetPrivate()->Add(Msg, ELMT_WARNING, errorContext);
+		}
+
+		// Get the information about the texture coordinates channel
+		int32 TexCoordsBufferIndex = -1;
+		int32 TexCoordsChannelIndex = -1;
+		Mesh->GetVertexBuffers().FindChannel(MBS_TEXCOORDS, LayoutChannel, &TexCoordsBufferIndex, &TexCoordsChannelIndex);
+		check(TexCoordsBufferIndex >= 0);
+		check(TexCoordsChannelIndex >= 0);
+
+		const FMeshBufferChannel& TexCoordsChannel = Mesh->VertexBuffers.Buffers[TexCoordsBufferIndex].Channels[TexCoordsChannelIndex];
+		check(TexCoordsChannel.Semantic == MBS_TEXCOORDS);
+
+		uint8* TexCoordData = Mesh->GetVertexBuffers().GetBufferData(TexCoordsBufferIndex);
+		int32 elemSize = Mesh->GetVertexBuffers().GetElementSize(TexCoordsBufferIndex);
+		int32 channelOffset = TexCoordsChannel.Offset;
+		TexCoordData += channelOffset;
+
+		// Get a copy of the UVs as FVector2f to work with them. 
+		TArray<FVector2f> TexCoords;
+		{
+			TexCoords.SetNumUninitialized(NumVertices);
+
+			bool bNonNormalizedUVs = false;
+			const bool bIsOverlayLayout = GeneratedLayout.Layout->GetLayoutPackingStrategy() == mu::EPackStrategy::Overlay;
+
+			const uint8* pVertices = TexCoordData;
+			for (int32 VertexIndex = 0; VertexIndex < NumVertices; ++VertexIndex)
+			{
+				FVector2f& UV = TexCoords[VertexIndex];
+				if (TexCoordsChannel.Format == MBF_FLOAT32)
+				{
+					UV = *((FVector2f*)pVertices);
+				}
+				else if (TexCoordsChannel.Format == MBF_FLOAT16)
+				{
+					const FFloat16* pUV = reinterpret_cast<const FFloat16*>(pVertices);
+					UV = FVector2f(float(pUV[0]), float(pUV[1]));
+				}
+
+				// Check that UVs are normalized. If not, clamp the values and throw a warning.
+				if (MeshOptions.bNormalizeUVs && !bIsOverlayLayout
+					&& (UV[0] < 0.f || UV[0] > 1.f || UV[1] < 0.f || UV[1] > 1.f))
+				{
+					UV[0] = FMath::Clamp(UV[0], 0.f, 1.f);
+					UV[1] = FMath::Clamp(UV[1], 0.f, 1.f);
+					bNonNormalizedUVs = true;
+				}
+
+				pVertices += elemSize;
+			}
+
+			// Mutable does not support non-normalized UVs
+			if (bNonNormalizedUVs && !bIsOverlayLayout)
+			{
+				FString Msg = FString::Printf(TEXT("Source mesh has non-normalized UVs in LOD %d"), CurrentParents.Last().Lod );
+				ErrorLog->GetPrivate()->Add(Msg, ELMT_WARNING, errorContext);
+			}
+		}
+
+
+		const int32 NumTriangles = Mesh->GetIndexCount() / 3;
+		TArray<FTriangleInfo> Triangles;
+
+		// Vertices mapped to unique vertex index
+		TArray<int32> CollapsedVertices;
+
+		// Vertex to face map used to speed up connectivity building
+		TMultiMap<int32, uint32> VertexToFaceMap;
+
+		// Find Unique Vertices
+		if (!bIsSingleFullBlock && MeshOptions.bClampUVIslands)
+		{
+			VertexToFaceMap.Reserve(NumVertices);
+			Triangles.SetNumUninitialized(NumTriangles);
+
+			MeshCreateCollapsedVertexMap(Mesh.get(), CollapsedVertices);
+		}
+
+		TArray<int32> ConflictiveTriangles;
+
+		const uint32 MaxGridX = MeshOptions.bNormalizeUVs ? MAX_uint32 : WorkingGrid.X - 1;
+		const uint32 MaxGridY = MeshOptions.bNormalizeUVs ? MAX_uint32 : WorkingGrid.Y - 1;
+
+		// Allocate the per-vertex layout block data
+		TArray<uint16> LayoutData;
+		constexpr uint16 NullBlockId = MAX_uint16 - 1;
+		LayoutData.Init(NullBlockId, NumVertices);
+
+		UntypedMeshBufferIteratorConst ItIndices(Mesh->GetIndexBuffers(), MBS_VERTEXINDEX);
+		for (int32 TriangleIndex = 0;  TriangleIndex < NumTriangles; ++TriangleIndex)
+		{
+			uint32 Index0 = ItIndices.GetAsUINT32();
+			++ItIndices;
+			uint32 Index1 = ItIndices.GetAsUINT32();
+			++ItIndices;
+			uint32 Index2 = ItIndices.GetAsUINT32();
+			++ItIndices;
+
+			auto AssignOneVertex = 
+				[LayoutGrid, MaxGridX, MaxGridY, NumBlocks, Layout, &GeneratedLayout, &GridBlockBlockId, &TexCoords, &LayoutData]
+				(int32 VertexIndex)
+				{
+					uint16& BlockIndex = LayoutData[VertexIndex];
+
+					// Was it previously assigned?
+					if (BlockIndex != NullBlockId)
+					{
+						return BlockIndex;
+					}
+
+					FVector2f UV = TexCoords[VertexIndex];
+
+					int32 VertexWorkingGridX = FMath::Clamp(LayoutGrid.X * UV[0], 0, LayoutGrid.X - 1);
+					int32 VertexWorkingGridY = FMath::Clamp(LayoutGrid.Y * UV[1], 0, LayoutGrid.Y - 1);
+
+					// First: Assign the vertices to masked blocks in order
+					for (int32 CandidateBlockIndex = 0; CandidateBlockIndex < NumBlocks; ++CandidateBlockIndex)
+					{
+						const mu::Image* Mask = GeneratedLayout.Source->Blocks[CandidateBlockIndex].Mask.get();
+						if (Mask)
+						{
+							// First discard with block limits.
+							FIntVector2 Min = Layout->Blocks[CandidateBlockIndex].Min;
+							FIntVector2 Size = Layout->Blocks[CandidateBlockIndex].Size;
+
+							bool bInBlock = 
+								(VertexWorkingGridX >= Min.X && VertexWorkingGridX < Min.X + Size.X)
+								&&
+								(VertexWorkingGridY >= Min.Y && VertexWorkingGridY < Min.Y + Size.Y);
+
+							if (bInBlock)
+							{
+								// TODO: This always clamps the UVs
+								FVector2f SampleUV;
+								SampleUV.X = FMath::Fmod(UV.X, 1.0);
+								SampleUV.Y = FMath::Fmod(UV.Y, 1.0);
+
+								FVector4f MaskValue = Mask->Sample(SampleUV);
+								if (MaskValue.X > 0.5f)
+								{
+									BlockIndex = CandidateBlockIndex;
+									break;
+								}
+							}
+						}
+					}
+
+					// Second: Assign to non-masked blocks if not assigned yet
+					if (BlockIndex == NullBlockId)
+					{
+						uint32 ClampedX = FMath::Min<uint32>(MaxGridX, FMath::Max<uint32>(0, VertexWorkingGridX));
+						uint32 ClampedY = FMath::Min<uint32>(MaxGridY, FMath::Max<uint32>(0, VertexWorkingGridY));
+						BlockIndex = GridBlockBlockId.Get(ClampedX, ClampedY);
+					}
+					return BlockIndex;
+				};
+
+			const uint16 BlockIndexV0 = AssignOneVertex(Index0);
+			const uint16 BlockIndexV1 = AssignOneVertex(Index1);
+			const uint16 BlockIndexV2 = AssignOneVertex(Index2);
+
+			if (!bIsSingleFullBlock && MeshOptions.bClampUVIslands)
+			{
+				if (BlockIndexV0 != BlockIndexV1 || BlockIndexV0 != BlockIndexV2)
+				{
+					ConflictiveTriangles.Add(TriangleIndex);
+				}
+
+				FTriangleInfo& Triangle = Triangles[TriangleIndex];
+
+				Triangle.Indices[0] = Index0;
+				Triangle.Indices[1] = Index1;
+				Triangle.Indices[2] = Index2;
+				Triangle.CollapsedIndices[0] = CollapsedVertices[Index0];
+				Triangle.CollapsedIndices[1] = CollapsedVertices[Index1];
+				Triangle.CollapsedIndices[2] = CollapsedVertices[Index2];
+
+				Triangle.BlockIndices[0] = BlockIndexV0;
+				Triangle.BlockIndices[1] = BlockIndexV1;
+				Triangle.BlockIndices[2] = BlockIndexV2;
+				Triangle.bUVsFixed = false;
+
+				VertexToFaceMap.Add(Triangle.CollapsedIndices[0], TriangleIndex);
+				VertexToFaceMap.Add(Triangle.CollapsedIndices[1], TriangleIndex);
+				VertexToFaceMap.Add(Triangle.CollapsedIndices[2], TriangleIndex);
+			}
+		}
+
+		// Clamp UV islands to the predominant block of each island. Will only happen if bClampUVIslands is true.
+		for (int32 ConflictiveTriangleIndex : ConflictiveTriangles)
+		{
+			FTriangleInfo& Triangle = Triangles[ConflictiveTriangleIndex];
+
+			// Skip the ones that have been fixed already
+			if (Triangle.bUVsFixed)
+			{
+				continue;
+			}
+
+			// Find triangles from the same UV Island
+			TArray<uint32> TriangleIndices;
+			GetUVIsland(Triangles, ConflictiveTriangleIndex, TriangleIndices, TexCoords, VertexToFaceMap);
+
+			// Get predominant BlockId != MAX_uint16
+			TArray<uint32> NumVerticesPerBlock;
+			NumVerticesPerBlock.SetNumZeroed(NumBlocks);
+
+			for (int32 TriangleIndex : TriangleIndices)
+			{
+				FTriangleInfo& OtherTriangle = Triangles[TriangleIndex];
+				for (int32 VertexIndex = 0; VertexIndex < 3; ++VertexIndex)
+				{
+					const uint16& BlockIndex = OtherTriangle.BlockIndices[VertexIndex];
+					if (BlockIndex != MAX_uint16)
+					{
+						NumVerticesPerBlock[BlockIndex]++;
+					}
+				}
+			}
+
+			uint16 BlockIndex = 0;
+			uint32 CurrentMaxVertices = 0;
+			for (int32 Index = 0; Index < NumBlocks; ++Index)
+			{
+				if (NumVerticesPerBlock[Index] > CurrentMaxVertices)
+				{
+					BlockIndex = Index;
+					CurrentMaxVertices = NumVerticesPerBlock[Index];
+				}
+			}
+
+			// Get the limits of the predominant block rect
+			const FLayoutBlock& LayoutBlock = Layout->Blocks[BlockIndex];
+
+			const float SmallNumber = 0.000001;
+			const float MinX = ((float)LayoutBlock.Min.X) / (float)LayoutGrid.X + SmallNumber;
+			const float MinY = ((float)LayoutBlock.Min.Y) / (float)LayoutGrid.Y + SmallNumber;
+			const float MaxX = (((float)LayoutBlock.Size.X + LayoutBlock.Min.X) / (float)LayoutGrid.X) - 2 * SmallNumber;
+			const float MaxY = (((float)LayoutBlock.Size.Y + LayoutBlock.Min.Y) / (float)LayoutGrid.Y) - 2 * SmallNumber;
+
+			// Iterate triangles and clamp the UVs
+			for (int32 TriangleIndex : TriangleIndices)
+			{
+				FTriangleInfo& OtherTriangle = Triangles[TriangleIndex];
+
+				for (int8 VertexIndex = 0; VertexIndex < 3; ++VertexIndex)
+				{
+					if (OtherTriangle.BlockIndices[VertexIndex] == BlockIndex)
+					{
+						continue;
+					}
+
+					OtherTriangle.BlockIndices[VertexIndex] = BlockIndex;
+
+					// Clamp UVs to the block they are assigned to
+					const int32 UVIndex = OtherTriangle.Indices[VertexIndex];
+					FVector2f& UV = TexCoords[UVIndex];
+					UV[0] = FMath::Clamp(UV[0], MinX, MaxX);
+					UV[1] = FMath::Clamp(UV[1], MinY, MaxY);
+					LayoutData[UVIndex] = BlockIndex;
+				}
+
+				OtherTriangle.bUVsFixed = true;
+			}
+		}
+
+		// Warn about vertices without a block id
+		int32 FirstLODToIgnoreWarnings = GeneratedLayout.Source->FirstLODToIgnoreWarnings;
+		if (FirstLODToIgnoreWarnings == -1 || CurrentParents.Last().Lod < FirstLODToIgnoreWarnings)
+		{
+			TArray<float> UnassignedUVs;
+			UnassignedUVs.Reserve(NumVertices / 100);
+
+			const FVector2f* UVs = TexCoords.GetData();
+			for (int32 VertexIndex = 0; VertexIndex < NumVertices; ++VertexIndex)
+			{
+				if (LayoutData[VertexIndex] == MAX_uint16)
+				{
+					UnassignedUVs.Add((*(UVs + VertexIndex))[0]);
+					UnassignedUVs.Add((*(UVs + VertexIndex))[1]);
+				}
+			}
+
+			if (!UnassignedUVs.IsEmpty())
+			{
+				FString Msg = FString::Printf(TEXT("Source mesh has %d vertices not assigned to any layout block in LOD %d"), UnassignedUVs.Num(), CurrentParents.Last().Lod);
+
+				ErrorLogMessageAttachedDataView attachedDataView;
+				attachedDataView.m_unassignedUVs = UnassignedUVs.GetData();
+				attachedDataView.m_unassignedUVsSize = (size_t)UnassignedUVs.Num();
+
+				ErrorLog->GetPrivate()->Add(Msg, attachedDataView, ELMT_WARNING, errorContext);
+			}
+		}
+
+		// Create the layout block vertex buffer
+		uint8* LayoutBufferPtr = nullptr;
+		{
+			const int32 LayoutBufferIndex = Mesh->GetVertexBuffers().GetBufferCount();
+			Mesh->GetVertexBuffers().SetBufferCount(LayoutBufferIndex + 1);
+
+			// TODO
+			check(Layout->GetBlockCount() < MAX_uint16);
+			const EMeshBufferSemantic LayoutSemantic = MBS_LAYOUTBLOCK;
+			const int32 LayoutSemanticIndex = int32(LayoutChannel);
+			const EMeshBufferFormat LayoutFormat = bUseAbsoluteBlockIds ? MBF_UINT64 : MBF_UINT16;
+			const int32 LayoutComponents = 1;
+			const int32 LayoutOffset = 0;
+			int32 ElementSize = bUseAbsoluteBlockIds ? sizeof(uint64) : sizeof(uint16);
+			Mesh->GetVertexBuffers().SetBuffer
+			(
+				LayoutBufferIndex,
+				ElementSize,
+				1,
+				&LayoutSemantic, &LayoutSemanticIndex,
+				&LayoutFormat, &LayoutComponents,
+				&LayoutOffset
+			);
+			LayoutBufferPtr = Mesh->GetVertexBuffers().GetBufferData(LayoutBufferIndex);
+		}
+
+		// Copy UVs back to the mesh
+		{
+			uint8* pVertices = TexCoordData;
+			FVector2f* UVs = TexCoords.GetData();
+
+			for (int32 VertexIndex = 0; VertexIndex < NumVertices; ++VertexIndex)
+			{
+				FVector2f& UV = TexCoords[VertexIndex];
+
+				uint16 LayoutBlockIndex = LayoutData[VertexIndex];
+				if (Layout->Blocks.IsValidIndex(LayoutBlockIndex))
+				{
+					uint64 LayoutBlockId = Layout->Blocks[LayoutBlockIndex].Id;
+
+					UV = BlockRects[LayoutBlockIndex].Homogenize(UV);
+
+					// Replace block index by the actual id of the block
+					if (bUseAbsoluteBlockIds)
+					{
+						uint64* Ptr = reinterpret_cast<uint64*>(LayoutBufferPtr) + VertexIndex;
+						*Ptr = LayoutBlockId;
+					}
+					else
+					{
+						uint16* Ptr = reinterpret_cast<uint16*>(LayoutBufferPtr) + VertexIndex;
+						*Ptr = uint16(LayoutBlockId & 0xffff);
+					}
+				}
+				else
+				{
+					// Map vertices without block
+					if (bUseAbsoluteBlockIds)
+					{
+						uint64* Ptr = reinterpret_cast<uint64*>(LayoutBufferPtr) + VertexIndex;
+						*Ptr = MeshOptions.bEnsureAllVerticesHaveLayoutBlock ? 0 : std::numeric_limits<uint64>::max();
+					}
+					else
+					{
+						uint16* Ptr = reinterpret_cast<uint16*>(LayoutBufferPtr) + VertexIndex;
+						*Ptr = MeshOptions.bEnsureAllVerticesHaveLayoutBlock ? 0 : std::numeric_limits<uint16>::max();
+					}
+				}
+
+				// Copy UVs
+				if (TexCoordsChannel.Format == MBF_FLOAT32)
+				{
+					FVector2f* pUV = reinterpret_cast<FVector2f*>(pVertices);
+					*pUV = UV;
+				}
+				else if (TexCoordsChannel.Format == MBF_FLOAT16)
+				{
+					FFloat16* pUV = reinterpret_cast<FFloat16*>(pVertices);
+					pUV[0] = FFloat16(UV[0]);
+					pUV[1] = FFloat16(UV[1]);
+				}
+
+				pVertices += elemSize;
+			}
+		}
+	}
+	
+
+    //---------------------------------------------------------------------------------------------
+    void CodeGenerator::GenerateMesh( const FMeshGenerationOptions& InOptions, FMeshGenerationResult& OutResult, const NodeMeshPtrConst& InUntypedNode)
+    {
+        if (!InUntypedNode)
+        {
+            OutResult = FMeshGenerationResult();
+            return;
+        }
+
+        // See if it was already generated
+		FGeneratedMeshCacheKey Key;
+		Key.Node = InUntypedNode;
+		Key.Options = InOptions;
+        GeneratedMeshMap::ValueType* it = GeneratedMeshes.Find(Key);
+        if ( it )
+        {
+			OutResult = *it;
+            return;
+        }
+
+		const NodeMesh* Node = InUntypedNode.get();
+
+        // Generate for each different type of node
+		switch (Node->GetType()->Type)
+		{
+		case Node::EType::MeshConstant: GenerateMesh_Constant(InOptions, OutResult, static_cast<const NodeMeshConstant*>(Node)); break;
+		case Node::EType::MeshFormat: GenerateMesh_Format(InOptions, OutResult, static_cast<const NodeMeshFormat*>(Node)); break;
+		case Node::EType::MeshMorph: GenerateMesh_Morph(InOptions, OutResult, static_cast<const NodeMeshMorph*>(Node)); break;
+		case Node::EType::MeshMakeMorph: GenerateMesh_MakeMorph(InOptions, OutResult, static_cast<const NodeMeshMakeMorph*>(Node)); break;
+		case Node::EType::MeshFragment: GenerateMesh_Fragment(InOptions, OutResult, static_cast<const NodeMeshFragment*>(Node)); break;
+		case Node::EType::MeshInterpolate: GenerateMesh_Interpolate(InOptions, OutResult, static_cast<const NodeMeshInterpolate*>(Node)); break;
+		case Node::EType::MeshSwitch: GenerateMesh_Switch(InOptions, OutResult, static_cast<const NodeMeshSwitch*>(Node)); break;
+		case Node::EType::MeshTransform: GenerateMesh_Transform(InOptions, OutResult, static_cast<const NodeMeshTransform*>(Node)); break;
+		case Node::EType::MeshClipMorphPlane: GenerateMesh_ClipMorphPlane(InOptions, OutResult, static_cast<const NodeMeshClipMorphPlane*>(Node)); break;
+		case Node::EType::MeshClipWithMesh: GenerateMesh_ClipWithMesh(InOptions, OutResult, static_cast<const NodeMeshClipWithMesh*>(Node)); break;
+		case Node::EType::MeshApplyPose: GenerateMesh_ApplyPose(InOptions, OutResult, static_cast<const NodeMeshApplyPose*>(Node)); break;
+		case Node::EType::MeshVariation: GenerateMesh_Variation(InOptions, OutResult, static_cast<const NodeMeshVariation*>(Node)); break;
+		case Node::EType::MeshTable: GenerateMesh_Table(InOptions, OutResult, static_cast<const NodeMeshTable*>(Node)); break;
+		case Node::EType::MeshGeometryOperation: GenerateMesh_GeometryOperation(InOptions, OutResult, static_cast<const NodeMeshGeometryOperation*>(Node)); break;
+		case Node::EType::MeshReshape: GenerateMesh_Reshape(InOptions, OutResult, static_cast<const NodeMeshReshape*>(Node)); break;
+		case Node::EType::MeshClipDeform: GenerateMesh_ClipDeform(InOptions, OutResult, static_cast<const NodeMeshClipDeform*>(Node)); break;
+		default: check(false);
+		}
+
+        // Cache the result
+        GeneratedMeshes.Add( Key, OutResult);
+    }
+
+
+    //---------------------------------------------------------------------------------------------
+    void CodeGenerator::GenerateMesh_Morph(
+		const FMeshGenerationOptions& InOptions, 
+		FMeshGenerationResult& OutResult, 
+		const NodeMeshMorph* InMorphNode 
+	)
+    {
+        NodeMeshMorph::Private& node = *InMorphNode->GetPrivate();
+
+        Ptr<ASTOpMeshMorph> OpMorph = new ASTOpMeshMorph();
+
+        // Factor
+        if ( node.Factor )
+        {
+            OpMorph->Factor = Generate_Generic( node.Factor.get(), InOptions );
+        }
+        else
+        {
+            // This argument is required
+            OpMorph->Factor = GenerateMissingScalarCode(TEXT("Morph factor"), 0.5f, InMorphNode->GetMessageContext());
+        }
+
+        // Base
+        FMeshGenerationResult BaseResult;
+        if ( node.Base )
+        {
+            GenerateMesh(InOptions,BaseResult, node.Base );
+            OpMorph->Base = BaseResult.MeshOp;
+        }
+        else
+        {
+            // This argument is required
+            ErrorLog->GetPrivate()->Add( "Mesh morph base node is not set.",
+                                            ELMT_ERROR, InMorphNode->GetMessageContext());
+        }        
+
+		if (node.Morph)
+        {
+            FMeshGenerationResult TargetResult;
+			FMeshGenerationOptions TargetOptions = InOptions;
+			TargetOptions.bLayouts = false;
+			// We need to override the layouts with the layouts that were generated for the base to make
+			// sure that we get the correct mesh when generating the target
+			TargetOptions.OverrideLayouts = BaseResult.GeneratedLayouts;
+			TargetOptions.ActiveTags.Empty();
+            GenerateMesh(TargetOptions, TargetResult, node.Morph);
+
+            // TODO: Make sure that the target is a mesh with the morph format
+            Ptr<ASTOp> target = TargetResult.MeshOp;
+
+            OpMorph->Target = target;
+        }
+ 
+        const bool bReshapeEnabled = node.bReshapeSkeleton || node.bReshapePhysicsVolumes;
+        
+        Ptr<ASTOpMeshMorphReshape> OpMorphReshape;
+        if ( bReshapeEnabled )
+        {
+		    Ptr<ASTOpMeshBindShape> OpBind = new ASTOpMeshBindShape();
+		    Ptr<ASTOpMeshApplyShape> OpApply = new ASTOpMeshApplyShape();
+
+			// Setting bReshapeVertices to false the bind op will remove all mesh members except 
+			// PhysicsBodies and the Skeleton.
+            OpBind->bReshapeVertices = false;
+            OpBind->bApplyLaplacian = false;
+            OpBind->bRecomputeNormals = false;
+		    OpBind->bReshapeSkeleton = node.bReshapeSkeleton;
+		    OpBind->BonesToDeform = node.BonesToDeform;
+    	    OpBind->bReshapePhysicsVolumes = node.bReshapePhysicsVolumes; 
+			OpBind->PhysicsToDeform = node.PhysicsToDeform;
+			OpBind->BindingMethod = static_cast<uint32>(EShapeBindingMethod::ReshapeClosestProject);
+            
+			OpBind->Mesh = BaseResult.MeshOp;
+            OpBind->Shape = BaseResult.MeshOp;
+           
+			OpApply->bReshapeVertices = OpBind->bReshapeVertices;
+			OpApply->bRecomputeNormals = OpBind->bRecomputeNormals;
+		    OpApply->bReshapeSkeleton = OpBind->bReshapeSkeleton;
+		    OpApply->bReshapePhysicsVolumes = OpBind->bReshapePhysicsVolumes;
+
+			OpApply->Mesh = OpBind;
+            OpApply->Shape = OpMorph;
+
+            OpMorphReshape = new ASTOpMeshMorphReshape();
+            OpMorphReshape->Morph = OpMorph;
+            OpMorphReshape->Reshape = OpApply;
+        }
+
+ 		if (OpMorphReshape)
+		{
+			OutResult.MeshOp = OpMorphReshape;
+		}
+		else
+		{
+			OutResult.MeshOp = OpMorph;
+		}
+
+        OutResult.BaseMeshOp = BaseResult.BaseMeshOp;
+		OutResult.GeneratedLayouts = BaseResult.GeneratedLayouts;
+}
+
+
+    //---------------------------------------------------------------------------------------------
+    void CodeGenerator::GenerateMesh_MakeMorph(const FMeshGenerationOptions& InOptions, FMeshGenerationResult& OutResult, const NodeMeshMakeMorph* InMakeMorphNode )
+    {
+        NodeMeshMakeMorph::Private& node = *InMakeMorphNode->GetPrivate();
+
+        Ptr<ASTOpMeshDifference> op = new ASTOpMeshDifference();
+
+        // \todo Texcoords are broken?
+        op->bIgnoreTextureCoords = true;
+	
+		// UE only has position and normal morph data, optimize for this case if indicated. 
+		if (node.bOnlyPositionAndNormal)
+		{
+			op->Channels = { {static_cast<uint8>(MBS_POSITION), 0}, {static_cast<uint8>(MBS_NORMAL), 0} };
+		}
+
+        // Base
+        FMeshGenerationResult BaseResult;
+        if ( node.m_pBase )
+        {
+			FMeshGenerationOptions BaseOptions = InOptions;
+			BaseOptions.bLayouts = false;
+			GenerateMesh(BaseOptions, BaseResult, node.m_pBase );
+
+            op->Base = BaseResult.MeshOp;
+        }
+        else
+        {
+            // This argument is required
+            ErrorLog->GetPrivate()->Add( "Mesh make morph base node is not set.",
+                                            ELMT_ERROR, InMakeMorphNode->GetMessageContext());
+        }
+
+        // Target
+		if ( node.m_pTarget )
+        {
+			FMeshGenerationOptions TargetOptions = InOptions;
+			TargetOptions.bLayouts = false;
+			TargetOptions.OverrideLayouts.Empty();
+			TargetOptions.ActiveTags.Empty();
+			FMeshGenerationResult TargetResult;
+            GenerateMesh( TargetOptions, TargetResult, node.m_pTarget );
+
+            op->Target = TargetResult.MeshOp;
+        }
+        else
+        {
+            // This argument is required
+            ErrorLog->GetPrivate()->Add( "Mesh make morph target node is not set.",
+                                            ELMT_ERROR, InMakeMorphNode->GetMessageContext());
+        }
+
+        OutResult.MeshOp = op;
+        OutResult.BaseMeshOp = BaseResult.BaseMeshOp;
+		OutResult.GeneratedLayouts = BaseResult.GeneratedLayouts;
+	}
+
+    //---------------------------------------------------------------------------------------------
+    void CodeGenerator::GenerateMesh_Fragment(const FMeshGenerationOptions& InOptions, FMeshGenerationResult& OutResult, const NodeMeshFragment* Node )
+    {
+        FMeshGenerationResult BaseResult;
+        if ( Node->SourceMesh )
+        {
+			Ptr<ASTOpMeshExtractLayoutBlocks> op = new ASTOpMeshExtractLayoutBlocks();
+			OutResult.MeshOp = op;
+
+			op->LayoutIndex = (uint16)Node->LayoutIndex;
+
+			// Generate the source mesh
+			FMeshGenerationOptions BaseOptions = InOptions;
+			BaseOptions.bLayouts = true;
+			BaseOptions.bEnsureAllVerticesHaveLayoutBlock = false;
+
+			if (Node->Layout)
+			{
+				// Generate the layout with blocks to extract
+				Ptr<const Layout> Layout = GenerateLayout(Node->Layout, 0);
+				BaseOptions.OverrideLayouts.Empty();
+				BaseOptions.OverrideLayouts.Add({ Layout, Node->Layout });
+			}
+
+            GenerateMesh( BaseOptions, BaseResult, Node->SourceMesh);
+            op->Source = BaseResult.MeshOp;
+        }
+        else
+        {
+            // This argument is required
+            ErrorLog->GetPrivate()->Add( "Mesh fragment source is not set.", ELMT_ERROR, Node->GetMessageContext());
+        }
+
+        OutResult.BaseMeshOp = BaseResult.BaseMeshOp;
+		OutResult.GeneratedLayouts = BaseResult.GeneratedLayouts;
+	}
+
+
+    //---------------------------------------------------------------------------------------------
+    void CodeGenerator::GenerateMesh_Interpolate(const FMeshGenerationOptions& InOptions, FMeshGenerationResult& OutResult,
+		const NodeMeshInterpolate* InterpolateNode )
+    {
+        NodeMeshInterpolate::Private& node = *InterpolateNode->GetPrivate();
+
+        // Generate the code
+        Ptr<ASTOpFixed> op = new ASTOpFixed();
+        op->op.type = OP_TYPE::ME_INTERPOLATE;
+        OutResult.MeshOp = op;
+
+        // Factor
+        if ( Node* pFactor = node.m_pFactor.get() )
+        {
+            op->SetChild( op->op.args.MeshInterpolate.factor, Generate_Generic( pFactor, InOptions ) );
+        }
+        else
+        {
+            // This argument is required
+            op->SetChild( op->op.args.MeshInterpolate.factor,
+                          GenerateMissingScalarCode(TEXT("Interpolation factor"), 0.5f, InterpolateNode->GetMessageContext()) );
+        }
+
+        //
+        Ptr<ASTOp> base = 0;
+        int32 count = 0;
+        for ( int32 t=0
+            ; t<node.m_targets.Num() && t<MUTABLE_OP_MAX_INTERPOLATE_COUNT-1
+            ; ++t )
+        {
+            if ( NodeMesh* pA = node.m_targets[t].get() )
+            {
+				FMeshGenerationOptions TargetOptions = InOptions;
+				TargetOptions.OverrideLayouts.Empty();
+
+                FMeshGenerationResult TargetResult;
+                GenerateMesh( TargetOptions, TargetResult, pA );
+
+
+                // The first target is the base
+                if (count==0)
+                {
+                    base = TargetResult.MeshOp;
+                    op->SetChild( op->op.args.MeshInterpolate.base, TargetResult.MeshOp );
+
+                    OutResult.BaseMeshOp = TargetResult.BaseMeshOp;
+					OutResult.GeneratedLayouts = TargetResult.GeneratedLayouts;
+				}
+                else
+                {
+                    Ptr<ASTOpMeshDifference> dop = new ASTOpMeshDifference();
+                    dop->Base = base;
+                    dop->Target = TargetResult.MeshOp;
+
+                    // \todo Texcoords are broken?
+                    dop->bIgnoreTextureCoords = true;
+
+                    for ( size_t c=0; c<node.m_channels.Num(); ++c)
+                    {
+                        check( node.m_channels[c].semantic < 256 );
+						check(node.m_channels[c].semanticIndex < 256);
+						
+						ASTOpMeshDifference::FChannel Channel;
+						Channel.Semantic = uint8(node.m_channels[c].semantic);
+						Channel.SemanticIndex = uint8(node.m_channels[c].semanticIndex);
+						dop->Channels.Add(Channel);
+                    }
+
+                    op->SetChild( op->op.args.MeshInterpolate.targets[count-1], dop );
+                }
+                count++;
+            }
+        }
+
+        // At least one mesh is required
+        if (!count)
+        {
+            // TODO
+            //op.args.MeshInterpolate.target[0] = GenerateMissingImageCode( "First mesh", IF_RGB_UBYTE );
+            ErrorLog->GetPrivate()->Add
+                ( "Mesh interpolation: at least the first mesh is required.",
+                  ELMT_ERROR, InterpolateNode->GetMessageContext());
+        }
+    }
+
+
+    //---------------------------------------------------------------------------------------------
+    void CodeGenerator::GenerateMesh_Switch(const FMeshGenerationOptions& InOptions, FMeshGenerationResult& OutResult, const NodeMeshSwitch* SwitchNode )
+    {
+        NodeMeshSwitch::Private& node = *SwitchNode->GetPrivate();
+
+        if (node.m_options.Num() == 0)
+        {
+            // No options in the switch!
+            // TODO
+            OutResult = FMeshGenerationResult();
+			return;
+        }
+
+        Ptr<ASTOpSwitch> op = new ASTOpSwitch();
+        op->type = OP_TYPE::ME_SWITCH;
+
+        // Factor
+        if ( node.m_pParameter )
+        {
+            op->variable = Generate_Generic( node.m_pParameter.get(), InOptions);
+        }
+        else
+        {
+            // This argument is required
+            op->variable = GenerateMissingScalarCode(TEXT("Switch variable"), 0.0f, SwitchNode->GetMessageContext());
+        }
+
+        // Options
+		bool bFirstValidConnectionFound = false;
+        for ( int32 t=0; t< node.m_options.Num(); ++t )
+        {
+			FMeshGenerationOptions TargetOptions = InOptions;
+
+            if ( node.m_options[t] )
+            {
+				// Take the layouts from the first non-null connection.
+				// \TODO: Take them from the first connection that actually returns layouts?
+				if (bFirstValidConnectionFound)
+				{
+					TargetOptions.OverrideLayouts = OutResult.GeneratedLayouts;
+				}
+				
+				FMeshGenerationResult BranchResults;
+                GenerateMesh(TargetOptions, BranchResults, node.m_options[t] );
+
+                Ptr<ASTOp> branch = BranchResults.MeshOp;
+                op->cases.Emplace((int16)t,op,branch);
+
+				if (!bFirstValidConnectionFound)
+				{
+					bFirstValidConnectionFound = true;
+                    OutResult = BranchResults;
+                }
+            }
+        }
+
+        OutResult.MeshOp = op;
+    }
+
+
+	//---------------------------------------------------------------------------------------------
+	void CodeGenerator::GenerateMesh_Table(const FMeshGenerationOptions& InOptions, FMeshGenerationResult& OutResult, const NodeMeshTable* TableNode)
+	{
+		//
+		FMeshGenerationResult NewResult = OutResult;
+		int32 t = 0;
+		bool bFirstRowGenerated = false;
+
+		Ptr<ASTOp> Op = GenerateTableSwitch<NodeMeshTable, ETableColumnType::Mesh, OP_TYPE::ME_SWITCH>(*TableNode,
+			[this, &NewResult, &bFirstRowGenerated, &InOptions] (const NodeMeshTable& node, int32 colIndex, int32 row, mu::ErrorLog* pErrorLog)
+			{
+				mu::Ptr<mu::Mesh> pMesh = node.Table->GetPrivate()->Rows[row].Values[colIndex].Mesh;
+				FMeshGenerationResult BranchResults;
+
+				if (pMesh)
+				{
+					NodeMeshConstantPtr pCell = new NodeMeshConstant();
+					pCell->SetValue(pMesh);
+
+					// TODO Take into account layout strategy
+					int32 numLayouts = node.Layouts.Num();
+					pCell->SetLayoutCount(numLayouts);
+					for (int32 i = 0; i < numLayouts; ++i)
+					{
+						pCell->SetLayout(i, node.Layouts[i]);
+					}
+
+					FMeshGenerationOptions TargetOptions = InOptions;
+
+					if (bFirstRowGenerated)
+					{
+						TargetOptions.OverrideLayouts = NewResult.GeneratedLayouts;
+					}
+
+					TargetOptions.OverrideContext = node.Table->GetPrivate()->Rows[row].Values[colIndex].ErrorContext;
+
+					GenerateMesh(TargetOptions, BranchResults, pCell);
+
+					if (!bFirstRowGenerated)
+					{
+						NewResult = BranchResults;
+						bFirstRowGenerated = true;
+					}
+				}
+
+				return BranchResults.MeshOp;
+			});
+
+		NewResult.MeshOp = Op;
+
+		OutResult = NewResult;
+	}
+
+
+    //---------------------------------------------------------------------------------------------
+    void CodeGenerator::GenerateMesh_Variation(const FMeshGenerationOptions& InOptions, FMeshGenerationResult& OutResult,
+		const NodeMeshVariation* VariationNode )
+    {
+        NodeMeshVariation::Private& node = *VariationNode->GetPrivate();
+
+        FMeshGenerationResult currentResult;
+        Ptr<ASTOp> currentMeshOp;
+
+        bool firstOptionProcessed = false;
+
+        // Default case
+        if ( node.m_defaultMesh )
+        {
+            FMeshGenerationResult BranchResults;
+			FMeshGenerationOptions DefaultOptions = InOptions;
+
+			GenerateMesh(DefaultOptions, BranchResults, node.m_defaultMesh );
+            currentMeshOp = BranchResults.MeshOp;
+            currentResult = BranchResults;
+            firstOptionProcessed = true;
+        }
+
+        // Process variations in reverse order, since conditionals are built bottom-up.
+        for ( int32 t = node.m_variations.Num()-1; t >= 0; --t )
+        {
+            int32 tagIndex = -1;
+            const FString& tag = node.m_variations[t].m_tag;
+            for ( int32 i = 0; i < FirstPass.Tags.Num(); ++i )
+            {
+                if ( FirstPass.Tags[i].Tag==tag)
+                {
+                    tagIndex = i;
+                }
+            }
+
+            if ( tagIndex < 0 )
+            {
+                ErrorLog->GetPrivate()->Add( 
+					FString::Printf(TEXT("Unknown tag found in mesh variation [%s]."), *tag),
+					ELMT_WARNING,
+					VariationNode->GetMessageContext(),
+					ELMSB_UNKNOWN_TAG
+				);
+                continue;
+            }
+
+            Ptr<ASTOp> variationMeshOp;
+            if ( node.m_variations[t].m_mesh )
+            {
+				FMeshGenerationOptions VariationOptions = InOptions;
+
+                if (firstOptionProcessed)
+                {
+					VariationOptions.OverrideLayouts = currentResult.GeneratedLayouts;
+                }
+         
+                FMeshGenerationResult BranchResults;
+				GenerateMesh(VariationOptions, BranchResults, node.m_variations[t].m_mesh );
+
+                variationMeshOp = BranchResults.MeshOp;
+
+                if ( !firstOptionProcessed )
+                {
+                    firstOptionProcessed = true;                   
+                    currentResult = BranchResults;
+                }
+            }
+
+            Ptr<ASTOpConditional> conditional = new ASTOpConditional;
+            conditional->type = OP_TYPE::ME_CONDITIONAL;
+            conditional->no = currentMeshOp;
+            conditional->yes = variationMeshOp;            
+            conditional->condition = FirstPass.Tags[tagIndex].GenericCondition;
+
+            currentMeshOp = conditional;
+        }
+
+        OutResult = currentResult;
+        OutResult.MeshOp = currentMeshOp;
+    }
+
+
+    //---------------------------------------------------------------------------------------------
+    void CodeGenerator::GenerateMesh_Constant(const FMeshGenerationOptions& InOptions, FMeshGenerationResult& OutResult, const NodeMeshConstant* InNode )
+    {
+		MUTABLE_CPUPROFILER_SCOPE(GenerateMesh_Constant);
+
+        NodeMeshConstant::Private& Node = *InNode->GetPrivate();
+
+        Ptr<ASTOpConstantResource> ConstantOp = new ASTOpConstantResource();
+		ConstantOp->Type = OP_TYPE::ME_CONSTANT;
+		OutResult.BaseMeshOp = ConstantOp;
+		OutResult.MeshOp = ConstantOp;
+		OutResult.GeneratedLayouts.Empty();
+
+		bool bIsOverridingLayouts = !InOptions.OverrideLayouts.IsEmpty();
+
+        Ptr<Mesh> pMesh = Node.Value.get();
+		if (!pMesh)
+		{
+			// This data is required
+			MeshPtr EmptyMesh = new Mesh();
+			ConstantOp->SetValue(EmptyMesh, CompilerOptions->OptimisationOptions.DiskCacheContext);
+			EmptyMesh->MeshIDPrefix = ConstantOp->GetValueHash();
+
+			// Log an error message
+			ErrorLog->GetPrivate()->Add("Constant mesh not set.", ELMT_WARNING, InNode->GetMessageContext());
+
+			return;
+		}
+
+		if (pMesh->IsReference())
+		{
+			Ptr<ASTOpReferenceResource> ReferenceOp = new ASTOpReferenceResource();
+			ReferenceOp->type = OP_TYPE::ME_REFERENCE;
+			ReferenceOp->ID = pMesh->GetReferencedMesh();
+			ReferenceOp->bForceLoad = pMesh->IsForceLoad();
+
+			OutResult.BaseMeshOp = ReferenceOp;
+			OutResult.MeshOp = ReferenceOp;
+
+			return;
+		}
+
+		// Separate the tags from the mesh
+		TArray<FString> Tags = pMesh->Tags;
+		if (Tags.Num())
+		{
+			Ptr<Mesh> TaglessMesh = CloneOrTakeOver(pMesh.get());
+			TaglessMesh->Tags.SetNum(0, EAllowShrinking::No);
+			pMesh = TaglessMesh;
+		}
+
+		// Find out if we can (or have to) reuse a mesh that we have already generated.
+		FGeneratedConstantMesh DuplicateOf;
+		uint32 ThisMeshHash = HashCombineFast( GetTypeHash(pMesh->GetVertexCount()), GetTypeHash(pMesh->GetIndexCount()) );
+		TArray<FGeneratedConstantMesh>& CachedCandidates = GeneratedConstantMeshes.FindOrAdd(ThisMeshHash,{});
+		for (const FGeneratedConstantMesh& Candidate : CachedCandidates)
+		{
+			bool bCompareLayouts = InOptions.bLayouts && !bIsOverridingLayouts;
+
+			if (Candidate.Mesh->IsSimilar(*pMesh, bCompareLayouts))
+			{
+				// If it is similar and we are overriding the layouts, we must compare the layouts of the candidate with the ones
+				// we are using to override.
+				if (bIsOverridingLayouts)
+				{
+					if (Candidate.Mesh->GetLayoutCount() != InOptions.OverrideLayouts.Num())
+					{
+						continue;
+					}
+
+					bool bLayoutsAreEqual = true;
+					for (int32 l = 0; l < Candidate.Mesh->GetLayoutCount(); ++l)
+					{
+						bLayoutsAreEqual = (*Candidate.Mesh->GetLayout(l) == *InOptions.OverrideLayouts[l].Layout);
+						if ( !bLayoutsAreEqual )
+						{
+							break;
+						}
+					}
+
+					if (!bLayoutsAreEqual)
+					{
+						continue;
+					}
+				}
+
+				DuplicateOf = Candidate;
+				break;
+			}
+		}
+
+		Ptr<ASTOp> LastMeshOp = ConstantOp;
+
+		if (DuplicateOf.Mesh)
+		{
+			// Make sure the source layouts of the mesh are mapped to the layouts of the duplicated mesh.
+			if (InOptions.bLayouts)
+			{
+				if (bIsOverridingLayouts)
+				{
+					OutResult.GeneratedLayouts = InOptions.OverrideLayouts;
+				}
+				else
+				{
+					for (int32 l = 0; l < DuplicateOf.Mesh->GetLayoutCount(); ++l)
+					{
+						const Layout* DuplicatedLayout = DuplicateOf.Mesh->GetLayout(l);
+						OutResult.GeneratedLayouts.Add({ DuplicatedLayout,0 });
+					}
+				}
+			}
+
+			LastMeshOp = DuplicateOf.LastMeshOp;
+			ConstantOp = nullptr;
+		}
+		else
+		{
+			// We need to clone the mesh in the node because we will modify it.
+			Ptr<Mesh> ClonedMesh = pMesh->Clone();
+			ClonedMesh->EnsureSurfaceData();
+
+			ConstantOp->SetValue(ClonedMesh, CompilerOptions->OptimisationOptions.DiskCacheContext);
+
+			// Add the unique vertex ID prefix in all cases, since it is free memory-wise
+			uint32 MeshIDPrefix = uint32(ConstantOp->GetValueHash());
+			{
+				// Ensure the ID group is unique
+				bool bValid = false;
+				do
+				{
+					bool bAlreadyPresent = false;
+					UniqueVertexIDGroups.FindOrAdd(MeshIDPrefix, &bAlreadyPresent);
+					bValid = !bAlreadyPresent && MeshIDPrefix != 0;
+					if (!bValid)
+					{
+						++MeshIDPrefix;
+					}
+				} while (bValid);
+
+				ClonedMesh->MeshIDPrefix = MeshIDPrefix;
+			}
+
+			// Add the constant data
+			FGeneratedConstantMesh MeshEntry;
+			MeshEntry.Mesh = ClonedMesh;
+			MeshEntry.LastMeshOp = LastMeshOp;
+			CachedCandidates.Add(MeshEntry);
+
+			if (InOptions.bLayouts)
+			{
+				if (!bIsOverridingLayouts)
+				{
+					// Apply whatever transform is necessary for every layout
+					for (int32 LayoutIndex = 0; LayoutIndex < Node.Layouts.Num(); ++LayoutIndex)
+					{
+						Ptr<NodeLayout> LayoutNode = Node.Layouts[LayoutIndex];
+						if (!LayoutNode)
+						{
+							continue;
+						}
+
+						FGeneratedLayout GeneratedData;
+						GeneratedData.Source = LayoutNode;
+						GeneratedData.Layout = GenerateLayout(LayoutNode, MeshIDPrefix);
+						const void* Context = InOptions.OverrideContext.Get(InNode->GetMessageContext());
+
+						bool bUseAbsoluteBlockIds = false;
+						PrepareMeshForLayout(GeneratedData, ClonedMesh, LayoutIndex, Context, InOptions, bUseAbsoluteBlockIds);
+
+						OutResult.GeneratedLayouts.Add(GeneratedData);
+					}
+				}
+				else
+				{
+					// We need to apply the transform of the layouts used to override
+					for (int32 LayoutIndex = 0; LayoutIndex < InOptions.OverrideLayouts.Num(); ++LayoutIndex)
+					{
+						const FGeneratedLayout& OverrideData = InOptions.OverrideLayouts[LayoutIndex];
+						Ptr<const Layout> GeneratedLayout = OverrideData.Layout;
+						const void* Context = InOptions.OverrideContext.Get(InNode->GetMessageContext());
+
+						// In this case we need the layout block ids to use the ids in the parent layout, and not be prefixed with
+						// the current mesh id prefix. For this reason we need them to be absolute.
+						bool bUseAbsoluteBlockIds = true;
+						PrepareMeshForLayout(OverrideData, ClonedMesh, LayoutIndex, Context, InOptions, bUseAbsoluteBlockIds);
+
+						OutResult.GeneratedLayouts.Add(OverrideData);
+					}
+				}
+			}
+		}
+
+		OutResult.BaseMeshOp = LastMeshOp;
+
+		// Add the tags operation
+		if (Tags.Num())
+		{
+			Ptr<ASTOpMeshAddTags> AddTagsOp = new ASTOpMeshAddTags;
+			AddTagsOp->Source = LastMeshOp;
+			AddTagsOp->Tags = Tags;
+			LastMeshOp = AddTagsOp;
+		}
+
+		OutResult.MeshOp = LastMeshOp;
+
+		// Apply the modifier for the pre-normal operations stage.
+		bool bModifiersForBeforeOperations = true;
+		OutResult.MeshOp = ApplyMeshModifiers(InOptions, OutResult, nullptr, bModifiersForBeforeOperations, InNode->GetMessageContext(), InNode);
+    }
+
+
+    //---------------------------------------------------------------------------------------------
+    void CodeGenerator::GenerateMesh_Format(const FMeshGenerationOptions& InOptions, FMeshGenerationResult& OutResult,
+		const NodeMeshFormat* format )
+    {
+        NodeMeshFormat::Private& node = *format->GetPrivate();
+
+        if ( node.Source )
+        {
+			FMeshGenerationOptions Options = InOptions;
+
+			FMeshGenerationResult baseResult;
+			GenerateMesh(Options,baseResult, node.Source);
+            Ptr<ASTOpMeshFormat> op = new ASTOpMeshFormat();
+            op->Source = baseResult.MeshOp;
+            op->Flags = 0;
+
+            Ptr<Mesh> FormatMesh = new Mesh();
+
+            if (node.VertexBuffers.GetBufferCount())
+            {
+                op->Flags |= OP::MeshFormatArgs::Vertex;
+                FormatMesh->VertexBuffers = node.VertexBuffers;
+            }
+
+            if (node.IndexBuffers.GetBufferCount())
+            {
+				op->Flags |= OP::MeshFormatArgs::Index;
+                FormatMesh->IndexBuffers = node.IndexBuffers;
+            }
+
+			if (node.bOptimizeBuffers)
+			{
+				op->Flags |= OP::MeshFormatArgs::OptimizeBuffers;
+			}
+
+            Ptr<ASTOpConstantResource> cop = new ASTOpConstantResource();
+            cop->Type = OP_TYPE::ME_CONSTANT;
+            cop->SetValue( FormatMesh, CompilerOptions->OptimisationOptions.DiskCacheContext );
+            op->Format = cop;
+
+            OutResult.MeshOp = op;
+            OutResult.BaseMeshOp = baseResult.BaseMeshOp;
+			OutResult.GeneratedLayouts = baseResult.GeneratedLayouts;
+		}
+        else
+        {
+            // Put something there
+            GenerateMesh(InOptions, OutResult, new NodeMeshConstant() );
+        }
+    }
+
+
+    //---------------------------------------------------------------------------------------------
+    void CodeGenerator::GenerateMesh_Transform(const FMeshGenerationOptions& InOptions, FMeshGenerationResult& OutResult,
+                                                const NodeMeshTransform* TransformNode )
+    {
+        const auto& node = *TransformNode->GetPrivate();
+
+        Ptr<ASTOpMeshTransform> op = new ASTOpMeshTransform();
+
+        // Base
+        if (node.Source)
+        {
+            GenerateMesh(InOptions, OutResult, node.Source);
+            op->source = OutResult.MeshOp;
+        }
+        else
+        {
+            // This argument is required
+            ErrorLog->GetPrivate()->Add("Mesh transform base node is not set.", ELMT_ERROR, TransformNode->GetMessageContext() );
+        }
+
+        op->matrix = node.Transform;
+
+        OutResult.MeshOp = op;
+    }
+
+
+    //---------------------------------------------------------------------------------------------
+    void CodeGenerator::GenerateMesh_ClipMorphPlane(const FMeshGenerationOptions& InOptions, FMeshGenerationResult& OutResult,
+                                                     const NodeMeshClipMorphPlane* ClipNode )
+    {
+        Ptr<ASTOpMeshClipMorphPlane> op = new ASTOpMeshClipMorphPlane();
+
+        // Base
+        if (ClipNode->Source)
+        {
+			FMeshGenerationOptions BaseOptions = InOptions;
+            GenerateMesh(BaseOptions, OutResult, ClipNode->Source);
+            op->source = OutResult.MeshOp;
+        }
+        else
+        {
+            // This argument is required
+            ErrorLog->GetPrivate()->Add("Mesh clip-morph-plane source node is not set.",
+                ELMT_ERROR, ClipNode->GetMessageContext());
+        }
+
+        // Morph to an ellipse
+        {
+            op->morphShape.type = (uint8_t)FShape::Type::Ellipse;
+            op->morphShape.position = ClipNode->Parameters.Origin;
+            op->morphShape.up = ClipNode->Parameters.Normal;
+            op->morphShape.size = FVector3f(ClipNode->Parameters.Radius1, ClipNode->Parameters.Radius2, ClipNode->Parameters.Rotation); // TODO: Move rotation to ellipse rotation reference base instead of passing it directly
+
+                                                                                      // Generate a "side" vector.
+                                                                                      // \todo: make generic and move to the vector class
+            {
+                // Generate vector perpendicular to normal for ellipse rotation reference base
+				FVector3f aux_base(0.f, 1.f, 0.f);
+
+                if (fabs(FVector3f::DotProduct(ClipNode->Parameters.Normal, aux_base)) > 0.95f)
+                {
+                    aux_base = FVector3f(0.f, 0.f, 1.f);
+                }
+
+                op->morphShape.side = FVector3f::CrossProduct(ClipNode->Parameters.Normal, aux_base);
+            }
+        }
+
+        // Selection by shape
+        if (ClipNode->Parameters.VertexSelectionType== FClipMorphPlaneParameters::VS_SHAPE)
+        {
+            op->vertexSelectionType = OP::MeshClipMorphPlaneArgs::VS_SHAPE;
+            op->selectionShape.type = (uint8_t)FShape::Type::AABox;
+            op->selectionShape.position = ClipNode->Parameters.SelectionBoxOrigin;
+            op->selectionShape.size = ClipNode->Parameters.SelectionBoxRadius;
+        }
+        else if (ClipNode->Parameters.VertexSelectionType == FClipMorphPlaneParameters::VS_BONE_HIERARCHY)
+        {
+            // Selection by bone hierarchy?
+            op->vertexSelectionType = OP::MeshClipMorphPlaneArgs::VS_BONE_HIERARCHY;
+            op->vertexSelectionBone = ClipNode->Parameters.VertexSelectionBone;
+			op->vertexSelectionBoneMaxRadius = ClipNode->Parameters.MaxEffectRadius;
+        }
+        else
+        {
+            op->vertexSelectionType = OP::MeshClipMorphPlaneArgs::VS_NONE;
+        }
+
+        // Parameters
+        op->dist = ClipNode->Parameters.DistanceToPlane;
+        op->factor = ClipNode->Parameters.LinearityFactor;
+
+        OutResult.MeshOp = op;
+    }
+
+
+    //---------------------------------------------------------------------------------------------
+    void CodeGenerator::GenerateMesh_ClipWithMesh(const FMeshGenerationOptions& InOptions, FMeshGenerationResult& OutResult,
+                                                   const NodeMeshClipWithMesh* ClipNode)
+    {
+        Ptr<ASTOpFixed> op = new ASTOpFixed();
+        op->op.type = OP_TYPE::ME_CLIPWITHMESH;
+
+        // Base
+        if (ClipNode->Source)
+        {
+            GenerateMesh(InOptions, OutResult, ClipNode->Source );
+            op->SetChild( op->op.args.MeshClipWithMesh.source, OutResult.MeshOp );
+        }
+        else
+        {
+            // This argument is required
+            ErrorLog->GetPrivate()->Add("Mesh clip-with-mesh source node is not set.",
+                ELMT_ERROR, ClipNode->GetMessageContext());
+        }
+
+        // Clipping mesh
+        if (ClipNode->ClipMesh)
+        {
+			FMeshGenerationOptions ClipOptions = InOptions;
+			ClipOptions.bLayouts = false;
+			ClipOptions.OverrideLayouts.Empty();
+			ClipOptions.ActiveTags.Empty();
+
+            FMeshGenerationResult clipResult;
+            GenerateMesh(ClipOptions, clipResult, ClipNode->ClipMesh);
+            op->SetChild( op->op.args.MeshClipWithMesh.clipMesh, clipResult.MeshOp );
+		}
+        else
+        {
+            // This argument is required
+            ErrorLog->GetPrivate()->Add("Mesh clip-with-mesh clipping mesh node is not set.",
+                ELMT_ERROR, ClipNode->GetMessageContext());
+        }
+
+        OutResult.MeshOp = op;
+    }
+
+	//---------------------------------------------------------------------------------------------
+	void CodeGenerator::GenerateMesh_ClipDeform(const FMeshGenerationOptions& InOptions, FMeshGenerationResult& Result, const NodeMeshClipDeform* ClipDeform)
+	{
+		const Ptr<ASTOpMeshBindShape> OpBind = new ASTOpMeshBindShape();
+		const Ptr<ASTOpMeshClipDeform> OpClipDeform = new ASTOpMeshClipDeform();
+
+		// Base Mesh
+		if (ClipDeform->BaseMesh)
+		{
+			GenerateMesh(InOptions, Result, ClipDeform->BaseMesh);
+			OpBind->Mesh = Result.MeshOp;
+		}
+		else
+		{
+			// This argument is required
+			ErrorLog->GetPrivate()->Add("Mesh Clip Deform base mesh node is not set.", ELMT_ERROR, ClipDeform->GetMessageContext());
+		}
+
+		// Base Shape
+		if (ClipDeform->ClipShape)
+		{
+			FMeshGenerationOptions ClipOptions = InOptions;
+			ClipOptions.bLayouts = false;
+			ClipOptions.OverrideLayouts.Empty();
+			ClipOptions.ActiveTags.Empty();
+
+			FMeshGenerationResult baseResult;
+			GenerateMesh(ClipOptions, baseResult, ClipDeform->ClipShape);
+			OpBind->Shape = baseResult.MeshOp;
+			OpClipDeform->ClipShape = baseResult.MeshOp;
+		}
+
+		OpClipDeform->Mesh = OpBind;
+
+		Result.MeshOp = OpClipDeform;
+	}
+
+    //---------------------------------------------------------------------------------------------
+    void CodeGenerator::GenerateMesh_ApplyPose(const FMeshGenerationOptions& InOptions, FMeshGenerationResult& OutResult,
+                                                const NodeMeshApplyPose* PoseNode )
+    {
+        const auto& node = *PoseNode->GetPrivate();
+
+        Ptr<ASTOpMeshApplyPose> op = new ASTOpMeshApplyPose();
+
+        // Base
+        if (node.m_pBase)
+        {
+            GenerateMesh(InOptions, OutResult, node.m_pBase );
+            op->base = OutResult.MeshOp;
+        }
+        else
+        {
+            // This argument is required
+            ErrorLog->GetPrivate()->Add("Mesh apply-pose base node is not set.",
+                ELMT_ERROR, PoseNode->GetMessageContext());
+        }
+
+        // Pose mesh
+        if (node.m_pPose)
+        {
+			FMeshGenerationOptions PoseOptions = InOptions;
+			PoseOptions.bLayouts = false;
+			PoseOptions.OverrideLayouts.Empty();
+			PoseOptions.ActiveTags.Empty();
+
+            FMeshGenerationResult poseResult;
+            GenerateMesh(PoseOptions, poseResult, node.m_pPose );
+            op->pose = poseResult.MeshOp;
+		}
+        else
+        {
+            // This argument is required
+            ErrorLog->GetPrivate()->Add("Mesh apply-pose pose node is not set.",
+                ELMT_ERROR, PoseNode->GetMessageContext());
+        }
+
+        OutResult.MeshOp = op;
+    }
+
+
+	//---------------------------------------------------------------------------------------------
+	void CodeGenerator::GenerateMesh_GeometryOperation(const FMeshGenerationOptions& InOptions, FMeshGenerationResult& OutResult, const NodeMeshGeometryOperation* GeomNode)
+	{
+		const auto& node = *GeomNode->GetPrivate();
+
+		Ptr<ASTOpMeshGeometryOperation> op = new ASTOpMeshGeometryOperation();
+
+		// Mesh A
+		if (node.m_pMeshA)
+		{
+			GenerateMesh(InOptions, OutResult, node.m_pMeshA);
+			op->meshA = OutResult.MeshOp;
+		}
+		else
+		{
+			// This argument is required
+			ErrorLog->GetPrivate()->Add("Mesh geometric op mesh-a node is not set.",
+				ELMT_ERROR, GeomNode->GetMessageContext());
+		}
+
+		// Mesh B
+		if (node.m_pMeshB)
+		{
+			FMeshGenerationOptions OtherOptions = InOptions;
+			OtherOptions.bLayouts = false;
+			OtherOptions.OverrideLayouts.Empty();
+			OtherOptions.ActiveTags.Empty();
+
+			FMeshGenerationResult bResult;
+			GenerateMesh(OtherOptions, bResult, node.m_pMeshB);
+			op->meshB = bResult.MeshOp;
+		}
+
+		op->scalarA = Generate_Generic(node.m_pScalarA, InOptions);
+		op->scalarB = Generate_Generic(node.m_pScalarB, InOptions);
+
+		OutResult.MeshOp = op;
+	}
+
+
+	//---------------------------------------------------------------------------------------------
+	void CodeGenerator::GenerateMesh_Reshape(const FMeshGenerationOptions& InOptions, FMeshGenerationResult& OutResult, const NodeMeshReshape* Reshape)
+	{
+		const NodeMeshReshape::Private& Node = *Reshape->GetPrivate();
+
+		Ptr<ASTOpMeshBindShape> OpBind = new ASTOpMeshBindShape();
+		Ptr<ASTOpMeshApplyShape> OpApply = new ASTOpMeshApplyShape();
+
+		OpBind->bReshapeSkeleton = Node.bReshapeSkeleton;	
+		OpBind->BonesToDeform = Node.BonesToDeform;
+    	OpBind->bReshapePhysicsVolumes = Node.bReshapePhysicsVolumes;
+		OpBind->PhysicsToDeform = Node.PhysicsToDeform;
+		OpBind->bReshapeVertices = Node.bReshapeVertices;
+		OpBind->bRecomputeNormals = Node.bRecomputeNormals;	
+		OpBind->bApplyLaplacian = Node.bApplyLaplacian;
+		OpBind->BindingMethod = static_cast<uint32>(EShapeBindingMethod::ReshapeClosestProject);
+
+		OpBind->RChannelUsage = Node.ColorRChannelUsage;
+		OpBind->GChannelUsage = Node.ColorGChannelUsage;
+		OpBind->BChannelUsage = Node.ColorBChannelUsage;
+		OpBind->AChannelUsage = Node.ColorAChannelUsage;
+
+		OpApply->bReshapeVertices = OpBind->bReshapeVertices;
+		OpApply->bRecomputeNormals = OpBind->bRecomputeNormals;
+		OpApply->bReshapeSkeleton = OpBind->bReshapeSkeleton;
+		OpApply->bApplyLaplacian = OpBind->bApplyLaplacian;
+		OpApply->bReshapePhysicsVolumes = OpBind->bReshapePhysicsVolumes;
+
+		// Base Mesh
+		if (Node.BaseMesh)
+		{
+			GenerateMesh(InOptions, OutResult, Node.BaseMesh);
+			OpBind->Mesh = OutResult.MeshOp;
+		}
+		else
+		{
+			// This argument is required
+			ErrorLog->GetPrivate()->Add("Mesh reshape base node is not set.", ELMT_ERROR, Reshape->GetMessageContext());
+		}
+
+		// Base and target shapes shouldn't have layouts or modifiers.
+		FMeshGenerationOptions ShapeOptions = InOptions;
+		ShapeOptions.bLayouts = false;
+		ShapeOptions.OverrideLayouts.Empty();
+		ShapeOptions.ActiveTags.Empty();
+
+		// Base Shape
+		if (Node.BaseShape)
+		{
+			FMeshGenerationResult baseResult;
+			GenerateMesh(ShapeOptions, baseResult, Node.BaseShape);
+			OpBind->Shape = baseResult.MeshOp;
+		}
+
+		OpApply->Mesh = OpBind;
+
+		// Target Shape
+		if (Node.TargetShape)
+		{
+			FMeshGenerationResult targetResult;
+			GenerateMesh(ShapeOptions, targetResult, Node.TargetShape);
+			OpApply->Shape = targetResult.MeshOp;
+		}
+
+		OutResult.MeshOp = OpApply;
+	}
+
+}
diff --git a/Engine/Plugins/Mutable/Source/MutableTools/Private/MuT/CodeGenerator_Projector.cpp b/Engine/Plugins/Mutable/Source/MutableTools/Private/MuT/CodeGenerator_Projector.cpp
new file mode 100644
--- /dev/null
+++ b/Engine/Plugins/Mutable/Source/MutableTools/Private/MuT/CodeGenerator_Projector.cpp
@@ -0,0 +1,182 @@
+// Copyright Epic Games, Inc. All Rights Reserved.
+
+#include "HAL/PlatformCrt.h"
+#include "Misc/AssertionMacros.h"
+#include "MuR/MutableMath.h"
+#include "MuR/Operations.h"
+#include "MuR/Parameters.h"
+#include "MuR/ParametersPrivate.h"
+#include "MuR/Ptr.h"
+#include "MuR/RefCounted.h"
+#include "MuR/Types.h"
+#include "MuT/AST.h"
+#include "MuT/ASTOpConstantProjector.h"
+#include "MuT/ASTOpParameter.h"
+#include "MuT/CodeGenerator.h"
+#include "MuT/ErrorLog.h"
+#include "MuT/ErrorLogPrivate.h"
+#include "MuT/NodeProjector.h"
+#include "MuT/NodeProjectorPrivate.h"
+#include "MuT/NodeRange.h"
+#include "map"
+
+
+namespace mu
+{
+
+	//-------------------------------------------------------------------------------------------------
+	FProjector ProjectorFromNode(const NodeProjectorConstant::Private& node)
+	{
+		FProjector p;
+		p.type = node.m_type;
+		p.position[0] = node.m_position[0];
+		p.position[1] = node.m_position[1];
+		p.position[2] = node.m_position[2];
+		p.direction[0] = node.m_direction[0];
+		p.direction[1] = node.m_direction[1];
+		p.direction[2] = node.m_direction[2];
+		p.up[0] = node.m_up[0];
+		p.up[1] = node.m_up[1];
+		p.up[2] = node.m_up[2];
+		p.scale[0] = node.m_scale[0];
+		p.scale[1] = node.m_scale[1];
+		p.scale[2] = node.m_scale[2];
+		p.projectionAngle = node.m_projectionAngle;
+
+		return p;
+	}
+
+
+	//-------------------------------------------------------------------------------------------------
+	void CodeGenerator::GenerateProjector(FProjectorGenerationResult& result, const FGenericGenerationOptions& Options, const Ptr<const NodeProjector>& Untyped)
+	{
+		if (!Untyped)
+		{
+			result = FProjectorGenerationResult();
+			return;
+		}
+
+		// See if it was already generated
+		FGeneratedCacheKey Key;
+		Key.Node = Untyped;
+		Key.Options = Options;
+		FGeneratedProjectorsMap::ValueType* it = GeneratedProjectors.Find(Key);
+		if (it)
+		{
+			result = *it;
+			return;
+		}
+
+		// Generate for each different type of node
+		if ( Untyped->GetType()== NodeProjectorConstant::GetStaticType() )
+		{
+			GenerateProjector_Constant(result, Options, static_cast<const NodeProjectorConstant*>(Untyped.get()));
+		}
+		else if (Untyped->GetType() == NodeProjectorParameter::GetStaticType() )
+		{
+			GenerateProjector_Parameter(result, Options, static_cast<const NodeProjectorParameter*>(Untyped.get()));
+		}
+		else
+		{
+			check(false);
+		}
+
+
+		// Cache the result
+		GeneratedProjectors.Add(Key, result);
+	}
+
+
+	//-------------------------------------------------------------------------------------------------
+	void CodeGenerator::GenerateProjector_Constant(FProjectorGenerationResult& result, const FGenericGenerationOptions& Options,
+		const Ptr<const NodeProjectorConstant>& constant)
+	{
+		const NodeProjectorConstant::Private& node = *constant->GetPrivate();
+
+		Ptr<ASTOpConstantProjector> op = new ASTOpConstantProjector();
+		op->value = ProjectorFromNode(node);
+
+		result.op = op;
+		result.type = op->value.type;
+	}
+
+
+	//-------------------------------------------------------------------------------------------------
+	void CodeGenerator::GenerateProjector_Parameter(FProjectorGenerationResult& result, const FGenericGenerationOptions& Options,
+		const Ptr<const NodeProjectorParameter>& paramn)
+	{
+		const NodeProjectorParameter::Private& node = *paramn->GetPrivate();
+
+		Ptr<ASTOpParameter> op;
+
+		Ptr<ASTOpParameter>* it = FirstPass.ParameterNodes.Find(node.m_pNode);
+		if (!it)
+		{
+			FParameterDesc param;
+			param.m_name = node.m_name;
+			const TCHAR* CStr = ToCStr(node.m_uid);
+			param.m_uid.ImportTextItem(CStr, 0, nullptr, nullptr);
+			param.m_type = PARAMETER_TYPE::T_PROJECTOR;
+
+			FProjector p = ProjectorFromNode(node);
+			param.m_defaultValue.Set<ParamProjectorType>(p);
+
+			op = new ASTOpParameter();
+			op->type = OP_TYPE::PR_PARAMETER;
+			op->parameter = param;
+
+			// Generate the code for the ranges
+			for (int32 a = 0; a < node.m_ranges.Num(); ++a)
+			{
+				FRangeGenerationResult rangeResult;
+				GenerateRange(rangeResult, Options, node.m_ranges[a]);
+				op->ranges.Emplace(op.get(), rangeResult.sizeOp, rangeResult.rangeName, rangeResult.rangeUID);
+			}
+
+			FirstPass.ParameterNodes.Add(node.m_pNode, op);
+		}
+		else
+		{
+			op = *it;
+		}
+
+		const FProjector& Projector = op->parameter.m_defaultValue.Get<ParamProjectorType>();
+		result.type = Projector.type;
+		result.op = op;
+	}
+
+
+	//-------------------------------------------------------------------------------------------------
+	void CodeGenerator::GenerateMissingProjectorCode(FProjectorGenerationResult& result,
+		const void* errorContext)
+	{
+		// Log an error message
+		ErrorLog->GetPrivate()->Add("Required projector connection not found.",
+			ELMT_ERROR, errorContext);
+
+		FProjector p;
+		p.type = PROJECTOR_TYPE::PLANAR;
+
+		p.direction[0] = 1;
+		p.direction[1] = 0;
+		p.direction[2] = 0;
+
+		p.up[0] = 0;
+		p.up[1] = 0;
+		p.up[2] = 1;
+
+		p.position[0] = 0.0f;
+		p.position[1] = 0.0f;
+		p.position[2] = 0.0f;
+
+		p.scale[0] = 1.0f;
+		p.scale[1] = 1.0f;
+
+		Ptr<ASTOpConstantProjector> op = new ASTOpConstantProjector();
+		op->value = p;
+
+		result.op = op;
+		result.type = p.type;
+	}
+
+}
diff --git a/Engine/Plugins/Mutable/Source/MutableTools/Private/MuT/CodeGenerator_Scalar.cpp b/Engine/Plugins/Mutable/Source/MutableTools/Private/MuT/CodeGenerator_Scalar.cpp
new file mode 100644
--- /dev/null
+++ b/Engine/Plugins/Mutable/Source/MutableTools/Private/MuT/CodeGenerator_Scalar.cpp
@@ -0,0 +1,442 @@
+// Copyright Epic Games, Inc. All Rights Reserved.
+
+#include "HAL/PlatformMath.h"
+#include "Misc/AssertionMacros.h"
+#include "MuR/Image.h"
+#include "MuR/MutableMath.h"
+#include "MuR/Operations.h"
+#include "MuR/Parameters.h"
+#include "MuR/ParametersPrivate.h"
+#include "MuR/Platform.h"
+#include "MuR/Ptr.h"
+#include "MuR/RefCounted.h"
+#include "MuR/Types.h"
+#include "MuT/AST.h"
+#include "MuT/ASTOpConditional.h"
+#include "MuT/ASTOpParameter.h"
+#include "MuT/ASTOpScalarCurve.h"
+#include "MuT/ASTOpSwitch.h"
+#include "MuT/CodeGenerator.h"
+#include "MuT/CodeGenerator_FirstPass.h"
+#include "MuT/ErrorLog.h"
+#include "MuT/ErrorLogPrivate.h"
+#include "MuT/NodeImage.h"
+#include "MuT/NodeMeshGeometryOperation.h"
+#include "MuT/NodeMeshTable.h"
+#include "MuT/NodeRange.h"
+#include "MuT/NodeScalar.h"
+#include "MuT/NodeScalarArithmeticOperation.h"
+#include "MuT/NodeScalarArithmeticOperationPrivate.h"
+#include "MuT/NodeScalarConstant.h"
+#include "MuT/NodeScalarConstantPrivate.h"
+#include "MuT/NodeScalarCurve.h"
+#include "MuT/NodeScalarCurvePrivate.h"
+#include "MuT/NodeScalarEnumParameter.h"
+#include "MuT/NodeScalarEnumParameterPrivate.h"
+#include "MuT/NodeScalarParameter.h"
+#include "MuT/NodeScalarParameterPrivate.h"
+#include "MuT/NodeScalarSwitch.h"
+#include "MuT/NodeScalarTable.h"
+#include "MuT/NodeScalarVariation.h"
+#include "MuT/NodeScalarVariationPrivate.h"
+#include "MuT/Table.h"
+#include "MuT/TablePrivate.h"
+
+
+namespace mu
+{
+
+	void CodeGenerator::GenerateScalar(FScalarGenerationResult& result, const FGenericGenerationOptions& Options, const Ptr<const NodeScalar>& Untyped)
+	{
+		if (!Untyped)
+		{
+			result = FScalarGenerationResult();
+			return;
+		}
+
+		// See if it was already generated
+		FGeneratedCacheKey Key;
+		Key.Node = Untyped;
+		Key.Options = Options;
+		FGeneratedScalarsMap::ValueType* it = GeneratedScalars.Find(Key);
+		if (it)
+		{
+			result = *it;
+			return;
+		}
+
+		// Generate for each different type of node
+		if (Untyped->GetType() == NodeScalarConstant::GetStaticType())
+		{
+			const NodeScalarConstant* Constant = static_cast<const NodeScalarConstant*>(Untyped.get());
+			GenerateScalar_Constant(result, Options, Constant);
+		}
+		else if (Untyped->GetType() == NodeScalarParameter::GetStaticType())
+		{
+			const NodeScalarParameter* Param = static_cast<const NodeScalarParameter*>(Untyped.get());
+			GenerateScalar_Parameter(result, Options, Param);
+		}
+		else if (Untyped->GetType() == NodeScalarSwitch::GetStaticType())
+		{
+			const NodeScalarSwitch* Switch = static_cast<const NodeScalarSwitch*>(Untyped.get());
+			GenerateScalar_Switch(result, Options, Switch);
+		}
+		else if (Untyped->GetType() == NodeScalarEnumParameter::GetStaticType())
+		{
+			const NodeScalarEnumParameter* EnumParam = static_cast<const NodeScalarEnumParameter*>(Untyped.get());
+			GenerateScalar_EnumParameter(result, Options, EnumParam);
+		}
+		else if (Untyped->GetType() == NodeScalarCurve::GetStaticType())
+		{
+			const NodeScalarCurve* Curve = static_cast<const NodeScalarCurve*>(Untyped.get());
+			GenerateScalar_Curve(result, Options, Curve);
+		}
+		else if (Untyped->GetType() == NodeScalarArithmeticOperation::GetStaticType())
+		{
+			const NodeScalarArithmeticOperation* Arithmetic = static_cast<const NodeScalarArithmeticOperation*>(Untyped.get());
+			GenerateScalar_Arithmetic(result, Options, Arithmetic);
+		}
+		else if (Untyped->GetType() == NodeScalarVariation::GetStaticType())
+		{
+			const NodeScalarVariation* Variation = static_cast<const NodeScalarVariation*>(Untyped.get());
+			GenerateScalar_Variation(result, Options, Variation);
+		}
+		else if (Untyped->GetType() == NodeScalarTable::GetStaticType())
+		{
+			const NodeScalarTable* Table = static_cast<const NodeScalarTable*>(Untyped.get());
+			GenerateScalar_Table(result, Options, Table);
+		}
+		else
+		{
+			check(false);
+			return;
+		}
+
+		// Cache the result
+		GeneratedScalars.Add(Key, result);
+	}
+
+
+	//-------------------------------------------------------------------------------------------------
+	void CodeGenerator::GenerateScalar_Constant(FScalarGenerationResult& result, const FGenericGenerationOptions& Options, const Ptr<const NodeScalarConstant>& Typed)
+	{
+		const NodeScalarConstant::Private& node = *Typed->GetPrivate();
+
+		Ptr<ASTOpFixed> op = new ASTOpFixed();
+		op->op.type = OP_TYPE::SC_CONSTANT;
+		op->op.args.ScalarConstant.value = node.m_value;
+
+		result.op = op;
+	}
+
+
+	//-------------------------------------------------------------------------------------------------
+	void CodeGenerator::GenerateScalar_Parameter(FScalarGenerationResult& result, const FGenericGenerationOptions& Options, const Ptr<const NodeScalarParameter>& Typed)
+	{
+		const NodeScalarParameter::Private& node = *Typed->GetPrivate();
+
+		Ptr<ASTOpParameter> op;
+
+		Ptr<ASTOpParameter>* it = FirstPass.ParameterNodes.Find(node.m_pNode);
+		if (!it)
+		{
+			FParameterDesc param;
+			param.m_name = node.m_name;
+			const TCHAR* CStr = ToCStr(node.m_uid);
+			param.m_uid.ImportTextItem(CStr, 0, nullptr, nullptr);
+			param.m_type = PARAMETER_TYPE::T_FLOAT;
+			param.m_defaultValue.Set<ParamFloatType>(node.m_defaultValue);
+
+			op = new ASTOpParameter();
+			op->type = OP_TYPE::SC_PARAMETER;
+			op->parameter = param;
+
+			// Generate the code for the ranges
+			for (int32 a = 0; a < node.m_ranges.Num(); ++a)
+			{
+				FRangeGenerationResult rangeResult;
+				GenerateRange(rangeResult, Options, node.m_ranges[a]);
+				op->ranges.Emplace(op.get(), rangeResult.sizeOp, rangeResult.rangeName, rangeResult.rangeUID);
+			}
+
+			FirstPass.ParameterNodes.Add(node.m_pNode, op);
+		}
+		else
+		{
+			op = *it;
+		}
+
+		result.op = op;
+	}
+
+
+	//-------------------------------------------------------------------------------------------------
+	void CodeGenerator::GenerateScalar_EnumParameter(FScalarGenerationResult& result, const FGenericGenerationOptions& Options, const Ptr<const NodeScalarEnumParameter>& Typed)
+	{
+		const NodeScalarEnumParameter::Private& node = *Typed->GetPrivate();
+
+		Ptr<ASTOpParameter> op;
+
+		Ptr<ASTOpParameter>* it = FirstPass.ParameterNodes.Find(node.m_pNode);
+		if (!it)
+		{
+			FParameterDesc param;
+			param.m_name = node.m_name;
+			const TCHAR* CStr = ToCStr(node.m_uid);
+			param.m_uid.ImportTextItem(CStr, 0, nullptr, nullptr);
+			param.m_type = PARAMETER_TYPE::T_INT;
+			param.m_defaultValue.Set<ParamIntType>(node.m_defaultValue);
+
+			param.m_possibleValues.SetNum(node.m_options.Num());
+			for (int32 i = 0; i < node.m_options.Num(); ++i)
+			{
+				param.m_possibleValues[i].m_value = (int16_t)node.m_options[i].value;
+				param.m_possibleValues[i].m_name = node.m_options[i].name;
+			}
+
+			op = new ASTOpParameter();
+			op->type = OP_TYPE::NU_PARAMETER;
+			op->parameter = param;
+
+			// Generate the code for the ranges
+			for (int32 a = 0; a < node.m_ranges.Num(); ++a)
+			{
+				FRangeGenerationResult rangeResult;
+				GenerateRange(rangeResult, Options, node.m_ranges[a]);
+				op->ranges.Emplace(op.get(), rangeResult.sizeOp, rangeResult.rangeName, rangeResult.rangeUID);
+			}
+
+			FirstPass.ParameterNodes.Add(node.m_pNode, op);
+		}
+		else
+		{
+			op = *it;
+		}
+
+		result.op = op;
+	}
+
+
+	//-------------------------------------------------------------------------------------------------
+	void CodeGenerator::GenerateScalar_Switch(FScalarGenerationResult& result, const FGenericGenerationOptions& Options, const Ptr<const NodeScalarSwitch>& Typed)
+	{
+		const NodeScalarSwitch& node = *Typed;
+
+		if (node.Options.Num() == 0)
+		{
+			// No options in the switch!
+			Ptr<ASTOp> missingOp = GenerateMissingScalarCode(TEXT("Switch option"),
+				1.0f,
+				Typed->GetMessageContext());
+			result.op = missingOp;
+			return;
+		}
+
+		Ptr<ASTOpSwitch> op = new ASTOpSwitch();
+		op->type = OP_TYPE::SC_SWITCH;
+
+		// Variable value
+		if (node.Parameter)
+		{
+			FScalarGenerationResult ChildResult;
+			GenerateScalar(ChildResult, Options, node.Parameter.get());
+			op->variable = ChildResult.op;
+		}
+		else
+		{
+			// This argument is required
+			op->variable = GenerateMissingScalarCode(TEXT("Switch variable"), 0.0f, Typed->GetMessageContext());
+		}
+
+		// Options
+		for (int32 t = 0; t < node.Options.Num(); ++t)
+		{
+			Ptr<ASTOp> branch;
+			if (node.Options[t])
+			{
+				FScalarGenerationResult ChildResult;
+				GenerateScalar(ChildResult, Options, node.Options[t].get());
+				branch = ChildResult.op;
+			}
+			else
+			{
+				// This argument is required
+				branch = GenerateMissingScalarCode(TEXT("Switch option"), 1.0f, Typed->GetMessageContext());
+			}
+			op->cases.Emplace((int16_t)t, op, branch);
+		}
+
+		result.op = op;
+	}
+
+
+	//-------------------------------------------------------------------------------------------------
+	void CodeGenerator::GenerateScalar_Variation(FScalarGenerationResult& result, const FGenericGenerationOptions& Options, const Ptr<const NodeScalarVariation>& Typed)
+	{
+		const NodeScalarVariation::Private& node = *Typed->GetPrivate();
+
+		Ptr<ASTOp> op;
+
+		// Default case
+		if (node.m_defaultScalar)
+		{
+			FMeshGenerationResult branchResults;
+
+			FScalarGenerationResult ChildResult;
+			GenerateScalar(ChildResult, Options, node.m_defaultScalar);
+			op = ChildResult.op;
+		}
+
+		// Process variations in reverse order, since conditionals are built bottom-up.
+		for (int32 t = node.m_variations.Num() - 1; t >= 0; --t)
+		{
+			int32 tagIndex = -1;
+			const FString& tag = node.m_variations[t].m_tag;
+			for (int32 i = 0; i < FirstPass.Tags.Num(); ++i)
+			{
+				if (FirstPass.Tags[i].Tag == tag)
+				{
+					tagIndex = i;
+				}
+			}
+
+			if (tagIndex < 0)
+			{
+				FString Msg = FString::Printf(TEXT("Unknown tag found in image variation [%s]."), *tag);
+
+				ErrorLog->GetPrivate()->Add(Msg, ELMT_WARNING, Typed->GetMessageContext());
+				continue;
+			}
+
+			Ptr<ASTOp> variationOp;
+			if (node.m_variations[t].m_scalar)
+			{
+				FScalarGenerationResult ChildResult;
+				GenerateScalar(ChildResult, Options, node.m_variations[t].m_scalar);
+				variationOp = ChildResult.op;
+			}
+			else
+			{
+				// This argument is required
+				variationOp = GenerateMissingScalarCode(TEXT("Variation option"), 0.0f, Typed->GetMessageContext());
+			}
+
+
+			Ptr<ASTOpConditional> conditional = new ASTOpConditional;
+			conditional->type = OP_TYPE::SC_CONDITIONAL;
+			conditional->no = op;
+			conditional->yes = variationOp;
+			conditional->condition = FirstPass.Tags[tagIndex].GenericCondition;
+
+			op = conditional;
+		}
+
+		result.op = op;
+	}
+
+
+	//-------------------------------------------------------------------------------------------------
+	void CodeGenerator::GenerateScalar_Curve(FScalarGenerationResult& result, const FGenericGenerationOptions& Options, const Ptr<const NodeScalarCurve>& Typed)
+	{
+		Ptr<ASTOpScalarCurve> op = new ASTOpScalarCurve();
+
+		// T
+		if (Node* pA = Typed->CurveSampleValue.get())
+		{
+			op->time = Generate_Generic(pA, Options);
+		}
+		else
+		{
+			op->time = CodeGenerator::GenerateMissingScalarCode(TEXT("Curve T"), 0.5f, Typed->GetMessageContext());
+		}
+
+		op->Curve = Typed->Curve;
+
+		result.op = op;
+	}
+
+
+	//-------------------------------------------------------------------------------------------------
+	void CodeGenerator::GenerateScalar_Arithmetic(FScalarGenerationResult& result, const FGenericGenerationOptions& Options, const Ptr<const NodeScalarArithmeticOperation>& Typed)
+	{
+		const NodeScalarArithmeticOperation::Private& node = *Typed->GetPrivate();
+
+		Ptr<ASTOpFixed> op = new ASTOpFixed();
+		op->op.type = OP_TYPE::SC_ARITHMETIC;
+
+		switch (node.m_operation)
+		{
+		case NodeScalarArithmeticOperation::AO_ADD: op->op.args.ScalarArithmetic.operation = OP::ArithmeticArgs::ADD; break;
+		case NodeScalarArithmeticOperation::AO_SUBTRACT: op->op.args.ScalarArithmetic.operation = OP::ArithmeticArgs::SUBTRACT; break;
+		case NodeScalarArithmeticOperation::AO_MULTIPLY: op->op.args.ScalarArithmetic.operation = OP::ArithmeticArgs::MULTIPLY; break;
+		case NodeScalarArithmeticOperation::AO_DIVIDE: op->op.args.ScalarArithmetic.operation = OP::ArithmeticArgs::DIVIDE; break;
+		default:
+			checkf(false, TEXT("Unknown arithmetic operation."));
+			op->op.args.ScalarArithmetic.operation = OP::ArithmeticArgs::NONE;
+			break;
+		}
+
+		// A
+		if (Node* pA = node.m_pA.get())
+		{
+			op->SetChild(op->op.args.ScalarArithmetic.a, Generate_Generic(pA, Options));
+		}
+		else
+		{
+			op->SetChild(op->op.args.ScalarArithmetic.a,
+				CodeGenerator::GenerateMissingScalarCode( TEXT("ScalarArithmetic A"), 1.0f, Typed->GetMessageContext() )
+			);
+		}
+
+		// B
+		if (Node* pB = node.m_pB.get())
+		{
+			op->SetChild(op->op.args.ScalarArithmetic.b, Generate_Generic(pB, Options));
+		}
+		else
+		{
+			op->SetChild(op->op.args.ScalarArithmetic.b,
+				CodeGenerator::GenerateMissingScalarCode( TEXT("ScalarArithmetic B"), 1.0f, Typed->GetMessageContext() )
+			);
+		}
+
+		result.op = op;
+	}
+
+
+	//-------------------------------------------------------------------------------------------------
+	void CodeGenerator::GenerateScalar_Table(FScalarGenerationResult& result, const FGenericGenerationOptions& Options, const Ptr<const NodeScalarTable>& Typed)
+	{
+		const NodeScalarTable& node = *Typed;
+
+		Ptr<ASTOp> Op = GenerateTableSwitch<NodeScalarTable, ETableColumnType::Scalar, OP_TYPE::SC_SWITCH>(node,
+			[this,&Options](const NodeScalarTable& node, int32 colIndex, int32 row, mu::ErrorLog* pErrorLog)
+			{
+				NodeScalarConstantPtr pCell = new NodeScalarConstant();
+				float scalar = node.Table->GetPrivate()->Rows[row].Values[colIndex].Scalar;
+				pCell->SetValue(scalar);
+				return Generate_Generic(pCell, Options);
+			});
+
+		result.op = Op;
+	}
+
+
+	//---------------------------------------------------------------------------------------------
+	mu::Ptr<ASTOp> CodeGenerator::GenerateMissingScalarCode(const TCHAR* strWhere, float value, const void* errorContext)
+	{
+		// Log a warning
+		FString Msg = FString::Printf(TEXT("Required connection not found: %s"), strWhere );
+		ErrorLog->GetPrivate()->Add(Msg, ELMT_ERROR, errorContext);
+
+		// Create a constant node
+		NodeScalarConstantPtr pNode = new NodeScalarConstant();
+		pNode->SetValue(value);
+
+		FGenericGenerationOptions Options;
+		Ptr<ASTOp> result = Generate_Generic(pNode, Options);
+
+		return result;
+	}
+
+}
\ No newline at end of file
diff --git a/Engine/Plugins/Mutable/Source/MutableTools/Private/MuT/CodeGenerator_SecondPass.cpp b/Engine/Plugins/Mutable/Source/MutableTools/Private/MuT/CodeGenerator_SecondPass.cpp
new file mode 100644
--- /dev/null
+++ b/Engine/Plugins/Mutable/Source/MutableTools/Private/MuT/CodeGenerator_SecondPass.cpp
@@ -0,0 +1,487 @@
+// Copyright Epic Games, Inc. All Rights Reserved.
+
+#include "MuT/CodeGenerator_SecondPass.h"
+
+#include "Misc/AssertionMacros.h"
+#include "MuR/MutableTrace.h"
+#include "MuR/Operations.h"
+#include "MuR/Types.h"
+#include "MuT/ASTOpConstantBool.h"
+#include "MuT/ErrorLog.h"
+#include "MuT/CodeGenerator_FirstPass.h"
+
+namespace mu
+{
+
+
+	//---------------------------------------------------------------------------------------------
+	SecondPassGenerator::SecondPassGenerator(
+		FirstPassGenerator* firstPass,
+		const CompilerOptions::Private* options)
+	{
+		check(firstPass);
+		check(options);
+		FirstPass = firstPass;
+		CompilerOptions = options;
+	}
+
+
+	//---------------------------------------------------------------------------------------------
+	mu::Ptr<ASTOp> SecondPassGenerator::GenerateTagCondition(size_t tagIndex,
+		const set<size_t>& posSurf,
+		const set<size_t>& negSurf,
+		const set<size_t>& posTag,
+		const set<size_t>& negTag)
+	{
+		FirstPassGenerator::FTag& t = FirstPass->Tags[tagIndex];
+
+		// If this tag is already in the list of positive tags, return true as condition
+		if (posTag.find(tagIndex) != posTag.end())
+		{
+			return m_opPool.Add(new ASTOpConstantBool(true));
+		}
+
+		// If this tag is already in the list of negative tags, return false as condition
+		if (negTag.find(tagIndex) != negTag.end())
+		{
+			return m_opPool.Add(new ASTOpConstantBool(false));
+		}
+
+		// Cached?
+		FConditionGenerationKey key;
+		key.tagOrSurfIndex = tagIndex;
+		//    key.negSurf = negSurf;
+		//    key.posSurf = posSurf;
+		//    key.negTag = negTag;
+		//    key.posTag = posTag;
+		for (auto s : negTag) { if (TagsPerTag[tagIndex].count(s) > 0) { key.negTag.insert(s); } }
+		for (auto s : posTag) { if (TagsPerTag[tagIndex].count(s) > 0) { key.posTag.insert(s); } }
+		for (auto s : negSurf) { if (SurfacesPerTag[tagIndex].count(s) > 0) { key.negSurf.insert(s); } }
+		for (auto s : posSurf) { if (SurfacesPerTag[tagIndex].count(s) > 0) { key.posSurf.insert(s); } }
+
+		{
+			auto it = TagConditionGenerationCache.find(key);
+			if (it != TagConditionGenerationCache.end())
+			{
+				return it->second;
+			}
+		}
+		
+		Ptr<ASTOp> c;
+
+		// Condition expression for all the surfaces that activate the tag
+		for (int32 surfIndex : t.Surfaces)
+		{
+			if (posSurf.find(surfIndex) != posSurf.end())
+			{
+				// This surface is already a positive requirement higher up in the condition so
+				// we can ignore it here.
+				continue;
+			}
+
+			if (negSurf.find(surfIndex) != negSurf.end())
+			{
+				// This surface is a negative requirement higher up in the condition so
+				// this branch never be true.
+				continue;
+			}
+
+			const FirstPassGenerator::FSurface& surface = FirstPass->Surfaces[surfIndex];
+
+			auto PositiveTags = posTag;
+			PositiveTags.insert(tagIndex);
+
+			Ptr<ASTOp> surfCondition = GenerateDataCodition(surfIndex,
+				FirstPass->Surfaces[surfIndex].PositiveTags,
+				FirstPass->Surfaces[surfIndex].NegativeTags,
+				posSurf,
+				negSurf,
+				PositiveTags,
+				negTag);
+
+			// If the surface is a constant false, we can skip adding it
+			if (surfCondition && surfCondition->GetOpType()==OP_TYPE::BO_CONSTANT)
+			{
+				const ASTOpConstantBool* constOp = static_cast<const ASTOpConstantBool*>(surfCondition.get());
+				if (constOp->value == false)
+				{
+					continue;
+				}
+			}
+
+			Ptr<ASTOp> fullCondition;
+			if (surfCondition)
+			{
+				Ptr<ASTOpFixed> f = new ASTOpFixed();
+				f->op.type = OP_TYPE::BO_AND;
+				f->SetChild(f->op.args.BoolBinary.a, surface.ObjectCondition);
+				f->SetChild(f->op.args.BoolBinary.b, surfCondition);
+				fullCondition = m_opPool.Add(f);
+			}
+			else
+			{
+				fullCondition = m_opPool.Add(surface.ObjectCondition);
+			}
+
+
+			if (!c)
+			{
+				c = fullCondition;
+			}
+			else
+			{
+				Ptr<ASTOpFixed> o = new ASTOpFixed();
+				o->op.type = OP_TYPE::BO_OR;
+				o->SetChild(o->op.args.BoolBinary.a, fullCondition);
+				o->SetChild(o->op.args.BoolBinary.b, c);
+				c = m_opPool.Add(o);
+			}
+
+			// Optimise the condition now.
+			//PartialOptimise( c, CompilerOptions->OptimisationOptions );
+		}
+
+		TagConditionGenerationCache.insert(std::make_pair<>(key, c));
+
+		return c;
+	}
+
+	//---------------------------------------------------------------------------------------------
+	mu::Ptr<ASTOp> SecondPassGenerator::GenerateDataCodition(size_t Index,
+		const TArray<FString>& PositiveTags,
+		const TArray<FString>& NegativeTags,
+		const set<size_t>& posSurf,
+		const set<size_t>& negSurf,
+		const set<size_t>& posTag,
+		const set<size_t>& negTag)
+	{
+		// If this surface is already in the list of positive surfaces, return true as condition
+		if (posSurf.find(Index) != posSurf.end())
+		{
+			return m_opPool.Add(new ASTOpConstantBool(true));
+		}
+
+		// If this surface is already in the list of negative surfaces, return false as condition
+		if (negSurf.find(Index) != negSurf.end())
+		{
+			return m_opPool.Add(new ASTOpConstantBool(false));
+		}
+
+		Ptr<ASTOp> c;
+
+		for (const FString& t : PositiveTags)
+		{
+			const FirstPassGenerator::FTag* it = FirstPass->Tags.FindByPredicate([&](const FirstPassGenerator::FTag& e) { return e.Tag == t; });
+			if (!it)
+			{
+				// This could happen if a tag is in a variation but noone defines it.
+				// This surface depends on a tag that will never be active, so it will never be used.
+				return m_opPool.Add(new ASTOpConstantBool(false));
+			}
+
+			size_t tagIndex = it - &FirstPass->Tags[0];
+
+			set<size_t> positiveSurfacesVisited = posSurf;
+			positiveSurfacesVisited.insert(Index);
+
+			Ptr<ASTOp> tagCondition = GenerateTagCondition(tagIndex,
+				positiveSurfacesVisited,
+				negSurf,
+				posTag,
+				negTag);
+
+			if (!tagCondition)
+			{
+				// This tag is unconditionally activated, so there's no condition logic to add
+				continue;
+			}
+
+			// TODO: Optimise the tag condition here
+
+			// If the tag is a constant ...
+			bool isConstant = false;
+			bool constantValue = false;
+			if (tagCondition->GetOpType()==OP_TYPE::BO_CONSTANT)
+			{
+				const ASTOpConstantBool* constOp = static_cast<const ASTOpConstantBool*>(tagCondition.get());
+				isConstant = true;
+				constantValue = constOp->value;
+			}
+
+			if (!isConstant)
+			{
+				if (!c)
+				{
+					c = tagCondition;
+				}
+				else
+				{
+					Ptr<ASTOpFixed> o = new ASTOpFixed;
+					o->op.type = OP_TYPE::BO_AND;
+					o->SetChild(o->op.args.BoolBinary.a, tagCondition);
+					o->SetChild(o->op.args.BoolBinary.b, c);
+					c = m_opPool.Add(o);
+				}
+			}
+			else if (constantValue == true)
+			{
+				// No need to add it to the AND
+			}
+			else //if (constantValue==false)
+			{
+				// Entire expression will be false
+				Ptr<ASTOpConstantBool> f = new ASTOpConstantBool(false);
+
+				// No need to evaluate anything else.
+				c = m_opPool.Add(f);
+				break;
+			}
+		}
+
+
+		for (const FString& t : NegativeTags)
+		{
+			const FirstPassGenerator::FTag* it = FirstPass->Tags.FindByPredicate([&](const FirstPassGenerator::FTag& e) { return e.Tag == t; });
+			if (!it)
+			{
+				// This could happen if a tag is in a variation but noone defines it.
+				continue;
+			}
+
+			size_t tagIndex = it - FirstPass->Tags.GetData();
+
+			set<size_t> positiveSurfacesVisited = negSurf;
+			set<size_t> negativeSurfacesVisited = posSurf;
+			negativeSurfacesVisited.insert(Index);
+			set<size_t> positiveTagsVisited = negTag;
+			set<size_t> negativeTagsVisited = posTag;
+			Ptr<ASTOp> tagCondition = GenerateTagCondition(tagIndex,
+				positiveSurfacesVisited,
+				negativeSurfacesVisited,
+				positiveTagsVisited,
+				negativeTagsVisited);
+
+			// No condition is equal to a conditional with a true constant
+			if (!tagCondition)
+			{
+				ASTOpConstantBool* ConstOp = new ASTOpConstantBool();
+				ConstOp->value = true;
+				tagCondition = ConstOp;
+			}
+
+			// TODO: Optimise the tag condition here
+
+			// If the tag is a constant ...
+			bool isConstant = false;
+			bool constantValue = false;
+			if (tagCondition && tagCondition->GetOpType()== OP_TYPE::BO_CONSTANT)
+			{
+				const ASTOpConstantBool* constOp = static_cast<const ASTOpConstantBool*>(tagCondition.get());
+				isConstant = true;
+				constantValue = constOp->value;
+			}
+
+
+			if (!isConstant && tagCondition)
+			{
+				Ptr<ASTOpFixed> n = new ASTOpFixed;
+				n->op.type = OP_TYPE::BO_NOT;
+				n->SetChild(n->op.args.BoolNot.source, tagCondition);
+
+				if (!c)
+				{
+					c = m_opPool.Add(n);
+				}
+				else
+				{
+					Ptr<ASTOpFixed> o = new ASTOpFixed;
+					o->op.type = OP_TYPE::BO_AND;
+					o->SetChild(o->op.args.BoolBinary.a, n);
+					o->SetChild(o->op.args.BoolBinary.b, c);
+					c = m_opPool.Add(o);
+				}
+			}
+			else if (isConstant && constantValue == true)
+			{
+				// No expression here means always true which becomes always false
+				Ptr<ASTOpConstantBool> f = new ASTOpConstantBool(false);
+
+				// No need to evaluate anything else.
+				c = m_opPool.Add(f);
+				break;
+			}
+		}
+
+		return c;
+	}
+
+
+	//---------------------------------------------------------------------------------------------
+	bool SecondPassGenerator::Generate( Ptr<mu::ErrorLog> InErrorLog, const Node* root)
+	{
+		MUTABLE_CPUPROFILER_SCOPE(SecondPassGenerate);
+
+		check(root);
+		ErrorLog = InErrorLog;
+
+		// Find the list of surfaces every tag depends on
+		SurfacesPerTag.Empty();
+		SurfacesPerTag.SetNum(FirstPass->Tags.Num());
+		TagsPerTag.Empty();
+		TagsPerTag.SetNum(FirstPass->Tags.Num());
+		for (size_t t = 0; t < FirstPass->Tags.Num(); ++t)
+		{
+			set<size_t> pendingSurfs;
+			for (size_t s : FirstPass->Tags[t].Surfaces)
+			{
+				pendingSurfs.insert(s);
+			}
+
+			set<size_t> processedSurfs;
+			while (!pendingSurfs.empty())
+			{
+				size_t cs = *pendingSurfs.begin();
+				pendingSurfs.erase(pendingSurfs.begin());
+
+				if (std::find(processedSurfs.begin(), processedSurfs.end(), cs)
+					!= processedSurfs.end())
+				{
+					continue;
+				}
+
+				processedSurfs.insert(cs);
+
+				SurfacesPerTag[t].insert(cs);
+
+				auto& csurf = FirstPass->Surfaces[cs];
+				for (const FString& sct : csurf.PositiveTags)
+				{
+					const FirstPassGenerator::FTag* it = FirstPass->Tags.FindByPredicate([&](const FirstPassGenerator::FTag& e) { return e.Tag == sct; });
+					if (!it)
+					{
+						// This could happen if a tag is in a variation but noone defines it.
+						continue;
+					}
+
+					size_t ct = it - &FirstPass->Tags[0];
+
+					TagsPerTag[t].insert(ct);
+
+					for (size_t s : FirstPass->Tags[ct].Surfaces)
+					{
+						if (SurfacesPerTag[t].find(s) != SurfacesPerTag[t].end())
+						{
+							pendingSurfs.insert(s);
+						}
+					}
+				}
+				for (const FString& sct : csurf.NegativeTags)
+				{
+					const FirstPassGenerator::FTag* it = FirstPass->Tags.FindByPredicate([&](const FirstPassGenerator::FTag& e) { return e.Tag == sct; });
+					if (!it)
+					{
+						// This could happen if a tag is in a variation but noone defines it.
+						continue;
+					}
+
+					size_t ct = it - &FirstPass->Tags[0];
+
+					TagsPerTag[t].insert(ct);
+
+					for (size_t s : FirstPass->Tags[ct].Surfaces)
+					{
+						if (SurfacesPerTag[t].find(s) != SurfacesPerTag[t].end())
+						{
+							pendingSurfs.insert(s);
+						}
+					}
+				}
+			}
+		}
+
+		// Create the conditions for every surface, modifier, component and individual tag.
+		TagConditionGenerationCache.clear();
+
+		set<size_t> Empty;
+
+		for (int32 SurfaceIndex = 0; SurfaceIndex < FirstPass->Surfaces.Num(); ++SurfaceIndex)
+		{
+			FirstPassGenerator::FSurface& Surface = FirstPass->Surfaces[SurfaceIndex];
+
+			{
+				Ptr<ASTOp> c = GenerateDataCodition(
+					SurfaceIndex, 
+					Surface.PositiveTags,
+					Surface.NegativeTags,
+					Empty, Empty, Empty, Empty);
+
+				Ptr<ASTOpFixed> ConditionOp = new ASTOpFixed();
+				ConditionOp->op.type = OP_TYPE::BO_AND;
+				ConditionOp->SetChild(ConditionOp->op.args.BoolBinary.a, Surface.ObjectCondition);
+				ConditionOp->SetChild(ConditionOp->op.args.BoolBinary.b, c);
+
+				Surface.FinalCondition = ConditionOp;
+			}
+			
+			// TODO: Convert to modifiers that enable tags?
+			//for (int32 EditIndex = 0; EditIndex < Surface.Edits.Num(); ++EditIndex)
+			//{
+			//	FirstPassGenerator::FSurface::FEdit& Edit = Surface.Edits[EditIndex];
+			//	
+			//	Ptr<ASTOp> c = GenerateDataCodition(
+			//		EditIndex, 
+			//		Surface.Edits[EditIndex].PositiveTags,
+			//		Surface.Edits[EditIndex].NegativeTags,
+			//		Empty, Empty, Empty, Empty);
+
+			//	Ptr<ASTOpFixed> OpAnd = new ASTOpFixed;
+			//	OpAnd->op.type = OP_TYPE::BO_AND;
+			//	OpAnd->SetChild(OpAnd->op.args.BoolBinary.a, Edit.Condition);
+			//	OpAnd->SetChild(OpAnd->op.args.BoolBinary.b, c);
+			//	c = m_opPool.Add(OpAnd);
+			//	
+			//	Edit.Condition = OpAnd;
+			//}
+		}
+
+		for (int32 ModifierIndex = 0; ModifierIndex < FirstPass->Modifiers.Num(); ++ModifierIndex)
+		{
+			Ptr<ASTOp> c = GenerateDataCodition(
+				ModifierIndex, 
+				FirstPass->Modifiers[ModifierIndex].PositiveTags,
+				FirstPass->Modifiers[ModifierIndex].NegativeTags,
+				Empty, Empty, Empty, Empty);
+			
+			Ptr<ASTOpFixed> ConditionOp = new ASTOpFixed();
+			ConditionOp->op.type = OP_TYPE::BO_AND;
+			ConditionOp->SetChild(ConditionOp->op.args.BoolBinary.a, FirstPass->Modifiers[ModifierIndex].ObjectCondition);
+			ConditionOp->SetChild(ConditionOp->op.args.BoolBinary.b, c);
+
+			FirstPass->Modifiers[ModifierIndex].FinalCondition = ConditionOp;
+		}
+
+		for (int32 ComponentIndex = 0; ComponentIndex < FirstPass->Components.Num(); ++ComponentIndex)
+		{
+			Ptr<ASTOp> c = GenerateDataCodition(
+				ComponentIndex, 
+				FirstPass->Components[ComponentIndex].PositiveTags,
+				FirstPass->Components[ComponentIndex].NegativeTags,
+				Empty, Empty, Empty, Empty);
+
+			FirstPass->Components[ComponentIndex].ComponentCondition = c;
+		}
+
+		// TODO: Do we really need the tag conditions from here on?
+		for (int32 s = 0; s < FirstPass->Tags.Num(); ++s)
+		{
+			Ptr<ASTOp> c = GenerateTagCondition(s, Empty, Empty, Empty, Empty);
+			FirstPass->Tags[s].GenericCondition = c;
+		}
+
+		FirstPass = nullptr;
+
+		return true;
+	}
+
+
+}
diff --git a/Engine/Plugins/Mutable/Source/MutableTools/Private/MuT/CodeGenerator_SecondPass.h b/Engine/Plugins/Mutable/Source/MutableTools/Private/MuT/CodeGenerator_SecondPass.h
new file mode 100644
--- /dev/null
+++ b/Engine/Plugins/Mutable/Source/MutableTools/Private/MuT/CodeGenerator_SecondPass.h
@@ -0,0 +1,100 @@
+// Copyright Epic Games, Inc. All Rights Reserved.
+
+#pragma once
+
+#include "HAL/PlatformCrt.h"
+#include "MuR/MutableMemory.h"
+#include "MuR/Ptr.h"
+#include "MuR/RefCounted.h"
+#include "MuT/AST.h"
+#include "MuT/CodeGenerator_FirstPass.h"
+#include "MuT/Compiler.h"
+#include "MuT/Node.h"
+#include "map"
+#include "set"
+
+
+namespace mu
+{
+
+	/** Second pass of the code generation process.
+     * It solves surface and modifier conditions from tags and variations
+	 */
+    class SecondPassGenerator
+	{
+	public:
+
+		SecondPassGenerator( FirstPassGenerator*, const CompilerOptions::Private*  );
+
+		// Return true on success.
+        bool Generate( Ptr<ErrorLog>, const Node* Root );
+
+	private:
+
+        FirstPassGenerator* FirstPass = nullptr;
+        const CompilerOptions::Private *CompilerOptions = nullptr;
+
+
+        //!
+        Ptr<ErrorLog> ErrorLog;
+
+        //!
+        struct FConditionGenerationKey
+        {
+            size_t tagOrSurfIndex = 0;
+            set<size_t> posSurf;
+            set<size_t> negSurf;
+            set<size_t> posTag;
+            set<size_t> negTag;
+
+            inline bool operator<(const FConditionGenerationKey& o) const
+            {
+                if (tagOrSurfIndex<o.tagOrSurfIndex) return true;
+                if (tagOrSurfIndex>o.tagOrSurfIndex) return false;
+                if (posSurf<o.posSurf) return true;
+                if (posSurf>o.posSurf) return false;
+                if (negSurf<o.negSurf) return true;
+                if (negSurf>o.negSurf) return false;
+                if (posTag<o.posTag) return true;
+                if (posTag>o.posTag) return false;
+                if (negTag<o.negTag) return true;
+                if (negTag>o.negTag) return false;
+                return false;
+            }
+        };
+
+        // List of surfaces that activate or deactivate every tag, or another surface that activates a tag in this set.
+		// \TODO: Change to UE containers
+		TArray< set<size_t> > SurfacesPerTag;
+		TArray< set<size_t> > TagsPerTag;
+
+        std::map<FConditionGenerationKey,Ptr<ASTOp>> TagConditionGenerationCache;
+
+        UniqueOpPool m_opPool;
+
+        //!
+        Ptr<ASTOp> GenerateTagCondition( size_t tagIndex,
+                                         const set<size_t>& posSurf,
+                                         const set<size_t>& negSurf,
+                                         const set<size_t>& posTag,
+                                         const set<size_t>& negTag );
+
+        /** Generate Surface, Edit or Modifier condition.
+    	 * @param Index Surface, Edit, Component or Modifier index.
+    	 * @param PositiveTags function that given the Surface, Edit or Modifier index, returns its positive tags.
+    	 * @param NegativeTags function that given the Surface, Edit or Modifier, returns its negative tags.
+    	 * @param posSurf already visited Surfaces, Edits, or Modifiers that participate positively in the condition.
+    	 * @param negSurf already visited Surfaces, Edits, or Modifiers that participate negatively in the condition.
+      	 * @param posSurf Tags that already belong to the condition (positively).
+		 * @param negSurf Tags that already belong to the condition (negatively). */
+        Ptr<ASTOp> GenerateDataCodition(size_t Index,
+									        const TArray<FString>& PositiveTags,
+											const TArray<FString>& NegativeTags,
+                                            const set<size_t>& posSurf,
+                                            const set<size_t>& negSurf,
+                                            const set<size_t>& posTag,
+                                            const set<size_t>& negTag);
+    };
+
+}
+
