[mutable] Moved the Mutable plugin out of Experimental status into Beta.

diff --git a/Engine/Plugins/Mutable/Source/MutableTools/Private/MuT/CodeGenerator_String.cpp b/Engine/Plugins/Mutable/Source/MutableTools/Private/MuT/CodeGenerator_String.cpp
new file mode 100644
--- /dev/null
+++ b/Engine/Plugins/Mutable/Source/MutableTools/Private/MuT/CodeGenerator_String.cpp
@@ -0,0 +1,118 @@
+// Copyright Epic Games, Inc. All Rights Reserved.
+
+#include "HAL/PlatformCrt.h"
+#include "HAL/UnrealMemory.h"
+#include "Misc/AssertionMacros.h"
+#include "MuR/Operations.h"
+#include "MuR/Parameters.h"
+#include "MuR/ParametersPrivate.h"
+#include "MuR/Ptr.h"
+#include "MuR/RefCounted.h"
+#include "MuR/Types.h"
+#include "MuT/AST.h"
+#include "MuT/ASTOpConstantString.h"
+#include "MuT/ASTOpParameter.h"
+#include "MuT/CodeGenerator.h"
+#include "MuT/NodeRange.h"
+#include "MuT/NodeString.h"
+#include "MuT/NodeStringConstant.h"
+#include "MuT/NodeStringConstantPrivate.h"
+#include "MuT/NodeStringParameter.h"
+#include "MuT/NodeStringParameterPrivate.h"
+
+namespace mu
+{
+
+
+	//-------------------------------------------------------------------------------------------------
+	void CodeGenerator::GenerateString(FStringGenerationResult& result, const FGenericGenerationOptions& Options, const Ptr<const NodeString>& Untyped)
+	{
+		if (!Untyped)
+		{
+			result = FStringGenerationResult();
+			return;
+		}
+
+		// See if it was already generated
+		FGeneratedCacheKey Key;
+		Key.Node = Untyped;
+		Key.Options = Options;
+		FGeneratedStringsMap::ValueType* it = GeneratedStrings.Find(Key);
+		if (it)
+		{
+			result = *it;
+			return;
+		}
+
+		// Generate for each different type of node
+		if (Untyped->GetType()==NodeStringConstant::GetStaticType())
+		{
+			GenerateString_Constant(result, Options, static_cast<const NodeStringConstant*>(Untyped.get()));
+		}
+		else if (Untyped->GetType() == NodeStringParameter::GetStaticType())
+		{
+			GenerateString_Parameter(result, Options, static_cast<const NodeStringParameter*>(Untyped.get()));
+		}
+		else
+		{
+			check(false);
+		}
+
+		// Cache the result
+		GeneratedStrings.Add(Key, result);
+	}
+
+
+	//-------------------------------------------------------------------------------------------------
+	void CodeGenerator::GenerateString_Constant(FStringGenerationResult& result, const FGenericGenerationOptions& Options, const Ptr<const NodeStringConstant>& Typed)
+	{
+		const NodeStringConstant::Private& node = *Typed->GetPrivate();
+
+		Ptr<ASTOpConstantString> op = new ASTOpConstantString();
+		op->value = node.m_value;
+
+		result.op = op;
+	}
+
+
+	//-------------------------------------------------------------------------------------------------
+	void CodeGenerator::GenerateString_Parameter(FStringGenerationResult& result, const FGenericGenerationOptions& Options, const Ptr<const NodeStringParameter>& Typed)
+	{
+		const NodeStringParameter::Private& node = *Typed->GetPrivate();
+
+		Ptr<ASTOpParameter> op;
+
+		Ptr<ASTOpParameter>* it = FirstPass.ParameterNodes.Find(node.m_pNode);
+		if (!it)
+		{
+			FParameterDesc param;
+			param.m_name = node.m_name;
+			const TCHAR* CStr = ToCStr(node.m_uid);
+			param.m_uid.ImportTextItem(CStr, 0, nullptr, nullptr);
+			param.m_type = PARAMETER_TYPE::T_FLOAT;
+			param.m_defaultValue.Set<ParamStringType>(node.m_defaultValue);
+
+			op = new ASTOpParameter();
+			op->type = OP_TYPE::ST_PARAMETER;
+			op->parameter = param;
+
+			// Generate the code for the ranges
+			for (int32 a = 0; a < node.m_ranges.Num(); ++a)
+			{
+				FRangeGenerationResult rangeResult;
+				GenerateRange(rangeResult, Options, node.m_ranges[a]);
+				op->ranges.Emplace(op.get(), rangeResult.sizeOp, rangeResult.rangeName, rangeResult.rangeUID);
+			}
+
+			FirstPass.ParameterNodes.Add(node.m_pNode, op);
+		}
+		else
+		{
+			op = *it;
+		}
+
+		result.op = op;
+	}
+
+
+}
diff --git a/Engine/Plugins/Mutable/Source/MutableTools/Private/MuT/CodeOptimiser.cpp b/Engine/Plugins/Mutable/Source/MutableTools/Private/MuT/CodeOptimiser.cpp
new file mode 100644
--- /dev/null
+++ b/Engine/Plugins/Mutable/Source/MutableTools/Private/MuT/CodeOptimiser.cpp
@@ -0,0 +1,1571 @@
+// Copyright Epic Games, Inc. All Rights Reserved.
+
+#include "MuT/CodeOptimiser.h"
+
+#include "MuT/ErrorLogPrivate.h"
+#include "MuT/AST.h"
+#include "MuT/ASTOpInstanceAdd.h"
+#include "MuT/ASTOpConditional.h"
+#include "MuT/ASTOpSwitch.h"
+#include "MuT/ASTOpConstantResource.h"
+#include "MuT/ASTOpConstantBool.h"
+#include "MuT/ASTOpMeshApplyShape.h"
+#include "MuT/ASTOpMeshBindShape.h"
+#include "MuT/ASTOpMeshClipMorphPlane.h"
+#include "MuT/ASTOpMeshApplyPose.h"
+#include "MuT/ASTOpImageRasterMesh.h"
+#include "MuT/ASTOpReferenceResource.h"
+
+#include "MuR/ModelPrivate.h"
+#include "MuR/SystemPrivate.h"
+#include "MuR/Operations.h"
+#include "MuR/OpMeshMerge.h"
+#include "MuR/MutableRuntimeModule.h"
+
+#include "Tasks/Task.h"
+
+#include <unordered_set>
+
+
+namespace mu
+{
+
+	namespace
+	{
+
+		struct FMeshEntry
+		{
+			Ptr<const Mesh> Mesh;
+			Ptr<ASTOpConstantResource> Op;
+
+			bool operator==(const FMeshEntry& o) const
+			{
+				return Mesh == o.Mesh || *Mesh ==*o.Mesh;
+			}
+		};
+
+		struct FImageEntry
+		{
+			Ptr< const Image> Image;
+			Ptr<ASTOpConstantResource> Op;
+
+			bool operator==(const FImageEntry& o) const
+			{
+				return Image == o.Image || *Image == *o.Image;
+			}
+		};
+
+		struct FLayoutEntry
+		{
+			Ptr< const Layout> Layout;
+			Ptr<ASTOpConstantResource> Op;
+
+			bool operator==(const FLayoutEntry& o) const
+			{
+				return Layout == o.Layout || *Layout == *o.Layout;
+			}
+		};
+
+		struct custom_mesh_equal
+		{
+			bool operator()( const MeshPtrConst& a, const MeshPtrConst& b ) const
+			{
+				return a==b || *a==*b;
+			}
+		};
+
+		struct custom_image_equal
+		{
+			bool operator()( const ImagePtrConst& a, const ImagePtrConst& b ) const
+			{
+				return a==b || *a==*b;
+			}
+		};
+
+		struct custom_layout_equal
+		{
+			bool operator()( const Ptr<const Layout>& a, const Ptr<const Layout>& b ) const
+			{
+				return a==b || *a==*b;
+			}
+		};
+	}
+
+
+	//-------------------------------------------------------------------------------------------------
+	void DuplicatedDataRemoverAST( ASTOpList& roots )
+	{
+		MUTABLE_CPUPROFILER_SCOPE(DuplicatedDataRemoverAST);
+
+		TArray<Ptr<ASTOpConstantResource>> AllMeshOps;
+		TArray<Ptr<ASTOpConstantResource>> AllImageOps;
+		TArray<Ptr<ASTOpConstantResource>> AllLayoutOps;
+
+		// Gather constants
+		{
+			MUTABLE_CPUPROFILER_SCOPE(Gather);
+
+			ASTOp::Traverse_TopRandom_Unique_NonReentrant( roots, [&](Ptr<ASTOp> n)
+			{
+				switch ( n->GetOpType() )
+				{
+
+				case OP_TYPE::ME_CONSTANT:
+				{
+					ASTOpConstantResource* typedNode = static_cast<ASTOpConstantResource*>(n.get());
+					AllMeshOps.Add(typedNode);
+					break;
+				}
+
+				case OP_TYPE::IM_CONSTANT:
+				{
+					ASTOpConstantResource* typedNode = static_cast<ASTOpConstantResource*>(n.get());
+					AllImageOps.Add(typedNode);
+					break;
+				}
+
+				case OP_TYPE::LA_CONSTANT:
+				{
+					ASTOpConstantResource* typedNode = static_cast<ASTOpConstantResource*>(n.get());
+					AllLayoutOps.Add(typedNode);
+					break;
+				}
+
+				//    These should be part of the duplicated code removal, in AST.
+				//            // Names
+				//            case OP_TYPE::IN_ADDMESH:
+				//            case OP_TYPE::IN_ADDIMAGE:
+				//            case OP_TYPE::IN_ADDVECTOR:
+				//            case OP_TYPE::IN_ADDSCALAR:
+				//            case OP_TYPE::IN_ADDCOMPONENT:
+				//            case OP_TYPE::IN_ADDSURFACE:
+
+				default:
+					break;
+				}
+
+				return true;
+			});
+		}
+
+
+		// Compare meshes
+		{
+			MUTABLE_CPUPROFILER_SCOPE(CompareMeshes);
+
+			TMultiMap< SIZE_T, FMeshEntry > Meshes;
+
+			for (Ptr<ASTOpConstantResource>& typedNode : AllMeshOps)
+			{
+				SIZE_T Key = typedNode->GetValueHash();
+
+				Ptr<ASTOp> Found;
+
+				TArray<FMeshEntry*, TInlineAllocator<4>> Candidates;
+				Meshes.MultiFindPointer(Key, Candidates, false);
+
+				if (!Candidates.IsEmpty())
+				{
+					Ptr<const Mesh> mesh = static_cast<const Mesh*>(typedNode->GetValue().get());
+
+					for (FMeshEntry* It : Candidates)
+					{
+						if (!It->Mesh)
+						{
+							It->Mesh = static_cast<const Mesh*>(It->Op->GetValue().get());
+						}
+
+						if (custom_mesh_equal()(mesh, It->Mesh))
+						{
+							Found = It->Op;
+							break;
+						}
+					}
+				}
+
+				if (Found)
+				{
+					ASTOp::Replace(typedNode, Found);
+				}
+				else
+				{
+					// The mesh will be loaded only if it needs to be compared
+					FMeshEntry e;
+					e.Op = typedNode;
+					Meshes.Add(Key, e);
+				}
+			}
+		}
+
+		// Compare images
+		{
+			MUTABLE_CPUPROFILER_SCOPE(CompareImages);
+
+			TMultiMap< SIZE_T, FImageEntry > Images;
+
+			for (Ptr<ASTOpConstantResource>& typedNode : AllImageOps)
+			{
+				SIZE_T Key = typedNode->GetValueHash();
+
+				Ptr<ASTOp> Found;
+
+				TArray<FImageEntry*,TInlineAllocator<4>> Candidates;
+				Images.MultiFindPointer(Key, Candidates, false);
+				
+				if (!Candidates.IsEmpty())
+				{
+					Ptr<const Image> image = static_cast<const Image*>(typedNode->GetValue().get());
+
+					for (FImageEntry* It: Candidates)
+					{
+						if (!It->Image)
+						{
+							It->Image = static_cast<const Image*>(It->Op->GetValue().get());
+						}
+
+						if (custom_image_equal()(image, It->Image))
+						{
+							Found = It->Op;
+							break;
+						}
+					}
+				}
+
+				if (Found)
+				{
+					ASTOp::Replace(typedNode, Found);
+				}
+				else
+				{
+					// The image will be loaded only if it needs to be compared
+					FImageEntry e;
+					e.Op = typedNode;
+					Images.Add(Key, e);
+				}
+			}
+		}
+
+		// Compare layouts
+		{
+			MUTABLE_CPUPROFILER_SCOPE(CompareLayouts);
+
+			TMultiMap< SIZE_T, FLayoutEntry > Layouts;
+
+			for (Ptr<ASTOpConstantResource>& typedNode : AllLayoutOps)
+			{
+				SIZE_T Key = typedNode->GetValueHash();
+
+				Ptr<ASTOp> Found;
+
+				TArray<FLayoutEntry*, TInlineAllocator<4>> Candidates;
+				Layouts.MultiFindPointer(Key, Candidates, false);
+
+				if (!Candidates.IsEmpty())
+				{
+					Ptr<const Layout> layout = static_cast<const Layout*>(typedNode->GetValue().get());
+
+					for (FLayoutEntry* It : Candidates)
+					{
+						if (!It->Layout)
+						{
+							It->Layout = static_cast<const Layout*>(It->Op->GetValue().get());
+						}
+
+						if (custom_layout_equal()(layout, It->Layout))
+						{
+							Found = It->Op;
+							break;
+						}
+					}
+				}
+
+				if (Found)
+				{
+					ASTOp::Replace(typedNode, Found);
+				}
+				else
+				{
+					FLayoutEntry e;
+					e.Op = typedNode;
+					Layouts.Add(Key, e);
+				}
+			}
+		}
+
+	}
+
+
+	//-------------------------------------------------------------------------------------------------
+	//-------------------------------------------------------------------------------------------------
+	//-------------------------------------------------------------------------------------------------
+	namespace
+	{
+		struct op_pointer_hash
+		{
+			inline size_t operator() (const mu::Ptr<ASTOp>& n) const
+			{
+				return n->Hash();
+			}
+		};
+
+		struct op_pointer_equal
+		{
+			inline bool operator() (const mu::Ptr<ASTOp>& a, const mu::Ptr<ASTOp>& b) const
+			{
+				return *a==*b;
+			}
+		};
+	}
+
+	void DuplicatedCodeRemoverAST( ASTOpList& roots )
+	{
+		MUTABLE_CPUPROFILER_SCOPE(DuplicatedCodeRemoverAST);
+
+		// Visited nodes, per type
+		std::unordered_set<Ptr<ASTOp>,op_pointer_hash,op_pointer_equal> visited[int32(OP_TYPE::COUNT)];
+
+		ASTOp::Traverse_BottomUp_Unique_NonReentrant( roots, [&](Ptr<ASTOp>& n)
+		{
+			auto& container = visited[(int32)n->GetOpType()];
+
+			// Insert will tell us if it was already there
+			auto it = container.insert(n);
+			if( !it.second )
+			{
+				// It wasn't inserted, so it was already there
+				ASTOp::Replace(n,*it.first);
+			}
+		});
+	}
+
+
+	//---------------------------------------------------------------------------------------------
+	//---------------------------------------------------------------------------------------------
+	//---------------------------------------------------------------------------------------------
+	bool FindOpTypeVisitor::Apply( FProgram& program, OP::ADDRESS rootAt )
+	{
+		bool found = false;
+
+		// If the program added new operations, we haven't visited them.
+		m_visited.SetNumZeroed( program.m_opAddress.Num() );
+
+		m_pending.Reset();
+		m_pending.Reserve( program.m_opAddress.Num()/4 );
+		m_pending.Add({ false,rootAt });
+
+		// Don't early out to be able to complete parent op cached flags
+		while ( !m_pending.IsEmpty() )
+		{
+			TPair<bool,int32> item = m_pending.Pop();
+			OP::ADDRESS at = item.Value;
+
+			// Not cached?
+			if (m_visited[at]<=1)
+			{
+				if (item.Key)
+				{
+					// Item indicating we finished with all the children of a parent
+					check(m_visited[at]==1);
+
+					bool subtreeFound = false;
+					ForEachReference( program, at, [&](OP::ADDRESS ref)
+					{
+						subtreeFound = subtreeFound || m_visited[ref]==3;
+					});
+
+					m_visited[at] = subtreeFound ? 3 : 2;
+				}
+				else if (!found)
+				{
+					if (m_typesToFind.Find(program.GetOpType(at)) != INDEX_NONE )
+					{
+						m_visited[at] = 3; // visited, ops found
+						found = true;
+					}
+					else
+					{
+						check(m_visited[at]==0);
+
+						m_visited[at] = 1;
+						m_pending.Add({ true,at });
+
+						ForEachReference( program, at, [&](OP::ADDRESS ref)
+						{
+							if (ref && m_visited[ref]==0)
+							{
+								m_pending.Add({ false,ref });
+							}
+						});
+					}
+				}
+				else
+				{
+					// We won't process it.
+					m_visited[at] = 0;
+				}
+			}
+		}
+
+		return m_visited[rootAt]==3;
+	}
+
+
+	//---------------------------------------------------------------------------------------------
+	//---------------------------------------------------------------------------------------------
+	//---------------------------------------------------------------------------------------------
+	IsConstantVisitor::IsConstantVisitor()
+	{
+		TArray<OP_TYPE> parameterTypes;
+		parameterTypes.Add( OP_TYPE::BO_PARAMETER );
+		parameterTypes.Add( OP_TYPE::NU_PARAMETER );
+		parameterTypes.Add( OP_TYPE::SC_PARAMETER );
+		parameterTypes.Add( OP_TYPE::CO_PARAMETER );
+		parameterTypes.Add( OP_TYPE::PR_PARAMETER );
+		parameterTypes.Add( OP_TYPE::IM_PARAMETER );
+		m_findOpTypeVisitor = MakeUnique<FindOpTypeVisitor>(parameterTypes);
+	}
+
+
+	//---------------------------------------------------------------------------------------------
+	bool IsConstantVisitor::Apply( FProgram& program, OP::ADDRESS at )
+	{
+		return !m_findOpTypeVisitor->Apply(program,at);
+	}
+
+
+	//---------------------------------------------------------------------------------------------
+	//---------------------------------------------------------------------------------------------
+	//---------------------------------------------------------------------------------------------
+	class FConstantTask 
+	{
+	public:
+
+		// input
+		Ptr<ASTOp> Source;
+		FProxyFileContext* DiskCacheContext = nullptr;
+		int32 ImageCompressionQuality = 0;
+		int32 OptimizationPass = 0;
+		FReferencedResourceFunc ReferencedResourceProvider;
+
+		// Intermediate
+		Ptr<ASTOp> SourceCloned;
+
+		// Result
+		Ptr<ASTOp> Result;
+
+	public:
+
+		FConstantTask( const Ptr<ASTOp>& InSource, const CompilerOptions::Private* InOptions, int32 InOptimizationPass )
+		{
+			OptimizationPass = InOptimizationPass;
+			Source = InSource;
+			DiskCacheContext = InOptions->OptimisationOptions.DiskCacheContext;
+			ImageCompressionQuality = InOptions->ImageCompressionQuality;
+			ReferencedResourceProvider = InOptions->OptimisationOptions.ReferencedResourceProvider;
+		}
+
+		void Run(FImageOperator ImOp)
+		{
+			MUTABLE_CPUPROFILER_SCOPE(ConstantTask_Run);
+
+			// This runs in a worker thread
+
+			OP_TYPE type = SourceCloned->GetOpType();
+			DATATYPE DataType = GetOpDataType(type);
+
+			Ptr<Settings> pSettings = new Settings;
+			pSettings->SetProfile( false );
+			pSettings->SetImageCompressionQuality( ImageCompressionQuality );
+			SystemPtr pSystem = new System( pSettings );
+
+			pSystem->GetPrivate()->ImagePixelFormatOverride = ImOp.FormatImageOverride;
+
+			FSourceDataDescriptor SourceDataDescriptor;
+			if (DataType == DT_IMAGE)
+			{
+				SourceDataDescriptor = SourceCloned->GetSourceDataDescriptor();
+				check(!SourceDataDescriptor.IsInvalid());
+			}
+
+			// Don't generate mips during linking here.
+			FLinkerOptions LinkerOptions(ImOp);
+			LinkerOptions.MinTextureResidentMipCount = 255;
+			LinkerOptions.bSeparateImageMips = false;
+
+			TSharedPtr<const Model> model = MakeShared<Model>();
+			OP::ADDRESS at = ASTOp::FullLink(SourceCloned, model->GetPrivate()->m_program, &LinkerOptions);
+
+			FProgram::FState state;
+			state.m_root = at;
+			model->GetPrivate()->m_program.m_states.Add(state);
+
+			Ptr<Parameters> LocalParams = Model::NewParameters(model);
+			pSystem->GetPrivate()->BeginBuild( model );
+
+			// Calculate the value and replace this op by a constant
+			switch( DataType )
+			{
+			case DT_MESH:
+			{
+				MUTABLE_CPUPROFILER_SCOPE(ConstantMesh);
+
+				mu::Ptr<const Mesh> pMesh = pSystem->GetPrivate()->BuildMesh( model, LocalParams.get(), at );
+
+				if (pMesh)
+				{
+					mu::Ptr<ASTOpConstantResource> constantOp = new ASTOpConstantResource();
+					constantOp->Type = OP_TYPE::ME_CONSTANT;
+					constantOp->SetValue( pMesh, DiskCacheContext );
+					Result = constantOp;
+				  }
+				break;
+			}
+
+			case DT_IMAGE:
+			{
+				MUTABLE_CPUPROFILER_SCOPE(ConstantImage);
+
+				mu::Ptr<const Image> pImage = pSystem->GetPrivate()->BuildImage( model, LocalParams.get(), at, 0, 0 );
+
+				if (pImage)
+				{
+					mu::Ptr<ASTOpConstantResource> ConstantOp = new ASTOpConstantResource();
+					ConstantOp->SourceDataDescriptor = SourceDataDescriptor;
+					ConstantOp->Type = OP_TYPE::IM_CONSTANT;
+					ConstantOp->SetValue( pImage, DiskCacheContext );
+					Result = ConstantOp;
+				}
+				break;
+			}
+
+			case DT_LAYOUT:
+			{
+				MUTABLE_CPUPROFILER_SCOPE(ConstantLayout);
+
+				mu::Ptr<const Layout> pLayout = pSystem->GetPrivate()->BuildLayout( model, LocalParams.get(), at );
+
+				if (pLayout)
+				{
+					mu::Ptr<ASTOpConstantResource> constantOp = new ASTOpConstantResource();
+					constantOp->Type = OP_TYPE::LA_CONSTANT;
+					constantOp->SetValue( pLayout, DiskCacheContext);
+					Result = constantOp;
+				}
+				break;
+			}
+
+			case DT_BOOL:
+			{
+				MUTABLE_CPUPROFILER_SCOPE(ConstantBool);
+
+				bool value = pSystem->GetPrivate()->BuildBool( model, LocalParams.get(), at );
+
+				{
+					mu::Ptr<ASTOpConstantBool> constantOp = new ASTOpConstantBool();
+					constantOp->value = value;
+					Result = constantOp;
+				}
+				break;
+			}
+
+			case DT_COLOUR:
+			{
+				MUTABLE_CPUPROFILER_SCOPE(ConstantBool);
+
+				FVector4f ResultColor(0, 0, 0, 0);
+				ResultColor = pSystem->GetPrivate()->BuildColour( model, LocalParams.get(), at );
+
+				{
+					mu::Ptr<ASTOpFixed> constantOp = new ASTOpFixed();
+					constantOp->op.type = OP_TYPE::CO_CONSTANT;
+					constantOp->op.args.ColourConstant.value[0] = ResultColor[0];
+					constantOp->op.args.ColourConstant.value[1] = ResultColor[1];
+					constantOp->op.args.ColourConstant.value[2] = ResultColor[2];
+					constantOp->op.args.ColourConstant.value[3] = ResultColor[3];
+					Result = constantOp;
+				}
+				break;
+			}
+
+			case DT_INT:
+			case DT_SCALAR:
+			case DT_STRING:
+			case DT_PROJECTOR:
+				// TODO
+				break;
+
+			default:
+				break;
+			}
+
+			pSystem->GetPrivate()->EndBuild();
+		}
+
+	};
+
+
+	//---------------------------------------------------------------------------------------------
+	bool ConstantGeneratorAST( const CompilerOptions::Private* InOptions, Ptr<ASTOp>& Root, int32 Pass )
+	{
+		MUTABLE_CPUPROFILER_SCOPE(ConstantGenerator);
+
+		// don't do this if constant optimization has been disabled, usually for debugging.
+		if (!InOptions->OptimisationOptions.bConstReduction)
+		{
+			return false;
+		}
+
+		// Gather the roots of all constant operations
+		struct FConstantSubgraph
+		{
+			Ptr<ASTOp> Root;
+			UE::Tasks::FTaskEvent CompletedEvent;
+
+			// This is only necessary for non-concurrent execution
+			TArray< UE::Tasks::FTask, TInlineAllocator<8> > Requisites;
+			TFunction<void()> NonConcurrentTask;
+			TUniquePtr<FConstantTask> TaskData;
+		};
+		TArray< FConstantSubgraph > ConstantSubgraphs;
+		ConstantSubgraphs.Reserve(256);
+		{
+			MUTABLE_CPUPROFILER_SCOPE(ConstantGenerator_GenerateTasks);
+
+			ASTOp::Traverse_BottomUp_Unique(Root,
+				[&ConstantSubgraphs, Pass]
+				(Ptr<ASTOp>& SubgraphRoot)
+				{
+					OP_TYPE SubgraphType = SubgraphRoot->GetOpType();
+
+					bool bIsConstantSubgraph = true;
+					switch (SubgraphType)
+					{
+					case OP_TYPE::BO_PARAMETER:
+					case OP_TYPE::NU_PARAMETER:
+					case OP_TYPE::SC_PARAMETER:
+					case OP_TYPE::CO_PARAMETER:
+					case OP_TYPE::PR_PARAMETER:
+					case OP_TYPE::IM_PARAMETER:
+						bIsConstantSubgraph = false;
+						break;
+					default:
+						// Propagate from children
+						SubgraphRoot->ForEachChild([&bIsConstantSubgraph](ASTChild& c)
+							{
+								if (c)
+								{
+									bIsConstantSubgraph = bIsConstantSubgraph && c->bIsConstantSubgraph;
+								}
+							});
+						break;
+					}
+					SubgraphRoot->bIsConstantSubgraph = bIsConstantSubgraph;
+
+					// We avoid generating constants for these operations, to avoid the memory explosion.
+					// TODO: Make compiler options for some of them
+					// TODO: Some of them are worth if the code below them is unique.
+					bool bHasSpecialOpInSubgraph = false;
+					switch (SubgraphType)
+					{
+					case OP_TYPE::IM_BLANKLAYOUT:
+					case OP_TYPE::IM_COMPOSE:
+					case OP_TYPE::ME_MERGE:
+					case OP_TYPE::ME_CLIPWITHMESH:
+					case OP_TYPE::ME_CLIPMORPHPLANE:
+					case OP_TYPE::ME_APPLYPOSE:
+					case OP_TYPE::ME_REMOVEMASK:
+					case OP_TYPE::ME_ADDTAGS:
+					case OP_TYPE::IM_PLAINCOLOUR:
+						bHasSpecialOpInSubgraph = true;
+						break;
+
+					case OP_TYPE::IM_RASTERMESH:
+					{
+						const ASTOpImageRasterMesh* Raster = static_cast<const ASTOpImageRasterMesh*>(SubgraphRoot.get());
+						// If this operation is only rastering the mesh UVs, reduce it to constant. Otherwise avoid reducing it
+						// for the case of a constant projector of a large set of possible images. We don't want to generate all the
+						// projected version of the images beforehand. TODO: Make it a compile-time option?
+						bHasSpecialOpInSubgraph = Raster->image.child().get() != nullptr;
+						break;
+					}
+
+					case OP_TYPE::ME_REFERENCE:
+					case OP_TYPE::IM_REFERENCE:
+						// If we are in a reference-resolution optimization phase, then the ops are not special.
+						if (Pass < 2)
+						{
+							bHasSpecialOpInSubgraph = true;
+						}
+						else
+						{
+							const ASTOpReferenceResource* Typed = static_cast<const ASTOpReferenceResource*>(SubgraphRoot.get());
+							bHasSpecialOpInSubgraph = !Typed->bForceLoad;
+						}
+						break;
+
+					default:
+						// Propagate from children
+						SubgraphRoot->ForEachChild([&](ASTChild& c)
+							{
+								if (c)
+								{
+									bHasSpecialOpInSubgraph = bHasSpecialOpInSubgraph || c->bHasSpecialOpInSubgraph;
+								}
+							});
+						break;
+					}
+					SubgraphRoot->bHasSpecialOpInSubgraph = bHasSpecialOpInSubgraph;
+
+					bool bIsDataTypeThanCanTurnIntoConst = false;
+					DATATYPE DataType = GetOpDataType(SubgraphType);
+					switch (DataType)
+					{
+					case DT_MESH:
+					case DT_IMAGE:
+					case DT_LAYOUT:
+					case DT_BOOL:
+					case DT_COLOUR:
+						bIsDataTypeThanCanTurnIntoConst = true;
+						break;
+					default:
+						break;
+					}
+
+
+					// See if it is worth generating this as constant
+					// ---------------------------------------------
+					bool bWorthGenerating = SubgraphRoot->bIsConstantSubgraph
+						&& !SubgraphRoot->bHasSpecialOpInSubgraph
+						&& !SubgraphRoot->IsConstantOp()
+						&& bIsDataTypeThanCanTurnIntoConst;
+
+					if (bWorthGenerating)						 
+					{
+						bool bCanBeGenerated = true;
+
+						// Check source data incompatiblities: when generating constants don't mix data that has different source descriptors (tags and other properties).
+						if (DataType == DT_IMAGE)
+						{
+							FSourceDataDescriptor SourceDescriptor = SubgraphRoot->GetSourceDataDescriptor();
+							if (SourceDescriptor.IsInvalid())
+							{
+								bCanBeGenerated = false;
+							}
+						}
+
+						if (bCanBeGenerated)
+						{
+							ConstantSubgraphs.Add({ SubgraphRoot, UE::Tasks::FTaskEvent(TEXT("MutableConstantSubgraph")) });
+						}
+					}
+				});
+		}
+
+		auto GetRequisites = [&ConstantSubgraphs](const Ptr<ASTOp>& SubgraphRoot, TArray< UE::Tasks::FTask, TInlineAllocator<8> >& OutRequisites)
+		{
+			MUTABLE_CPUPROFILER_SCOPE(ConstantGenerator_GetRequisites);
+
+			TArray< Ptr<ASTOp> > ScanRoots;
+			ScanRoots.Add(SubgraphRoot);
+			ASTOp::Traverse_TopDown_Unique_Imprecise(ScanRoots, [&SubgraphRoot, &OutRequisites, &ConstantSubgraphs](Ptr<ASTOp>& ChildNode)
+				{
+					bool bRecurse = true;
+
+					// Subgraph root?
+					if (SubgraphRoot == ChildNode)
+					{
+						return bRecurse;
+					}
+
+					FConstantSubgraph* DependencyFound = ConstantSubgraphs.FindByPredicate([&ChildNode](const FConstantSubgraph& Candidate) { return Candidate.Root == ChildNode; });
+					if (DependencyFound)
+					{
+						bRecurse = false;
+						OutRequisites.Add(DependencyFound->CompletedEvent);
+					}
+
+					return bRecurse;
+				});
+		};
+
+		bool bUseConcurrency = InOptions->bUseConcurrency;
+
+		if (bUseConcurrency)
+		{
+			// Launch the tasks.
+			UE::Tasks::FTask LaunchTask = UE::Tasks::Launch(TEXT("ConstantGeneratorLaunchTasks"), 
+				[&ConstantSubgraphs, &GetRequisites, Pass, InOptions]()
+				{
+					MUTABLE_CPUPROFILER_SCOPE(ConstantGenerator_LaunchTasks);
+
+					FImageOperator ImOp = FImageOperator::GetDefault(InOptions->ImageFormatFunc);
+
+					// Traverse list of constants to generate. It is ordered in a bottom-up way.
+					int32 SubgraphCount = ConstantSubgraphs.Num();
+					for (int32 OrderIndex = 0; OrderIndex < SubgraphCount; ++OrderIndex)
+					{
+						int32 Index = SubgraphCount - 1 - OrderIndex;
+
+						Ptr<ASTOp> SubgraphRoot = ConstantSubgraphs[Index].Root;
+						UE::Tasks::FTaskEvent& SubgraphCompletionEvent = ConstantSubgraphs[Index].CompletedEvent;
+
+						// Referenced images are resolved in its own task to prevent requesting them twice (and loading them twice)
+						bool bIsCompileTimeReferenceImage = (SubgraphRoot->GetOpType() == OP_TYPE::IM_REFERENCE);
+
+						// Launch the task with its dependencies
+						if (bIsCompileTimeReferenceImage)
+						{
+							// Instead of generating the constant we resolve the reference, which also replaces the ASTOp.
+							const ASTOpReferenceResource* Typed = static_cast<const ASTOpReferenceResource*>(SubgraphRoot.get());
+							uint32 ImageID = Typed->ID;
+
+							TSharedPtr< Ptr<Image> > ResolveImage = MakeShared<Ptr<Image>>();
+
+							constexpr bool bRunImmediatlyIfPossible = false;
+							UE::Tasks::FTask ReferenceCompletion = InOptions->OptimisationOptions.ReferencedResourceProvider(ImageID, ResolveImage, bRunImmediatlyIfPossible);
+
+							UE::Tasks::FTask CompleteTask = UE::Tasks::Launch(TEXT("MutableResolveComplete"),
+								[SubgraphRoot, InOptions, ResolveImage]()
+								{
+									Ptr<ASTOpConstantResource> ConstantOp;
+									{
+										MUTABLE_CPUPROFILER_SCOPE(MutableResolveComplete_CreateConstant);
+										ConstantOp = new ASTOpConstantResource;
+										ConstantOp->Type = OP_TYPE::IM_CONSTANT;
+										{
+											MUTABLE_CPUPROFILER_SCOPE(GetSourceDataDescriptor);
+											ConstantOp->SourceDataDescriptor = SubgraphRoot->GetSourceDataDescriptor();
+										}
+										ConstantOp->SetValue(ResolveImage->get(), InOptions->OptimisationOptions.DiskCacheContext);
+									}
+									{
+										MUTABLE_CPUPROFILER_SCOPE(MutableResolveComplete_Replace);
+										ASTOp::Replace(SubgraphRoot, ConstantOp);
+									}
+								},
+								ReferenceCompletion,
+								LowLevelTasks::ETaskPriority::BackgroundNormal);
+
+							SubgraphCompletionEvent.AddPrerequisites(CompleteTask);
+						}
+
+						else
+						{
+							// Scan for requisites
+							TArray< UE::Tasks::FTask, TInlineAllocator<8> > Requisites;
+							GetRequisites(SubgraphRoot,Requisites);
+
+							TUniquePtr<FConstantTask> Task(new FConstantTask(SubgraphRoot, InOptions, Pass));
+
+							// Launch the preparation on the AST-modification pipe
+							UE::Tasks::FTask CompleteTask = UE::Tasks::Launch(TEXT("MutableConstant"), [TaskPtr = MoveTemp(Task), ImOp]()
+								{
+									MUTABLE_CPUPROFILER_SCOPE(MutableConstantPrepare);
+
+									// We need the clone because linking modifies ASTOp state and also to be safe for concurrency.
+									TaskPtr->SourceCloned = ASTOp::DeepClone(TaskPtr->Source);
+
+									TaskPtr->Run(ImOp);
+
+									ASTOp::Replace(TaskPtr->Source, TaskPtr->Result);
+									TaskPtr->Result = nullptr;
+									TaskPtr->Source = nullptr;
+								},
+								Requisites,
+								LowLevelTasks::ETaskPriority::BackgroundHigh);
+
+							SubgraphCompletionEvent.AddPrerequisites(CompleteTask);
+						}
+
+						ConstantSubgraphs[Index].Root = nullptr;
+						SubgraphCompletionEvent.Trigger();
+						
+						UE::Tasks::AddNested(SubgraphCompletionEvent);
+					}
+
+				});
+
+			// Wait for pending tasks
+			{
+				MUTABLE_CPUPROFILER_SCOPE(Waiting);
+				LaunchTask.Wait();
+			}
+		}
+
+		else
+		{
+			// Non-concurrent version: Prepare the tasks.
+			{
+				MUTABLE_CPUPROFILER_SCOPE(ConstantGenerator_LaunchTasks);
+
+				FImageOperator ImOp = FImageOperator::GetDefault(InOptions->ImageFormatFunc);
+
+				// Traverse list of constants to generate. It is ordered in a bottom-up way.
+				int32 SubgraphCount = ConstantSubgraphs.Num();
+				for (int32 OrderIndex = 0; OrderIndex < SubgraphCount; ++OrderIndex)
+				{
+					int32 Index = SubgraphCount - 1 - OrderIndex;
+
+					Ptr<ASTOp> SubgraphRoot = ConstantSubgraphs[Index].Root;
+					UE::Tasks::FTaskEvent& SubgraphCompletionEvent = ConstantSubgraphs[Index].CompletedEvent;
+
+					// Referenced images are resolved in its own task to prevent requesting them twice (and loading them twice)
+					bool bIsCompileTimeReferenceImage = (SubgraphRoot->GetOpType() == OP_TYPE::IM_REFERENCE);
+
+					// Launch the task with its dependencies
+					if (bIsCompileTimeReferenceImage)
+					{
+						// Instead of generating the constant we resolve the reference, which also replaces the ASTOp.
+						const ASTOpReferenceResource* Typed = static_cast<const ASTOpReferenceResource*>(SubgraphRoot.get());
+						uint32 ImageID = Typed->ID;
+
+						auto ImmediateCompleteFunc = [SubgraphRoot, InOptions, ImageID]()
+							{
+								TSharedPtr< Ptr<Image> > ResolveImage = MakeShared<Ptr<Image>>();
+								constexpr bool bRunImmediatlyIfPossible = true;
+								UE::Tasks::FTask ReferenceCompletion = InOptions->OptimisationOptions.ReferencedResourceProvider(ImageID, ResolveImage, bRunImmediatlyIfPossible);
+								ReferenceCompletion.Wait();
+								Ptr<ASTOpConstantResource> ConstantOp = new ASTOpConstantResource;
+								ConstantOp->SourceDataDescriptor = SubgraphRoot->GetSourceDataDescriptor();
+								ConstantOp->Type = OP_TYPE::IM_CONSTANT;
+								ConstantOp->SetValue(ResolveImage->get(), InOptions->OptimisationOptions.DiskCacheContext);
+								ASTOp::Replace(SubgraphRoot, ConstantOp);
+							};
+						ConstantSubgraphs[Index].NonConcurrentTask = ImmediateCompleteFunc;
+					}
+
+					else
+					{
+						// Scan for requisites
+						TArray< UE::Tasks::FTask, TInlineAllocator<8> > Requisites;
+						GetRequisites(SubgraphRoot, Requisites);
+
+						TUniquePtr<FConstantTask> Task(new FConstantTask(SubgraphRoot, InOptions, Pass));
+						FConstantTask* TaskPtr = Task.Get();
+
+						ConstantSubgraphs[Index].Requisites = Requisites;
+						ConstantSubgraphs[Index].TaskData = MoveTemp(Task);
+						ConstantSubgraphs[Index].NonConcurrentTask = [TaskPtr, ImOp]()
+							{
+								TaskPtr->SourceCloned = ASTOp::DeepClone(TaskPtr->Source);
+								TaskPtr->Run(ImOp);
+								ASTOp::Replace(TaskPtr->Source, TaskPtr->Result);
+							};
+					}
+
+					ConstantSubgraphs[Index].Root = nullptr;
+				}
+			}
+
+			auto TryRunFunc = [&ConstantSubgraphs](int32 Index) -> bool
+				{
+					bool bCanRun = true;
+					for (const UE::Tasks::FTask& Req : ConstantSubgraphs[Index].Requisites)
+					{
+						if (!Req.IsCompleted())
+						{
+							bCanRun = false;
+							break;
+						}
+					}
+
+					if (bCanRun)
+					{
+						ConstantSubgraphs[Index].NonConcurrentTask();
+						ConstantSubgraphs[Index].NonConcurrentTask.Reset();
+						ConstantSubgraphs[Index].TaskData.Reset();
+						ConstantSubgraphs[Index].CompletedEvent.Trigger();
+					}
+
+					return bCanRun;
+				};
+
+			// Resolve tasks top-down while as long as they have the prerequisites completed.
+			int32 SubgraphCount = ConstantSubgraphs.Num();
+			TArray<int32> Pending;
+			for (int32 OrderIndex = 0; OrderIndex < SubgraphCount; ++OrderIndex)
+			{
+				int32 Index = SubgraphCount - 1 - OrderIndex;
+
+				bool bRan = TryRunFunc(Index);
+
+				if (bRan)
+				{
+					// We may have unlocked one of the pending tasks. Review them all.
+					bool bModified = true;
+					while (bModified)
+					{
+						bModified = false;
+						for (int32 PendingIndexIndex = 0; PendingIndexIndex < Pending.Num(); )
+						{
+							int32 PendingIndex = Pending[PendingIndexIndex];
+							bRan = TryRunFunc(PendingIndex);
+							if (bRan)
+							{
+								bModified = true;
+								Pending.RemoveAtSwap(PendingIndexIndex);
+								break;
+							}
+							else
+							{
+								++PendingIndexIndex;
+							}
+						}
+					}
+				}
+				else
+				{
+					// We couldn't run the task yet, remember it for later.
+					Pending.Add(Index);
+				}
+			}
+
+			check(Pending.IsEmpty());
+		}
+
+		bool bSomethingModified = ConstantSubgraphs.Num() > 0;
+		return bSomethingModified;
+	}
+
+
+	//-------------------------------------------------------------------------------------------------
+	//-------------------------------------------------------------------------------------------------
+	//-------------------------------------------------------------------------------------------------
+	CodeOptimiser::CodeOptimiser(Ptr<CompilerOptions> InOptions, TArray<FStateCompilationData>& InStates )
+		: m_states( InStates )
+	{
+		m_options = InOptions;
+	}
+
+
+	//-------------------------------------------------------------------------------------------------
+	void CodeOptimiser::FullOptimiseAST( ASTOpList& roots, int32 Pass )
+	{
+		bool modified = true;
+		int32 numIterations = 0;
+		while (modified && (!m_optimizeIterationsMax || (m_optimizeIterationsLeft>0) || !numIterations) )
+		{
+			--m_optimizeIterationsLeft;
+			++numIterations;
+			UE_LOG(LogMutableCore, Verbose, TEXT("Main optimise iteration %d, max %d, left %d"),
+					numIterations, m_optimizeIterationsMax, m_optimizeIterationsLeft);
+
+			modified = false;
+
+			// All kind of optimisations that depend on the meaning of each operation
+			// \TODO: We are doing it for all states.
+			UE_LOG(LogMutableCore, Verbose, TEXT(" - semantic optimiser"));
+			modified |= SemanticOptimiserAST( roots, m_options->GetPrivate()->OptimisationOptions, Pass );
+			//AXE_INT_VALUE("Mutable", Verbose, "ast size", (int64_t)ASTOp::CountNodes(roots));
+			//UE_LOG(LogMutableCore, Verbose, TEXT("(int) %s : %ld"), TEXT("ast size"), int64(ASTOp::CountNodes(roots)));
+			ASTOp::LogHistogram(roots);
+
+			UE_LOG(LogMutableCore, Verbose, TEXT(" - sink optimiser"));
+			modified |= SinkOptimiserAST( roots, m_options->GetPrivate()->OptimisationOptions );
+			//AXE_INT_VALUE("Mutable", Verbose, "ast size", (int64_t)ASTOp::CountNodes(roots));
+			ASTOp::LogHistogram(roots);
+
+			// Image size operations are treated separately
+			UE_LOG(LogMutableCore, Verbose, TEXT(" - size optimiser"));
+			modified |= SizeOptimiserAST( roots );
+			//AXE_INT_VALUE("Mutable", Verbose, "ast size", (int64_t)ASTOp::CountNodes(roots));
+
+			// We cannot run the logic optimiser here because the constant memory explodes.
+			// TODO: Conservative logic optimiser
+			if (!modified)
+			{
+	//                AXE_INT_VALUE("Mutable", Verbose, "program size", (int64_t)program.m_opAddress.Num());
+	//                UE_LOG(LogMutableCore, Verbose, " - logic optimiser");
+	//                LocalLogicOptimiser log;
+	//                modified |= log.Apply( program, (int)s );
+			}
+		}
+
+	//    ASTOp::LogHistogram(roots);
+
+		UE_LOG(LogMutableCore, Verbose, TEXT(" - duplicated code remover"));
+		DuplicatedCodeRemoverAST( roots );
+		//UE_LOG(LogMutableCore, Verbose, TEXT("(int) %s : %ld"), TEXT("ast size"), int64(ASTOp::CountNodes(roots)));
+
+		ASTOp::LogHistogram(roots);
+
+		UE_LOG(LogMutableCore, Verbose, TEXT(" - duplicated data remover"));
+		DuplicatedDataRemoverAST( roots );
+		//UE_LOG(LogMutableCore, Verbose, TEXT("(int) %s : %ld"), TEXT("ast size"), int64(ASTOp::CountNodes(roots)));
+
+		ASTOp::LogHistogram(roots);
+
+		// Generate constants
+		for ( Ptr<ASTOp>& Root: roots )
+		{
+			//UE_LOG(LogMutableCore, Verbose, TEXT("(int) %s : %ld"), TEXT("ast size"), int64(ASTOp::CountNodes(roots)));
+			UE_LOG(LogMutableCore, Verbose, TEXT(" - constant generator"));
+
+			// Constant subtree generation
+			modified = ConstantGeneratorAST( m_options->GetPrivate(), Root, Pass );
+		}
+
+		//UE_LOG(LogMutableCore, Verbose, TEXT("(int) %s : %ld"), TEXT("ast size"), int64(ASTOp::CountNodes(roots)));
+
+	//    ASTOp::LogHistogram(roots);
+
+		// \todo: this code deduplication here, takes very long and doesn't do much... investigate!
+	//    UE_LOG(LogMutableCore, Verbose, " - duplicated code remover");
+	//    DuplicatedCodeRemoverAST( roots );
+	//    AXE_INT_VALUE("Mutable", Verbose, "ast size", (int64_t)ASTOp::CountNodes(roots));
+
+		ASTOp::LogHistogram(roots);
+
+		UE_LOG(LogMutableCore, Verbose, TEXT(" - duplicated data remover"));
+		DuplicatedDataRemoverAST( roots );
+		//UE_LOG(LogMutableCore, Verbose, TEXT("(int) %s : %ld"), TEXT("ast size"), int64(ASTOp::CountNodes(roots)));
+
+		//if (!modified)
+		{
+			//for ( size_t s=0;  s<program.m_states.size(); ++s )
+			{
+				ASTOp::LogHistogram(roots);
+				UE_LOG(LogMutableCore, Verbose, TEXT(" - logic optimiser"));
+				modified |= LocalLogicOptimiserAST( roots );
+				ASTOp::LogHistogram(roots);
+			}
+
+			//AXE_INT_VALUE("Mutable", Verbose, "ast size", (int64_t)ASTOp::CountNodes(roots));
+		}
+
+		ASTOp::LogHistogram(roots);
+	}
+
+
+	//---------------------------------------------------------------------------------------------
+	// The state represents if there is a parent operation requiring skeleton for current mesh subtree.
+	class CollectAllMeshesForSkeletonVisitorAST : public Visitor_TopDown_Unique_Const<uint8_t>
+	{
+	public:
+
+		CollectAllMeshesForSkeletonVisitorAST( const ASTOpList& roots  )
+		{
+			Traverse( roots, false );
+		}
+
+		// List of meshes that require a skeleton
+		TArray<mu::Ptr<ASTOpConstantResource>> MeshesRequiringSkeleton;
+
+	private:
+
+		// Visitor_TopDown_Unique_Const<uint8_t> interface
+		bool Visit( const mu::Ptr<ASTOp>& node ) override
+		{
+			// \todo: refine to avoid instruction branches with irrelevant skeletons.
+
+			uint8_t currentProtected = GetCurrentState();
+
+			switch (node->GetOpType())
+			{
+
+			case OP_TYPE::ME_CONSTANT:
+			{
+				mu::Ptr<ASTOpConstantResource> typedOp = static_cast<ASTOpConstantResource*>(node.get());
+
+				if (currentProtected)
+				{
+					MeshesRequiringSkeleton.AddUnique(typedOp);
+				}
+
+				return false;
+			}
+
+			case OP_TYPE::ME_CLIPMORPHPLANE:
+			{
+				ASTOpMeshClipMorphPlane* typedOp = static_cast<ASTOpMeshClipMorphPlane*>(node.get());
+				if (typedOp->vertexSelectionType == OP::MeshClipMorphPlaneArgs::VS_BONE_HIERARCHY)
+				{
+					// We need the skeleton for the source mesh
+					RecurseWithState( typedOp->source.child(), true );
+					return false;
+				}
+
+				return true;
+			}
+
+			case OP_TYPE::ME_APPLYPOSE:
+			{
+				ASTOpMeshApplyPose* typedOp = static_cast<ASTOpMeshApplyPose*>(node.get());
+
+				// We need the skeleton for both meshes
+				RecurseWithState(typedOp->base.child(), true);
+				RecurseWithState(typedOp->pose.child(), true);
+				return false;
+
+				break;
+			}
+
+			case OP_TYPE::ME_BINDSHAPE:
+			{
+				ASTOpMeshBindShape* typedOp = static_cast<ASTOpMeshBindShape*>(node.get());
+				if (typedOp->bReshapeSkeleton)
+				{
+					RecurseWithState(typedOp->Mesh.child(), true);
+					return false;
+				}
+
+				break;
+			}
+
+			case OP_TYPE::ME_APPLYSHAPE:
+			{
+				ASTOpMeshApplyShape* typedOp = static_cast<ASTOpMeshApplyShape*>(node.get());
+				if (typedOp->bReshapeSkeleton)
+				{
+					RecurseWithState(typedOp->Mesh.child(), true);
+					return false;
+				}
+
+				break;
+			}
+
+			default:
+				break;
+			}
+
+			return true;
+		}
+
+	};
+
+
+	//---------------------------------------------------------------------------------------------
+	// This stores an ADD_MESH op with the child meshes collected and the final skeleton to use
+	// for this op.
+	struct FAddMeshSkeleton
+	{
+		mu::Ptr<ASTOp> m_pAddMeshOp;
+		TArray<mu::Ptr<ASTOpConstantResource>> m_contributingMeshes;
+		mu::Ptr<Skeleton> m_pFinalSkeleton;
+
+		FAddMeshSkeleton( const mu::Ptr<ASTOp>& pAddMeshOp,
+						  TArray<mu::Ptr<ASTOpConstantResource>>& contributingMeshes,
+						  const mu::Ptr<Skeleton>& pFinalSkeleton )
+		{
+			m_pAddMeshOp = pAddMeshOp;
+			m_contributingMeshes = MoveTemp(contributingMeshes);
+			m_pFinalSkeleton = pFinalSkeleton;
+		}
+	};
+
+
+	//---------------------------------------------------------------------------------------------
+	void SkeletonCleanerAST( TArray<mu::Ptr<ASTOp>>& roots, const FModelOptimizationOptions& options )
+	{
+		// This collects all the meshes that require a skeleton because they are used in operations
+		// that require it.
+		CollectAllMeshesForSkeletonVisitorAST requireSkeletonCollector( roots );
+
+		TArray<FAddMeshSkeleton> replacementsFound;
+
+		ASTOp::Traverse_TopDown_Unique_Imprecise( roots, [&](mu::Ptr<ASTOp>& at )
+		{
+			// Only recurse instance construction ops.
+			bool processChildren = GetOpDataType(at->GetOpType())==DT_INSTANCE;
+
+			if ( at->GetOpType() == OP_TYPE::IN_ADDMESH )
+			{
+				ASTOpInstanceAdd* typedNode = static_cast<ASTOpInstanceAdd*>(at.get());
+				mu::Ptr<ASTOp> meshRoot = typedNode->value.child();
+
+				if (meshRoot)
+				{
+					// Gather constant meshes contributing to the final mesh
+					TArray<mu::Ptr<ASTOpConstantResource>> subtreeMeshes;
+					TArray<mu::Ptr<ASTOp>> tempRoots;
+					tempRoots.Add(meshRoot);
+					ASTOp::Traverse_TopDown_Unique_Imprecise( tempRoots, [&](mu::Ptr<ASTOp>& lat )
+					{
+						// \todo: refine to avoid instruction branches with irrelevant skeletons.
+						if ( lat->GetOpType() == OP_TYPE::ME_CONSTANT )
+						{
+							mu::Ptr<ASTOpConstantResource> typedOp = static_cast<ASTOpConstantResource*>(lat.get());
+							if ( subtreeMeshes.Find(typedOp)
+								 ==
+								 INDEX_NONE )
+							{
+								subtreeMeshes.Add(typedOp);
+							}
+						}
+						return true;
+					});
+
+					// Create a mesh just with the unified skeleton
+					SkeletonPtr pFinalSkeleton = new Skeleton;
+					for (const auto& meshAt: subtreeMeshes)
+					{
+						MeshPtrConst pMesh = static_cast<const Mesh*>(meshAt->GetValue().get());
+						mu::Ptr<const Skeleton> sourceSkeleton = pMesh ? pMesh->GetSkeleton() : nullptr;
+						if (sourceSkeleton)
+						{
+							ExtendSkeleton(pFinalSkeleton.get(),sourceSkeleton.get());
+						}
+					}
+
+					replacementsFound.Emplace( at, subtreeMeshes, pFinalSkeleton );
+				}
+			}
+
+			return processChildren;
+		});
+
+
+		// Iterate all meshes again
+		ASTOp::Traverse_TopDown_Unique_Imprecise( roots, [&](mu::Ptr<ASTOp>& at )
+		{
+			if (at->GetOpType()==OP_TYPE::ME_CONSTANT)
+			{
+				ASTOpConstantResource* typedOp = static_cast<ASTOpConstantResource*>(at.get());
+
+				for(FAddMeshSkeleton& Rep: replacementsFound)
+				{
+					if (Rep.m_contributingMeshes.Contains(at))
+					{
+						mu::Ptr<const Mesh> pMesh =  static_cast<const Mesh*>(typedOp->GetValue().get());
+						pMesh->CheckIntegrity();
+
+						Ptr<Mesh> NewMesh = new Mesh();
+						bool bOutSuccess = false;
+						MeshRemapSkeleton(NewMesh.get(), pMesh.get(), Rep.m_pFinalSkeleton.get(), bOutSuccess);
+
+						if (bOutSuccess)
+						{
+							NewMesh->CheckIntegrity();
+							mu::Ptr<ASTOpConstantResource> newOp = new ASTOpConstantResource();
+							newOp->Type = OP_TYPE::ME_CONSTANT;
+							newOp->SetValue(NewMesh, options.DiskCacheContext);
+
+							ASTOp::Replace(at, newOp);
+						}
+					}
+				}
+			}
+			return true;
+		});
+
+	//    for( auto& rep: replacementsFound )
+	//    {
+	//        check(rep.m_pAddMeshOp->GetOpType()==OP_TYPE::IN_ADDMESH);
+	//        auto newAddMeshRootTyped = dynamic_cast<ASTOpInstanceAdd*>(rep.m_pAddMeshOp.get());
+	//        auto meshRoot = newAddMeshRootTyped->value.child();
+
+	//        // Add new skeleton constant instruction
+	//        Ptr<ASTOpConstantResource> cop = new ASTOpConstantResource();
+	//        cop->type = OP_TYPE::ME_CONSTANT;
+	//        cop->SetValue( rep.m_pFinalSkeleton, options.bUseDiskCache );
+
+	//        // Add new "apply skeleton" instruction
+	//        Ptr<ASTOpFixed> skop = new ASTOpFixed();
+	//        skop->op.type = OP_TYPE::ME_SETSKELETON;
+	//        skop->SetChild( skop->op.args.MeshSetSkeleton.source, meshRoot);
+	//        skop->SetChild( skop->op.args.MeshSetSkeleton.skeleton, cop);
+
+	//        // We need a new mesh add instruction, it is not safe to edit the current one.
+	//        auto newOp = mu::Clone<ASTOpInstanceAdd>(newAddMeshRootTyped);
+	//        newOp->value = skop;
+
+	//        // TODO BLEH UNSAFE?
+	//        ASTOp::Replace(rep.m_pAddMeshOp,newOp);
+	//    }
+	}
+
+
+	//---------------------------------------------------------------------------------------------
+	void CodeOptimiser::OptimiseAST()
+	{
+		MUTABLE_CPUPROFILER_SCOPE(OptimiseAST);
+
+		// Gather all the roots (one for each state)
+		TArray<Ptr<ASTOp>> roots;
+		for(const FStateCompilationData& s:m_states)
+		{
+			roots.Add(s.root);
+		}
+
+		//UE_LOG(LogMutableCore, Verbose, TEXT("(int) %s : %ld"), TEXT("ast size"), int64(ASTOp::CountNodes(roots)));
+
+		if ( m_options->GetPrivate()->OptimisationOptions.bEnabled )
+		{
+			// We use 4 times the count because at the time we moved to sharing this count it
+			// was being used 4 times, and we want to keep the tests consistent.
+			m_optimizeIterationsLeft = m_options->GetPrivate()->OptimisationOptions.MaxOptimisationLoopCount * 4;
+			m_optimizeIterationsMax = m_optimizeIterationsLeft;
+
+			// The first duplicated data remover has the special mission of removing
+			// duplicated data (meshes) that may have been specified in the source
+			// data, before we make it diverge because of different uses, like layout
+			// creation
+			UE_LOG(LogMutableCore, Verbose, TEXT(" - duplicated data remover"));
+			DuplicatedDataRemoverAST( roots );
+			//UE_LOG(LogMutableCore, Verbose, TEXT("(int) %s : %ld"), TEXT("ast size"), int64(ASTOp::CountNodes(roots)));
+
+			ASTOp::LogHistogram(roots);
+
+			UE_LOG(LogMutableCore, Verbose, TEXT(" - duplicated code remover"));
+			DuplicatedCodeRemoverAST( roots );
+			//UE_LOG(LogMutableCore, Verbose, TEXT("(int) %s : %ld"), TEXT("ast size"), int64(ASTOp::CountNodes(roots)));
+
+			// Special optimization stages
+			if ( m_options->GetPrivate()->OptimisationOptions.bUniformizeSkeleton )
+			{
+				UE_LOG(LogMutableCore, Verbose, TEXT(" - skeleton cleaner"));
+				ASTOp::LogHistogram(roots);
+
+				SkeletonCleanerAST( roots, m_options->GetPrivate()->OptimisationOptions );
+				//UE_LOG(LogMutableCore, Verbose, TEXT("(int) %s : %ld"), TEXT("ast size"), int64(ASTOp::CountNodes(roots)));
+				ASTOp::LogHistogram(roots);
+			}
+
+			// First optimisation stage. It tries to resolve all the image sizes. This is necessary
+			// because some operations cannot be applied correctly until the image size is known
+			// like the grow-map generation.
+			bool modified = true;
+			int32 numIterations = 0;
+			while (modified)
+			{
+				MUTABLE_CPUPROFILER_SCOPE(FirstStage);
+
+				--m_optimizeIterationsLeft;
+				++numIterations;
+				UE_LOG(LogMutableCore, Verbose, TEXT("First optimise iteration %d, max %d, left %d"),
+						numIterations, m_optimizeIterationsMax, m_optimizeIterationsLeft);
+
+				modified = false;
+
+				UE_LOG(LogMutableCore, Verbose, TEXT(" - size optimiser"));
+				modified |= SizeOptimiserAST( roots );
+			}
+
+			//UE_LOG(LogMutableCore, Verbose, TEXT("(int) %s : %ld"), TEXT("ast size"), int64(ASTOp::CountNodes(roots)));
+			ASTOp::LogHistogram(roots);
+
+			// Main optimisation stage
+			{
+				MUTABLE_CPUPROFILER_SCOPE(MainStage);
+				FullOptimiseAST( roots, 0 );
+				//UE_LOG(LogMutableCore, Verbose, TEXT("(int) %s : %ld"), TEXT("ast size"), int64(ASTOp::CountNodes(roots)));
+				ASTOp::LogHistogram(roots);
+
+				FullOptimiseAST( roots, 1 );
+				//UE_LOG(LogMutableCore, Verbose, TEXT("(int) %s : %ld"), TEXT("ast size"), int64(ASTOp::CountNodes(roots)));
+				ASTOp::LogHistogram(roots);
+			}
+
+			// Constant resolution stage: resolve referenced assets.
+			{
+				MUTABLE_CPUPROFILER_SCOPE(ReferenceResolution);
+				
+				constexpr int32 Pass = 2;
+
+				//FullOptimiseAST(roots, 2);
+
+				// Generate constants
+				for (Ptr<ASTOp>& Root : roots)
+				{
+					// Constant subtree generation
+					modified = ConstantGeneratorAST(m_options->GetPrivate(), Root, Pass);
+				}
+
+				DuplicatedDataRemoverAST(roots);
+			}
+
+			// Main optimisation stage again for data-aware optimizations
+			{
+				MUTABLE_CPUPROFILER_SCOPE(FinalStage);
+				FullOptimiseAST(roots, 0);
+				//UE_LOG(LogMutableCore, Verbose, TEXT("(int) %s : %ld"), TEXT("ast size"), int64(ASTOp::CountNodes(roots)));
+				ASTOp::LogHistogram(roots);
+
+				FullOptimiseAST(roots, 1);
+				//UE_LOG(LogMutableCore, Verbose, TEXT("(int) %s : %ld"), TEXT("ast size"), int64(ASTOp::CountNodes(roots)));
+				ASTOp::LogHistogram(roots);
+			}
+
+			// Analyse mesh constants to see which of them are in optimised mesh formats, and set the flags.
+			ASTOp::Traverse_BottomUp_Unique_NonReentrant( roots, [&](Ptr<ASTOp>& n)
+			{
+				if (n->GetOpType()==OP_TYPE::ME_CONSTANT)
+				{
+					ASTOpConstantResource* typed = static_cast<ASTOpConstantResource*>(n.get());
+					auto pMesh = static_cast<const Mesh*>(typed->GetValue().get());
+					pMesh->ResetStaticFormatFlags();
+					typed->SetValue( pMesh, m_options->GetPrivate()->OptimisationOptions.DiskCacheContext);
+				}
+			});
+
+			ASTOp::LogHistogram(roots);
+
+			// Reset the state root operations in case they have changed due to optimization
+			for (int32 RootIndex = 0; RootIndex < m_states.Num(); ++RootIndex)
+			{
+				m_states[RootIndex].root = roots[RootIndex];
+			}
+
+			{
+				MUTABLE_CPUPROFILER_SCOPE(StatesStage);
+
+				// Optimise for every state
+				OptimiseStatesAST( );
+
+				// Optimise the data formats (TODO)
+				//OperationFlagGenerator flagGen( pResult.get() );
+			}
+
+			ASTOp::LogHistogram(roots);
+		}
+
+		// Minimal optimisation of constant subtrees
+		else if ( m_options->GetPrivate()->OptimisationOptions.bConstReduction )
+		{
+			// The first duplicated data remover has the special mission of removing
+			// duplicated data (meshes) that may have been specified in the source
+			// data, before we make it diverge because of different uses, like layout
+			// creation
+			UE_LOG(LogMutableCore, Verbose, TEXT(" - duplicated data remover"));
+			DuplicatedDataRemoverAST( roots );
+			//AXE_INT_VALUE("Mutable", Verbose, "ast size", (int64_t)ASTOp::CountNodes(roots));
+
+			UE_LOG(LogMutableCore, Verbose, TEXT(" - duplicated code remover"));
+			DuplicatedCodeRemoverAST( roots );
+			//UE_LOG(LogMutableCore, Verbose, TEXT("(int) %s : %ld"), TEXT("ast size"), int64(ASTOp::CountNodes(roots)));
+
+			// Constant resolution stage: resolve referenced assets.
+			{
+				MUTABLE_CPUPROFILER_SCOPE(ReferenceResolution);
+				FullOptimiseAST(roots, 2);
+			}
+
+			for ( int32 StateIndex=0; StateIndex <m_states.Num(); ++StateIndex)
+			{
+				constexpr int32 Pass = 1;
+
+				UE_LOG(LogMutableCore, Verbose, TEXT(" - constant generator"));
+				ConstantGeneratorAST( m_options->GetPrivate(), roots[StateIndex], Pass);
+				//AXE_INT_VALUE("Mutable", Verbose, "ast size", (int64_t)ASTOp::CountNodes(roots));
+			}
+
+			UE_LOG(LogMutableCore, Verbose, TEXT(" - duplicated data remover"));
+			DuplicatedDataRemoverAST( roots );
+			//AXE_INT_VALUE("Mutable", Verbose, "ast size", (int64_t)ASTOp::CountNodes(roots));
+
+			UE_LOG(LogMutableCore, Verbose, TEXT(" - duplicated code remover"));
+			DuplicatedCodeRemoverAST( roots );
+			//AXE_INT_VALUE("Mutable", Verbose, "ast size", (int64_t)ASTOp::CountNodes(roots));
+
+			// Reset the state root operations in case they have changed due to optimization
+			for (int32 RootIndex = 0; RootIndex < m_states.Num(); ++RootIndex)
+			{
+				m_states[RootIndex].root = roots[RootIndex];
+			}
+		}
+
+		ASTOp::LogHistogram(roots);
+
+	}
+
+}
diff --git a/Engine/Plugins/Mutable/Source/MutableTools/Private/MuT/CodeOptimiser.h b/Engine/Plugins/Mutable/Source/MutableTools/Private/MuT/CodeOptimiser.h
new file mode 100644
--- /dev/null
+++ b/Engine/Plugins/Mutable/Source/MutableTools/Private/MuT/CodeOptimiser.h
@@ -0,0 +1,311 @@
+// Copyright Epic Games, Inc. All Rights Reserved.
+
+#pragma once
+
+#include "MuR/Operations.h"
+#include "MuR/CodeVisitor.h"
+#include "MuT/CompilerPrivate.h"
+#include "MuT/AST.h"
+
+namespace mu
+{
+
+    //---------------------------------------------------------------------------------------------
+    //! Code optimiser
+    //---------------------------------------------------------------------------------------------
+    class CodeOptimiser
+    {
+    public:
+        CodeOptimiser( Ptr<CompilerOptions> options, TArray<FStateCompilationData>& states );
+
+        //! Optimise the virtual machine code, using several transforms.
+        void OptimiseAST();
+
+    private:
+
+		Ptr<CompilerOptions> m_options;
+
+        //!
+		TArray<FStateCompilationData>& m_states;
+
+        //! The max number of optimize iterations is shared across several stages now.
+        //! This is how many are left
+        int m_optimizeIterationsMax=0;
+        int m_optimizeIterationsLeft=0;
+
+        //! Full optimisation pass
+        void FullOptimiseAST( ASTOpList& roots, int32 Pass);
+
+        //! Optimise the code of a model for a specific state, generating new instructions and
+        //! state information.
+        void OptimiseStatesAST();
+
+    };
+
+
+    //---------------------------------------------------------------------------------------------
+    //! Scan a subtree and return true if any op of a type in the provided subset is found.
+    //! Intermediate data is used between calls to apply, so don't remove program code or directly
+    //! change the instructions. Adding new instructions is ok.
+    //---------------------------------------------------------------------------------------------
+    class FindOpTypeVisitor
+    {
+    public:
+
+        FindOpTypeVisitor(const TArray<OP_TYPE>& types)
+            : m_typesToFind(types)
+        {
+        }
+
+        bool Apply( FProgram& program, OP::ADDRESS at );
+
+
+    private:
+
+        TArray<OP_TYPE> m_typesToFind;
+
+		TArray< TPair<bool,OP::ADDRESS> > m_pending;
+
+        // 0 not visited
+        // 1 children pending
+        // 2 visited and not found
+        // 3 visited and found
+		TArray<uint8> m_visited;
+    };
+
+
+    //---------------------------------------------------------------------------------------------
+    //! Scan a subtree and return true if the expression is constant.
+    //! It has the same restrictions that FindOpTypeVisitor
+    //---------------------------------------------------------------------------------------------
+    class IsConstantVisitor
+    {
+    public:
+
+        IsConstantVisitor();
+
+        bool Apply( FProgram& program, OP::ADDRESS at );
+
+    private:
+
+        TUniquePtr<FindOpTypeVisitor> m_findOpTypeVisitor;
+    };
+
+
+    //---------------------------------------------------------------------------------------------
+    //! ConstantGenerator replaces constant subtrees of operations with an equivalent single
+	//! constant value operation. 
+    //---------------------------------------------------------------------------------------------
+    extern bool ConstantGeneratorAST( const CompilerOptions::Private* options, Ptr<ASTOp>& root, int32 Pass );
+
+    //---------------------------------------------------------------------------------------------
+    //! \TODO: shapes, projectors, others? but not switches (they must be unique)
+    //---------------------------------------------------------------------------------------------
+    extern void DuplicatedDataRemoverAST( ASTOpList& roots );
+
+
+    //---------------------------------------------------------------------------------------------
+    //! Mark all the duplicated code instructions to point at the same operation, leaving the
+    //! copies unreachable.
+    //---------------------------------------------------------------------------------------------
+    extern void DuplicatedCodeRemoverAST( ASTOpList& roots );
+
+
+    //---------------------------------------------------------------------------------------------
+    //! All kinds of optimisations that depend on the meaning of each operation
+    //---------------------------------------------------------------------------------------------
+    extern bool SemanticOptimiserAST( ASTOpList& roots, const FModelOptimizationOptions& optimisationOptions, int32 Pass);
+
+    //---------------------------------------------------------------------------------------------
+    //! Semantic operator that reorders instructions moving expensive ones down to the
+    //! leaves of the expressions trying to turn them into constants.
+    //---------------------------------------------------------------------------------------------
+    extern bool SinkOptimiserAST( ASTOpList& roots, const FModelOptimizationOptions& );
+
+    //---------------------------------------------------------------------------------------------
+    //!
+    //---------------------------------------------------------------------------------------------
+    extern bool SizeOptimiserAST( ASTOpList& roots );
+
+    //---------------------------------------------------------------------------------------------
+    //!
+    //---------------------------------------------------------------------------------------------
+    extern bool LocalLogicOptimiserAST(ASTOpList& roots);
+
+    //---------------------------------------------------------------------------------------------
+    //! Discard all LODs beyond the given lod count.
+    //---------------------------------------------------------------------------------------------
+    class LODCountReducerAST : public Visitor_TopDown_Unique_Cloning
+    {
+    public:
+
+        LODCountReducerAST( Ptr<ASTOp>& root, int lodCount );
+
+    protected:
+
+        Ptr<ASTOp> Visit( Ptr<ASTOp> at, bool& processChildren ) override;
+
+        int m_lodCount;
+    };
+
+
+    //---------------------------------------------------------------------------------------------
+    //! Scan the code in the given subtree and return true if a state runtime parameters is
+    //! found.
+    //! Intermediate data is used betwen calls to apply, so don't remove program code or directly
+    //! change the instructions. Adding new instructions is ok.
+    //---------------------------------------------------------------------------------------------
+    class RuntimeParameterVisitorAST
+    {
+    public:
+
+        RuntimeParameterVisitorAST(const FStateCompilationData* pState);
+
+        bool HasAny( const Ptr<ASTOp>& root );
+
+    private:
+
+        const FStateCompilationData* m_pState;
+
+        //!
+        struct PENDING_ITEM
+        {
+            //! 0: indicate subtree pending
+            //! 1: indicate children finished
+            uint8_t itemType;
+
+            //! 0: everything is relevant
+            //! 1: only layouts are relevant
+            uint8_t onlyLayoutsRelevant;
+
+            //! Address to visit
+            Ptr<ASTOp> at;
+        };
+
+        //!
+        vector< PENDING_ITEM > m_pending;
+
+        //! Possible op state
+        enum class OP_STATE : uint8_t
+        {
+            NOT_VISITED = 0,
+            CHILDREN_PENDING_FULL,
+            CHILDREN_PENDING_PARTIAL,
+            VISITED_HASRUNTIME,
+            VISITED_FULL_DOESNTHAVERUNTIME,
+            VISITED_PARTIAL_DOESNTHAVERUNTIME
+        };
+
+        std::unordered_map<Ptr<ASTOp>,OP_STATE> m_visited;
+
+        //!
+        void AddIfNeeded( const PENDING_ITEM& item );
+
+    };
+
+
+    //---------------------------------------------------------------------------------------------
+    //! Remove all texture compression operations that would happen for runtime parameter changes.
+    //---------------------------------------------------------------------------------------------
+    class RuntimeTextureCompressionRemoverAST : public Visitor_TopDown_Unique_Cloning
+    {
+    public:
+
+        RuntimeTextureCompressionRemoverAST(FStateCompilationData* state, bool bInAlwaysUncompress );
+
+    protected:
+
+        Ptr<ASTOp> Visit( Ptr<ASTOp> at, bool& processChildren ) override;
+
+    private:
+
+        RuntimeParameterVisitorAST m_hasRuntimeParamVisitor;
+		bool bAlwaysUncompress = false;
+
+    };
+
+    //---------------------------------------------------------------------------------------------
+    //! Restructure the code to move operations involving runtime parameters as high as possible
+    //! in the code tree.
+    //---------------------------------------------------------------------------------------------
+    class ParameterOptimiserAST : public Visitor_TopDown_Unique_Cloning
+    {
+    private:
+
+		FStateCompilationData& m_stateProps;
+        // unused const GPU_PLATFORM_PROPS& m_gpuPlatformProps;
+
+    public:
+
+        ParameterOptimiserAST(FStateCompilationData&s,
+                               const FModelOptimizationOptions& optimisationOptions );
+
+        bool Apply();
+
+
+    private:
+
+        Ptr<ASTOp> Visit( Ptr<ASTOp> at, bool& processChildren ) override;
+
+        bool m_modified;
+
+        FModelOptimizationOptions OptimisationOptions;
+
+        RuntimeParameterVisitorAST m_hasRuntimeParamVisitor;
+
+    };
+
+
+    //---------------------------------------------------------------------------------------------
+    //! Some masks are optional. If they are null, replace them by a white plain image of the right
+    //! size.
+    //---------------------------------------------------------------------------------------------
+    extern Ptr<ASTOp> EnsureValidMask( Ptr<ASTOp> mask, Ptr<ASTOp> base );
+
+
+    //---------------------------------------------------------------------------------------------
+    //! Calculate all the parameters found relevant under a particular operation. This may not
+    //! incldue all the parameters in the subtree (if because of the operations they are not
+    //! relevant)
+    //! It has an internal cache, so don't reuse if the program changes.
+    //---------------------------------------------------------------------------------------------
+    class SubtreeRelevantParametersVisitorAST
+    {
+    public:
+
+        void Run( Ptr<ASTOp> root );
+
+        //! After Run, list of relevant parameters.
+        TSet< FString > m_params;
+
+    private:
+
+        struct FState
+        {
+            Ptr<ASTOp> op;
+            bool bOnlyLayoutIsRelevant=false;
+
+			FState( Ptr<ASTOp> o=nullptr, bool l=false) : op(o), bOnlyLayoutIsRelevant(l) {}
+
+            bool operator==(const FState& o) const
+            {
+                return  op == o.op &&
+					bOnlyLayoutIsRelevant == o.bOnlyLayoutIsRelevant;
+            }
+
+			friend FORCEINLINE uint32 GetTypeHash(const FState& InKey)
+			{
+				uint32 KeyHash = 0;
+				KeyHash = HashCombineFast(KeyHash, ::GetTypeHash(InKey.op.get()));
+				KeyHash = HashCombineFast(KeyHash, ::GetTypeHash(InKey.bOnlyLayoutIsRelevant));
+				return KeyHash;
+			}
+		};
+
+
+        // Result cache
+        // \todo optimise by storing unique lists separately and an index here.
+        TMap< FState, TSet<FString> > m_resultCache;
+    };
+
+}
diff --git a/Engine/Plugins/Mutable/Source/MutableTools/Private/MuT/CodeOptimiser_GPU.cpp b/Engine/Plugins/Mutable/Source/MutableTools/Private/MuT/CodeOptimiser_GPU.cpp
new file mode 100644
--- /dev/null
+++ b/Engine/Plugins/Mutable/Source/MutableTools/Private/MuT/CodeOptimiser_GPU.cpp
@@ -0,0 +1,6 @@
+// Copyright Epic Games, Inc. All Rights Reserved.
+
+#include "MuT/CodeOptimiser_GPU.h"
+
+#define MUTABLE_GPU_CODE_BUFFER_SIZE	512
+
diff --git a/Engine/Plugins/Mutable/Source/MutableTools/Private/MuT/CodeOptimiser_GPU.h b/Engine/Plugins/Mutable/Source/MutableTools/Private/MuT/CodeOptimiser_GPU.h
new file mode 100644
--- /dev/null
+++ b/Engine/Plugins/Mutable/Source/MutableTools/Private/MuT/CodeOptimiser_GPU.h
@@ -0,0 +1,5 @@
+// Copyright Epic Games, Inc. All Rights Reserved.
+
+#pragma once
+
+
