[mutable] Moved the Mutable plugin out of Experimental status into Beta.

diff --git a/Engine/Plugins/Mutable/Source/MutableTools/Private/MuT/CodeOptimiser_Semantic.cpp b/Engine/Plugins/Mutable/Source/MutableTools/Private/MuT/CodeOptimiser_Semantic.cpp
new file mode 100644
--- /dev/null
+++ b/Engine/Plugins/Mutable/Source/MutableTools/Private/MuT/CodeOptimiser_Semantic.cpp
@@ -0,0 +1,1535 @@
+// Copyright Epic Games, Inc. All Rights Reserved.
+
+#include "MuT/NodeImageFormatPrivate.h"
+#include "HAL/PlatformMath.h"
+#include "Misc/AssertionMacros.h"
+#include "MuR/Image.h"
+#include "MuR/MutableMath.h"
+#include "MuR/MutableTrace.h"
+#include "MuR/Operations.h"
+#include "MuR/Ptr.h"
+#include "MuR/RefCounted.h"
+#include "MuT/AST.h"
+#include "MuT/ASTOpConditional.h"
+#include "MuT/ASTOpConstantBool.h"
+#include "MuT/ASTOpImageCompose.h"
+#include "MuT/ASTOpImageMipmap.h"
+#include "MuT/ASTOpImageLayer.h"
+#include "MuT/ASTOpImageLayerColor.h"
+#include "MuT/ASTOpImageMultiLayer.h"
+#include "MuT/ASTOpImagePatch.h"
+#include "MuT/ASTOpImagePixelFormat.h"
+#include "MuT/ASTOpImageMakeGrowMap.h"
+#include "MuT/ASTOpImageSwizzle.h"
+#include "MuT/ASTOpImageTransform.h"
+#include "MuT/ASTOpImageRasterMesh.h"
+#include "MuT/ASTOpMeshMorph.h"
+#include "MuT/ASTOpMeshAddTags.h"
+#include "MuT/ASTOpSwitch.h"
+#include "MuT/CodeOptimiser.h"
+#include "MuT/Table.h"
+
+
+namespace mu
+{
+
+    //---------------------------------------------------------------------------------------------
+    //---------------------------------------------------------------------------------------------
+    //---------------------------------------------------------------------------------------------
+    bool SemanticOptimiserAST(
+		ASTOpList& Roots,
+		const FModelOptimizationOptions& OptimisationOptions,
+		int32 Pass
+	)
+    {
+        MUTABLE_CPUPROFILER_SCOPE(SemanticOptimiserAST);
+
+        bool bModified = false;
+
+        // TODO: isn't top down better suited?
+        ASTOp::Traverse_BottomUp_Unique( Roots, [&](Ptr<ASTOp>& CurrentOp)
+        {
+            Ptr<ASTOp> OptimizedOp = CurrentOp->OptimiseSemantic(OptimisationOptions, Pass);
+
+            // If the returned value is null it means no change.
+            if (OptimizedOp && OptimizedOp !=CurrentOp)
+            {
+				bModified = true;
+                ASTOp::Replace(CurrentOp, OptimizedOp);
+
+				// Check if we are replacing one of the root operations and update the parameter array.
+				for (int32 RootIndex=0; RootIndex<Roots.Num(); ++RootIndex)
+				{
+					if (Roots[RootIndex]== CurrentOp)
+					{
+						Roots[RootIndex] = OptimizedOp;
+					}
+				}
+            }
+        });
+
+        return bModified;
+    }
+
+
+    //---------------------------------------------------------------------------------------------
+    Ptr<ASTOp> ASTOpFixed::OptimiseSemantic( const FModelOptimizationOptions& options, int32 Pass) const
+    {
+        Ptr<ASTOp> at;
+
+        OP_TYPE type = GetOpType();
+        switch ( type )
+        {
+
+        case OP_TYPE::BO_NOT:
+        {
+			Ptr<ASTOp> sourceAt = children[op.args.BoolNot.source].child();
+            if  ( sourceAt->GetOpType() == OP_TYPE::BO_CONSTANT )
+            {
+                Ptr<ASTOpConstantBool> newAt = new ASTOpConstantBool();
+                newAt->value = !newAt->value;
+                at=newAt;
+            }
+            break;
+        }
+
+        case OP_TYPE::BO_AND:
+        {
+            bool changed = false;
+			Ptr<ASTOp> aAt = children[op.args.BoolBinary.a].child();
+			Ptr<ASTOp> bAt = children[op.args.BoolBinary.b].child();
+            if ( !aAt )
+            {
+                at = bAt;
+                changed = true;
+            }
+            else if ( !bAt )
+            {
+                at = aAt;
+                changed = true;
+            }
+            else if ( aAt->GetOpType() == OP_TYPE::BO_CONSTANT )
+            {
+                if (static_cast<const ASTOpConstantBool*>(aAt.get())->value)
+                {
+                    at = bAt;
+                    changed = true;
+                }
+                else
+                {
+                    at = aAt;
+                    changed = true;
+                }
+            }
+            else if ( bAt->GetOpType() == OP_TYPE::BO_CONSTANT )
+            {
+                if (static_cast<const ASTOpConstantBool*>(bAt.get())->value)
+                {
+                    at = aAt;
+                    changed = true;
+                }
+                else
+                {
+                    at = bAt;
+                    changed = true;
+                }
+            }
+
+            // Common cases of repeated branch in children
+            else if ( aAt->GetOpType() == type )
+            {
+				const ASTOpFixed* typedA = static_cast<const ASTOpFixed*>(aAt.get());
+                if ( typedA->children[typedA->op.args.BoolBinary.a].child()==bAt
+                     ||
+                     typedA->children[typedA->op.args.BoolBinary.b].child()==bAt )
+                {
+                    at = aAt;
+                    changed = true;
+                }
+            }
+            else if ( bAt->GetOpType() == type )
+            {
+				const ASTOpFixed* typedB = static_cast<const ASTOpFixed*>(bAt.get());
+                if ( typedB->children[typedB->op.args.BoolBinary.b].child()==aAt
+                     ||
+                     typedB->children[typedB->op.args.BoolBinary.b].child()==bAt )
+                {
+                    at = bAt;
+                    changed = true;
+                }
+            }
+
+            else if (aAt == bAt || *aAt==*bAt)
+            {
+                at = aAt;
+                changed = true;
+            }
+
+            // if it became null, it means true (neutral AND argument)
+            if (changed && !at)
+            {
+                Ptr<ASTOpConstantBool> newAt = new ASTOpConstantBool();
+                newAt->value = true;
+                at = newAt;
+            }
+            break;
+        }
+
+        case OP_TYPE::BO_OR:
+        {
+            bool changed = false;
+			Ptr<ASTOp> aAt = children[op.args.BoolBinary.a].child();
+			Ptr<ASTOp> bAt = children[op.args.BoolBinary.b].child();
+            if ( !aAt )
+            {
+                at = bAt;
+                changed = true;
+            }
+            else if ( !bAt )
+            {
+                at = aAt;
+                changed = true;
+            }
+            else if ( aAt->GetOpType() == OP_TYPE::BO_CONSTANT )
+            {
+                if (static_cast<const ASTOpConstantBool*>(aAt.get())->value)
+                {
+                    at = aAt;
+                    changed = true;
+                }
+                else
+                {
+                    at = bAt;
+                    changed = true;
+                }
+            }
+            else if ( bAt->GetOpType() == OP_TYPE::BO_CONSTANT )
+            {
+                if (static_cast<const ASTOpConstantBool*>(bAt.get())->value)
+                {
+                    at = bAt;
+                    changed = true;
+                }
+                else
+                {
+                    at = aAt;
+                    changed = true;
+                }
+            }
+
+            // Common cases of repeated branch in children
+            else if ( aAt->GetOpType() == type )
+            {
+				const ASTOpFixed* typedA = static_cast<const ASTOpFixed*>(aAt.get());
+                if ( typedA->children[typedA->op.args.BoolBinary.a].child()==bAt
+                     ||
+                     typedA->children[typedA->op.args.BoolBinary.b].child()==bAt )
+                {
+                    at = aAt;
+                    changed = true;
+                }
+            }
+            else if ( bAt->GetOpType() == type )
+            {
+				const ASTOpFixed* typedB = static_cast<const ASTOpFixed*>(bAt.get());
+                if ( typedB->children[typedB->op.args.BoolBinary.b].child()==aAt
+                     ||
+                     typedB->children[typedB->op.args.BoolBinary.b].child()==bAt )
+                {
+                    at = bAt;
+                    changed = true;
+                }
+            }
+
+            else if (aAt == bAt || *aAt==*bAt)
+            {
+                at = aAt;
+                changed = true;
+            }
+
+            // if it became null, it means false (neutral OR argument)
+            if (changed && !at)
+            {
+                Ptr<ASTOpConstantBool> newAt = new ASTOpConstantBool();
+                newAt->value = false;
+                at = newAt;
+            }
+
+            break;
+        }
+
+		case OP_TYPE::CO_SWIZZLE:
+		{
+			// Optimizations that can be applied per-channel
+			{
+				Ptr<ASTOpFixed> NewAt;
+
+				for (int32 ChannelIndex = 0; ChannelIndex < MUTABLE_OP_MAX_SWIZZLE_CHANNELS; ++ChannelIndex)
+				{
+					Ptr<ASTOp> Candidate = children[op.args.ColourSwizzle.sources[ChannelIndex]].child();
+					if (!Candidate)
+					{
+						continue;
+					}
+
+					switch (Candidate->GetOpType())
+					{
+
+						// Swizzle + swizzle = swizzle
+					case OP_TYPE::CO_SWIZZLE:
+					{
+						if (!NewAt)
+						{
+							NewAt = mu::Clone<ASTOpFixed>(this);
+						}
+						const ASTOpFixed* TypedCandidate = static_cast<const ASTOpFixed*>(Candidate.get());
+						int32 CandidateChannel = op.args.ColourSwizzle.sourceChannels[ChannelIndex];
+
+						NewAt->SetChild(NewAt->op.args.ColourSwizzle.sources[ChannelIndex], TypedCandidate->children[TypedCandidate->op.args.ColourSwizzle.sources[CandidateChannel]]);
+						NewAt->op.args.ColourSwizzle.sourceChannels[ChannelIndex] = TypedCandidate->op.args.ColourSwizzle.sourceChannels[CandidateChannel];
+						break;
+					}
+
+					default:
+						break;
+
+					}
+				}
+
+				at = NewAt;
+			}
+
+			// Not optimized yet?
+			if (!at)
+			{
+				// Optimizations that depend on all channels.
+				bool bAllChannelsSameType = true;
+				OP_TYPE AllChannelsType = OP_TYPE::NONE;
+				for (int32 ChannelIndex = 0; ChannelIndex < MUTABLE_OP_MAX_SWIZZLE_CHANNELS; ++ChannelIndex)
+				{
+					Ptr<ASTOp> Candidate = children[op.args.ColourSwizzle.sources[ChannelIndex]].child();
+					if (!Candidate)
+					{
+						continue;
+					}
+
+					if (ChannelIndex == 0)
+					{
+						AllChannelsType = Candidate->GetOpType();
+					}
+					else if (Candidate->GetOpType()!=AllChannelsType)
+					{
+						bAllChannelsSameType = false;
+						break;
+					}
+				}
+
+				if (bAllChannelsSameType)
+				{
+					switch (AllChannelsType)
+					{
+					case OP_TYPE::CO_FROMSCALARS:
+					{
+						// We can remove the swizzle and replace it with a new FromScalars actually swizzling the inputs.
+
+						Ptr<ASTOpFixed> NewAt = new ASTOpFixed();
+						NewAt->op.type = OP_TYPE::CO_FROMSCALARS;
+
+						for (int32 ChannelIndex = 0; ChannelIndex < MUTABLE_OP_MAX_SWIZZLE_CHANNELS; ++ChannelIndex)
+						{
+							const ASTOp* SelectedSourceGeneric = children[op.args.ColourSwizzle.sources[ChannelIndex]].child().get();
+							if (SelectedSourceGeneric)
+							{
+								const ASTOpFixed* SelectedSource = static_cast<const ASTOpFixed*>(SelectedSourceGeneric);
+								int32 SelectedChannel = op.args.ColourSwizzle.sourceChannels[ChannelIndex];
+								Ptr<ASTOp> SelectedFloatInput = SelectedSource->children[SelectedSource->op.args.ColourFromScalars.v[SelectedChannel]].child();
+								NewAt->SetChild(NewAt->op.args.ColourFromScalars.v[ChannelIndex], SelectedFloatInput);
+							}
+						}
+
+						at = NewAt;
+						break;
+					}
+
+					default: break;
+					}
+				}
+			}
+			
+			break;
+		}
+
+		//-------------------------------------------------------------------------------------
+		case OP_TYPE::IM_RESIZEREL:
+		{
+			Ptr<ASTOp> sourceAt = children[op.args.ImageResizeRel.source].child();
+
+			// The instruction can be sunk
+			OP_TYPE sourceType = sourceAt->GetOpType();
+			switch (sourceType)
+			{
+			// This is done here instead of in the OptimizeSize step to prevent removing 
+			// resize_rel too early.
+			case OP_TYPE::IM_RESIZE:
+			{
+				Ptr<ASTOpFixed> NewOp = mu::Clone<ASTOpFixed>(sourceAt.get());
+				NewOp->op.args.ImageResize.size[0] = int16(NewOp->op.args.ImageResize.size[0] * op.args.ImageResizeRel.factor[0]);
+				NewOp->op.args.ImageResize.size[1] = int16(NewOp->op.args.ImageResize.size[1] * op.args.ImageResizeRel.factor[1]);
+
+				at = NewOp;
+				break;
+			}
+
+
+			default:
+				break;
+
+			}
+
+			break;
+		}
+
+        //-------------------------------------------------------------------------------------
+        case OP_TYPE::IM_PIXELFORMAT:
+		case OP_TYPE::IM_SWIZZLE:
+		case OP_TYPE::IM_LAYER:
+		{
+			check(false); // moved to its own operation class
+			break;
+		}
+
+		//-------------------------------------------------------------------------------------
+		case OP_TYPE::ME_MASKDIFF:
+		{
+			Ptr<ASTOp> Fragment = children[op.args.MeshMaskDiff.fragment].child();
+			OP_TYPE FragmentType = Fragment->GetOpType();
+			switch (FragmentType)
+			{
+			case OP_TYPE::ME_ADDTAGS:
+			{
+				// Tags in the fragment can be ignored.
+				const ASTOpMeshAddTags* Add = static_cast<const ASTOpMeshAddTags*>(Fragment.get());
+
+				Ptr<ASTOpFixed> NewAt = mu::Clone<ASTOpFixed>(this);
+				NewAt->SetChild(NewAt->op.args.MeshMaskDiff.fragment, Add->Source);
+				at = NewAt;
+				break;
+			}
+			case OP_TYPE::ME_SWITCH:
+			{
+				// Move the mask diff down all the paths
+				Ptr<ASTOpSwitch> nop = mu::Clone<ASTOpSwitch>(Fragment.get());
+
+				if (nop->def)
+				{
+					Ptr<ASTOpFixed> defOp = mu::Clone<ASTOpFixed>(this);
+					defOp->SetChild(defOp->op.args.MeshMaskDiff.fragment, nop->def);
+					nop->def = defOp;
+				}
+
+				// We need to copy the options because we change them
+				for (int32 v = 0; v < nop->cases.Num(); ++v)
+				{
+					if (nop->cases[v].branch)
+					{
+						Ptr<ASTOpFixed> bOp = mu::Clone<ASTOpFixed>(this);
+						bOp->SetChild(bOp->op.args.MeshMaskDiff.fragment, nop->cases[v].branch);
+						nop->cases[v].branch = bOp;
+					}
+				}
+
+				at = nop;
+				break;
+			}
+			default:
+				break;
+			}
+
+			break;
+		}
+
+		//-------------------------------------------------------------------------------------
+		case OP_TYPE::ME_APPLYLAYOUT:
+		{
+			Ptr<ASTOp> Base = children[op.args.MeshApplyLayout.mesh].child();
+			OP_TYPE BaseType = Base->GetOpType();
+			switch (BaseType)
+			{
+			case OP_TYPE::ME_ADDTAGS:
+			{
+				// Add the tags after layout
+				Ptr<ASTOpMeshAddTags> NewAddTags = mu::Clone<ASTOpMeshAddTags>(Base);
+
+				if (NewAddTags->Source)
+				{
+					Ptr<ASTOpFixed> NewAt = mu::Clone<ASTOpFixed>(this);
+					NewAt->SetChild(NewAt->op.args.MeshApplyLayout.mesh, NewAddTags->Source);
+					NewAddTags->Source = NewAt;
+				}
+
+				at = NewAddTags;
+				break;
+			}
+
+			default:
+				break;
+			}
+
+			break;
+		}
+
+        default:
+            break;
+        }
+
+        return at;
+    }
+
+
+    //---------------------------------------------------------------------------------------------
+    //---------------------------------------------------------------------------------------------
+    //---------------------------------------------------------------------------------------------
+    bool SinkOptimiserAST
+    (
+            ASTOpList& InRoots,
+            const FModelOptimizationOptions& InOptimisationOptions
+    )
+    {
+        MUTABLE_CPUPROFILER_SCOPE(SinkOptimiserAST);
+
+        bool bModified = false;
+
+		FOptimizeSinkContext Context;
+
+        ASTOp::Traverse_TopDown_Unique_Imprecise(InRoots, [&](Ptr<ASTOp>& n)
+        {
+            auto o = n->OptimiseSink(InOptimisationOptions, Context);
+            if (o && n!=o)
+            {
+				bModified = true;
+                ASTOp::Replace(n,o);
+            }
+
+            return true;
+        });
+
+        return bModified;
+    }
+
+
+	//---------------------------------------------------------------------------------------------
+	//---------------------------------------------------------------------------------------------
+	//---------------------------------------------------------------------------------------------
+	Ptr<ASTOp> ASTOpFixed::OptimiseSink( const FModelOptimizationOptions& options, FOptimizeSinkContext& context ) const
+    {
+        Ptr<ASTOp> at;
+
+        OP_TYPE type = GetOpType();
+        switch ( type )
+        {
+
+        case OP_TYPE::ME_PROJECT:
+        {
+            auto sourceAt = children[op.args.MeshProject.mesh].child();
+            if (!sourceAt) { break; }
+
+            auto projectorAt = children[op.args.MeshProject.projector].child();
+            if (!projectorAt) { break; }
+
+            OP_TYPE sourceType = sourceAt->GetOpType();
+            switch ( sourceType )
+            {
+
+            case OP_TYPE::ME_CONDITIONAL:
+            {
+                if (projectorAt->GetOpType()==OP_TYPE::PR_CONSTANT)
+                {
+                    // We move the project down the two paths
+                    auto nop = mu::Clone<ASTOpConditional>(sourceAt);
+
+                    auto aOp = mu::Clone<ASTOpFixed>(this);
+                    aOp->SetChild( aOp->op.args.MeshProject.mesh, nop->yes );
+                    nop->yes = aOp;
+
+                    auto bOp = mu::Clone<ASTOpFixed>(this);
+                    bOp->SetChild( bOp->op.args.MeshProject.mesh, nop->no);
+                    nop->no = bOp;
+
+                    at = nop;
+                }
+                break;
+            }
+
+            case OP_TYPE::ME_SWITCH:
+            {
+                if (projectorAt->GetOpType()==OP_TYPE::PR_CONSTANT)
+                {
+                    // Move the format down all the paths
+                    auto nop = mu::Clone<ASTOpSwitch>(sourceAt);
+
+                    if (nop->def)
+                    {
+                        auto defOp = mu::Clone<ASTOpFixed>(this);
+                        defOp->SetChild( defOp->op.args.MeshProject.mesh, nop->def );
+                        nop->def = defOp;
+                    }
+
+                    // We need to copy the options because we change them
+                    for ( int32 v=0; v<nop->cases.Num(); ++v )
+                    {
+                        if ( nop->cases[v].branch )
+                        {
+                            auto bOp = mu::Clone<ASTOpFixed>(this);
+                            bOp->SetChild( bOp->op.args.MeshProject.mesh, nop->cases[v].branch );
+                            nop->cases[v].branch = bOp;
+                        }
+                    }
+
+                    at = nop;
+                }
+                break;
+            }
+
+			case OP_TYPE::ME_ADDTAGS:
+			{
+				// Add the tags after layout
+				Ptr<ASTOpMeshAddTags> NewAddTags = mu::Clone<ASTOpMeshAddTags>(sourceAt);
+
+				if (NewAddTags->Source)
+				{
+					Ptr<ASTOpFixed> NewAt = mu::Clone<ASTOpFixed>(this);
+					NewAt->SetChild(NewAt->op.args.MeshProject.mesh, NewAddTags->Source);
+					NewAddTags->Source = NewAt;
+				}
+
+				at = NewAddTags;
+				break;
+			}
+
+            default:
+                break;
+            }
+
+            break;
+        }
+
+
+        //-----------------------------------------------------------------------------------------
+        //-----------------------------------------------------------------------------------------
+        //-----------------------------------------------------------------------------------------
+		case OP_TYPE::IM_SWIZZLE:
+		case OP_TYPE::LA_REMOVEBLOCKS:
+		case OP_TYPE::IM_RASTERMESH:
+		{
+			// Moved to their own operation class
+			check(false);
+            break;
+        }
+
+
+        //-----------------------------------------------------------------------------------------
+        //-----------------------------------------------------------------------------------------
+        //-----------------------------------------------------------------------------------------
+        case OP_TYPE::IM_DISPLACE:
+        {
+			Ptr<ASTOp> OriginalAt = at;
+			Ptr<ASTOp> sourceAt = children[op.args.ImageDisplace.source].child();
+			Ptr<ASTOp> displaceMapAt = children[op.args.ImageDisplace.displacementMap].child();
+
+            switch ( sourceAt->GetOpType() )
+            {
+
+            case OP_TYPE::IM_CONDITIONAL:
+            {
+                if (displaceMapAt->GetOpType()==OP_TYPE::IM_CONDITIONAL)
+                {
+					const ASTOpConditional* typedSource = static_cast<const ASTOpConditional*>(sourceAt.get());
+					const ASTOpConditional* typedDisplacementMap = static_cast<const ASTOpConditional*>(displaceMapAt.get());
+
+                    if (typedSource->condition==typedDisplacementMap->condition)
+                    {
+                        auto nop = mu::Clone<ASTOpConditional>(sourceAt);
+
+                        auto aOp = mu::Clone<ASTOpFixed>(this);
+                        aOp->SetChild( aOp->op.args.ImageDisplace.source, typedSource->yes );
+                        aOp->SetChild( aOp->op.args.ImageDisplace.displacementMap, typedDisplacementMap->yes );
+                        nop->yes = aOp;
+
+                        auto bOp = mu::Clone<ASTOpFixed>(this);
+                        bOp->SetChild( bOp->op.args.ImageDisplace.source, typedSource->no);
+                        bOp->SetChild( bOp->op.args.ImageDisplace.displacementMap, typedDisplacementMap->no);
+                        nop->no = bOp;
+
+                        at = nop;
+                    }
+                }
+                break;
+            }
+
+            case OP_TYPE::IM_SWITCH:
+            {
+                if (displaceMapAt->GetOpType()==OP_TYPE::IM_SWITCH)
+                {
+					const ASTOpSwitch* typedSource = static_cast<const ASTOpSwitch*>(sourceAt.get());
+					const ASTOpSwitch* typedDisplacementMap = static_cast<const ASTOpSwitch*>(displaceMapAt.get());
+
+                    if (typedSource->IsCompatibleWith(typedDisplacementMap))
+                    {
+                        // Move the format down all the paths
+                        auto nop = mu::Clone<ASTOpSwitch>(sourceAt);
+
+                        if (nop->def)
+                        {
+							Ptr<ASTOpFixed> defOp = mu::Clone<ASTOpFixed>(this);
+                            defOp->SetChild( defOp->op.args.ImageDisplace.source, typedSource->def );
+                            defOp->SetChild( defOp->op.args.ImageDisplace.displacementMap, typedDisplacementMap->def );
+                            nop->def = defOp;
+                        }
+
+                        // We need to copy the options because we change them
+                        for ( size_t v=0; v<nop->cases.Num(); ++v )
+                        {
+                            if ( nop->cases[v].branch )
+                            {
+								Ptr<ASTOpFixed> bOp = mu::Clone<ASTOpFixed>(this);
+                                bOp->SetChild( bOp->op.args.ImageDisplace.source, typedSource->cases[v].branch );
+                                bOp->SetChild( bOp->op.args.ImageDisplace.displacementMap, typedDisplacementMap->FindBranch(typedSource->cases[v].condition) );
+                                nop->cases[v].branch = bOp;
+                            }
+                        }
+
+                        at = nop;
+                    }
+                }
+
+				// If we didn't optimize already, try to simply sink the source.
+				if (OriginalAt == at)
+				{
+					// Make a project for every path
+					auto nop = mu::Clone<ASTOpSwitch>(sourceAt.get());
+
+					if (nop->def)
+					{
+						Ptr<ASTOpFixed> defOp = mu::Clone<ASTOpFixed>(this);
+						defOp->SetChild(defOp->op.args.ImageDisplace.source, nop->def);
+						nop->def = defOp;
+					}
+
+					// We need to copy the options because we change them
+					for (size_t o = 0; o < nop->cases.Num(); ++o)
+					{
+						if (nop->cases[o].branch)
+						{
+							Ptr<ASTOpFixed> bOp = mu::Clone<ASTOpFixed>(this);
+							bOp->SetChild(bOp->op.args.ImageDisplace.source, nop->cases[o].branch);
+							nop->cases[o].branch = bOp;
+						}
+					}
+
+					at = nop;
+				}
+
+                break;
+            }
+
+            default:
+                break;
+            }
+
+            break;
+        }
+
+
+        default:
+            break;
+        }
+
+        return at;
+    }
+
+    //---------------------------------------------------------------------------------------------
+    //---------------------------------------------------------------------------------------------
+    //---------------------------------------------------------------------------------------------
+    bool SizeOptimiserAST( ASTOpList& roots )
+    {
+        MUTABLE_CPUPROFILER_SCOPE(SizeOptimiser);
+
+        bool modified = false;
+
+        // TODO: isn't top down better suited?
+        ASTOp::Traverse_BottomUp_Unique( roots, [&](Ptr<ASTOp>& n)
+        {
+            auto o = n->OptimiseSize();
+            if (o && o!=n)
+            {
+                modified = true;
+                ASTOp::Replace(n,o);
+            }
+        });
+
+        return modified;
+    }
+
+
+    //---------------------------------------------------------------------------------------------
+    //!
+    //---------------------------------------------------------------------------------------------
+    class Sink_ImageResizeRelAST
+    {
+    public:
+
+        Ptr<ASTOp> Apply( const ASTOp* root )
+        {
+            m_root = root;
+            OldToNew.Empty();
+
+            check(root->GetOpType()==OP_TYPE::IM_RESIZEREL);
+
+			const ASTOpFixed* typedRoot = static_cast<const ASTOpFixed*>(root);
+            m_initialSource = typedRoot->children[typedRoot->op.args.ImageResizeRel.source].child();
+            Ptr<ASTOp> newSource = Visit( m_initialSource, typedRoot );
+
+            // If there is any change, it is the new root.
+            if (newSource!=m_initialSource)
+            {
+                return newSource;
+            }
+
+            return nullptr;
+        }
+
+    protected:
+
+        const ASTOp* m_root;
+        Ptr<ASTOp> m_initialSource;
+        //! For each operation we sink, the map from old instructions to new instructions.
+		TMap<FSinkerOldToNewKey, Ptr<ASTOp>> OldToNew;
+
+		Ptr<ASTOp> Visit(Ptr<ASTOp> at, const ASTOpFixed* currentSinkingOp)
+		{
+			if (!at) return nullptr;
+
+			// Already visited?
+			const Ptr<ASTOp>* Cached = OldToNew.Find({ at,currentSinkingOp });
+			if (Cached)
+			{
+				return *Cached;
+			}
+
+			float scaleX = currentSinkingOp->op.args.ImageResizeRel.factor[0];
+			float scaleY = currentSinkingOp->op.args.ImageResizeRel.factor[1];
+
+			Ptr<ASTOp> newAt = at;
+			switch (at->GetOpType())
+			{
+
+			case OP_TYPE::IM_CONDITIONAL:
+			{
+				// We move the mask creation down the two paths
+				auto newOp = mu::Clone<ASTOpConditional>(at);
+				newOp->yes = Visit(newOp->yes.child(), currentSinkingOp);
+				newOp->no = Visit(newOp->no.child(), currentSinkingOp);
+				newAt = newOp;
+				break;
+			}
+
+			case OP_TYPE::IM_PIXELFORMAT:
+			{
+				// \todo: only if shrinking?
+				Ptr<ASTOpImagePixelFormat> newOp = mu::Clone<ASTOpImagePixelFormat>(at);
+				newOp->Source = Visit(newOp->Source.child(), currentSinkingOp);
+				newAt = newOp;
+				break;
+			}
+
+			case OP_TYPE::IM_SWITCH:
+			{
+				// We move the mask creation down all the paths
+				Ptr<ASTOpSwitch> newOp = mu::Clone<ASTOpSwitch>(at);
+				newOp->def = Visit(newOp->def.child(), currentSinkingOp);
+				for (ASTOpSwitch::FCase& c : newOp->cases)
+				{
+					c.branch = Visit(c.branch.child(), currentSinkingOp);
+				}
+				newAt = newOp;
+				break;
+			}
+
+			case OP_TYPE::IM_SWIZZLE:
+			{
+				Ptr<ASTOpImageSwizzle> newOp = mu::Clone<ASTOpImageSwizzle>(at);
+				for (int s = 0; s < MUTABLE_OP_MAX_SWIZZLE_CHANNELS; ++s)
+				{
+					Ptr<ASTOp> channelOp = newOp->Sources[s].child();
+					if (channelOp)
+					{
+						newOp->Sources[s] = Visit(channelOp, currentSinkingOp);
+					}
+				}
+				newAt = newOp;
+				break;
+			}
+
+			case OP_TYPE::IM_COMPOSE:
+			{
+				// We can only optimise if the layout grid blocks size in pixels is
+				// still an integer after relative scale
+				bool acceptable = false;
+				{
+					const ASTOpImageCompose* typedAt = static_cast<const ASTOpImageCompose*>(at.get());
+					Ptr<ASTOp> originalBaseOp = typedAt->Base.child();
+
+					// \todo: recursion-proof cache?
+					int layoutBlockPixelsX = 0;
+					int layoutBlockPixelsY = 0;
+					originalBaseOp->GetLayoutBlockSize(&layoutBlockPixelsX, &layoutBlockPixelsY);
+
+					int scaledLayoutBlockPixelsX = int(layoutBlockPixelsX * scaleX);
+					int scaledLayoutBlockPixelsY = int(layoutBlockPixelsY * scaleY);
+					int unscaledLayoutBlockPixelsX = int(scaledLayoutBlockPixelsX / scaleX);
+					int unscaledLayoutBlockPixelsY = int(scaledLayoutBlockPixelsY / scaleY);
+					acceptable =
+						(layoutBlockPixelsX != 0 && layoutBlockPixelsY != 0)
+						&&
+						(layoutBlockPixelsX == unscaledLayoutBlockPixelsX)
+						&&
+						(layoutBlockPixelsY == unscaledLayoutBlockPixelsY);
+				}
+
+				if (acceptable)
+				{
+					Ptr<ASTOpImageCompose> newOp = mu::Clone<ASTOpImageCompose>(at);
+
+					Ptr<ASTOp> baseOp = newOp->Base.child();
+					newOp->Base = Visit(baseOp, currentSinkingOp);
+
+					Ptr<ASTOp> blockOp = newOp->BlockImage.child();
+					newOp->BlockImage = Visit(blockOp, currentSinkingOp);
+
+					newAt = newOp;
+				}
+
+				break;
+			}
+
+			case OP_TYPE::IM_PATCH:
+			{
+				Ptr<ASTOpImagePatch> newOp = mu::Clone<ASTOpImagePatch>(at);
+
+				newOp->base = Visit(newOp->base.child(), currentSinkingOp);
+				newOp->patch = Visit(newOp->patch.child(), currentSinkingOp);
+
+				// todo: review if this is always correct, or we need some "divisible" check
+				newOp->location[0] = uint16(newOp->location[0] * scaleX);
+				newOp->location[1] = uint16(newOp->location[1] * scaleY);
+
+				newAt = newOp;
+
+				break;
+			}
+
+			case OP_TYPE::IM_MIPMAP:
+			{
+				Ptr<ASTOpImageMipmap> newOp = mu::Clone<ASTOpImageMipmap>(at.get());
+				Ptr<ASTOp> baseOp = newOp->Source.child();
+				newOp->Source = Visit(baseOp, currentSinkingOp);
+				newAt = newOp;
+				break;
+			}
+
+			case OP_TYPE::IM_INTERPOLATE:
+			{
+				Ptr<ASTOpFixed> newOp = mu::Clone<ASTOpFixed>(at);
+
+				for (int v = 0; v < MUTABLE_OP_MAX_INTERPOLATE_COUNT; ++v)
+				{
+					Ptr<ASTOp> child = newOp->children[newOp->op.args.ImageInterpolate.targets[v]].child();
+					Ptr<ASTOp> bOp = Visit(child, currentSinkingOp);
+					newOp->SetChild(newOp->op.args.ImageInterpolate.targets[v], bOp);
+				}
+
+				newAt = newOp;
+				break;
+			}
+
+			case OP_TYPE::IM_MULTILAYER:
+			{
+				Ptr<ASTOpImageMultiLayer> nop = mu::Clone<ASTOpImageMultiLayer>(at);
+				nop->base = Visit(nop->base.child(), currentSinkingOp);
+				nop->mask = Visit(nop->mask.child(), currentSinkingOp);
+				nop->blend = Visit(nop->blend.child(), currentSinkingOp);
+				newAt = nop;
+				break;
+			}
+
+			case OP_TYPE::IM_LAYER:
+			{
+				Ptr<ASTOpImageLayer> nop = mu::Clone<ASTOpImageLayer>(at);
+
+				Ptr<ASTOp> aOp = nop->base.child();
+				nop->base = Visit(aOp, currentSinkingOp);
+
+				Ptr<ASTOp> bOp = nop->blend.child();
+				nop->blend = Visit(bOp, currentSinkingOp);
+
+				Ptr<ASTOp> mOp = nop->mask.child();
+				nop->mask = Visit(mOp, currentSinkingOp);
+
+				newAt = nop;
+				break;
+			}
+
+			case OP_TYPE::IM_LAYERCOLOUR:
+			{
+				Ptr<ASTOpImageLayerColor> nop = mu::Clone<ASTOpImageLayerColor>(at);
+
+				Ptr<ASTOp> aOp = nop->base.child();
+				nop->base = Visit(aOp, currentSinkingOp);
+
+				Ptr<ASTOp> mOp = nop->mask.child();
+				nop->mask = Visit(mOp, currentSinkingOp);
+
+				newAt = nop;
+				break;
+			}
+
+			case OP_TYPE::IM_DISPLACE:
+			{
+				Ptr<ASTOpFixed> nop = mu::Clone<ASTOpFixed>(at);
+
+				Ptr<ASTOp> sourceOp = nop->children[nop->op.args.ImageDisplace.source].child();
+				nop->SetChild(nop->op.args.ImageDisplace.source, Visit(sourceOp, currentSinkingOp));
+
+				Ptr<ASTOp> mapOp = nop->children[nop->op.args.ImageDisplace.displacementMap].child();
+				nop->SetChild(nop->op.args.ImageDisplace.displacementMap, Visit(mapOp, currentSinkingOp));
+
+				// Make sure we don't scale a constant displacement map, which is very wrong.
+				if (mapOp->GetOpType() == OP_TYPE::IM_MAKEGROWMAP)
+				{
+					newAt = nop;
+				}
+				else
+				{
+					// We cannot resize an already calculated displacement map.
+				}
+
+				break;
+			}
+
+			case OP_TYPE::IM_MAKEGROWMAP:
+			{
+				Ptr<ASTOpImageMakeGrowMap> nop = mu::Clone<ASTOpImageMakeGrowMap>(at);
+				Ptr<ASTOp> maskOp = nop->Mask.child();
+				nop->Mask = Visit(maskOp, currentSinkingOp);
+				newAt = nop;
+				break;
+			}
+
+			case OP_TYPE::IM_INVERT:
+			{
+				Ptr<ASTOpFixed> nop = mu::Clone<ASTOpFixed>(at);
+				Ptr<ASTOp> maskOp = nop->children[nop->op.args.ImageInvert.base].child();
+				nop->SetChild(nop->op.args.ImageInvert.base, Visit(maskOp, currentSinkingOp));
+				newAt = nop;
+				break;
+			}
+
+			case OP_TYPE::IM_SATURATE:
+			{
+				Ptr<ASTOpFixed> nop = mu::Clone<ASTOpFixed>(at);
+				Ptr<ASTOp> maskOp = nop->children[nop->op.args.ImageSaturate.base].child();
+				nop->SetChild(nop->op.args.ImageSaturate.base, Visit(maskOp, currentSinkingOp));
+				newAt = nop;
+				break;
+			}
+
+			case OP_TYPE::IM_TRANSFORM:
+			{
+				// It can only sink in the transform if it doesn't have it's own size.
+				const ASTOpImageTransform* TypedAt = static_cast<const ASTOpImageTransform*>(at.get());
+				if (TypedAt->SizeX == 0 && TypedAt->SizeY == 0)
+				{
+					Ptr<ASTOpImageTransform> NewOp = mu::Clone<ASTOpImageTransform>(at);
+					Ptr<ASTOp> MaskOp = NewOp->Base.child();
+					NewOp->Base = Visit(MaskOp, currentSinkingOp);
+					newAt = NewOp;
+				}
+
+				break;
+			}
+
+			case OP_TYPE::IM_RASTERMESH:
+			{
+				Ptr<ASTOpImageRasterMesh> nop = mu::Clone<ASTOpImageRasterMesh>(at);
+				Ptr<ASTOp> maskOp = nop->mask.child();
+				nop->mask = Visit(maskOp, currentSinkingOp);
+
+				// Resize the image to project as well, assuming that since the target has a different resolution
+				// it make sense for the source image to have a similar resize.
+				// Actually, don't do it because the LODBias will be applied separetely at graph generation time.
+				//auto imageOp = nop->image.child();
+				//nop->image = Visit(imageOp, currentSinkingOp);
+
+				nop->SizeX = uint16(nop->SizeX * scaleX + 0.5f);
+				nop->SizeY = uint16(nop->SizeY * scaleY + 0.5f);
+				newAt = nop;
+				break;
+			}
+
+			default:
+				break;
+			}
+
+			// end on line, replace with sinking op
+			if (at == newAt && at != m_initialSource)
+			{
+				Ptr<ASTOpFixed> newOp = mu::Clone<ASTOpFixed>(currentSinkingOp);
+				check(newOp->GetOpType() == OP_TYPE::IM_RESIZEREL);
+
+				newOp->SetChild(newOp->op.args.ImageResizeRel.source, at);
+
+				newAt = newOp;
+			}
+
+			OldToNew.Add({ at, currentSinkingOp }, newAt);
+
+            return newAt;
+        }
+    };
+
+
+    //---------------------------------------------------------------------------------------------
+    Ptr<ASTOp> ASTOpFixed::OptimiseSize() const
+    {
+        Ptr<ASTOp> at;
+
+        OP_TYPE type = GetOpType();
+        switch ( type )
+        {
+
+        //-------------------------------------------------------------------------------------
+        case OP_TYPE::IM_RESIZE:
+        {
+            auto sourceAt = children[op.args.ImageResize.source].child();
+
+            // The instruction can be sunk
+            OP_TYPE sourceType = sourceAt->GetOpType();
+            switch ( sourceType )
+            {
+            case OP_TYPE::IM_RESIZE:
+            {
+                // Keep top resize
+                Ptr<const ASTOpFixed> sourceOp = static_cast<const ASTOpFixed*>(sourceAt.get());
+
+                Ptr<ASTOpFixed> newOp = mu::Clone<ASTOpFixed>(this);
+                newOp->SetChild( newOp->op.args.ImageResize.source, sourceOp->children[sourceOp->op.args.ImageResize.source]);
+
+                at = newOp;
+                break;
+            }
+
+            case OP_TYPE::IM_PLAINCOLOUR:
+            {
+                // Set the size in the children and remove resize
+				Ptr<ASTOpFixed> sourceOp = mu::Clone<ASTOpFixed>(sourceAt.get());
+                sourceOp->op.args.ImagePlainColour.size[0] = op.args.ImageResize.size[0];
+				sourceOp->op.args.ImagePlainColour.size[1] = op.args.ImageResize.size[1];
+				sourceOp->op.args.ImagePlainColour.LODs = 1; // TODO
+				at = sourceOp;
+                break;
+            }
+
+			case OP_TYPE::IM_TRANSFORM:
+			{
+				// Set the size in the children and remove resize
+				Ptr<ASTOpImageTransform> sourceOp = mu::Clone<ASTOpImageTransform>(sourceAt.get());
+				sourceOp->SizeX = op.args.ImageResize.size[0];
+				sourceOp->SizeY = op.args.ImageResize.size[1];
+				at = sourceOp;
+				break;
+			}
+
+            case OP_TYPE::IM_CONDITIONAL:
+            {
+                // We move the resize down the two paths
+                auto newOp = mu::Clone<ASTOpConditional>(sourceAt);
+
+                Ptr<ASTOpFixed> aOp = mu::Clone<ASTOpFixed>(this);
+                aOp->SetChild( aOp->op.args.ImageResize.source,newOp->yes);
+                newOp->yes = aOp;
+
+                Ptr<ASTOpFixed> bOp = mu::Clone<ASTOpFixed>(this);
+                bOp->SetChild( bOp->op.args.ImageResize.source,newOp->no);
+                newOp->no = bOp;
+
+                at = newOp;
+                break;
+            }
+
+            case OP_TYPE::IM_SWITCH:
+            {
+                // Move the resize down all the paths
+                Ptr<ASTOpSwitch> newOp = mu::Clone<ASTOpSwitch>(sourceAt);
+
+                if (newOp->def)
+                {
+                    Ptr<ASTOpFixed> defOp = mu::Clone<ASTOpFixed>(this);
+                    defOp->SetChild( defOp->op.args.ImageResize.source, newOp->def);
+                    newOp->def = defOp;
+                }
+
+                for ( auto& cas:newOp->cases )
+                {
+                    if ( cas.branch )
+                    {
+                        Ptr<ASTOpFixed> bOp = mu::Clone<ASTOpFixed>(this);
+                        bOp->SetChild( bOp->op.args.ImageResize.source, cas.branch);
+                        cas.branch = bOp;
+                    }
+                }
+
+                at = newOp;
+                break;
+            }
+
+			case OP_TYPE::IM_SWIZZLE:
+			{
+				Ptr<ASTOpImageSwizzle> newOp = mu::Clone<ASTOpImageSwizzle>(sourceAt);
+				for (int s = 0; s < MUTABLE_OP_MAX_SWIZZLE_CHANNELS; ++s)
+				{
+					Ptr<ASTOp> OldChannelOp = newOp->Sources[s].child();
+					if (OldChannelOp)
+					{
+						Ptr<ASTOpFixed> ChannelResize = mu::Clone<ASTOpFixed>(this);
+						ChannelResize->SetChild(ChannelResize->op.args.ImageResize.source, OldChannelOp);
+						newOp->Sources[s] = ChannelResize;
+					}
+				}
+				at = newOp;
+				break;
+			}
+
+            case OP_TYPE::IM_COMPOSE:
+            {
+                Ptr<ASTOpImageCompose> newOp = mu::Clone<ASTOpImageCompose>(sourceAt);
+
+                auto aOp = mu::Clone<ASTOpFixed>(this);
+                aOp->SetChild( aOp->op.args.ImageResize.source, newOp->Base.child());
+                newOp->Base = aOp;
+
+                auto bOp = mu::Clone<ASTOpFixed>(this);
+                bOp->SetChild( bOp->op.args.ImageResize.source, newOp->BlockImage.child());
+                newOp->BlockImage = bOp;
+
+                if ( newOp->Mask )
+                {
+                    auto maskOp = mu::Clone<ASTOpFixed>(this);
+                    maskOp->SetChild( maskOp->op.args.ImageResize.source, newOp->Mask.child());
+                    newOp->Mask = maskOp;
+                }
+
+                at = newOp;
+                break;
+            }
+
+            case OP_TYPE::IM_RASTERMESH:
+            {
+                Ptr<ASTOpImageRasterMesh> newOp = mu::Clone<ASTOpImageRasterMesh>(sourceAt);
+
+                //if ( newOp->op.args.ImageRasterMesh.sizeX != op.args.ImageResize.size[0]
+                //     ||
+                //     newOp->op.args.ImageRasterMesh.sizeY != op.args.ImageResize.size[1] )
+                {
+                    newOp->SizeX = op.args.ImageResize.size[0];
+                    newOp->SizeY = op.args.ImageResize.size[1];
+
+					if (newOp->mask)
+					{
+						Ptr<ASTOpFixed> mop = mu::Clone<ASTOpFixed>(this);
+						mop->SetChild(mop->op.args.ImageResize.source, newOp->mask.child());
+						newOp->mask = mop;
+					}
+
+					// Don't apply absolute resizes to the image to raster: it could even enlarge it.
+					// This should only be scaled with relative resizes, which come from LOD biases, etc.
+					//if (newOp->op.args.ImageRasterMesh.image)
+					//{
+					//	Ptr<ASTOpFixed> mop = mu::Clone<ASTOpFixed>(this);
+					//	mop->SetChild(mop->op.args.ImageResize.source, newOp->children[newOp->op.args.ImageRasterMesh.image].child());
+					//	newOp->SetChild(newOp->op.args.ImageRasterMesh.image, mop);
+					//}
+				}
+
+                at = newOp;
+                break;
+            }
+
+            case OP_TYPE::IM_INTERPOLATE:
+            {
+                Ptr<ASTOpFixed> newOp = mu::Clone<ASTOpFixed>(sourceAt);
+
+                for ( int s=0; s<MUTABLE_OP_MAX_INTERPOLATE_COUNT; ++s )
+                {
+                    auto targetAt = newOp->children[ newOp->op.args.ImageInterpolate.targets[s] ].child();
+                    if ( targetAt )
+                    {
+                        Ptr<ASTOpFixed> sourceOp = mu::Clone<ASTOpFixed>(this);
+                        sourceOp->SetChild( sourceOp->op.args.ImageResize.source, targetAt);
+                        newOp->SetChild( newOp->op.args.ImageInterpolate.targets[s], sourceOp );
+                    }
+                }
+
+                at = newOp;
+                break;
+            }
+
+			case OP_TYPE::IM_INVERT:
+			{
+				Ptr<ASTOpFixed> NewOp = mu::Clone<ASTOpFixed>(sourceAt);
+				Ptr<ASTOp> BaseAt = NewOp->children[NewOp->op.args.ImageInvert.base].child();
+
+				Ptr<ASTOpFixed> NewBase = mu::Clone<ASTOpFixed>(this);
+				NewBase->SetChild(NewBase->op.args.ImageResize.source, BaseAt);
+
+				NewOp->SetChild(NewOp->op.args.ImageInvert.base, NewBase);
+
+				at = NewOp;
+				break;
+			}
+
+			case OP_TYPE::IM_PIXELFORMAT:
+			{				
+				// \todo: only if shrinking?
+				
+				// Only sink the resize if we know that the pixelformat source image is uncompressed.
+				Ptr<ASTOpImagePixelFormat> SourceTyped = static_cast<ASTOpImagePixelFormat*>(sourceAt.get());
+				FImageDesc PixelFormatSourceDesc = SourceTyped ->Source->GetImageDesc();
+				if (PixelFormatSourceDesc.m_format!=EImageFormat::IF_NONE
+					&&
+					!mu::IsCompressedFormat(PixelFormatSourceDesc.m_format) )
+				{
+					Ptr<ASTOpImagePixelFormat> NewOp = mu::Clone<ASTOpImagePixelFormat>(sourceAt);
+					Ptr<ASTOp> BaseAt = NewOp->Source.child();
+
+					Ptr<ASTOpFixed> NewBase = mu::Clone<ASTOpFixed>(this);
+					NewBase->SetChild(NewBase->op.args.ImageResize.source, BaseAt);
+
+					NewOp->Source = NewBase;
+
+					at = NewOp;
+				}
+				break;
+			}
+
+            case OP_TYPE::IM_LAYER:
+            {
+                Ptr<ASTOpImageLayer> newOp = mu::Clone<ASTOpImageLayer>(sourceAt);
+
+                Ptr<ASTOpFixed> baseOp = mu::Clone<ASTOpFixed>(this);
+                baseOp->SetChild( baseOp->op.args.ImageResize.source, newOp->base );
+                newOp->base = baseOp;
+
+                Ptr<ASTOpFixed> blendOp = mu::Clone<ASTOpFixed>(this);
+                blendOp->SetChild( blendOp->op.args.ImageResize.source, newOp->blend );
+                newOp->blend = blendOp;
+
+                auto maskAt = newOp->mask.child();
+                if (maskAt)
+                {
+                    Ptr<ASTOpFixed> maskOp = mu::Clone<ASTOpFixed>(this);
+                    maskOp->SetChild( maskOp->op.args.ImageResize.source, maskAt );
+                    newOp->mask = maskOp;
+                }
+
+                at = newOp;
+                break;
+            }
+
+            case OP_TYPE::IM_LAYERCOLOUR:
+            {
+                Ptr<ASTOpImageLayerColor> newOp = mu::Clone<ASTOpImageLayerColor>(sourceAt);
+
+                Ptr<ASTOpFixed> baseOp = mu::Clone<ASTOpFixed>(this);
+                baseOp->SetChild( baseOp->op.args.ImageResize.source, newOp->base );
+                newOp->base = baseOp;
+
+                auto maskAt = newOp->mask.child();
+                if (maskAt)
+                {
+                    Ptr<ASTOpFixed> maskOp = mu::Clone<ASTOpFixed>(this);
+                    maskOp->SetChild( maskOp->op.args.ImageResize.source, maskAt );
+                    newOp->mask = maskOp;
+                }
+
+                at = newOp;
+                break;
+            }
+
+			case OP_TYPE::IM_DISPLACE:
+			{
+				// In the size optimization phase we can optimize the resize with the displace
+				// because the constants have not been collapsed yet.
+				// We will still check it and sink the size directly below the IM_MAKEGROWMAP op
+				Ptr<ASTOpFixed> SourceTyped = static_cast<ASTOpFixed*>(sourceAt.get());
+				Ptr<ASTOp> OriginalDisplacementMapOp = SourceTyped->children[SourceTyped->op.args.ImageDisplace.displacementMap].m_child;
+				if (OriginalDisplacementMapOp->GetOpType() == OP_TYPE::IM_MAKEGROWMAP)
+				{
+					Ptr<ASTOpFixed> newOp = mu::Clone<ASTOpFixed>(sourceAt);
+
+					Ptr<ASTOpFixed> baseOp = mu::Clone<ASTOpFixed>(this);
+					baseOp->SetChild(baseOp->op.args.ImageResize.source,
+						newOp->children[newOp->op.args.ImageDisplace.source]);
+					newOp->SetChild(newOp->op.args.ImageDisplace.source, baseOp);
+
+					// Clone the map op and replace its children
+					Ptr<ASTOpImageMakeGrowMap> mapOp = mu::Clone<ASTOpImageMakeGrowMap>(OriginalDisplacementMapOp);
+					newOp->SetChild(newOp->op.args.ImageDisplace.displacementMap, mapOp);
+
+					Ptr<ASTOpFixed> mapSourceOp = mu::Clone<ASTOpFixed>(this);
+					mapSourceOp->SetChild(mapSourceOp->op.args.ImageResize.source, mapOp->Mask);
+					mapOp->Mask = mapSourceOp;
+
+					at = newOp;
+				}
+				break;
+			}
+
+            default:
+                break;
+            }
+
+            break;
+        }
+
+
+
+        //-------------------------------------------------------------------------------------
+        case OP_TYPE::IM_RESIZEREL:
+        {
+            auto sourceAt = children[op.args.ImageResizeRel.source].child();
+
+            // The instruction can be sunk
+            OP_TYPE sourceType = sourceAt->GetOpType();
+            switch ( sourceType )
+            {
+
+			case OP_TYPE::IM_BLANKLAYOUT:
+			{
+				Ptr<ASTOpFixed> newOp = mu::Clone<ASTOpFixed>(sourceAt);
+
+				newOp->op.args.ImageBlankLayout.blockSize[0] =
+					uint16(newOp->op.args.ImageBlankLayout.blockSize[0]
+						* op.args.ImageResizeRel.factor[0]
+						+ 0.5f);
+				newOp->op.args.ImageBlankLayout.blockSize[1] =
+					uint16(newOp->op.args.ImageBlankLayout.blockSize[1]
+						* op.args.ImageResizeRel.factor[1]
+						+ 0.5f);
+				at = newOp;
+				break;
+			}
+
+			case OP_TYPE::IM_PLAINCOLOUR:
+			{
+				Ptr<ASTOpFixed> newOp = mu::Clone<ASTOpFixed>(sourceAt);
+
+				newOp->op.args.ImagePlainColour.size[0] = FMath::CeilToInt( newOp->op.args.ImagePlainColour.size[0] * op.args.ImageResizeRel.factor[0] );
+				newOp->op.args.ImagePlainColour.size[1] = FMath::CeilToInt( newOp->op.args.ImagePlainColour.size[1] * op.args.ImageResizeRel.factor[1] );
+				newOp->op.args.ImagePlainColour.LODs = 1; // TODO
+				at = newOp;
+				break;
+			}
+
+			case OP_TYPE::IM_TRANSFORM:
+			{
+				// We can only optimize here if we know the transform result size, otherwise, we will sink the op in the sinker.
+				const ASTOpImageTransform* typedAt = static_cast<const ASTOpImageTransform*>(sourceAt.get());
+				if (typedAt->SizeX != 0 && typedAt->SizeY != 0)
+				{
+					// Set the size in the children and remove resize
+					Ptr<ASTOpImageTransform> sourceOp = mu::Clone<ASTOpImageTransform>(sourceAt.get());
+					sourceOp->SizeX = FMath::CeilToInt32(sourceOp->SizeX * op.args.ImageResizeRel.factor[0]);
+					sourceOp->SizeY = FMath::CeilToInt32(sourceOp->SizeY * op.args.ImageResizeRel.factor[1]);
+					at = sourceOp;
+				}
+				break;
+			}
+
+				// Don't combine. ResizeRel sometimes can resize more children than Resize can do. (see RasterMesh)
+				// It can be combined in an optimization step further in the process, when normal sizes may have been 
+				// optimized already (see OptimizeSemantic)
+//                case OP_TYPE::IM_RESIZE:
+//                {
+//                    break;
+//                }
+
+
+            default:
+            {
+                Sink_ImageResizeRelAST sinker;
+                at = sinker.Apply(this);
+
+                break;
+            }
+
+            }
+
+            break;
+        }
+
+
+        //-------------------------------------------------------------------------------------
+        case OP_TYPE::IM_RESIZELIKE:
+        {
+            Ptr<ImageSizeExpression> sourceSize =
+                    children[op.args.ImageResizeLike.source]->GetImageSizeExpression();
+            Ptr<ImageSizeExpression> sizeSourceSize =
+                    children[op.args.ImageResizeLike.sizeSource]->GetImageSizeExpression();
+
+            if ( *sourceSize == *sizeSourceSize )
+            {
+                at = children[op.args.ImageResizeLike.source].child();
+            }
+
+            else if ( sizeSourceSize->type == ImageSizeExpression::ISET_CONSTANT )
+            {
+                Ptr<ASTOpFixed> newAt = new ASTOpFixed;
+                newAt->op.type = OP_TYPE::IM_RESIZE;
+                newAt->SetChild( newAt->op.args.ImageResize.source, children[op.args.ImageResizeLike.source] );
+                newAt->op.args.ImageResize.size[0] = sizeSourceSize->size[0];
+                newAt->op.args.ImageResize.size[1] = sizeSourceSize->size[1];
+                at = newAt;
+            }
+            else if ( sizeSourceSize->type == ImageSizeExpression::ISET_LAYOUTFACTOR )
+            {
+                // TODO
+                // Skip intermediate ops until the layout
+//                    if ( program.m_code[sizeSourceAt].type != OP_TYPE::IM_BLANKLAYOUT )
+//                    {
+//                        m_modified = true;
+
+//                        OP blackLayoutOp;
+//                        blackLayoutOp.type = OP_TYPE::IM_BLANKLAYOUT;
+//                        blackLayoutOp.args.ImageBlankLayout.layout = sizeSourceSize->layout;
+//                        blackLayoutOp.args.ImageBlankLayout.blockSize[0] = sizeSourceSize->factor[0];
+//                        blackLayoutOp.args.ImageBlankLayout.blockSize[1] = sizeSourceSize->factor[1];
+//                        blackLayoutOp.args.ImageBlankLayout.format = IF_L_UBYTE;
+
+//                        OP opResize = program.m_code[at];
+//                        opResize.args.ImageResizeLike.sizeSource = program.AddOp( blackLayoutOp );
+//                        at = program.AddOp( opResize );
+//                    }
+            }
+
+            break;
+        }
+
+        default:
+            break;
+
+        }
+
+        return at;
+    }
+
+}
diff --git a/Engine/Plugins/Mutable/Source/MutableTools/Private/MuT/CompilerPrivate.h b/Engine/Plugins/Mutable/Source/MutableTools/Private/MuT/CompilerPrivate.h
new file mode 100644
--- /dev/null
+++ b/Engine/Plugins/Mutable/Source/MutableTools/Private/MuT/CompilerPrivate.h
@@ -0,0 +1,113 @@
+// Copyright Epic Games, Inc. All Rights Reserved.
+
+#pragma once
+
+#include "MuT/Compiler.h"
+
+#include "MuT/AST.h"
+#include "MuT/ErrorLogPrivate.h"
+#include "MuT/NodeObjectPrivate.h"
+#include "MuR/Operations.h"
+
+
+namespace mu
+{
+	/** Statistics about the proxy file usage. */
+	struct FProxyFileContext
+	{
+		FProxyFileContext();
+
+		/** Options */
+
+		/** Minimum data size in bytes to dump it to the disk. */
+		uint64 MinProxyFileSize = 64 * 1024;
+
+		/** When creating temporary files, number of retries in case the OS-level call fails. */
+		uint64 MaxFileCreateAttempts = 256;
+
+		/** Statistics */
+		std::atomic<uint64> FilesWritten = 0;
+		std::atomic<uint64> FilesRead = 0;
+		std::atomic<uint64> BytesWritten = 0;
+		std::atomic<uint64> BytesRead = 0;
+
+		/** Internal data. */
+		std::atomic<uint64> CurrentFileIndex = 0;
+	};
+
+		
+    //!
+    class CompilerOptions::Private
+    {
+    public:
+
+        //! Detailed optimization options
+        FModelOptimizationOptions OptimisationOptions;
+		FProxyFileContext DiskCacheContext;
+
+		uint64 EmbeddedDataBytesLimit = 1024;
+		uint64 PackagedDataBytesLimit = 1024*1024*64;
+
+		// \TODO: Unused?
+		int32 MinTextureResidentMipCount = 3;
+
+        int32 ImageCompressionQuality = 0;
+		int32 ImageTiling = 0 ;
+
+		/** If this flag is enabled, the compiler can use concurrency to reduce compile time at the cost of higher CPU and memory usage. */
+		bool bUseConcurrency = false;
+
+		bool bIgnoreStates = false;
+		bool bLog = false;
+
+		FImageOperator::FImagePixelFormatFunc ImageFormatFunc;
+    };
+
+
+    //!
+    struct FStateCompilationData
+    {
+        FObjectState nodeState;
+        Ptr<ASTOp> root;
+        FProgram::FState state;
+		//FStateOptimizationOptions optimisationFlags;
+
+        //! List of instructions that need to be cached to efficiently update this state
+        TArray<Ptr<ASTOp>> m_updateCache;
+
+        //! List of root instructions for the dynamic resources that depend on the runtime
+        //! parameters of this state.
+		TArray< TPair<Ptr<ASTOp>, TArray<FString> > > m_dynamicResources;
+    };
+
+
+    //!
+    class Compiler::Private
+    {
+    public:
+
+        Private()
+        {
+            ErrorLog = new mu::ErrorLog();
+        }
+
+        Ptr<ErrorLog> ErrorLog;
+
+        /** */
+        Ptr<CompilerOptions> m_options;
+
+		/** Store for additional data generated during compilation, but not necessary for the runtime. */
+		struct FAdditionalData
+		{
+			/** Source data descriptor for every image constant that has been generated. 
+			* It must have the same size than the Program::ConstantImages array.
+			*/
+			TArray<FSourceDataDescriptor> SourceImagePerConstant;
+		};
+
+		//!
+		void GenerateRoms(Model*, const CompilerOptions*, const FLinkerOptions::FAdditionalData&);
+
+    };
+
+}
diff --git a/Engine/Plugins/Mutable/Source/MutableTools/Private/MuT/DataPacker.cpp b/Engine/Plugins/Mutable/Source/MutableTools/Private/MuT/DataPacker.cpp
new file mode 100644
--- /dev/null
+++ b/Engine/Plugins/Mutable/Source/MutableTools/Private/MuT/DataPacker.cpp
@@ -0,0 +1,833 @@
+// Copyright Epic Games, Inc. All Rights Reserved.
+
+#include "MuT/DataPacker.h"
+
+#include "Containers/Array.h"
+#include "HAL/PlatformCrt.h"
+#include "HAL/UnrealMemory.h"
+#include "MuR/CodeVisitor.h"
+#include "MuR/Image.h"
+#include "MuR/Layout.h"
+#include "MuR/Mesh.h"
+#include "MuR/MeshBufferSet.h"
+#include "MuR/ModelPrivate.h"
+#include "MuR/MutableTrace.h"
+#include "MuR/Operations.h"
+#include "MuR/Ptr.h"
+#include "MuR/RefCounted.h"
+#include "MuT/ASTOpConstantResource.h"
+#include "MuT/ASTOpImageCompose.h"
+#include "MuT/ASTOpImageLayer.h"
+#include "MuT/ASTOpImageLayerColor.h"
+#include "MuT/ASTOpImageMultiLayer.h"
+#include "MuT/ASTOpInstanceAdd.h"
+#include "MuT/ASTOpMeshExtractLayoutBlocks.h"
+#include "MuT/ASTOpMeshRemoveMask.h"
+#include "MuT/ASTOpMeshDifference.h"
+#include "MuT/ASTOpMeshMorph.h"
+#include "MuT/ASTOpLayoutFromMesh.h"
+#include "MuT/ASTOpImageRasterMesh.h"
+#include "MuT/CompilerPrivate.h"
+
+
+namespace mu
+{
+
+    //---------------------------------------------------------------------------------------------
+    //---------------------------------------------------------------------------------------------
+    //---------------------------------------------------------------------------------------------
+    SubtreeSearchConstantVisitor::SubtreeSearchConstantVisitor( FProgram& program,
+                                                                OP::ADDRESS constant,
+                                                                OP_TYPE optype )
+        : m_program( program ),
+          m_constant( constant ),
+          m_opType( optype )
+    {
+        m_visited.SetNum( program.m_opAddress.Num() );
+        if ( m_visited.Num() )
+        {
+            FMemory::Memset( &m_visited[0], 0, m_visited.Num() );
+        }
+    }
+
+
+    //---------------------------------------------------------------------------------------------
+    bool SubtreeSearchConstantVisitor::Run( OP::ADDRESS root )
+    {
+        // too spammy
+        // \todo: called too many times!
+        // MUTABLE_CPUPROFILER_SCOPE(SubtreeSearchConstantVisitor);
+
+        bool found = false;
+
+        TArray< TPair<bool,int> > pending;
+        pending.Reserve(m_program.m_opAddress.Num()/4);
+
+		pending.Add({ false,root });
+
+        while( pending.Num() )
+        {
+			TPair<bool,int> item = pending.Pop();
+
+            if (item.Key)
+            {
+                // Item indicating we finished with all the children of a parent.
+
+                // Propagate the "found" state upwards
+                m_visited[item.Value]=1;
+
+                //OP cop = m_program.m_code[item.Value];
+                ForEachReference( m_program, (OP::ADDRESS)item.Value, [&](OP::ADDRESS ref)
+                {
+                    if (ref && m_visited[ref]==2)
+                    {
+                        m_visited[item.Value]=2;
+                    }
+                });
+            }
+            else
+            {
+                if ( !m_visited[item.Value] )
+                {
+                    OP_TYPE thisOpType = m_program.GetOpType(item.Value);
+                    if ( m_opType == thisOpType )
+                    {
+						OP::ResourceConstantArgs args = m_program.GetOpArgs<OP::ResourceConstantArgs>(item.Value);
+                        if ( m_constant == args.value )
+                        {
+                            found = true;
+                            m_visited[item.Value] = 2;
+                        }
+                    }
+
+                    if (!found)
+                    {
+						pending.Add({ true,item.Value });
+
+                        //OP cop = m_program.m_code[item.Value];
+                        ForEachReference( m_program, item.Value, [&](OP::ADDRESS ref)
+                        {
+                            if (ref && !m_visited[ref])
+                            {
+								pending.Add({ false,ref });
+                            }
+                        });
+                    }
+                }
+                else
+                {
+                    found = (m_visited[item.Value]==2);
+                }
+            }
+        }
+
+        return found;
+    }
+
+
+
+    //---------------------------------------------------------------------------------------------
+    //! Get all the parameters that affect the constant.
+    //! "Affect" means that the constant may be used or not depending on the parameter.
+    //---------------------------------------------------------------------------------------------
+    class GatherParametersVisitor : public UniqueConstCodeVisitorIterative< TArray<int> >
+    {
+    public:
+
+        GatherParametersVisitor( FProgram& program, OP::ADDRESS constant, OP_TYPE opType )
+            : m_constSearch( program, constant, opType )
+        {
+            MUTABLE_CPUPROFILER_SCOPE(GatherParametersVisitor);
+
+			TArray<int> currentParams;
+            currentParams.SetNumZeroed( program.m_parameters.Num() );
+            SetDefaultState(currentParams);
+
+            // TODO: we can optimize by precalculating what ops have the required optype below
+            m_opType = opType;
+            m_constant = constant;
+
+            m_allParams.SetNumZeroed( program.m_parameters.Num() );
+            FullTraverse( program );
+
+            for (int32 i=0; i<m_allParams.Num(); ++i )
+            {
+                if ( m_allParams[i] )
+                {
+                    m_sortedParams.Add( int(i) );
+                }
+            }
+        }
+
+
+        virtual bool Visit( OP::ADDRESS at, FProgram& program )
+        {
+            bool recurse = true;
+
+            if ( m_constSearch.Run( at ) )
+            {
+                switch ( program.GetOpType(at) )
+                {
+                case OP_TYPE::NU_CONDITIONAL:
+                case OP_TYPE::SC_CONDITIONAL:
+                case OP_TYPE::CO_CONDITIONAL:
+                case OP_TYPE::IM_CONDITIONAL:
+                case OP_TYPE::ME_CONDITIONAL:
+                case OP_TYPE::LA_CONDITIONAL:
+                case OP_TYPE::IN_CONDITIONAL:
+				case OP_TYPE::ED_CONDITIONAL:
+                {
+					OP::ConditionalArgs args = program.GetOpArgs<OP::ConditionalArgs>(at);
+
+                    // If the constant is present in only one of the 2 branches
+                    bool foundYes = m_constSearch.Run( args.yes  );
+                    bool foundNo = m_constSearch.Run( args.no );
+
+                    if ( foundYes!=foundNo )
+                    {
+                        m_conditionVisitor.Run( args.condition, program );
+
+                        TArray<int32> currentParams = GetCurrentState();
+                        for (int32 p=0; p<m_conditionVisitor.RelevantParams.Num(); ++p )
+                        {
+                            currentParams[ m_conditionVisitor.RelevantParams[p] ]++;
+                        }
+
+                        // Set state for child recursion
+                        SetCurrentState( currentParams );
+                    }
+
+                    break;
+                }
+
+                case OP_TYPE::NU_SWITCH:
+                case OP_TYPE::SC_SWITCH:
+                case OP_TYPE::CO_SWITCH:
+                case OP_TYPE::IM_SWITCH:
+                case OP_TYPE::ME_SWITCH:
+                case OP_TYPE::LA_SWITCH:
+				case OP_TYPE::ED_SWITCH:
+                {
+					const uint8* data = program.GetOpArgsPointer(at);
+					
+					OP::ADDRESS VarAddress;
+					FMemory::Memcpy( &VarAddress, data, sizeof(OP::ADDRESS) );
+
+                    m_conditionVisitor.Run( VarAddress, program );
+
+					TArray<int32> currentParams = GetCurrentState();
+                    for (int32 p=0; p<m_conditionVisitor.RelevantParams.Num(); ++p )
+                    {
+                        currentParams[ m_conditionVisitor.RelevantParams[p] ]++;
+                    }
+
+                    // Set state for child recursion
+                    SetCurrentState( currentParams );
+
+                    break;
+                }
+
+                default:
+                {
+                    // If we hit the constant we are analysing
+                    if ( m_opType == OP_TYPE::ME_CONSTANT
+                         &&
+                         program.GetOpType(at) == m_opType )
+                    {
+						OP::MeshConstantArgs args = program.GetOpArgs<OP::MeshConstantArgs>(at);
+                        if ( args.value == m_constant )
+                        {
+                            // Accumulate the currently relevant parameters
+                            const TArray<int>& currentParams = GetCurrentState();
+                            for (int32 i=0; i<currentParams.Num(); ++i )
+                            {
+                                m_allParams[i] += currentParams[i];
+                            }
+                        }
+                    }
+                    else if ( program.GetOpType(at) == m_opType )
+                    {
+						OP::ResourceConstantArgs args = program.GetOpArgs<OP::ResourceConstantArgs>(at);
+                        if ( args.value == m_constant )
+                        {
+                            // Accumulate the currently relevant parameters
+                            const TArray<int>& currentParams = GetCurrentState();
+                            for (int32 i=0; i<currentParams.Num(); ++i )
+                            {
+                                m_allParams[i] += currentParams[i];
+                            }
+                        }
+                    }
+
+                    break;
+                }
+
+                }
+
+            }
+
+            return recurse;
+        }
+
+    public:
+
+        //! Output result
+		TArray<int> m_sortedParams;
+
+    private:
+
+        OP::ADDRESS m_constant;
+        OP_TYPE m_opType;
+
+		TArray<int> m_allParams;
+
+        SubtreeParametersVisitor m_conditionVisitor;
+        SubtreeSearchConstantVisitor m_constSearch;
+
+    };
+
+
+    //---------------------------------------------------------------------------------------------
+    //---------------------------------------------------------------------------------------------
+    //---------------------------------------------------------------------------------------------
+    class AccumulateImageFormatsAST : public Visitor_TopDown_Unique_Const< TArray<bool> >
+    {
+    public:
+
+        void Run( const ASTOpList& roots )
+        {
+            MUTABLE_CPUPROFILER_SCOPE(AccumulateImageFormatsAST);
+
+            // TODO: Use statically sized vectors (std::array)
+            //vector<bool> initial( IF_COUNT, true );
+            //std::fill( m_supportedFormats.begin(), m_supportedFormats.end(), initial );
+
+			TArray<bool> defaultState;
+			defaultState.SetNumZeroed(int32(EImageFormat::IF_COUNT));
+
+            Traverse( roots, defaultState );
+        }
+
+
+        bool Visit( const Ptr<ASTOp>& node ) override
+        {
+            bool recurse = true;
+
+            const TArray<bool>& currentFormats = GetCurrentState();
+
+			TArray<bool> defaultState;
+			defaultState.SetNumZeroed(int32(EImageFormat::IF_COUNT));
+			bool allFalse = currentFormats == defaultState;
+
+            // Can we use the cache?
+            if (allFalse)
+            {
+                if (m_visited.count(node))
+                {
+                    return false;
+                }
+
+                m_visited.insert(node);
+            }
+
+            switch ( node->GetOpType() )
+            {
+            case OP_TYPE::IM_CONSTANT:
+            {
+                // Remove unsupported formats
+				const ASTOpConstantResource* op = static_cast<const ASTOpConstantResource*>(node.get());
+                if (!m_supportedFormats.count(op))
+                {
+					TArray<bool> initial;
+					initial.Init( true, int32(EImageFormat::IF_COUNT) );
+                    m_supportedFormats.insert( std::make_pair(op, std::move(initial)) );
+                }
+
+                for ( unsigned f=0; f< unsigned(EImageFormat::IF_COUNT); ++f )
+                {
+                    if ( !currentFormats[f] )
+                    {
+                        m_supportedFormats[op][f] = false;
+                    }
+                }
+                recurse = false;
+                break;
+            }
+
+            case OP_TYPE::IM_SWITCH:
+            case OP_TYPE::IM_CONDITIONAL:
+                // Switches and conditionals don't change the supported formats
+                break;
+
+            case OP_TYPE::IM_COMPOSE:
+            {
+                recurse = false;
+				const ASTOpImageCompose* op = static_cast<const ASTOpImageCompose*>(node.get());
+
+				TArray<bool> newState;
+				newState.Init(false, int32(EImageFormat::IF_COUNT));
+                RecurseWithState( op->Layout.child(), newState );
+                RecurseWithState( op->Base.child(), newState );
+                RecurseWithState( op->BlockImage.child(), newState );
+
+                if ( op->Mask )
+                {
+                    newState[(int32)EImageFormat::IF_L_UBIT_RLE] = true;
+                    RecurseWithState( op->Mask.child(), newState );
+                }
+                break;
+            }
+
+            case OP_TYPE::IM_LAYERCOLOUR:
+            {
+                recurse = false;
+				const ASTOpImageLayerColor* op = static_cast<const ASTOpImageLayerColor*>(node.get());
+				TArray<bool> newState;
+				newState.Init(false, int32(EImageFormat::IF_COUNT));
+				RecurseWithState( op->base.child(), newState );
+                RecurseWithState( op->color.child(), newState );
+
+                if ( op->mask )
+                {
+                    newState[(int32)EImageFormat::IF_L_UBYTE] = true;
+                    newState[(int32)EImageFormat::IF_L_UBYTE_RLE] = true;
+
+                    RecurseWithState( op->mask.child(), newState );
+                }
+                break;
+            }
+
+            case OP_TYPE::IM_LAYER:
+            {
+                recurse = false;
+				const ASTOpImageLayer* op = static_cast<const ASTOpImageLayer*>(node.get());
+				TArray<bool> newState;
+				newState.Init(false, int32(EImageFormat::IF_COUNT));
+				RecurseWithState( op->base.child(), newState );
+                RecurseWithState( op->blend.child(), newState );
+
+                if (op->mask)
+                {
+                    newState[(int32)EImageFormat::IF_L_UBYTE] = true;
+                    newState[(int32)EImageFormat::IF_L_UBYTE_RLE] = true;
+
+                    RecurseWithState( op->mask.child(), newState );
+                }
+                break;
+            }
+
+            case OP_TYPE::IM_MULTILAYER:
+            {
+                recurse = false;
+				const ASTOpImageMultiLayer* op = static_cast<const ASTOpImageMultiLayer*>(node.get());
+				TArray<bool> newState;
+				newState.Init(false, int32(EImageFormat::IF_COUNT));
+				RecurseWithState( op->base.child(), newState );
+                RecurseWithState( op->blend.child(), newState );
+
+                if (op->mask)
+                {
+                    newState[(size_t)EImageFormat::IF_L_UBYTE] = true;
+                    newState[(size_t)EImageFormat::IF_L_UBYTE_RLE] = true;
+
+                    RecurseWithState( op->mask.child(), newState );
+                }
+                break;
+            }
+
+            case OP_TYPE::IM_DISPLACE:
+            {
+                recurse = false;
+				const ASTOpFixed* op = static_cast<const ASTOpFixed*>(node.get());
+				TArray<bool> newState;
+				newState.Init(false, int32(EImageFormat::IF_COUNT));
+				RecurseWithState( op->children[op->op.args.ImageDisplace.source].child(), newState );
+
+                newState[(int32)EImageFormat::IF_L_UBYTE ] = true;
+                newState[(int32)EImageFormat::IF_L_UBYTE_RLE ] = true;
+
+                RecurseWithState( op->children[op->op.args.ImageDisplace.displacementMap].child(), newState );
+                break;
+            }
+
+            default:
+            {
+                //m_currentFormats.Add(vector<bool>(IF_COUNT, false));
+                //Recurse(at, program);
+                //m_currentFormats.pop_back();
+
+				TArray<bool> newState;
+				newState.Init(false, int32(EImageFormat::IF_COUNT));
+				if (currentFormats != newState)
+                {
+                    RecurseWithState(node, newState);
+                    recurse = false;
+                }
+                else
+                {
+                    recurse = true;
+                }
+                break;
+            }
+
+            }
+
+            return recurse;
+        }
+
+    public:
+
+        //! Result of this visitor:
+        //! Formats known to be supported by every constant image.
+        std::unordered_map< Ptr<const ASTOpConstantResource>, TArray<bool> > m_supportedFormats;
+
+    private:
+
+        //! Cache. Only valid is current formats are all false.
+        std::unordered_set<Ptr<ASTOp>> m_visited;
+
+    };
+
+
+
+
+    //---------------------------------------------------------------------------------------------
+    //---------------------------------------------------------------------------------------------
+    //---------------------------------------------------------------------------------------------
+    class AccumulateMeshChannelUsageAST : public Visitor_TopDown_Unique_Const< uint64_t >
+    {
+    public:
+
+        void Run( const ASTOpList& roots )
+        {
+            MUTABLE_CPUPROFILER_SCOPE(AccumulateMeshChannelUsageAST);
+
+            // Sanity check in case we add more semantics
+            static_assert(MBS_COUNT<sizeof(uint64)*8, "Too many mesh buffer semantics." );
+
+            // Default state: we need everything except internal semantics
+            uint64 defaultState = 0xffffffffffffffff;
+            defaultState ^= (UINT64_C(1)<<MBS_LAYOUTBLOCK);
+            defaultState ^= (UINT64_C(1)<<MBS_VERTEXINDEX);
+
+            Traverse(roots,defaultState);
+        }
+
+
+        bool Visit( const Ptr<ASTOp>& node ) override
+        {
+            bool recurse = true;
+
+            uint64 currentSemantics = GetCurrentState();
+
+            switch ( node->GetOpType() )
+            {
+
+            case OP_TYPE::ME_CONSTANT:
+            {
+                // Accumulate necessary semantics
+				const ASTOpConstantResource* op = static_cast<const ASTOpConstantResource*>(node.get());
+                uint64 currentFlags = 0;
+                if (m_requiredSemantics.count(op))
+                {
+                    currentFlags = m_requiredSemantics[op];
+                }
+                else
+                {
+                    m_requiredSemantics.insert(std::make_pair(op,currentFlags));
+                }
+
+                currentFlags |= currentSemantics;
+                m_requiredSemantics[op] = currentFlags;
+                recurse = false;
+                break;
+            }
+
+            // TODO: These could probably optimise something
+            //case OP_TYPE::IM_RASTERMESH: break;
+
+            case OP_TYPE::ME_DIFFERENCE:
+            {
+                recurse = false;
+
+				const ASTOpMeshDifference* op = static_cast<const ASTOpMeshDifference*>(node.get());
+
+                uint64 newState = currentSemantics;
+                newState |= (UINT64_C(1)<<MBS_VERTEXINDEX);
+                RecurseWithState( op->Base.child(), newState );
+
+                RecurseWithState( op->Target.child(), currentSemantics );
+                break;
+             }
+
+            case OP_TYPE::ME_REMOVEMASK:
+            {
+                recurse = false;
+
+				const ASTOpMeshRemoveMask* op = static_cast<const ASTOpMeshRemoveMask*>(node.get());
+
+                uint64 newState = currentSemantics;
+                newState |= (UINT64_C(1)<<MBS_VERTEXINDEX);
+                RecurseWithState( op->source.child(), newState );
+                for( const TPair<ASTChild, ASTChild>& r: op->removes )
+                {
+                    RecurseWithState( r.Value.child(), newState );
+                }
+                break;
+             }
+
+            case OP_TYPE::ME_MORPH:
+            {
+                recurse = false;
+
+				const ASTOpMeshMorph* op = static_cast<const ASTOpMeshMorph*>(node.get());
+
+                uint64 newState = currentSemantics;
+                newState |= (UINT64_C(1)<<MBS_VERTEXINDEX);
+                RecurseWithState( op->Base.child(), newState );
+                RecurseWithState( op->Target.child(), newState );
+                break;
+             }
+
+            case OP_TYPE::ME_APPLYLAYOUT:
+            {
+                recurse = false;
+
+				const ASTOpFixed* op = static_cast<const ASTOpFixed*>(node.get());
+
+                uint64 newState = currentSemantics;
+                newState |= (UINT64_C(1)<<MBS_LAYOUTBLOCK);
+                RecurseWithState( op->children[op->op.args.MeshApplyLayout.mesh].child(), newState );
+
+                RecurseWithState( op->children[op->op.args.MeshApplyLayout.layout].child(), currentSemantics );
+                break;
+            }
+
+			case OP_TYPE::ME_PROJECT:
+			{
+				recurse = false;
+
+				const ASTOpFixed* op = static_cast<const ASTOpFixed*>(node.get());
+
+				uint64 newState = currentSemantics;
+				newState |= (UINT64_C(1) << MBS_LAYOUTBLOCK);
+				RecurseWithState(op->children[op->op.args.MeshProject.mesh].child(), newState);
+
+				RecurseWithState(op->children[op->op.args.MeshProject.projector].child(), currentSemantics);
+				break;
+			}
+
+			case OP_TYPE::IM_RASTERMESH:
+			{
+				recurse = false;
+
+				const ASTOpImageRasterMesh* op = static_cast<const ASTOpImageRasterMesh*>(node.get());
+
+				uint64 newState = currentSemantics;
+				newState |= (UINT64_C(1) << MBS_LAYOUTBLOCK);
+				RecurseWithState(op->mesh.child(), newState);
+
+				RecurseWithState(op->image.child(), currentSemantics);
+				RecurseWithState(op->angleFadeProperties.child(), currentSemantics);
+				RecurseWithState(op->mask.child(), currentSemantics);
+				RecurseWithState(op->projector.child(), currentSemantics);
+				break;
+			}
+
+            case OP_TYPE::ME_EXTRACTLAYOUTBLOCK:
+            {
+                recurse = false;
+
+				const ASTOpMeshExtractLayoutBlocks* op = static_cast<const ASTOpMeshExtractLayoutBlocks*>(node.get());
+
+                // todo: check if we really need all of them
+                uint64 newState = currentSemantics;
+                newState |= (UINT64_C(1)<<MBS_LAYOUTBLOCK);
+                newState |= (UINT64_C(1)<<MBS_VERTEXINDEX);
+
+                RecurseWithState( op->Source.child(), newState );
+                break;
+            }
+
+			case OP_TYPE::LA_FROMMESH:
+			{
+				recurse = false;
+
+				const ASTOpLayoutFromMesh* op = static_cast<const ASTOpLayoutFromMesh*>(node.get());
+
+				uint64 newState = currentSemantics;
+				newState |= (UINT64_C(1) << MBS_LAYOUTBLOCK);
+
+				RecurseWithState(op->Mesh.child(), newState);
+				break;
+			}
+
+            case OP_TYPE::IN_ADDMESH:
+            {
+                recurse = false;
+
+				const ASTOpInstanceAdd* op = static_cast<const ASTOpInstanceAdd*>(node.get());
+
+                RecurseWithState( op->instance.child(), currentSemantics );
+
+                uint64 newState = GetDefaultState();
+                RecurseWithState( op->value.child(), newState );
+                break;
+            }
+
+            default:
+                // Unhandled op, we may need everything? Recurse with current state?
+                //uint64 newState = 0xffffffffffffffff;
+                break;
+
+            }
+
+            return recurse;
+        }
+
+    public:
+
+        //! Result of this visitor:
+        //! Used mesh channel semantics for each constant mesh
+        std::unordered_map< Ptr<const ASTOpConstantResource>, uint64 > m_requiredSemantics;
+
+    };
+
+
+    //---------------------------------------------------------------------------------------------
+    //---------------------------------------------------------------------------------------------
+    //---------------------------------------------------------------------------------------------
+
+    // Todo: move to its own file
+    inline void MeshRemoveUnusedBufferSemantics( Mesh* pMesh, uint64 usedSemantics )
+    {
+        // right now we only remove entire buffers if no channel is used
+        // TODO: remove from inside the buffer?
+        for (int v=0; v<pMesh->GetVertexBuffers().GetBufferCount(); )
+        {
+            bool used = false;
+            for (int c=0; !used && c<pMesh->GetVertexBuffers().GetBufferChannelCount(v); ++c)
+            {
+                EMeshBufferSemantic semantic;
+                pMesh->GetVertexBuffers().GetChannel(v,c,&semantic,
+                                                      nullptr, nullptr, nullptr, nullptr);
+                used = (( (UINT64_C(1)<<semantic) ) & usedSemantics) != 0;
+            }
+
+            if (!used)
+            {
+                TArray<FMeshBuffer>& buffers = pMesh->GetVertexBuffers().Buffers;
+                buffers.RemoveAt(v);
+            }
+            else
+            {
+                ++v;
+            }
+        }
+
+        // TODO: hack, if we don't need layouts, remove them.
+        {
+            uint64 layoutSemantics = 0;
+            layoutSemantics |= (UINT64_C(1)<<MBS_LAYOUTBLOCK);
+
+            if ( (usedSemantics&layoutSemantics) == 0)
+            {
+                pMesh->Layouts.Empty();
+            }
+        }
+    }
+
+
+    //---------------------------------------------------------------------------------------------
+    void DataOptimise( const CompilerOptions* Options, ASTOpList& roots )
+    {
+		int32 ImageCompressionQuality = Options->GetPrivate()->ImageCompressionQuality;
+		const FModelOptimizationOptions& OptimizeOptions = Options->GetPrivate()->OptimisationOptions;
+
+        // Images
+        AccumulateImageFormatsAST accFormat;
+        accFormat.Run( roots );
+
+        // See if we can convert some constants to more efficient formats
+        ASTOp::Traverse_BottomUp_Unique_NonReentrant( roots, [&](Ptr<ASTOp>& n)
+        {
+            if (n->GetOpType()==OP_TYPE::IM_CONSTANT)
+            {
+				ASTOpConstantResource* typed = static_cast<ASTOpConstantResource*>(n.get());
+                Ptr<const Image> pOld = static_cast<const Image*>(typed->GetValue().get());
+
+				FImageOperator ImOp = FImageOperator::GetDefault( Options->GetPrivate()->ImageFormatFunc );
+
+				// See if there is a better format for this image
+				FVector4f PlainColor;
+				if ( pOld->IsPlainColour(PlainColor) )
+				{
+					// It is more efficient to just have an instruction for it instead, to avoid the overhead
+					// of data loading. 
+					// Warning This eliminates the mips. \TODO: Add support for mips in plaincolour instruction?
+					Ptr<ASTOpFixed> NewColor = new ASTOpFixed;
+					NewColor->op.type = OP_TYPE::CO_CONSTANT;
+					NewColor->op.args.ColourConstant.value[0] = PlainColor[0];
+					NewColor->op.args.ColourConstant.value[1] = PlainColor[1];
+					NewColor->op.args.ColourConstant.value[2] = PlainColor[2];
+					NewColor->op.args.ColourConstant.value[3] = PlainColor[3];
+
+					Ptr<ASTOpFixed> NewPlain = new ASTOpFixed;
+					NewPlain->op.type = OP_TYPE::IM_PLAINCOLOUR;
+					NewPlain->SetChild( NewPlain->op.args.ImagePlainColour.colour, NewColor );
+					NewPlain->op.args.ImagePlainColour.format = pOld->GetFormat();
+					NewPlain->op.args.ImagePlainColour.size[0] = pOld->GetSizeX();
+					NewPlain->op.args.ImagePlainColour.size[1] = pOld->GetSizeY();
+					NewPlain->op.args.ImagePlainColour.LODs = 1;
+
+					ASTOp::Replace(n, NewPlain);
+
+				}
+				else if ( accFormat.m_supportedFormats[typed][(int32)EImageFormat::IF_L_UBIT_RLE] )
+                {
+                    ImagePtr pNew = ImOp.ImagePixelFormat( ImageCompressionQuality, pOld.get(), EImageFormat::IF_L_UBIT_RLE );
+
+                    // Only replace if the compression was worth!
+                    int32 oldSize = pOld->GetDataSize();
+					int32 newSize = pNew->GetDataSize();
+                    if (float(oldSize) > float(newSize) * OptimizeOptions.MinRLECompressionGain)
+                    {
+                        typed->SetValue(pNew, OptimizeOptions.DiskCacheContext);
+                    }
+                }
+                else if ( accFormat.m_supportedFormats[typed][(int32)EImageFormat::IF_L_UBYTE_RLE] )
+                {
+                    ImagePtr pNew = ImOp.ImagePixelFormat( ImageCompressionQuality, pOld.get(), EImageFormat::IF_L_UBYTE_RLE );
+
+                    // Only replace if the compression was worth!
+					int32 oldSize = pOld->GetDataSize();
+					int32 newSize = pNew->GetDataSize();
+                    if (float(oldSize) > float(newSize) * OptimizeOptions.MinRLECompressionGain)
+                    {
+                        typed->SetValue(pNew, OptimizeOptions.DiskCacheContext);
+                    }
+                }
+            }
+        });
+
+
+        // Meshes
+        AccumulateMeshChannelUsageAST meshSemanticsVisitor;
+        meshSemanticsVisitor.Run( roots );
+
+        // See if we can remove some buffers from the constants
+        ASTOp::Traverse_BottomUp_Unique_NonReentrant( roots, [&](Ptr<ASTOp>& n)
+        {
+            if (n->GetOpType()==OP_TYPE::ME_CONSTANT)
+            {
+				ASTOpConstantResource* typed = static_cast<ASTOpConstantResource*>(n.get());
+                Ptr<Mesh> pMesh = static_cast<const Mesh*>(typed->GetValue().get())->Clone();
+                MeshRemoveUnusedBufferSemantics( pMesh.get(), meshSemanticsVisitor.m_requiredSemantics[typed]);
+                typed->SetValue(pMesh, OptimizeOptions.DiskCacheContext);
+            }
+        });
+
+    }
+
+
+}
diff --git a/Engine/Plugins/Mutable/Source/MutableTools/Private/MuT/DataPacker.h b/Engine/Plugins/Mutable/Source/MutableTools/Private/MuT/DataPacker.h
new file mode 100644
--- /dev/null
+++ b/Engine/Plugins/Mutable/Source/MutableTools/Private/MuT/DataPacker.h
@@ -0,0 +1,45 @@
+// Copyright Epic Games, Inc. All Rights Reserved.
+
+#pragma once
+
+#include "MuR/Operations.h"
+#include "MuT/AST.h"
+
+namespace mu { struct FProgram; }
+
+
+#define MUTABLE_HASH_SEED					((uint32)0xcadababa)
+
+
+namespace mu
+{
+	class CompilerOptions;
+
+    /** Convert constant data to different formats, based on their usage. */
+    extern void DataOptimise( const CompilerOptions*, ASTOpList& roots );
+
+    //---------------------------------------------------------------------------------------------
+    //! Find the given constant in a subtree
+    //---------------------------------------------------------------------------------------------
+    class SubtreeSearchConstantVisitor
+    {
+    public:
+
+        SubtreeSearchConstantVisitor( FProgram& program, OP::ADDRESS constant, OP_TYPE optype );
+
+        bool Run( OP::ADDRESS root );
+
+
+    private:
+
+        FProgram& m_program;
+        OP::ADDRESS m_constant;
+        OP_TYPE m_opType;
+
+        //! 0: not visited
+        //! 1: visited and not found
+        //! 2: visited and found
+        TArray<uint8> m_visited;
+    };
+
+}
diff --git a/Engine/Plugins/Mutable/Source/MutableTools/Private/MuT/ErrorLog.cpp b/Engine/Plugins/Mutable/Source/MutableTools/Private/MuT/ErrorLog.cpp
new file mode 100644
--- /dev/null
+++ b/Engine/Plugins/Mutable/Source/MutableTools/Private/MuT/ErrorLog.cpp
@@ -0,0 +1,341 @@
+// Copyright Epic Games, Inc. All Rights Reserved.
+
+#include "MuT/ErrorLog.h"
+
+#include "Containers/Array.h"
+#include "Containers/Map.h"
+#include "Containers/UnrealString.h"
+#include "HAL/PlatformMath.h"
+#include "HAL/UnrealMemory.h"
+#include "Logging/LogCategory.h"
+#include "Logging/LogMacros.h"
+#include "Misc/AssertionMacros.h"
+#include "MuR/Layout.h"
+#include "MuR/Mesh.h"
+#include "MuR/Model.h"
+#include "MuR/ModelPrivate.h"
+#include "MuR/MutableMath.h"
+#include "MuR/Operations.h"
+#include "MuR/ParametersPrivate.h"
+#include "MuR/Platform.h"
+#include "MuR/Skeleton.h"
+#include "MuR/System.h"
+#include "MuR/MutableRuntimeModule.h"
+#include "MuT/ErrorLogPrivate.h"
+#include "Templates/Tuple.h"
+#include "Trace/Detail/Channel.h"
+
+namespace mu
+{
+
+
+	//---------------------------------------------------------------------------------------------
+	ErrorLog::ErrorLog()
+	{
+		m_pD = new Private();
+	}
+
+
+	//---------------------------------------------------------------------------------------------
+	ErrorLog::~ErrorLog()
+	{
+        check( m_pD );
+		delete m_pD;
+		m_pD = 0;
+	}
+
+
+	//---------------------------------------------------------------------------------------------
+	ErrorLog::Private* ErrorLog::GetPrivate() const
+	{
+		return m_pD;
+	}
+
+
+	//---------------------------------------------------------------------------------------------
+    int ErrorLog::GetMessageCount( ErrorLogMessageType ) const
+	{
+        // \todo: by type argument
+		return (int)m_pD->m_messages.Num();
+	}
+
+
+	//---------------------------------------------------------------------------------------------
+	const FString& ErrorLog::GetMessageText( int index ) const
+	{
+		const char* strResult = "";
+
+		if ( index >=0 && index<(int)m_pD->m_messages.Num() )
+		{
+			return m_pD->m_messages[index].m_text;
+		}
+
+		static FString Empty;
+		return Empty;
+	}
+
+
+	//---------------------------------------------------------------------------------------------
+	const void* ErrorLog::GetMessageContext( int index ) const
+	{
+		const void* result = 0;
+
+		if ( index >=0 && index<(int)m_pD->m_messages.Num() )
+		{
+			result = m_pD->m_messages[index].m_context;
+		}
+
+		return result;
+	}
+
+
+	//---------------------------------------------------------------------------------------------
+	ErrorLogMessageType ErrorLog::GetMessageType( int index ) const
+	{
+		ErrorLogMessageType result = ELMT_NONE;
+
+		if ( index >=0 && index<(int)m_pD->m_messages.Num() )
+		{
+			result = m_pD->m_messages[index].m_type;
+		}
+
+		return result;
+	}
+
+
+	//---------------------------------------------------------------------------------------------
+	ErrorLogMessageSpamBin ErrorLog::GetMessageSpamBin(int index) const
+	{
+		ErrorLogMessageSpamBin result = ELMSB_ALL;
+
+		if (index >= 0 && index < (int)m_pD->m_messages.Num())
+		{
+			result = m_pD->m_messages[index].m_spam;
+		}
+
+		return result;
+	}
+
+	//---------------------------------------------------------------------------------------------
+    ErrorLogMessageAttachedDataView ErrorLog::GetMessageAttachedData( int index ) const
+	{
+        ErrorLogMessageAttachedDataView result;
+
+		if ( index >=0 && index<(int)m_pD->m_messages.Num() )
+		{
+			const Private::FMessage& message = m_pD->m_messages[index];
+            
+            if ( message.m_data ) 
+            {
+                result.m_unassignedUVs = message.m_data->m_unassignedUVs.GetData();
+			    result.m_unassignedUVsSize = message.m_data->m_unassignedUVs.Num();
+            }
+		}
+
+		return result;
+	}
+
+	//---------------------------------------------------------------------------------------------
+	void ErrorLog::Private::Add(const FString& InMessage,
+								ErrorLogMessageType InType,
+								const void* InContext,
+								ErrorLogMessageSpamBin InSpamBin)
+	{
+		m_messages.Add(FMessage());
+		m_messages.Last().m_type = InType;
+		m_messages.Last().m_spam = InSpamBin;
+		m_messages.Last().m_text = InMessage;
+		m_messages.Last().m_context = InContext;
+	}
+
+	//---------------------------------------------------------------------------------------------
+	void ErrorLog::Private::Add(const FString& InMessage,
+                                const ErrorLogMessageAttachedDataView& InDataView,
+                                ErrorLogMessageType InType, const void* InContext,
+								ErrorLogMessageSpamBin InSpamBin)
+	{
+		m_messages.Add(FMessage() );
+		m_messages.Last().m_type = InType;
+		m_messages.Last().m_spam = InSpamBin;
+		m_messages.Last().m_text = InMessage;
+		m_messages.Last().m_context = InContext;
+		m_messages.Last().m_data = MakeShared<FErrorData>();
+
+        if ( InDataView.m_unassignedUVs && InDataView.m_unassignedUVsSize > 0 )
+        {
+			// \TODO: Review
+			m_messages.Last().m_data->m_unassignedUVs.Append(InDataView.m_unassignedUVs, InDataView.m_unassignedUVsSize);
+        }
+	}
+	
+	
+	//---------------------------------------------------------------------------------------------
+	void ErrorLog::Log() const
+	{
+		UE_LOG(LogMutableCore, Log, TEXT(" Error Log :\n"));
+
+		for ( const Private::FMessage& msg : m_pD->m_messages )
+		{
+			switch ( msg.m_type )
+			{
+			case ELMT_ERROR: 	UE_LOG(LogMutableCore, Log, TEXT("  ERR  %s\n"), *msg.m_text); break;
+			case ELMT_WARNING: 	UE_LOG(LogMutableCore, Log, TEXT("  WRN  %s\n"), *msg.m_text); break;
+			case ELMT_INFO: 	UE_LOG(LogMutableCore, Log, TEXT("  INF  %s\n"), *msg.m_text); break;
+			default: 			UE_LOG(LogMutableCore, Log, TEXT("  NON  %s\n"), *msg.m_text); break;
+			}
+		}
+	}
+
+
+	//---------------------------------------------------------------------------------------------
+	void ErrorLog::Merge( const ErrorLog* pOther )
+	{
+		m_pD->m_messages.Append(pOther->GetPrivate()->m_messages);
+	}
+
+
+	//---------------------------------------------------------------------------------------------
+	//---------------------------------------------------------------------------------------------
+	//---------------------------------------------------------------------------------------------
+    
+    // clang-format off
+
+    const TCHAR* s_opNames[] =
+	{
+		TEXT("NONE             "),
+
+		TEXT("BO_CONSTANT      "),
+        TEXT("NU_CONSTANT      "),
+        TEXT("SC_CONSTANT      "),
+        TEXT("CO_CONSTANT      "),
+		TEXT("IM_CONSTANT      "),
+		TEXT("ME_CONSTANT      "),
+		TEXT("LA_CONSTANT      "),
+		TEXT("PR_CONSTANT      "),
+		TEXT("ST_CONSTANT      "),
+		TEXT("ED_CONSTANT      "),
+
+		TEXT("BO_PARAMETER     "),
+		TEXT("NU_PARAMETER     "),
+		TEXT("SC_PARAMETER     "),
+		TEXT("CO_PARAMETER     "),
+		TEXT("PR_PARAMETER     "),
+		TEXT("IM_PARAMETER     "),
+		TEXT("ST_PARAMETER     "),
+		
+		TEXT("IM_REFERENCE     "),
+		TEXT("ME_REFERENCE     "),
+
+		TEXT("NU_CONDITIONAL   "),
+		TEXT("SC_CONDITIONAL   "),
+		TEXT("CO_CONDITIONAL   "),
+		TEXT("IM_CONDITIONAL   "),
+		TEXT("ME_CONDITIONAL   "),
+		TEXT("LA_CONDITIONAL   "),
+		TEXT("IN_CONDITIONAL   "),
+		TEXT("ED_CONDITIONAL   "),
+		
+		TEXT("NU_SWITCH        "),
+		TEXT("SC_SWITCH        "),
+		TEXT("CO_SWITCH        "),
+		TEXT("IM_SWITCH        "),
+		TEXT("ME_SWITCH        "),
+		TEXT("LA_SWITCH        "),
+		TEXT("IN_SWITCH        "),
+		TEXT("ED_SWITCH        "),
+
+		TEXT("BO_LESS          "),
+		TEXT("BO_EQUAL_SC_CONST"),
+		TEXT("BO_AND           "),
+		TEXT("BO_OR            "),
+		TEXT("BO_NOT           "),
+		
+		TEXT("SC_MULTIPLYADD   "),
+		TEXT("SC_ARITHMETIC    "),
+		TEXT("SC_CURVE         "),
+		
+		TEXT("CO_SAMPLEIMAGE   "),
+		TEXT("CO_SWIZZLE       "),
+		TEXT("CO_FROMSCALARS   "),
+		TEXT("CO_ARITHMETIC    "),
+
+		TEXT("IM_LAYER         "),
+		TEXT("IM_LAYERCOLOUR   "),
+		TEXT("IM_PIXELFORMAT   "),
+		TEXT("IM_MIPMAP        "),
+		TEXT("IM_RESIZE        "),
+		TEXT("IM_RESIZELIKE    "),
+		TEXT("IM_RESIZEREL     "),
+		TEXT("IM_BLANKLAYOUT   "),
+		TEXT("IM_COMPOSE       "),
+		TEXT("IM_INTERPOLATE   "),
+		TEXT("IM_SATURATE      "),
+		TEXT("IM_LUMINANCE     "),
+		TEXT("IM_SWIZZLE       "),
+		TEXT("IM_COLOURMAP     "),
+		TEXT("IM_GRADIENT      "),
+		TEXT("IM_BINARISE      "),
+		TEXT("IM_PLAINCOLOUR   "),
+		TEXT("IM_CROP          "),
+		TEXT("IM_PATCH         "),
+		TEXT("IM_RASTERMESH    "),
+		TEXT("IM_MAKEGROWMAP   "),
+		TEXT("IM_DISPLACE      "),
+		TEXT("IM_MULTILAYER    "),
+		TEXT("IM_INVERT        "),
+		TEXT("IM_NORMAL_COMPO  "),
+		TEXT("IM_TRANSFORM     "),
+
+		TEXT("ME_APPLYLAYOUT   "),
+		TEXT("ME_DIFFERENCE    "),
+		TEXT("ME_MORPH         "),
+		TEXT("ME_MERGE         "),
+		TEXT("ME_INTERPOLATE   "),
+		TEXT("ME_MASKCLIPMESH  "),
+		TEXT("ME_MASKCLIPUVMASK"),
+		TEXT("ME_MASKDIFF      "),
+        TEXT("ME_REMOVEMASK    "),
+        TEXT("ME_FORMAT        "),
+        TEXT("ME_EXTRACTLABLOCK"),
+		TEXT("ME_TRANSFORM     "),
+		TEXT("ME_CLIPMORPHPLANE"),
+        TEXT("ME_CLIPWITHMESH  "),
+        TEXT("ME_SETSKELETON   "),
+        TEXT("ME_PROJECT       "),
+        TEXT("ME_APPLYPOSE     "),
+		TEXT("ME_GEOMETRYOP	   "),
+		TEXT("ME_BINDSHAPE	   "),
+		TEXT("ME_APPLYSHAPE	   "),
+		TEXT("ME_CLIPDEFORM	   "),
+		TEXT("ME_MORPHRESHAPE  "),
+		TEXT("ME_OPTIMIZESKIN  "),
+		TEXT("ME_ADDTAGS       "),
+
+		TEXT("IN_ADDMESH       "),
+		TEXT("IN_ADDIMAGE      "),
+		TEXT("IN_ADDVECTOR     "),
+		TEXT("IN_ADDSCALAR     "),
+		TEXT("IN_ADDSTRING     "),
+        TEXT("IN_ADDSURFACE    "),
+        TEXT("IN_ADDCOMPONENT  "),
+        TEXT("IN_ADDLOD        "),
+		TEXT("IN_ADDEXTENSIDATA"),
+		
+		TEXT("LA_PACK          "),
+		TEXT("LA_MERGE         "),
+		TEXT("LA_REMOVEBLOCKS  "),
+		TEXT("LA_FROMMESH	   "),
+	};
+
+	static_assert(sizeof(s_opNames) / sizeof(void*) == int32(OP_TYPE::COUNT));
+
+    // clang-format on
+
+    //---------------------------------------------------------------------------------------------
+    const TCHAR* GetOpName( OP_TYPE type )
+    {
+        //check( type>=0 && type<(int)OP_TYPE::COUNT );
+        return s_opNames[(int)type];
+    }
+
+}
