[mutable] Moved the Mutable plugin out of Experimental status into Beta.

diff --git a/Engine/Plugins/Mutable/Source/MutableTools/Private/MuT/Node.cpp b/Engine/Plugins/Mutable/Source/MutableTools/Private/MuT/Node.cpp
new file mode 100644
--- /dev/null
+++ b/Engine/Plugins/Mutable/Source/MutableTools/Private/MuT/Node.cpp
@@ -0,0 +1,113 @@
+// Copyright Epic Games, Inc. All Rights Reserved.
+
+
+#include "MuT/Node.h"
+#include "MuT/NodePrivate.h"
+#include "Misc/AssertionMacros.h"
+
+#include "MuT/NodeComponent.h"
+#include "MuT/NodeComponentNew.h"
+#include "MuT/NodeComponentSwitch.h"
+#include "MuT/NodeComponentVariation.h"
+#include "MuT/NodeLOD.h"
+#include "MuT/NodeSurface.h"
+#include "MuT/NodeSurfaceNew.h"
+#include "MuT/NodeSurfaceSwitch.h"
+#include "MuT/NodeSurfaceVariation.h"
+#include "MuT/NodeExtensionData.h"
+#include "MuT/NodeExtensionDataConstant.h"
+#include "MuT/NodeImageTable.h"
+#include "MuT/NodeColour.h"
+#include "MuT/NodeColourSwitch.h"
+#include "MuT/NodeColourVariation.h"
+#include "MuT/NodeColourConstant.h"
+#include "MuT/NodeColourParameter.h"
+#include "MuT/NodeColourSampleImage.h"
+#include "MuT/NodeColourArithmeticOperation.h"
+#include "MuT/NodeColourFromScalars.h"
+#include "MuT/NodeColourTable.h"
+#include "MuT/NodeScalarSwitch.h"
+#include "MuT/NodeMeshFragment.h"
+#include "MuT/NodeMeshClipMorphPlane.h"
+#include "MuT/NodeMeshClipDeform.h"
+#include "MuT/NodeMeshClipWithMesh.h"
+#include "MuT/NodeModifier.h"
+#include "MuT/NodeModifierMeshClipDeform.h"
+#include "MuT/NodeModifierMeshClipMorphPlane.h"
+#include "MuT/NodeModifierMeshClipWithMesh.h"
+#include "MuT/NodeModifierMeshClipWithUVMask.h"
+#include "MuT/NodeModifierSurfaceEdit.h"
+
+namespace mu
+{
+
+	// Static initialisation
+	FNodeType Node::StaticType = FNodeType(Node::EType::Node, nullptr);
+
+	FNodeType NodeComponent::StaticType = FNodeType(Node::EType::Component, Node::GetStaticType());
+	FNodeType NodeComponentNew::StaticType = FNodeType(Node::EType::ComponentNew, NodeComponent::GetStaticType());
+	FNodeType NodeComponentSwitch::StaticType = FNodeType(Node::EType::ComponentSwitch, NodeComponent::GetStaticType());
+	FNodeType NodeComponentVariation::StaticType = FNodeType(Node::EType::ComponentVariation, NodeComponent::GetStaticType());
+
+	FNodeType NodeScalarSwitch::StaticType = FNodeType(Node::EType::ScalarSwitch, NodeScalar::GetStaticType());
+
+	FNodeType NodeSurface::StaticType = FNodeType(Node::EType::Surface, Node::GetStaticType());
+	FNodeType NodeSurfaceNew::StaticType = FNodeType(Node::EType::SurfaceNew, NodeSurface::GetStaticType());
+	FNodeType NodeSurfaceSwitch::StaticType = FNodeType(Node::EType::SurfaceSwitch, NodeSurface::GetStaticType());
+	FNodeType NodeSurfaceVariation::StaticType = FNodeType(Node::EType::SurfaceVariation, NodeSurface::GetStaticType());
+
+	FNodeType NodeLOD::StaticType = FNodeType(Node::EType::LOD, Node::GetStaticType());
+	FNodeType NodeExtensionData::StaticType = FNodeType(Node::EType::ExtensionData, Node::GetStaticType());
+	FNodeType NodeExtensionDataConstant::StaticType = FNodeType(Node::EType::ExtensionDataConstant, NodeExtensionData::GetStaticType());
+	FNodeType NodeImageTable::StaticType = FNodeType(Node::EType::ImageTable, NodeImage::GetStaticType());
+
+	FNodeType NodeColour::StaticType = FNodeType(Node::EType::Color, Node::GetStaticType());
+	FNodeType NodeColourConstant::StaticType = FNodeType(Node::EType::ColorConstant, NodeColour::GetStaticType());
+	FNodeType NodeColourParameter::StaticType = FNodeType(Node::EType::ColorParameter, NodeColour::GetStaticType());
+	FNodeType NodeColourSwitch::StaticType = FNodeType(Node::EType::ColorSwitch, NodeColour::GetStaticType());
+	FNodeType NodeColourVariation::StaticType = FNodeType(Node::EType::ColorVariation, NodeColour::GetStaticType());
+	FNodeType NodeColourTable::StaticType = FNodeType(Node::EType::ColorTable, NodeColour::GetStaticType());
+	FNodeType NodeColourArithmeticOperation::StaticType = FNodeType(Node::EType::ColorArithmeticOperation, NodeColour::GetStaticType());
+	FNodeType NodeColourSampleImage::StaticType = FNodeType(Node::EType::ColorSampleImage, NodeColour::GetStaticType());
+	FNodeType NodeColourFromScalars::StaticType = FNodeType(Node::EType::ColorFromScalars, NodeColour::GetStaticType());
+
+	FNodeType NodeMeshFragment::StaticType = FNodeType(Node::EType::MeshFragment, NodeMesh::GetStaticType());
+	FNodeType NodeMeshClipMorphPlane::StaticType = FNodeType(Node::EType::MeshClipMorphPlane, NodeMesh::GetStaticType());
+	FNodeType NodeMeshClipDeform::StaticType = FNodeType(Node::EType::MeshClipDeform, NodeMesh::GetStaticType());
+	FNodeType NodeMeshClipWithMesh::StaticType = FNodeType(Node::EType::MeshClipWithMesh, NodeMesh::GetStaticType());
+
+	FNodeType NodeModifier::StaticType = FNodeType(Node::EType::Modifier, Node::GetStaticType());
+	FNodeType NodeModifierMeshClipDeform::StaticType = FNodeType(Node::EType::ModifierMeshClipDeform, NodeModifier::GetStaticType());
+	FNodeType NodeModifierMeshClipMorphPlane::StaticType = FNodeType(Node::EType::ModifierMeshClipMorphPlane, NodeModifier::GetStaticType());
+	FNodeType NodeModifierMeshClipWithMesh::StaticType = FNodeType(Node::EType::ModifierMeshClipWithMesh, NodeModifier::GetStaticType());
+	FNodeType NodeModifierMeshClipWithUVMask::StaticType = FNodeType(Node::EType::ModifierMeshClipWithUVMask, NodeModifier::GetStaticType());
+	FNodeType NodeModifierSurfaceEdit::StaticType = FNodeType(Node::EType::ModifierSurfaceEdit, NodeSurface::GetStaticType());
+
+
+	FNodeType::FNodeType()
+	{
+		Type = Node::EType::None;
+		Parent = nullptr;
+	}
+
+
+	FNodeType::FNodeType(Node::EType InType, const FNodeType* pParent )
+	{
+		Type = InType;
+		Parent = pParent;
+	}
+
+
+	void Node::SetMessageContext( const void* context )
+	{
+		MessageContext = context;
+	}
+
+	const void* Node::GetMessageContext() const 
+	{ 
+		return MessageContext; 
+	}
+
+}
+
+
diff --git a/Engine/Plugins/Mutable/Source/MutableTools/Private/MuT/NodeBool.cpp b/Engine/Plugins/Mutable/Source/MutableTools/Private/MuT/NodeBool.cpp
new file mode 100644
--- /dev/null
+++ b/Engine/Plugins/Mutable/Source/MutableTools/Private/MuT/NodeBool.cpp
@@ -0,0 +1,31 @@
+// Copyright Epic Games, Inc. All Rights Reserved.
+
+
+#include "MuT/NodeBool.h"
+
+#include "Misc/AssertionMacros.h"
+
+
+namespace mu
+{
+
+	static FNodeType s_nodeBoolType = FNodeType(Node::EType::Bool, Node::GetStaticType() );
+
+
+	//---------------------------------------------------------------------------------------------
+	const FNodeType* NodeBool::GetType() const
+	{
+		return GetStaticType();
+	}
+
+
+	//---------------------------------------------------------------------------------------------
+	const FNodeType* NodeBool::GetStaticType()
+	{
+		return &s_nodeBoolType;
+	}
+
+
+}
+
+
diff --git a/Engine/Plugins/Mutable/Source/MutableTools/Private/MuT/NodeBoolAnd.cpp b/Engine/Plugins/Mutable/Source/MutableTools/Private/MuT/NodeBoolAnd.cpp
new file mode 100644
--- /dev/null
+++ b/Engine/Plugins/Mutable/Source/MutableTools/Private/MuT/NodeBoolAnd.cpp
@@ -0,0 +1,58 @@
+// Copyright Epic Games, Inc. All Rights Reserved.
+
+
+#include "Misc/AssertionMacros.h"
+#include "MuT/Node.h"
+#include "MuT/NodeBool.h"
+#include "MuT/NodeBoolPrivate.h"
+#include "MuT/NodePrivate.h"
+
+
+namespace mu
+{
+
+	//---------------------------------------------------------------------------------------------
+	// Static initialisation
+	//---------------------------------------------------------------------------------------------
+	FNodeType NodeBoolAnd::Private::s_type = FNodeType(Node::EType::BoolParameter, NodeBool::GetStaticType() );
+
+
+	//---------------------------------------------------------------------------------------------
+	//!
+	//---------------------------------------------------------------------------------------------
+
+	MUTABLE_IMPLEMENT_NODE( NodeBoolAnd );
+
+
+	//---------------------------------------------------------------------------------------------
+	// Own Interface
+	//---------------------------------------------------------------------------------------------
+	Ptr<NodeBool> NodeBoolAnd::GetA() const
+	{
+		return m_pD->m_pA;
+	}
+
+
+	//---------------------------------------------------------------------------------------------
+	void NodeBoolAnd::SetA(Ptr<NodeBool> p )
+	{
+		m_pD->m_pA = p;
+	}
+
+
+	//---------------------------------------------------------------------------------------------
+	Ptr<NodeBool> NodeBoolAnd::GetB() const
+	{
+		return m_pD->m_pB;
+	}
+
+
+	//---------------------------------------------------------------------------------------------
+	void NodeBoolAnd::SetB(Ptr<NodeBool> p )
+	{
+		m_pD->m_pB = p;
+	}
+
+}
+
+
diff --git a/Engine/Plugins/Mutable/Source/MutableTools/Private/MuT/NodeBoolConstant.cpp b/Engine/Plugins/Mutable/Source/MutableTools/Private/MuT/NodeBoolConstant.cpp
new file mode 100644
--- /dev/null
+++ b/Engine/Plugins/Mutable/Source/MutableTools/Private/MuT/NodeBoolConstant.cpp
@@ -0,0 +1,43 @@
+// Copyright Epic Games, Inc. All Rights Reserved.
+
+
+#include "Misc/AssertionMacros.h"
+#include "MuT/Node.h"
+#include "MuT/NodeBool.h"
+#include "MuT/NodeBoolPrivate.h"
+#include "MuT/NodePrivate.h"
+
+
+
+namespace mu
+{
+
+	//---------------------------------------------------------------------------------------------
+	// Static initialisation
+	//---------------------------------------------------------------------------------------------
+	FNodeType NodeBoolConstant::Private::s_type = FNodeType(Node::EType::BoolConstant, NodeBool::GetStaticType() );
+
+
+	//---------------------------------------------------------------------------------------------
+	//!
+	//---------------------------------------------------------------------------------------------
+
+	MUTABLE_IMPLEMENT_NODE( NodeBoolConstant );
+
+
+	//---------------------------------------------------------------------------------------------
+	// Own Interface
+	//---------------------------------------------------------------------------------------------
+	bool NodeBoolConstant::GetValue() const
+	{
+		return m_pD->m_value;
+	}
+
+	//---------------------------------------------------------------------------------------------
+	void NodeBoolConstant::SetValue( bool v )
+	{
+		m_pD->m_value = v;
+	}
+
+
+}
diff --git a/Engine/Plugins/Mutable/Source/MutableTools/Private/MuT/NodeBoolNot.cpp b/Engine/Plugins/Mutable/Source/MutableTools/Private/MuT/NodeBoolNot.cpp
new file mode 100644
--- /dev/null
+++ b/Engine/Plugins/Mutable/Source/MutableTools/Private/MuT/NodeBoolNot.cpp
@@ -0,0 +1,44 @@
+// Copyright Epic Games, Inc. All Rights Reserved.
+
+
+#include "Misc/AssertionMacros.h"
+#include "MuT/Node.h"
+#include "MuT/NodeBool.h"
+#include "MuT/NodeBoolPrivate.h"
+#include "MuT/NodePrivate.h"
+
+
+namespace mu
+{
+
+	//---------------------------------------------------------------------------------------------
+	// Static initialisation
+	//---------------------------------------------------------------------------------------------
+	FNodeType NodeBoolNot::Private::s_type = FNodeType(Node::EType::BoolParameter, NodeBool::GetStaticType() );
+
+
+	//---------------------------------------------------------------------------------------------
+	//!
+	//---------------------------------------------------------------------------------------------
+
+	MUTABLE_IMPLEMENT_NODE( NodeBoolNot );
+
+
+	//---------------------------------------------------------------------------------------------
+	// Own Interface
+	//---------------------------------------------------------------------------------------------
+	Ptr<NodeBool> NodeBoolNot::GetInput() const
+	{
+		return m_pD->m_pSource;
+	}
+
+
+	//---------------------------------------------------------------------------------------------
+	void NodeBoolNot::SetInput(Ptr<NodeBool> p )
+	{
+		m_pD->m_pSource = p;
+	}
+
+}
+
+
diff --git a/Engine/Plugins/Mutable/Source/MutableTools/Private/MuT/NodeBoolParameter.cpp b/Engine/Plugins/Mutable/Source/MutableTools/Private/MuT/NodeBoolParameter.cpp
new file mode 100644
--- /dev/null
+++ b/Engine/Plugins/Mutable/Source/MutableTools/Private/MuT/NodeBoolParameter.cpp
@@ -0,0 +1,66 @@
+// Copyright Epic Games, Inc. All Rights Reserved.
+
+
+#include "Misc/AssertionMacros.h"
+#include "MuR/Ptr.h"
+#include "MuR/RefCounted.h"
+#include "MuT/Node.h"
+#include "MuT/NodeBool.h"
+#include "MuT/NodeBoolPrivate.h"
+#include "MuT/NodePrivate.h"
+#include "MuT/NodeRange.h"
+
+
+namespace mu
+{
+
+	//---------------------------------------------------------------------------------------------
+	// Static initialisation
+	//---------------------------------------------------------------------------------------------
+	FNodeType NodeBoolParameter::Private::s_type = FNodeType(Node::EType::BoolParameter, NodeBool::GetStaticType() );
+
+
+	//---------------------------------------------------------------------------------------------
+	//!
+	//---------------------------------------------------------------------------------------------
+
+	MUTABLE_IMPLEMENT_NODE( NodeBoolParameter );
+
+
+	//---------------------------------------------------------------------------------------------
+	// Own Interface
+	//---------------------------------------------------------------------------------------------
+	void NodeBoolParameter::SetName( const FString& strName )
+	{
+		m_pD->m_name = strName;
+	}
+
+
+	//---------------------------------------------------------------------------------------------
+	void NodeBoolParameter::SetDefaultValue( bool v )
+	{
+		m_pD->m_defaultValue = v;
+	}
+
+
+    //---------------------------------------------------------------------------------------------
+    void NodeBoolParameter::SetRangeCount( int i )
+    {
+        check(i>=0);
+        m_pD->m_ranges.SetNum(i);
+    }
+
+
+    //---------------------------------------------------------------------------------------------
+    void NodeBoolParameter::SetRange( int i, NodeRangePtr pRange )
+    {
+        check( i>=0 && i<int(m_pD->m_ranges.Num()) );
+        if ( i>=0 && i<int(m_pD->m_ranges.Num()) )
+        {
+            m_pD->m_ranges[i] = pRange;
+        }
+    }
+
+}
+
+
diff --git a/Engine/Plugins/Mutable/Source/MutableTools/Private/MuT/NodeBoolPrivate.h b/Engine/Plugins/Mutable/Source/MutableTools/Private/MuT/NodeBoolPrivate.h
new file mode 100644
--- /dev/null
+++ b/Engine/Plugins/Mutable/Source/MutableTools/Private/MuT/NodeBoolPrivate.h
@@ -0,0 +1,64 @@
+// Copyright Epic Games, Inc. All Rights Reserved.
+
+#pragma once
+
+#include "MuT/NodePrivate.h"
+
+#include "MuT/NodeBool.h"
+#include "MuT/NodeImage.h"
+#include "MuT/NodeRange.h"
+
+namespace mu
+{
+
+
+	class NodeBool::Private : public Node::Private
+	{
+	};
+
+
+	class NodeBoolConstant::Private : public NodeBool::Private
+	{
+	public:
+
+		static FNodeType s_type;
+
+		bool m_value;
+	};
+
+
+	class NodeBoolParameter::Private : public NodeBool::Private
+	{
+	public:
+
+		static FNodeType s_type;
+
+		bool m_defaultValue;
+		FString m_name;
+		FString m_uid;
+
+        TArray<Ptr<NodeRange>> m_ranges;
+	};
+
+
+	class NodeBoolNot::Private : public NodeBool::Private
+	{
+	public:
+
+		static FNodeType s_type;
+
+		Ptr<NodeBool> m_pSource;
+	};
+
+
+	class NodeBoolAnd::Private : public NodeBool::Private
+	{
+	public:
+
+		static FNodeType s_type;
+
+		Ptr<NodeBool> m_pA;
+		Ptr<NodeBool> m_pB;
+	};
+
+}
diff --git a/Engine/Plugins/Mutable/Source/MutableTools/Private/MuT/NodeComponentEdit.cpp b/Engine/Plugins/Mutable/Source/MutableTools/Private/MuT/NodeComponentEdit.cpp
new file mode 100644
--- /dev/null
+++ b/Engine/Plugins/Mutable/Source/MutableTools/Private/MuT/NodeComponentEdit.cpp
@@ -0,0 +1,49 @@
+// Copyright Epic Games, Inc. All Rights Reserved.
+
+
+#include "MuT/NodeComponentEdit.h"
+
+#include "Misc/AssertionMacros.h"
+#include "MuT/NodeComponentNew.h"
+#include "MuT/NodePrivate.h"
+#include "MuT/NodeSurface.h"
+
+
+namespace mu
+{
+
+	FNodeType NodeComponentEdit::StaticType = FNodeType(Node::EType::ComponentEdit, NodeComponent::GetStaticType() );
+
+
+	//---------------------------------------------------------------------------------------------
+	// Own Interface
+	//---------------------------------------------------------------------------------------------
+    void NodeComponentEdit::SetParent( NodeComponent* p )
+	{
+		m_pParent = p;
+	}
+
+
+	//---------------------------------------------------------------------------------------------
+    NodeComponent* NodeComponentEdit::GetParent() const
+	{
+        return m_pParent.get();
+	}
+
+
+	//---------------------------------------------------------------------------------------------
+	const NodeComponentNew* NodeComponentEdit::GetParentComponentNew() const
+	{
+		const NodeComponentNew* Parent = nullptr;
+		if (m_pParent)
+		{
+			Parent = m_pParent->GetParentComponentNew();
+		}
+
+		check(Parent);
+		return Parent;
+	}
+
+}
+
+
diff --git a/Engine/Plugins/Mutable/Source/MutableTools/Private/MuT/NodeImage.cpp b/Engine/Plugins/Mutable/Source/MutableTools/Private/MuT/NodeImage.cpp
new file mode 100644
--- /dev/null
+++ b/Engine/Plugins/Mutable/Source/MutableTools/Private/MuT/NodeImage.cpp
@@ -0,0 +1,32 @@
+// Copyright Epic Games, Inc. All Rights Reserved.
+
+#include "MuT/NodeImage.h"
+
+#include "Misc/AssertionMacros.h"
+
+namespace mu
+{
+
+	//---------------------------------------------------------------------------------------------
+	// Static initialisation
+	//---------------------------------------------------------------------------------------------
+	static FNodeType s_nodeImageType = 	FNodeType( Node::EType::Image, Node::GetStaticType() );
+
+
+	//---------------------------------------------------------------------------------------------
+	const FNodeType* NodeImage::GetType() const
+	{
+		return GetStaticType();
+	}
+
+
+	//---------------------------------------------------------------------------------------------
+	const FNodeType* NodeImage::GetStaticType()
+	{
+		return &s_nodeImageType;
+	}
+
+
+}
+
+
diff --git a/Engine/Plugins/Mutable/Source/MutableTools/Private/MuT/NodeImageBinarise.cpp b/Engine/Plugins/Mutable/Source/MutableTools/Private/MuT/NodeImageBinarise.cpp
new file mode 100644
--- /dev/null
+++ b/Engine/Plugins/Mutable/Source/MutableTools/Private/MuT/NodeImageBinarise.cpp
@@ -0,0 +1,55 @@
+// Copyright Epic Games, Inc. All Rights Reserved.
+
+
+#include "MuT/NodeImageBinarise.h"
+
+#include "Misc/AssertionMacros.h"
+#include "MuT/NodeImageBinarisePrivate.h"
+#include "MuT/NodePrivate.h"
+#include "MuT/NodeScalar.h"
+
+
+namespace mu
+{
+
+	//---------------------------------------------------------------------------------------------
+	// Static initialisation
+	//---------------------------------------------------------------------------------------------
+	FNodeType NodeImageBinarise::Private::s_type = FNodeType(Node::EType::ImageBinarise, NodeImage::GetStaticType() );
+
+	//---------------------------------------------------------------------------------------------
+
+	MUTABLE_IMPLEMENT_NODE( NodeImageBinarise )
+
+
+	//---------------------------------------------------------------------------------------------
+	// Own Interface
+	//---------------------------------------------------------------------------------------------
+	NodeImagePtr NodeImageBinarise::GetBase() const
+	{
+		return m_pD->m_pBase;
+	}
+
+
+	//---------------------------------------------------------------------------------------------
+	void NodeImageBinarise::SetBase( NodeImagePtr pNode )
+	{
+		m_pD->m_pBase = pNode;
+	}
+
+
+	//---------------------------------------------------------------------------------------------
+	NodeScalarPtr NodeImageBinarise::GetThreshold() const
+	{
+		return m_pD->m_pThreshold;
+	}
+
+
+	//---------------------------------------------------------------------------------------------
+	void NodeImageBinarise::SetThreshold( NodeScalarPtr pNode )
+	{
+		m_pD->m_pThreshold = pNode;
+	}
+
+}
+
diff --git a/Engine/Plugins/Mutable/Source/MutableTools/Private/MuT/NodeImageBinarisePrivate.h b/Engine/Plugins/Mutable/Source/MutableTools/Private/MuT/NodeImageBinarisePrivate.h
new file mode 100644
--- /dev/null
+++ b/Engine/Plugins/Mutable/Source/MutableTools/Private/MuT/NodeImageBinarisePrivate.h
@@ -0,0 +1,25 @@
+// Copyright Epic Games, Inc. All Rights Reserved.
+
+#pragma once
+
+#include "MuT/NodeImagePrivate.h"
+#include "MuT/NodeImageBinarise.h"
+#include "MuT/NodeScalar.h"
+
+
+namespace mu
+{
+
+
+	class NodeImageBinarise::Private : public NodeImage::Private
+	{
+	public:
+
+		static FNodeType s_type;
+
+		NodeImagePtr m_pBase;
+		NodeScalarPtr m_pThreshold;
+	};
+
+
+}
diff --git a/Engine/Plugins/Mutable/Source/MutableTools/Private/MuT/NodeImageColourMap.cpp b/Engine/Plugins/Mutable/Source/MutableTools/Private/MuT/NodeImageColourMap.cpp
new file mode 100644
--- /dev/null
+++ b/Engine/Plugins/Mutable/Source/MutableTools/Private/MuT/NodeImageColourMap.cpp
@@ -0,0 +1,69 @@
+// Copyright Epic Games, Inc. All Rights Reserved.
+
+
+#include "MuT/NodeImageColourMap.h"
+
+#include "Misc/AssertionMacros.h"
+#include "MuT/NodeImageColourMapPrivate.h"
+#include "MuT/NodePrivate.h"
+
+
+namespace mu
+{
+
+	//---------------------------------------------------------------------------------------------
+	// Static initialisation
+	//---------------------------------------------------------------------------------------------
+	FNodeType NodeImageColourMap::Private::s_type = FNodeType(Node::EType::ImageColorMap, NodeImage::GetStaticType() );
+
+
+	//---------------------------------------------------------------------------------------------
+	//!
+	//---------------------------------------------------------------------------------------------
+
+	MUTABLE_IMPLEMENT_NODE( NodeImageColourMap )
+
+
+	//---------------------------------------------------------------------------------------------
+	// Own Interface
+	//---------------------------------------------------------------------------------------------
+	NodeImagePtr NodeImageColourMap::GetBase() const
+	{
+		return m_pD->m_pBase.get();
+	}
+
+
+	//---------------------------------------------------------------------------------------------
+	void NodeImageColourMap::SetBase( NodeImagePtr pNode )
+	{
+		m_pD->m_pBase = pNode;
+	}
+
+
+	//---------------------------------------------------------------------------------------------
+	NodeImagePtr NodeImageColourMap::GetMap() const
+	{
+		return m_pD->m_pMap.get();
+	}
+
+
+	//---------------------------------------------------------------------------------------------
+	void NodeImageColourMap::SetMap( NodeImagePtr pNode )
+	{
+		m_pD->m_pMap = pNode;
+	}
+
+
+	//---------------------------------------------------------------------------------------------
+	NodeImagePtr NodeImageColourMap::GetMask() const
+	{
+		return m_pD->m_pMask.get();
+	}
+
+	//---------------------------------------------------------------------------------------------
+	void NodeImageColourMap::SetMask( NodeImagePtr pNode )
+	{
+		m_pD->m_pMask = pNode;
+	}
+
+}
diff --git a/Engine/Plugins/Mutable/Source/MutableTools/Private/MuT/NodeImageColourMapPrivate.h b/Engine/Plugins/Mutable/Source/MutableTools/Private/MuT/NodeImageColourMapPrivate.h
new file mode 100644
--- /dev/null
+++ b/Engine/Plugins/Mutable/Source/MutableTools/Private/MuT/NodeImageColourMapPrivate.h
@@ -0,0 +1,25 @@
+// Copyright Epic Games, Inc. All Rights Reserved.
+
+#pragma once
+
+#include "MuT/NodeImagePrivate.h"
+#include "MuT/NodeImageColourMap.h"
+
+namespace mu
+{
+
+	class NodeImageColourMap::Private : public NodeImage::Private
+	{
+	public:
+
+		static FNodeType s_type;
+
+		NodeImagePtr m_pBase;
+		NodeImagePtr m_pMask;
+		NodeImagePtr m_pMap;
+
+	};
+
+
+}
+
diff --git a/Engine/Plugins/Mutable/Source/MutableTools/Private/MuT/NodeImageConditional.cpp b/Engine/Plugins/Mutable/Source/MutableTools/Private/MuT/NodeImageConditional.cpp
new file mode 100644
--- /dev/null
+++ b/Engine/Plugins/Mutable/Source/MutableTools/Private/MuT/NodeImageConditional.cpp
@@ -0,0 +1,75 @@
+// Copyright Epic Games, Inc. All Rights Reserved.
+
+
+#include "MuT/NodeImageConditional.h"
+
+#include "Misc/AssertionMacros.h"
+#include "MuT/NodeBool.h"
+#include "MuT/NodeImageConditionalPrivate.h"
+#include "MuT/NodePrivate.h"
+
+
+namespace mu
+{
+
+	//---------------------------------------------------------------------------------------------
+	// Static initialisation
+	//---------------------------------------------------------------------------------------------
+    FNodeType NodeImageConditional::Private::s_type = FNodeType(Node::EType::ImageConditional, NodeImage::GetStaticType() );
+
+
+	//---------------------------------------------------------------------------------------------
+	//!
+	//---------------------------------------------------------------------------------------------
+
+    MUTABLE_IMPLEMENT_NODE( NodeImageConditional )
+
+
+	//---------------------------------------------------------------------------------------------
+	// Own Interface
+	//---------------------------------------------------------------------------------------------
+    NodeBool* NodeImageConditional::GetParameter() const
+	{
+        return m_pD->m_parameter.get();
+	}
+
+
+	//---------------------------------------------------------------------------------------------
+    void NodeImageConditional::SetParameter( NodeBool* pNode )
+	{
+        m_pD->m_parameter = pNode;
+	}
+
+
+    //---------------------------------------------------------------------------------------------
+    NodeImage* NodeImageConditional::GetOptionTrue() const
+    {
+        return m_pD->m_true.get();
+    }
+
+
+    //---------------------------------------------------------------------------------------------
+    void NodeImageConditional::SetOptionTrue( NodeImage* pNode )
+    {
+        m_pD->m_true = pNode;
+    }
+
+
+    //---------------------------------------------------------------------------------------------
+    NodeImage* NodeImageConditional::GetOptionFalse() const
+    {
+        return m_pD->m_false.get();
+    }
+
+
+    //---------------------------------------------------------------------------------------------
+    void NodeImageConditional::SetOptionFalse( NodeImage* pNode )
+    {
+        m_pD->m_false = pNode;
+    }
+
+
+
+}
+
+
diff --git a/Engine/Plugins/Mutable/Source/MutableTools/Private/MuT/NodeImageConditionalPrivate.h b/Engine/Plugins/Mutable/Source/MutableTools/Private/MuT/NodeImageConditionalPrivate.h
new file mode 100644
--- /dev/null
+++ b/Engine/Plugins/Mutable/Source/MutableTools/Private/MuT/NodeImageConditionalPrivate.h
@@ -0,0 +1,28 @@
+// Copyright Epic Games, Inc. All Rights Reserved.
+
+
+#pragma once
+
+#include "MuT/NodeImageConditional.h"
+#include "MuT/NodeImagePrivate.h"
+#include "MuT/NodeBool.h"
+
+namespace mu
+{
+
+
+    class NodeImageConditional::Private : public NodeImage::Private
+	{
+	public:
+
+		static FNodeType s_type;
+
+        Ptr<NodeBool> m_parameter;
+        NodeImagePtr m_true;
+        NodeImagePtr m_false;
+
+	};
+
+
+}
+
diff --git a/Engine/Plugins/Mutable/Source/MutableTools/Private/MuT/NodeImageConstant.cpp b/Engine/Plugins/Mutable/Source/MutableTools/Private/MuT/NodeImageConstant.cpp
new file mode 100644
--- /dev/null
+++ b/Engine/Plugins/Mutable/Source/MutableTools/Private/MuT/NodeImageConstant.cpp
@@ -0,0 +1,34 @@
+// Copyright Epic Games, Inc. All Rights Reserved.
+
+
+#include "MuT/NodeImageConstant.h"
+
+#include "Misc/AssertionMacros.h"
+#include "MuR/Image.h"
+#include "MuR/Serialisation.h"
+#include "MuT/NodePrivate.h"
+
+
+namespace mu
+{
+
+	//---------------------------------------------------------------------------------------------
+	// Static initialisation
+	//---------------------------------------------------------------------------------------------
+	FNodeType NodeImageConstant::StaticType = FNodeType(Node::EType::ImageConstant, NodeImage::GetStaticType());
+
+
+	//---------------------------------------------------------------------------------------------
+	void NodeImageConstant::SetValue(const Image* Value)
+	{
+		Proxy = new ResourceProxyMemory<Image>(Value);
+	}
+
+
+	//---------------------------------------------------------------------------------------------
+	void NodeImageConstant::SetValue(Ptr<ResourceProxy<Image>> InProxy)
+	{
+		Proxy = InProxy;
+	}
+
+}
diff --git a/Engine/Plugins/Mutable/Source/MutableTools/Private/MuT/NodeImageFormat.cpp b/Engine/Plugins/Mutable/Source/MutableTools/Private/MuT/NodeImageFormat.cpp
new file mode 100644
--- /dev/null
+++ b/Engine/Plugins/Mutable/Source/MutableTools/Private/MuT/NodeImageFormat.cpp
@@ -0,0 +1,57 @@
+// Copyright Epic Games, Inc. All Rights Reserved.
+
+
+#include "MuT/NodeImageFormat.h"
+
+#include "Misc/AssertionMacros.h"
+#include "MuT/NodeImageFormatPrivate.h"
+#include "MuT/NodePrivate.h"
+
+
+namespace mu
+{
+
+	//---------------------------------------------------------------------------------------------
+	// Static initialisation
+	//---------------------------------------------------------------------------------------------
+	FNodeType NodeImageFormat::Private::s_type = FNodeType(Node::EType::ImageFormat, NodeImage::GetStaticType() );
+
+
+	//---------------------------------------------------------------------------------------------
+	//!
+	//---------------------------------------------------------------------------------------------
+
+	MUTABLE_IMPLEMENT_NODE( NodeImageFormat )
+
+
+	//---------------------------------------------------------------------------------------------
+	// Own Interface
+	//---------------------------------------------------------------------------------------------
+	EImageFormat NodeImageFormat::GetFormat() const
+	{
+		return m_pD->m_format;
+	}
+
+
+	//---------------------------------------------------------------------------------------------
+    void NodeImageFormat::SetFormat(EImageFormat format, EImageFormat formatIfAlpha )
+	{
+		m_pD->m_format = format;
+        m_pD->m_formatIfAlpha = formatIfAlpha;
+	}
+
+
+	//---------------------------------------------------------------------------------------------
+	NodeImagePtr NodeImageFormat::GetSource() const
+	{
+		return m_pD->m_source.get();
+	}
+
+
+	//---------------------------------------------------------------------------------------------
+	void NodeImageFormat::SetSource(NodeImagePtr pNode )
+	{
+		m_pD->m_source = pNode;
+	}
+
+}
diff --git a/Engine/Plugins/Mutable/Source/MutableTools/Private/MuT/NodeImageGradient.cpp b/Engine/Plugins/Mutable/Source/MutableTools/Private/MuT/NodeImageGradient.cpp
new file mode 100644
--- /dev/null
+++ b/Engine/Plugins/Mutable/Source/MutableTools/Private/MuT/NodeImageGradient.cpp
@@ -0,0 +1,80 @@
+// Copyright Epic Games, Inc. All Rights Reserved.
+
+
+#include "MuT/NodeImageGradient.h"
+
+#include "Misc/AssertionMacros.h"
+#include "MuR/MutableMath.h"
+#include "MuT/NodeColour.h"
+#include "MuT/NodeImageGradientPrivate.h"
+#include "MuT/NodePrivate.h"
+
+
+namespace mu
+{
+
+	//---------------------------------------------------------------------------------------------
+	// Static initialisation
+	//---------------------------------------------------------------------------------------------
+	FNodeType NodeImageGradient::Private::s_type = FNodeType(Node::EType::ImageGradient, NodeImage::GetStaticType() );
+
+
+	//---------------------------------------------------------------------------------------------
+	//!
+	//---------------------------------------------------------------------------------------------
+
+	MUTABLE_IMPLEMENT_NODE( NodeImageGradient )
+
+
+	//---------------------------------------------------------------------------------------------
+	// Own Interface
+	//---------------------------------------------------------------------------------------------
+	NodeColourPtr NodeImageGradient::GetColour0() const
+	{
+		return m_pD->m_pColour0.get();
+	}
+
+
+	//---------------------------------------------------------------------------------------------
+	void NodeImageGradient::SetColour0( NodeColourPtr pNode )
+	{
+		m_pD->m_pColour0 = pNode;
+	}
+
+
+	//---------------------------------------------------------------------------------------------
+	NodeColourPtr NodeImageGradient::GetColour1() const
+	{
+		return m_pD->m_pColour1.get();
+	}
+
+
+	//---------------------------------------------------------------------------------------------
+	void NodeImageGradient::SetColour1( NodeColourPtr pNode )
+	{
+		m_pD->m_pColour1 = pNode;
+	}
+
+
+	//---------------------------------------------------------------------------------------------
+	int NodeImageGradient::GetSizeX() const
+	{
+		return m_pD->m_size[0];
+	}
+
+
+	//---------------------------------------------------------------------------------------------
+	int NodeImageGradient::GetSizeY() const
+	{
+		return m_pD->m_size[1];
+	}
+
+	//---------------------------------------------------------------------------------------------
+	void NodeImageGradient::SetSize( int x, int y )
+	{
+		m_pD->m_size[0] = x;
+		m_pD->m_size[1] = y;
+	}
+
+}
+
diff --git a/Engine/Plugins/Mutable/Source/MutableTools/Private/MuT/NodeImageGradientPrivate.h b/Engine/Plugins/Mutable/Source/MutableTools/Private/MuT/NodeImageGradientPrivate.h
new file mode 100644
--- /dev/null
+++ b/Engine/Plugins/Mutable/Source/MutableTools/Private/MuT/NodeImageGradientPrivate.h
@@ -0,0 +1,27 @@
+// Copyright Epic Games, Inc. All Rights Reserved.
+
+#pragma once
+
+#include "MuT/NodeColour.h"
+#include "MuT/NodeImagePrivate.h"
+#include "MuT/NodeImageGradient.h"
+#include "MuR/MutableMath.h"
+
+namespace mu
+{
+
+
+	class NodeImageGradient::Private : public NodeImage::Private
+	{
+	public:
+
+		static FNodeType s_type;
+
+		NodeColourPtr m_pColour0;
+		NodeColourPtr m_pColour1;
+		FIntVector2 m_size = { 256,1 };		
+
+	};
+
+
+}
diff --git a/Engine/Plugins/Mutable/Source/MutableTools/Private/MuT/NodeImageInterpolate.cpp b/Engine/Plugins/Mutable/Source/MutableTools/Private/MuT/NodeImageInterpolate.cpp
new file mode 100644
--- /dev/null
+++ b/Engine/Plugins/Mutable/Source/MutableTools/Private/MuT/NodeImageInterpolate.cpp
@@ -0,0 +1,77 @@
+// Copyright Epic Games, Inc. All Rights Reserved.
+
+
+#include "MuT/NodeImageInterpolate.h"
+
+#include "Misc/AssertionMacros.h"
+#include "MuT/NodeImageInterpolatePrivate.h"
+#include "MuT/NodePrivate.h"
+#include "MuT/NodeScalar.h"
+
+
+namespace mu
+{
+
+	//---------------------------------------------------------------------------------------------
+	// Static initialisation
+	//---------------------------------------------------------------------------------------------
+	FNodeType NodeImageInterpolate::Private::s_type = FNodeType(Node::EType::ImageInterpolate, NodeImage::GetStaticType() );
+
+
+	//---------------------------------------------------------------------------------------------
+	//!
+	//---------------------------------------------------------------------------------------------
+
+	MUTABLE_IMPLEMENT_NODE( NodeImageInterpolate )
+
+
+	//---------------------------------------------------------------------------------------------
+	// Own Interface
+	//---------------------------------------------------------------------------------------------
+	NodeScalarPtr NodeImageInterpolate::GetFactor() const
+	{
+		return m_pD->m_pFactor.get();
+	}
+
+
+	//---------------------------------------------------------------------------------------------
+	void NodeImageInterpolate::SetFactor( NodeScalarPtr pNode )
+	{
+		m_pD->m_pFactor = pNode;
+	}
+
+
+	//---------------------------------------------------------------------------------------------
+	NodeImagePtr NodeImageInterpolate::GetTarget( int t ) const
+	{
+		check( t>=0 && t<(int)m_pD->m_targets.Num() );
+		return m_pD->m_targets[t].get();
+	}
+
+
+	//---------------------------------------------------------------------------------------------
+	void NodeImageInterpolate::SetTarget( int t, NodeImagePtr pNode )
+	{
+		check( t>=0 && t<(int)m_pD->m_targets.Num() );
+		m_pD->m_targets[t] = pNode;
+	}
+
+
+	//---------------------------------------------------------------------------------------------
+	void NodeImageInterpolate::SetTargetCount( int t )
+	{
+		m_pD->m_targets.SetNum(t);
+	}
+
+
+	//---------------------------------------------------------------------------------------------
+	int NodeImageInterpolate::GetTargetCount() const
+	{
+		return int(m_pD->m_targets.Num());
+	}
+
+
+
+}
+
+
diff --git a/Engine/Plugins/Mutable/Source/MutableTools/Private/MuT/NodeImageInvert.cpp b/Engine/Plugins/Mutable/Source/MutableTools/Private/MuT/NodeImageInvert.cpp
new file mode 100644
--- /dev/null
+++ b/Engine/Plugins/Mutable/Source/MutableTools/Private/MuT/NodeImageInvert.cpp
@@ -0,0 +1,36 @@
+// Copyright Epic Games, Inc. All Rights Reserved.
+
+
+#include "MuT/NodeImageInvert.h"
+
+#include "Misc/AssertionMacros.h"
+#include "MuT/NodeImageInvertPrivate.h"
+#include "MuT/NodePrivate.h"
+
+
+namespace mu
+{
+
+	//---------------------------------------------------------------------------------------------
+	// Static initialisation
+	//---------------------------------------------------------------------------------------------
+	FNodeType NodeImageInvert::Private::s_type = FNodeType(Node::EType::ImageInvert, NodeImage::GetStaticType());
+
+	//---------------------------------------------------------------------------------------------
+
+	MUTABLE_IMPLEMENT_NODE( NodeImageInvert )
+
+	//---------------------------------------------------------------------------------------------
+	// Own Interface
+	//---------------------------------------------------------------------------------------------
+	NodeImagePtr NodeImageInvert::GetBase()const
+	{
+		return m_pD->m_pBase;
+	}
+
+	void NodeImageInvert::SetBase(NodeImagePtr pNode)
+	{
+		m_pD->m_pBase = pNode;
+	}
+	
+}
diff --git a/Engine/Plugins/Mutable/Source/MutableTools/Private/MuT/NodeImageLayer.cpp b/Engine/Plugins/Mutable/Source/MutableTools/Private/MuT/NodeImageLayer.cpp
new file mode 100644
--- /dev/null
+++ b/Engine/Plugins/Mutable/Source/MutableTools/Private/MuT/NodeImageLayer.cpp
@@ -0,0 +1,85 @@
+// Copyright Epic Games, Inc. All Rights Reserved.
+
+
+#include "MuT/NodeImageLayer.h"
+
+#include "Misc/AssertionMacros.h"
+#include "MuT/NodeImageLayerPrivate.h"
+#include "MuT/NodePrivate.h"
+
+
+namespace mu
+{
+
+	//---------------------------------------------------------------------------------------------
+	// Static initialisation
+	//---------------------------------------------------------------------------------------------
+	FNodeType NodeImageLayer::Private::s_type = FNodeType(Node::EType::ImageLayer, NodeImage::GetStaticType() );
+
+
+	//---------------------------------------------------------------------------------------------
+	//!
+	//---------------------------------------------------------------------------------------------
+
+	MUTABLE_IMPLEMENT_NODE( NodeImageLayer )
+
+
+	//---------------------------------------------------------------------------------------------
+	// Own Interface
+	//---------------------------------------------------------------------------------------------
+	NodeImagePtr NodeImageLayer::GetBase() const
+	{
+		return m_pD->m_pBase.get();
+	}
+
+
+	//---------------------------------------------------------------------------------------------
+	void NodeImageLayer::SetBase( NodeImagePtr pNode )
+	{
+		m_pD->m_pBase = pNode;
+	}
+
+
+	//---------------------------------------------------------------------------------------------
+	NodeImagePtr NodeImageLayer::GetMask() const
+	{
+		return m_pD->m_pMask.get();
+	}
+
+
+	//---------------------------------------------------------------------------------------------
+	void NodeImageLayer::SetMask( NodeImagePtr pNode )
+	{
+		m_pD->m_pMask = pNode;
+	}
+
+
+	//---------------------------------------------------------------------------------------------
+	NodeImagePtr NodeImageLayer::GetBlended() const
+	{
+		return m_pD->m_pBlended.get();
+	}
+
+
+	//---------------------------------------------------------------------------------------------
+	void NodeImageLayer::SetBlended( NodeImagePtr pNode )
+	{
+		m_pD->m_pBlended = pNode;
+	}
+
+
+	//---------------------------------------------------------------------------------------------
+	EBlendType NodeImageLayer::GetBlendType() const
+	{
+		return m_pD->m_type;
+	}
+
+
+	//---------------------------------------------------------------------------------------------
+	void NodeImageLayer::SetBlendType(EBlendType t)
+	{
+		m_pD->m_type = t;
+	}
+
+
+}
diff --git a/Engine/Plugins/Mutable/Source/MutableTools/Private/MuT/NodeImageLayerColour.cpp b/Engine/Plugins/Mutable/Source/MutableTools/Private/MuT/NodeImageLayerColour.cpp
new file mode 100644
--- /dev/null
+++ b/Engine/Plugins/Mutable/Source/MutableTools/Private/MuT/NodeImageLayerColour.cpp
@@ -0,0 +1,81 @@
+// Copyright Epic Games, Inc. All Rights Reserved.
+
+
+#include "MuT/NodeImageLayerColour.h"
+
+#include "Misc/AssertionMacros.h"
+#include "MuT/NodeColour.h"
+#include "MuT/NodeImageLayerColourPrivate.h"
+#include "MuT/NodePrivate.h"
+
+
+namespace mu
+{
+
+	//---------------------------------------------------------------------------------------------
+	// Static initialisation
+	//---------------------------------------------------------------------------------------------
+	FNodeType NodeImageLayerColour::Private::s_type = FNodeType(Node::EType::ImageLayerColour, NodeImage::GetStaticType() );
+
+
+	//---------------------------------------------------------------------------------------------
+	//!
+	//---------------------------------------------------------------------------------------------
+
+	MUTABLE_IMPLEMENT_NODE( NodeImageLayerColour );
+
+
+	//---------------------------------------------------------------------------------------------
+	// Own Interface
+	//---------------------------------------------------------------------------------------------
+	NodeImagePtr NodeImageLayerColour::GetBase() const
+	{
+		return m_pD->m_pBase.get();
+	}
+
+	//---------------------------------------------------------------------------------------------
+	void NodeImageLayerColour::SetBase( NodeImagePtr pNode )
+	{
+		m_pD->m_pBase = pNode;
+	}
+
+	//---------------------------------------------------------------------------------------------
+	NodeImagePtr NodeImageLayerColour::GetMask() const
+	{
+		return m_pD->m_pMask.get();
+	}
+
+	//---------------------------------------------------------------------------------------------
+	void NodeImageLayerColour::SetMask( NodeImagePtr pNode )
+	{
+		m_pD->m_pMask = pNode;
+	}
+
+	//---------------------------------------------------------------------------------------------
+	NodeColourPtr NodeImageLayerColour::GetColour() const
+	{
+		return m_pD->m_pColour.get();
+	}
+
+	//---------------------------------------------------------------------------------------------
+	void NodeImageLayerColour::SetColour( NodeColourPtr pNode )
+	{
+		m_pD->m_pColour = pNode;
+	}
+
+
+	//---------------------------------------------------------------------------------------------
+	EBlendType NodeImageLayerColour::GetBlendType() const
+	{
+		return m_pD->m_type;
+	}
+
+
+	//---------------------------------------------------------------------------------------------
+	void NodeImageLayerColour::SetBlendType(EBlendType t)
+	{
+		m_pD->m_type = t;
+	}
+
+}
+
diff --git a/Engine/Plugins/Mutable/Source/MutableTools/Private/MuT/NodeImageLuminance.cpp b/Engine/Plugins/Mutable/Source/MutableTools/Private/MuT/NodeImageLuminance.cpp
new file mode 100644
--- /dev/null
+++ b/Engine/Plugins/Mutable/Source/MutableTools/Private/MuT/NodeImageLuminance.cpp
@@ -0,0 +1,44 @@
+// Copyright Epic Games, Inc. All Rights Reserved.
+
+
+#include "MuT/NodeImageLuminance.h"
+
+#include "Misc/AssertionMacros.h"
+#include "MuT/NodeImageLuminancePrivate.h"
+#include "MuT/NodePrivate.h"
+
+
+namespace mu
+{
+
+	//---------------------------------------------------------------------------------------------
+	// Static initialisation
+	//---------------------------------------------------------------------------------------------
+	FNodeType NodeImageLuminance::Private::s_type = FNodeType(Node::EType::ImageLuminance, NodeImage::GetStaticType() );
+
+
+	//---------------------------------------------------------------------------------------------
+	//!
+	//---------------------------------------------------------------------------------------------
+
+	MUTABLE_IMPLEMENT_NODE( NodeImageLuminance )
+
+
+	//---------------------------------------------------------------------------------------------
+	// Own Interface
+	//---------------------------------------------------------------------------------------------
+	NodeImagePtr NodeImageLuminance::GetSource() const
+	{
+		return m_pD->m_pSource.get();
+	}
+
+
+	//---------------------------------------------------------------------------------------------
+	void NodeImageLuminance::SetSource( NodeImagePtr pNode )
+	{
+		m_pD->m_pSource = pNode;
+	}
+
+
+
+}
diff --git a/Engine/Plugins/Mutable/Source/MutableTools/Private/MuT/NodeImageLuminancePrivate.h b/Engine/Plugins/Mutable/Source/MutableTools/Private/MuT/NodeImageLuminancePrivate.h
new file mode 100644
--- /dev/null
+++ b/Engine/Plugins/Mutable/Source/MutableTools/Private/MuT/NodeImageLuminancePrivate.h
@@ -0,0 +1,22 @@
+// Copyright Epic Games, Inc. All Rights Reserved.
+
+#pragma once
+
+#include "MuT/NodeImagePrivate.h"
+#include "MuT/NodeImageLuminance.h"
+
+
+namespace mu
+{
+
+	class NodeImageLuminance::Private : public NodeImage::Private
+	{
+	public:
+
+		static FNodeType s_type;
+
+		NodeImagePtr m_pSource;
+	};
+
+
+}
diff --git a/Engine/Plugins/Mutable/Source/MutableTools/Private/MuT/NodeImageMipmap.cpp b/Engine/Plugins/Mutable/Source/MutableTools/Private/MuT/NodeImageMipmap.cpp
new file mode 100644
--- /dev/null
+++ b/Engine/Plugins/Mutable/Source/MutableTools/Private/MuT/NodeImageMipmap.cpp
@@ -0,0 +1,49 @@
+// Copyright Epic Games, Inc. All Rights Reserved.
+
+
+#include "MuT/NodeImageMipmap.h"
+
+#include "Misc/AssertionMacros.h"
+#include "MuR/ImagePrivate.h"
+#include "MuT/NodeImageMipmapPrivate.h"
+#include "MuT/NodePrivate.h"
+
+
+namespace mu
+{
+
+	//---------------------------------------------------------------------------------------------
+	// Static initialisation
+	//---------------------------------------------------------------------------------------------
+    FNodeType NodeImageMipmap::Private::s_type = FNodeType(Node::EType::ImageMipmap, NodeImage::GetStaticType() );
+
+
+	//---------------------------------------------------------------------------------------------
+	//!
+	//---------------------------------------------------------------------------------------------
+
+    MUTABLE_IMPLEMENT_NODE( NodeImageMipmap )
+
+
+	//---------------------------------------------------------------------------------------------
+	// Own Interface
+	//---------------------------------------------------------------------------------------------
+    NodeImagePtr NodeImageMipmap::GetSource() const
+	{
+		return m_pD->m_pSource.get();
+	}
+
+
+	//---------------------------------------------------------------------------------------------
+    void NodeImageMipmap::SetSource( NodeImagePtr pNode )
+	{
+		m_pD->m_pSource = pNode;
+	}
+
+	//---------------------------------------------------------------------------------------------
+	void NodeImageMipmap::SetMipmapGenerationSettings(EMipmapFilterType FilterType, EAddressMode AddressMode)
+	{
+		m_pD->m_settings = FMipmapGenerationSettings{ FilterType, AddressMode };
+	}
+
+}
diff --git a/Engine/Plugins/Mutable/Source/MutableTools/Private/MuT/NodeImageMultiLayer.cpp b/Engine/Plugins/Mutable/Source/MutableTools/Private/MuT/NodeImageMultiLayer.cpp
new file mode 100644
--- /dev/null
+++ b/Engine/Plugins/Mutable/Source/MutableTools/Private/MuT/NodeImageMultiLayer.cpp
@@ -0,0 +1,99 @@
+// Copyright Epic Games, Inc. All Rights Reserved.
+
+
+#include "MuT/NodeImageMultiLayer.h"
+
+#include "Misc/AssertionMacros.h"
+#include "MuT/NodeImageMultiLayerPrivate.h"
+#include "MuT/NodePrivate.h"
+
+
+namespace mu
+{
+
+	//---------------------------------------------------------------------------------------------
+	// Static initialisation
+	//---------------------------------------------------------------------------------------------
+    FNodeType NodeImageMultiLayer::Private::s_type = FNodeType(Node::EType::ImageMultiLayer, NodeImage::GetStaticType() );
+
+
+	//---------------------------------------------------------------------------------------------
+	//!
+	//---------------------------------------------------------------------------------------------
+
+    MUTABLE_IMPLEMENT_NODE( NodeImageMultiLayer )
+
+
+	//---------------------------------------------------------------------------------------------
+	// Own Interface
+	//---------------------------------------------------------------------------------------------
+    NodeImagePtr NodeImageMultiLayer::GetBase() const
+	{
+		return m_pD->m_pBase.get();
+	}
+
+
+	//---------------------------------------------------------------------------------------------
+    void NodeImageMultiLayer::SetBase( NodeImagePtr pNode )
+	{
+		m_pD->m_pBase = pNode;
+	}
+
+
+	//---------------------------------------------------------------------------------------------
+    NodeImagePtr NodeImageMultiLayer::GetMask() const
+	{
+		return m_pD->m_pMask.get();
+	}
+
+
+	//---------------------------------------------------------------------------------------------
+    void NodeImageMultiLayer::SetMask( NodeImagePtr pNode )
+	{
+		m_pD->m_pMask = pNode;
+	}
+
+
+    //---------------------------------------------------------------------------------------------
+    NodeImagePtr NodeImageMultiLayer::GetBlended() const
+    {
+        return m_pD->m_pBlended.get();
+    }
+
+
+    //---------------------------------------------------------------------------------------------
+    void NodeImageMultiLayer::SetBlended( NodeImagePtr pNode )
+    {
+        m_pD->m_pBlended = pNode;
+    }
+
+
+    //---------------------------------------------------------------------------------------------
+    NodeRangePtr NodeImageMultiLayer::GetRange() const
+    {
+        return m_pD->m_pRange.get();
+    }
+
+
+    //---------------------------------------------------------------------------------------------
+    void NodeImageMultiLayer::SetRange( NodeRangePtr pNode )
+    {
+        m_pD->m_pRange = pNode;
+    }
+
+
+	//---------------------------------------------------------------------------------------------
+	EBlendType NodeImageMultiLayer::GetBlendType() const
+	{
+		return m_pD->m_type;
+	}
+
+
+	//---------------------------------------------------------------------------------------------
+    void NodeImageMultiLayer::SetBlendType(EBlendType t)
+	{
+		m_pD->m_type = t;
+	}
+
+
+}
diff --git a/Engine/Plugins/Mutable/Source/MutableTools/Private/MuT/NodeImageNormalComposite.cpp b/Engine/Plugins/Mutable/Source/MutableTools/Private/MuT/NodeImageNormalComposite.cpp
new file mode 100644
--- /dev/null
+++ b/Engine/Plugins/Mutable/Source/MutableTools/Private/MuT/NodeImageNormalComposite.cpp
@@ -0,0 +1,79 @@
+// Copyright Epic Games, Inc. All Rights Reserved.
+
+
+#include "MuT/NodeImageNormalComposite.h"
+
+#include "Misc/AssertionMacros.h"
+#include "MuT/NodeImageNormalCompositePrivate.h"
+#include "MuT/NodePrivate.h"
+#include "MuR/Image.h"
+
+
+namespace mu
+{
+
+	//---------------------------------------------------------------------------------------------
+	// Static initialisation
+	//---------------------------------------------------------------------------------------------
+	FNodeType NodeImageNormalComposite::Private::s_type = FNodeType(Node::EType::ImageNormalComposite, NodeImage::GetStaticType() );
+
+
+	//---------------------------------------------------------------------------------------------
+	//!
+	//---------------------------------------------------------------------------------------------
+
+	MUTABLE_IMPLEMENT_NODE( NodeImageNormalComposite )
+
+
+	//---------------------------------------------------------------------------------------------
+	// Own Interface
+	//---------------------------------------------------------------------------------------------
+	NodeImagePtr NodeImageNormalComposite::GetBase() const
+	{
+		return m_pD->m_pBase.get();
+	}
+
+	//---------------------------------------------------------------------------------------------
+	void NodeImageNormalComposite::SetBase( NodeImagePtr pNode )
+	{
+		m_pD->m_pBase = pNode;
+	}
+
+	//---------------------------------------------------------------------------------------------
+	NodeImagePtr NodeImageNormalComposite::GetNormal() const
+	{
+		return m_pD->m_pNormal.get();
+	}
+
+	//---------------------------------------------------------------------------------------------
+	void NodeImageNormalComposite::SetNormal( NodeImagePtr pNode )
+	{
+		m_pD->m_pNormal = pNode;
+	}
+
+	//---------------------------------------------------------------------------------------------
+	float NodeImageNormalComposite::GetPower() const
+	{
+		return m_pD->m_power;
+	}
+
+	
+	//---------------------------------------------------------------------------------------------
+	void NodeImageNormalComposite::SetPower( float power )
+	{
+		m_pD->m_power = power;
+	}
+
+
+	//---------------------------------------------------------------------------------------------
+	ECompositeImageMode NodeImageNormalComposite::GetMode() const
+	{
+		return m_pD->m_mode;
+	}
+
+	//---------------------------------------------------------------------------------------------
+	void NodeImageNormalComposite::SetMode( ECompositeImageMode mode )
+	{
+		m_pD->m_mode = mode;
+	}
+}
diff --git a/Engine/Plugins/Mutable/Source/MutableTools/Private/MuT/NodeImageNormalCompositePrivate.h b/Engine/Plugins/Mutable/Source/MutableTools/Private/MuT/NodeImageNormalCompositePrivate.h
new file mode 100644
--- /dev/null
+++ b/Engine/Plugins/Mutable/Source/MutableTools/Private/MuT/NodeImageNormalCompositePrivate.h
@@ -0,0 +1,25 @@
+// Copyright Epic Games, Inc. All Rights Reserved.
+
+#pragma once
+
+#include "MuT/NodeImagePrivate.h"
+#include "MuT/NodeImageNormalComposite.h"
+
+
+namespace mu
+{
+
+	class NodeImageNormalComposite::Private : public NodeImage::Private
+	{
+	public:
+
+		static FNodeType s_type;
+
+		NodeImagePtr m_pBase;
+		NodeImagePtr m_pNormal;
+		float m_power;
+		ECompositeImageMode m_mode;
+
+	};
+
+}
diff --git a/Engine/Plugins/Mutable/Source/MutableTools/Private/MuT/NodeImageParameter.cpp b/Engine/Plugins/Mutable/Source/MutableTools/Private/MuT/NodeImageParameter.cpp
new file mode 100644
--- /dev/null
+++ b/Engine/Plugins/Mutable/Source/MutableTools/Private/MuT/NodeImageParameter.cpp
@@ -0,0 +1,70 @@
+// Copyright Epic Games, Inc. All Rights Reserved.
+
+
+#include "MuT/NodeImageParameter.h"
+
+#include "Misc/AssertionMacros.h"
+#include "MuT/NodeImageParameterPrivate.h"
+#include "MuT/NodePrivate.h"
+
+
+namespace mu
+{
+
+	//---------------------------------------------------------------------------------------------
+	// Static initialisation
+	//---------------------------------------------------------------------------------------------
+    FNodeType NodeImageParameter::Private::s_type = FNodeType(Node::EType::ImageParameter, NodeImage::GetStaticType() );
+
+
+	//---------------------------------------------------------------------------------------------
+	//!
+	//---------------------------------------------------------------------------------------------
+
+    MUTABLE_IMPLEMENT_NODE( NodeImageParameter )
+
+
+	//---------------------------------------------------------------------------------------------
+	// Own Interface
+	//---------------------------------------------------------------------------------------------
+    void NodeImageParameter::SetName( const FString& Name )
+	{
+		m_pD->m_name = Name;
+	}
+
+
+	//---------------------------------------------------------------------------------------------
+	void NodeImageParameter::SetUid(const FString& Uid)
+	{
+		m_pD->m_uid = Uid;
+	}
+
+	
+	//---------------------------------------------------------------------------------------------
+	void NodeImageParameter::SetDefaultValue(FName Value)
+	{
+    	m_pD->m_defaultValue = Value;
+	}
+
+
+	//---------------------------------------------------------------------------------------------
+	void NodeImageParameter::SetRangeCount(int i)
+	{
+		check(i >= 0);
+		m_pD->m_ranges.SetNum(i);
+	}
+
+
+	//---------------------------------------------------------------------------------------------
+	void NodeImageParameter::SetRange(int i, NodeRangePtr pRange)
+	{
+		check(i >= 0 && i<int(m_pD->m_ranges.Num()));
+		if (i >= 0 && i<int(m_pD->m_ranges.Num()))
+		{
+			m_pD->m_ranges[i] = pRange;
+		}
+	}
+
+}
+
+
diff --git a/Engine/Plugins/Mutable/Source/MutableTools/Private/MuT/NodeImageParameterPrivate.h b/Engine/Plugins/Mutable/Source/MutableTools/Private/MuT/NodeImageParameterPrivate.h
new file mode 100644
--- /dev/null
+++ b/Engine/Plugins/Mutable/Source/MutableTools/Private/MuT/NodeImageParameterPrivate.h
@@ -0,0 +1,26 @@
+// Copyright Epic Games, Inc. All Rights Reserved.
+
+#pragma once
+
+#include "MuT/NodeImageParameter.h"
+#include "MuT/NodeImagePrivate.h"
+#include "MuT/NodeRange.h"
+
+
+namespace mu
+{
+
+    class NodeImageParameter::Private : public NodeImage::Private
+	{
+	public:
+
+		static FNodeType s_type;
+
+    	FName m_defaultValue;
+		FString m_name;
+		FString m_uid;
+
+		TArray<Ptr<NodeRange>> m_ranges;
+	};
+
+}
diff --git a/Engine/Plugins/Mutable/Source/MutableTools/Private/MuT/NodeImagePlainColour.cpp b/Engine/Plugins/Mutable/Source/MutableTools/Private/MuT/NodeImagePlainColour.cpp
new file mode 100644
--- /dev/null
+++ b/Engine/Plugins/Mutable/Source/MutableTools/Private/MuT/NodeImagePlainColour.cpp
@@ -0,0 +1,50 @@
+// Copyright Epic Games, Inc. All Rights Reserved.
+
+
+#include "MuT/NodeImagePlainColour.h"
+
+#include "Misc/AssertionMacros.h"
+#include "MuT/NodeColour.h"
+#include "MuT/NodeImagePlainColourPrivate.h"
+#include "MuT/NodePrivate.h"
+
+
+namespace mu
+{
+
+	//---------------------------------------------------------------------------------------------
+	// Static initialisation
+	//---------------------------------------------------------------------------------------------
+	FNodeType NodeImagePlainColour::Private::s_type = FNodeType(Node::EType::ImagePlainColour, NodeImage::GetStaticType() );
+
+
+	//---------------------------------------------------------------------------------------------
+
+	MUTABLE_IMPLEMENT_NODE( NodeImagePlainColour )
+
+
+	//---------------------------------------------------------------------------------------------
+	// Own Interface
+	//---------------------------------------------------------------------------------------------
+	NodeColourPtr NodeImagePlainColour::GetColour() const
+	{
+		return m_pD->m_pColour;
+	}
+
+
+	//---------------------------------------------------------------------------------------------
+	void NodeImagePlainColour::SetColour( NodeColourPtr pNode )
+	{
+		m_pD->m_pColour = pNode;
+	}
+
+
+	//---------------------------------------------------------------------------------------------
+	void NodeImagePlainColour::SetSize( int x, int y )
+	{
+		m_pD->m_sizeX = x;
+		m_pD->m_sizeY = y;
+	}
+
+}
+
diff --git a/Engine/Plugins/Mutable/Source/MutableTools/Private/MuT/NodeImageResize.cpp b/Engine/Plugins/Mutable/Source/MutableTools/Private/MuT/NodeImageResize.cpp
new file mode 100644
--- /dev/null
+++ b/Engine/Plugins/Mutable/Source/MutableTools/Private/MuT/NodeImageResize.cpp
@@ -0,0 +1,54 @@
+// Copyright Epic Games, Inc. All Rights Reserved.
+
+
+#include "MuT/NodeImageResize.h"
+
+#include "Misc/AssertionMacros.h"
+#include "MuT/NodeImageResizePrivate.h"
+#include "MuT/NodePrivate.h"
+
+
+namespace mu
+{
+
+	//---------------------------------------------------------------------------------------------
+	// Static initialisation
+	//---------------------------------------------------------------------------------------------
+	FNodeType NodeImageResize::Private::s_type = FNodeType(Node::EType::ImageResize, NodeImage::GetStaticType() );
+
+	//---------------------------------------------------------------------------------------------
+
+	MUTABLE_IMPLEMENT_NODE( NodeImageResize )
+
+
+	//---------------------------------------------------------------------------------------------
+	// Own Interface
+	//---------------------------------------------------------------------------------------------
+	NodeImagePtr NodeImageResize::GetBase() const
+	{
+		return m_pD->m_pBase;
+	}
+
+
+	//---------------------------------------------------------------------------------------------
+	void NodeImageResize::SetBase( NodeImagePtr pNode )
+	{
+		m_pD->m_pBase = pNode;
+	}
+
+
+	//---------------------------------------------------------------------------------------------
+	void NodeImageResize::SetRelative( bool rel )
+	{
+		m_pD->m_relative = rel;
+	}
+
+
+	//---------------------------------------------------------------------------------------------
+	void NodeImageResize::SetSize( float x, float y )
+	{
+		m_pD->m_sizeX = x;
+		m_pD->m_sizeY = y;
+	}
+
+}
diff --git a/Engine/Plugins/Mutable/Source/MutableTools/Private/MuT/NodeImageSaturate.cpp b/Engine/Plugins/Mutable/Source/MutableTools/Private/MuT/NodeImageSaturate.cpp
new file mode 100644
--- /dev/null
+++ b/Engine/Plugins/Mutable/Source/MutableTools/Private/MuT/NodeImageSaturate.cpp
@@ -0,0 +1,60 @@
+// Copyright Epic Games, Inc. All Rights Reserved.
+
+
+#include "MuT/NodeImageSaturate.h"
+
+#include "Misc/AssertionMacros.h"
+#include "MuT/NodeImageSaturatePrivate.h"
+#include "MuT/NodePrivate.h"
+#include "MuT/NodeScalar.h"
+
+
+namespace mu
+{
+
+	//---------------------------------------------------------------------------------------------
+	// Static initialisation
+	//---------------------------------------------------------------------------------------------
+	FNodeType NodeImageSaturate::Private::s_type = FNodeType(Node::EType::ImageSaturate, NodeImage::GetStaticType() );
+
+
+	//---------------------------------------------------------------------------------------------
+	//!
+	//---------------------------------------------------------------------------------------------
+
+	MUTABLE_IMPLEMENT_NODE( NodeImageSaturate )
+
+
+	//---------------------------------------------------------------------------------------------
+	// Own Interface
+	//---------------------------------------------------------------------------------------------
+	NodeScalarPtr NodeImageSaturate::GetFactor() const
+	{
+		return m_pD->m_pFactor.get();
+	}
+
+
+	//---------------------------------------------------------------------------------------------
+	void NodeImageSaturate::SetFactor( NodeScalarPtr pNode )
+	{
+		m_pD->m_pFactor = pNode;
+	}
+
+
+	//---------------------------------------------------------------------------------------------
+	NodeImagePtr NodeImageSaturate::GetSource() const
+	{
+		return m_pD->m_pSource.get();
+	}
+
+
+	//---------------------------------------------------------------------------------------------
+	void NodeImageSaturate::SetSource( NodeImagePtr pNode )
+	{
+		m_pD->m_pSource = pNode;
+	}
+
+
+
+}
+
diff --git a/Engine/Plugins/Mutable/Source/MutableTools/Private/MuT/NodeImageSaturatePrivate.h b/Engine/Plugins/Mutable/Source/MutableTools/Private/MuT/NodeImageSaturatePrivate.h
new file mode 100644
--- /dev/null
+++ b/Engine/Plugins/Mutable/Source/MutableTools/Private/MuT/NodeImageSaturatePrivate.h
@@ -0,0 +1,23 @@
+// Copyright Epic Games, Inc. All Rights Reserved.
+
+#pragma once
+
+#include "MuT/NodeImagePrivate.h"
+#include "MuT/NodeImageSaturate.h"
+#include "MuT/NodeScalar.h"
+
+namespace mu
+{
+
+	class NodeImageSaturate::Private : public NodeImage::Private
+	{
+	public:
+
+		static FNodeType s_type;
+
+		NodeImagePtr m_pSource;
+		NodeScalarPtr m_pFactor;
+	};
+
+
+}
diff --git a/Engine/Plugins/Mutable/Source/MutableTools/Private/MuT/NodeImageSwitch.cpp b/Engine/Plugins/Mutable/Source/MutableTools/Private/MuT/NodeImageSwitch.cpp
new file mode 100644
--- /dev/null
+++ b/Engine/Plugins/Mutable/Source/MutableTools/Private/MuT/NodeImageSwitch.cpp
@@ -0,0 +1,77 @@
+// Copyright Epic Games, Inc. All Rights Reserved.
+
+
+#include "MuT/NodeImageSwitch.h"
+
+#include "Misc/AssertionMacros.h"
+#include "MuT/NodeImageSwitchPrivate.h"
+#include "MuT/NodePrivate.h"
+#include "MuT/NodeScalar.h"
+
+
+namespace mu
+{
+
+	//---------------------------------------------------------------------------------------------
+	// Static initialisation
+	//---------------------------------------------------------------------------------------------
+	FNodeType NodeImageSwitch::Private::s_type = FNodeType(Node::EType::ImageSwitch, NodeImage::GetStaticType() );
+
+
+	//---------------------------------------------------------------------------------------------
+	//!
+	//---------------------------------------------------------------------------------------------
+
+	MUTABLE_IMPLEMENT_NODE( NodeImageSwitch )
+
+
+	//---------------------------------------------------------------------------------------------
+	// Own Interface
+	//---------------------------------------------------------------------------------------------
+	NodeScalarPtr NodeImageSwitch::GetParameter() const
+	{
+		return m_pD->m_pParameter.get();
+	}
+
+
+	//---------------------------------------------------------------------------------------------
+	void NodeImageSwitch::SetParameter( NodeScalarPtr pNode )
+	{
+		m_pD->m_pParameter = pNode;
+	}
+
+
+	//---------------------------------------------------------------------------------------------
+	void NodeImageSwitch::SetOptionCount( int t )
+	{
+		m_pD->m_options.SetNum(t);
+	}
+
+
+	//---------------------------------------------------------------------------------------------
+	int NodeImageSwitch::GetOptionCount() const 
+	{
+		return int(m_pD->m_options.Num());
+	}
+
+
+	//---------------------------------------------------------------------------------------------
+	NodeImagePtr NodeImageSwitch::GetOption( int t ) const
+	{
+		check( t>=0 && t<(int)m_pD->m_options.Num() );
+		return m_pD->m_options[t].get();
+	}
+
+
+	//---------------------------------------------------------------------------------------------
+	void NodeImageSwitch::SetOption( int t, NodeImagePtr pNode )
+	{
+		check( t>=0 && t<(int)m_pD->m_options.Num() );
+		m_pD->m_options[t] = pNode;
+	}
+
+
+
+}
+
+
diff --git a/Engine/Plugins/Mutable/Source/MutableTools/Private/MuT/NodeImageSwizzle.cpp b/Engine/Plugins/Mutable/Source/MutableTools/Private/MuT/NodeImageSwizzle.cpp
new file mode 100644
--- /dev/null
+++ b/Engine/Plugins/Mutable/Source/MutableTools/Private/MuT/NodeImageSwizzle.cpp
@@ -0,0 +1,81 @@
+// Copyright Epic Games, Inc. All Rights Reserved.
+
+
+#include "MuT/NodeImageSwizzle.h"
+
+#include "Misc/AssertionMacros.h"
+#include "MuR/ImagePrivate.h"
+#include "MuT/NodeImageSwizzlePrivate.h"
+#include "MuT/NodePrivate.h"
+
+
+namespace mu
+{
+
+	//---------------------------------------------------------------------------------------------
+	// Static initialisation
+	//---------------------------------------------------------------------------------------------
+	FNodeType NodeImageSwizzle::Private::s_type = FNodeType(Node::EType::ImageSwizzle, NodeImage::GetStaticType() );
+
+
+	//---------------------------------------------------------------------------------------------
+	//!
+	//---------------------------------------------------------------------------------------------
+
+	MUTABLE_IMPLEMENT_NODE( NodeImageSwizzle )
+
+
+	//---------------------------------------------------------------------------------------------
+	// Own Interface
+	//---------------------------------------------------------------------------------------------
+	EImageFormat NodeImageSwizzle::GetFormat() const
+	{
+		return m_pD->m_format;
+	}
+
+
+	//---------------------------------------------------------------------------------------------
+	void NodeImageSwizzle::SetFormat(EImageFormat format )
+	{
+		m_pD->m_format = format;
+
+		int channelCount = GetImageFormatData( format ).Channels;
+		m_pD->m_sources.SetNum( channelCount );
+		m_pD->m_sourceChannels.SetNum( channelCount );
+	}
+
+
+	//---------------------------------------------------------------------------------------------
+	NodeImagePtr NodeImageSwizzle::GetSource( int t ) const
+	{
+		check( t>=0 && t<(int)m_pD->m_sources.Num() );
+		return m_pD->m_sources[t].get();
+	}
+
+
+	//---------------------------------------------------------------------------------------------
+	void NodeImageSwizzle::SetSource( int t, NodeImagePtr pNode )
+	{
+		check( t>=0 && t<(int)m_pD->m_sources.Num() );
+		m_pD->m_sources[t] = pNode;
+	}
+
+
+	//---------------------------------------------------------------------------------------------
+	int NodeImageSwizzle::GetSourceChannel( int t ) const
+	{
+		check( t>=0 && t<(int)m_pD->m_sourceChannels.Num() );
+		return m_pD->m_sourceChannels[t];
+	}
+
+
+	//---------------------------------------------------------------------------------------------
+	void NodeImageSwizzle::SetSourceChannel( int t, int channel )
+	{
+		check( t>=0 && t<(int)m_pD->m_sourceChannels.Num() );
+		m_pD->m_sourceChannels[t] = channel;
+	}
+
+
+
+}
diff --git a/Engine/Plugins/Mutable/Source/MutableTools/Private/MuT/NodeImageTransform.cpp b/Engine/Plugins/Mutable/Source/MutableTools/Private/MuT/NodeImageTransform.cpp
new file mode 100644
--- /dev/null
+++ b/Engine/Plugins/Mutable/Source/MutableTools/Private/MuT/NodeImageTransform.cpp
@@ -0,0 +1,147 @@
+// Copyright Epic Games, Inc. All Rights Reserved.
+
+
+#include "MuT/NodeImageTransform.h"
+
+#include "Misc/AssertionMacros.h"
+#include "MuT/NodeImageTransformPrivate.h"
+#include "MuT/NodePrivate.h"
+#include "MuT/NodeScalar.h"
+
+
+namespace mu
+{
+
+	//---------------------------------------------------------------------------------------------
+	// Static initialisation
+	//---------------------------------------------------------------------------------------------
+	FNodeType NodeImageTransform::Private::s_type = FNodeType(Node::EType::ImageTransform, NodeImage::GetStaticType() );
+
+	//---------------------------------------------------------------------------------------------
+
+	MUTABLE_IMPLEMENT_NODE( NodeImageTransform )
+
+
+	//---------------------------------------------------------------------------------------------
+	// Own Interface
+	//---------------------------------------------------------------------------------------------
+	NodeImagePtr NodeImageTransform::GetBase() const
+	{
+		return m_pD->m_pBase;
+	}
+
+	//---------------------------------------------------------------------------------------------
+	NodeScalarPtr NodeImageTransform::GetOffsetX() const
+	{
+		return m_pD->m_pOffsetX;
+	}
+
+	//---------------------------------------------------------------------------------------------
+	NodeScalarPtr NodeImageTransform::GetOffsetY() const
+	{
+		return m_pD->m_pOffsetY;
+	}
+
+	//---------------------------------------------------------------------------------------------
+	NodeScalarPtr NodeImageTransform::GetScaleX() const
+	{
+		return m_pD->m_pScaleX;
+	}
+
+	//---------------------------------------------------------------------------------------------
+	NodeScalarPtr NodeImageTransform::GetScaleY() const
+	{
+		return m_pD->m_pScaleY;
+	}
+
+	//---------------------------------------------------------------------------------------------
+	NodeScalarPtr NodeImageTransform::GetRotation() const
+	{
+		return m_pD->m_pRotation;
+	}
+
+	//---------------------------------------------------------------------------------------------
+	void NodeImageTransform::SetBase( NodeImagePtr pNode )
+	{
+		m_pD->m_pBase = pNode;
+	}
+
+	//---------------------------------------------------------------------------------------------
+	void NodeImageTransform::SetOffsetX( NodeScalarPtr pNode )
+	{
+		m_pD->m_pOffsetX = pNode;
+	}
+
+	//---------------------------------------------------------------------------------------------
+	void NodeImageTransform::SetOffsetY( NodeScalarPtr pNode )
+	{
+		m_pD->m_pOffsetY = pNode;
+	}
+
+	//---------------------------------------------------------------------------------------------
+	void NodeImageTransform::SetScaleX( NodeScalarPtr pNode )
+	{
+		m_pD->m_pScaleX = pNode;
+	}
+
+	//---------------------------------------------------------------------------------------------
+	void NodeImageTransform::SetScaleY( NodeScalarPtr pNode )
+	{
+		m_pD->m_pScaleY = pNode;
+	}
+
+	//---------------------------------------------------------------------------------------------
+	void NodeImageTransform::SetRotation( NodeScalarPtr pNode )
+	{
+		m_pD->m_pRotation = pNode;
+	}
+
+	//---------------------------------------------------------------------------------------------
+	void NodeImageTransform::SetAddressMode(EAddressMode AddressMode)
+	{
+		m_pD->AddressMode = AddressMode;
+	}
+
+	//---------------------------------------------------------------------------------------------
+	EAddressMode NodeImageTransform::GetAddressMode() const
+	{
+		return m_pD->AddressMode;
+	}
+
+	//---------------------------------------------------------------------------------------------
+	uint16 NodeImageTransform::GetSizeX() const
+	{
+		return m_pD->SizeX;
+	}
+
+	//---------------------------------------------------------------------------------------------
+	void NodeImageTransform::SetSizeX(uint16 SizeX)
+	{
+		m_pD->SizeX = SizeX;
+	}
+
+	//---------------------------------------------------------------------------------------------
+	uint16 NodeImageTransform::GetSizeY() const
+	{
+		return m_pD->SizeY;
+	}
+
+	//---------------------------------------------------------------------------------------------
+	void NodeImageTransform::SetSizeY(uint16 SizeY)
+	{
+		m_pD->SizeY = SizeY;
+	}
+
+	//---------------------------------------------------------------------------------------------
+	bool NodeImageTransform::GetKeepAspectRatio() const
+	{
+		return m_pD->bKeepAspectRatio;
+	}
+
+	//---------------------------------------------------------------------------------------------
+	void NodeImageTransform::SetKeepAspectRatio(bool bKeepAspectRatio)
+	{
+		m_pD->bKeepAspectRatio = bKeepAspectRatio;
+	}
+}
+
diff --git a/Engine/Plugins/Mutable/Source/MutableTools/Private/MuT/NodeImageTransformPrivate.h b/Engine/Plugins/Mutable/Source/MutableTools/Private/MuT/NodeImageTransformPrivate.h
new file mode 100644
--- /dev/null
+++ b/Engine/Plugins/Mutable/Source/MutableTools/Private/MuT/NodeImageTransformPrivate.h
@@ -0,0 +1,39 @@
+// Copyright Epic Games, Inc. All Rights Reserved.
+
+
+#pragma once
+
+#include "MuR/ImagePrivate.h"
+#include "MuT/NodeImagePrivate.h"
+#include "MuT/NodeImageTransform.h"
+#include "MuT/NodeScalar.h"
+
+namespace mu
+{
+
+
+	class NodeImageTransform::Private : public NodeImage::Private
+	{
+	public:
+
+		static FNodeType s_type;
+
+		NodeImagePtr m_pBase;
+		NodeScalarPtr m_pOffsetX;
+		NodeScalarPtr m_pOffsetY;
+		NodeScalarPtr m_pScaleX;
+		NodeScalarPtr m_pScaleY;
+		NodeScalarPtr m_pRotation;
+
+		EAddressMode AddressMode = EAddressMode::Wrap;
+		uint32 SizeX = 0;
+		uint32 SizeY = 0;
+		
+		bool bKeepAspectRatio = false;
+		uint8 UnusedPadding[sizeof(NodeImagePtr) - sizeof(bool)] = {0}; 
+		static_assert(sizeof(NodeImagePtr) - sizeof(bool) >= 1);
+
+	};
+
+
+}
diff --git a/Engine/Plugins/Mutable/Source/MutableTools/Private/MuT/NodeImageVariation.cpp b/Engine/Plugins/Mutable/Source/MutableTools/Private/MuT/NodeImageVariation.cpp
new file mode 100644
--- /dev/null
+++ b/Engine/Plugins/Mutable/Source/MutableTools/Private/MuT/NodeImageVariation.cpp
@@ -0,0 +1,65 @@
+// Copyright Epic Games, Inc. All Rights Reserved.
+
+
+#include "MuT/NodeImageVariation.h"
+
+#include "Misc/AssertionMacros.h"
+#include "MuT/NodeImageVariationPrivate.h"
+#include "MuT/NodePrivate.h"
+
+
+namespace mu
+{
+
+    //---------------------------------------------------------------------------------------------
+    // Static initialisation
+    //---------------------------------------------------------------------------------------------
+    FNodeType NodeImageVariation::Private::s_type = FNodeType(Node::EType::ImageVariation, NodeImage::GetStaticType() );
+
+
+    //---------------------------------------------------------------------------------------------
+    //---------------------------------------------------------------------------------------------
+    //---------------------------------------------------------------------------------------------
+
+    MUTABLE_IMPLEMENT_NODE( NodeImageVariation )
+
+
+    //---------------------------------------------------------------------------------------------
+    // Own Interface
+    //---------------------------------------------------------------------------------------------
+    void NodeImageVariation::SetDefaultImage( NodeImage* p ) 
+	{ 
+		m_pD->m_defaultImage = p; 
+	}
+
+
+	int NodeImageVariation::GetVariationCount() const 
+	{ 
+		return int( m_pD->m_variations.Num() ); 
+	}
+
+
+    void NodeImageVariation::SetVariationCount( int num )
+    {
+        check( num >= 0 );
+        m_pD->m_variations.SetNum( num );
+    }
+
+
+	void NodeImageVariation::SetVariationTag( int index, const FString& Tag )
+    {
+        check( index >= 0 && index < (int)m_pD->m_variations.Num() );
+
+		m_pD->m_variations[index].m_tag = Tag;
+    }
+
+
+    void NodeImageVariation::SetVariationImage( int index, NodeImage* pNode )
+    {
+        check( index >= 0 && index < (int)m_pD->m_variations.Num() );
+
+        m_pD->m_variations[index].m_image = pNode;
+    }
+
+
+} // namespace mu
diff --git a/Engine/Plugins/Mutable/Source/MutableTools/Private/MuT/NodeImageVariationPrivate.h b/Engine/Plugins/Mutable/Source/MutableTools/Private/MuT/NodeImageVariationPrivate.h
new file mode 100644
--- /dev/null
+++ b/Engine/Plugins/Mutable/Source/MutableTools/Private/MuT/NodeImageVariationPrivate.h
@@ -0,0 +1,33 @@
+// Copyright Epic Games, Inc. All Rights Reserved.
+
+#pragma once
+
+#include "MuT/NodePrivate.h"
+#include "MuT/NodeImageVariation.h"
+
+namespace mu
+{
+
+    //---------------------------------------------------------------------------------------------
+    //!
+    //---------------------------------------------------------------------------------------------
+    class NodeImageVariation::Private : public Node::Private
+    {
+    public:
+        Private() {}
+
+        static FNodeType s_type;
+
+        NodeImagePtr m_defaultImage;
+
+        struct FVariation
+        {
+            NodeImagePtr m_image;
+			FString m_tag;
+        };
+
+		TArray<FVariation> m_variations;
+    };
+
+} // namespace mu
+
diff --git a/Engine/Plugins/Mutable/Source/MutableTools/Private/MuT/NodeMeshApplyPose.cpp b/Engine/Plugins/Mutable/Source/MutableTools/Private/MuT/NodeMeshApplyPose.cpp
new file mode 100644
--- /dev/null
+++ b/Engine/Plugins/Mutable/Source/MutableTools/Private/MuT/NodeMeshApplyPose.cpp
@@ -0,0 +1,60 @@
+// Copyright Epic Games, Inc. All Rights Reserved.
+
+#include "MuT/NodeMeshApplyPose.h"
+
+#include "Misc/AssertionMacros.h"
+#include "MuT/NodeLayout.h"
+#include "MuT/NodeMeshApplyPosePrivate.h"
+#include "MuT/NodeMeshPrivate.h"
+#include "MuT/NodePrivate.h"
+
+
+namespace mu
+{
+
+	//---------------------------------------------------------------------------------------------
+	// Static initialisation
+	//---------------------------------------------------------------------------------------------
+    FNodeType NodeMeshApplyPose::Private::s_type = FNodeType(Node::EType::MeshApplyPose, NodeMesh::GetStaticType() );
+
+
+	//---------------------------------------------------------------------------------------------
+	//!
+	//---------------------------------------------------------------------------------------------
+
+    MUTABLE_IMPLEMENT_NODE( NodeMeshApplyPose )
+
+
+	//---------------------------------------------------------------------------------------------
+	// Own Interface
+	//---------------------------------------------------------------------------------------------
+    NodeMeshPtr NodeMeshApplyPose::GetBase() const
+    {
+        return m_pD->m_pBase.get();
+    }
+
+
+    //---------------------------------------------------------------------------------------------
+    void NodeMeshApplyPose::SetBase( NodeMeshPtr pNode )
+    {
+        m_pD->m_pBase = pNode;
+    }
+
+
+    //---------------------------------------------------------------------------------------------
+    NodeMeshPtr NodeMeshApplyPose::GetPose() const
+    {
+        return m_pD->m_pPose.get();
+    }
+
+
+    //---------------------------------------------------------------------------------------------
+    void NodeMeshApplyPose::SetPose( NodeMeshPtr pNode )
+    {
+        m_pD->m_pPose = pNode;
+    }
+
+
+}
+
+
diff --git a/Engine/Plugins/Mutable/Source/MutableTools/Private/MuT/NodeMeshApplyPosePrivate.h b/Engine/Plugins/Mutable/Source/MutableTools/Private/MuT/NodeMeshApplyPosePrivate.h
new file mode 100644
--- /dev/null
+++ b/Engine/Plugins/Mutable/Source/MutableTools/Private/MuT/NodeMeshApplyPosePrivate.h
@@ -0,0 +1,26 @@
+// Copyright Epic Games, Inc. All Rights Reserved.
+
+#pragma once
+
+#include "MuT/NodeMeshApplyPose.h"
+#include "MuT/NodeScalar.h"
+#include "MuT/NodeMeshPrivate.h"
+
+namespace mu
+{
+
+
+    class NodeMeshApplyPose::Private : public NodeMesh::Private
+	{
+	public:
+
+		static FNodeType s_type;
+
+        NodeMeshPtr m_pBase;
+        NodeMeshPtr m_pPose;
+
+	};
+
+
+}
+
diff --git a/Engine/Plugins/Mutable/Source/MutableTools/Private/MuT/NodeMeshConstant.cpp b/Engine/Plugins/Mutable/Source/MutableTools/Private/MuT/NodeMeshConstant.cpp
new file mode 100644
--- /dev/null
+++ b/Engine/Plugins/Mutable/Source/MutableTools/Private/MuT/NodeMeshConstant.cpp
@@ -0,0 +1,95 @@
+// Copyright Epic Games, Inc. All Rights Reserved.
+
+
+#include "MuT/NodeMeshConstant.h"
+
+#include "Misc/AssertionMacros.h"
+#include "MuT/NodeLayout.h"
+#include "MuT/NodeMeshConstantPrivate.h"
+#include "MuT/NodePrivate.h"
+#include "MuR/Mesh.h"
+
+
+namespace mu
+{
+	FNodeType NodeMeshConstant::Private::s_type = FNodeType(Node::EType::MeshConstant, NodeMesh::GetStaticType() );
+
+
+	MUTABLE_IMPLEMENT_NODE( NodeMeshConstant );
+
+
+	Ptr<Mesh> NodeMeshConstant::GetValue() const
+	{
+		return m_pD->Value;
+	}
+
+
+	void NodeMeshConstant::SetValue( Ptr<Mesh> pValue )
+	{
+		m_pD->Value = pValue;
+
+        if (m_pD->Value)
+        {
+            // Ensure mesh is well formed
+            m_pD->Value->EnsureSurfaceData();
+        }
+    }
+
+
+	void NodeMeshConstant::AddMorph(const FString& Name, Ptr<Mesh> Morphed)
+	{
+		m_pD->Morphs.Add({Name,Morphed});
+	}
+
+
+	Ptr<Mesh> NodeMeshConstant::FindMorph(const FString& Name) const
+	{
+		for (Private::FMorph& Morph : m_pD->Morphs )
+		{
+			if (Morph.Name == Name)
+			{
+				return Morph.MorphedMesh;
+			}
+		}
+		return nullptr;
+	}
+
+
+	int32 NodeMeshConstant::GetLayoutCount() const
+	{
+		return m_pD->Layouts.Num();
+	}
+
+
+	void NodeMeshConstant::SetLayoutCount( int32 num )
+	{
+		check( num >=0 );
+		m_pD->Layouts.SetNum( num );
+	}
+
+
+	Ptr<NodeLayout> NodeMeshConstant::GetLayout( int32 index ) const
+	{
+		check( index >=0 && index < m_pD->Layouts.Num() );
+
+		Ptr<NodeLayout> pResult;
+
+		if (index >= 0 && index < m_pD->Layouts.Num())
+		{
+			pResult = m_pD->Layouts[index];
+		}
+		return pResult;
+	}
+
+
+	void NodeMeshConstant::SetLayout( int32 index, Ptr<NodeLayout> pLayout )
+	{
+		check( index >=0 && index < m_pD->Layouts.Num() );
+
+		m_pD->Layouts[ index ] = pLayout;
+	}
+
+
+}
+
+
diff --git a/Engine/Plugins/Mutable/Source/MutableTools/Private/MuT/NodeMeshFormat.cpp b/Engine/Plugins/Mutable/Source/MutableTools/Private/MuT/NodeMeshFormat.cpp
new file mode 100644
--- /dev/null
+++ b/Engine/Plugins/Mutable/Source/MutableTools/Private/MuT/NodeMeshFormat.cpp
@@ -0,0 +1,64 @@
+// Copyright Epic Games, Inc. All Rights Reserved.
+
+
+#include "MuT/NodeMeshFormat.h"
+
+#include "Misc/AssertionMacros.h"
+#include "MuT/NodeLayout.h"
+#include "MuT/NodeMeshFormatPrivate.h"
+#include "MuT/NodeMeshPrivate.h"
+#include "MuT/NodePrivate.h"
+
+
+namespace mu
+{
+class FMeshBufferSet;
+
+	//---------------------------------------------------------------------------------------------
+	// Static initialisation
+	//---------------------------------------------------------------------------------------------
+	FNodeType NodeMeshFormat::Private::s_type = FNodeType(Node::EType::MeshFormat, NodeMesh::GetStaticType() );
+
+
+	//---------------------------------------------------------------------------------------------
+	//!
+	//---------------------------------------------------------------------------------------------
+
+	MUTABLE_IMPLEMENT_NODE( NodeMeshFormat )
+
+
+	//---------------------------------------------------------------------------------------------
+	// Own Interface
+	//---------------------------------------------------------------------------------------------
+	NodeMeshPtr NodeMeshFormat::GetSource() const
+	{
+		return m_pD->Source;
+	}
+
+
+	void NodeMeshFormat::SetSource( NodeMesh* pValue )
+	{
+		m_pD->Source = pValue;
+	}
+
+
+	FMeshBufferSet& NodeMeshFormat::GetVertexBuffers()
+	{
+		return m_pD->VertexBuffers;
+	}
+
+
+	FMeshBufferSet& NodeMeshFormat::GetIndexBuffers()
+	{
+		return m_pD->IndexBuffers;
+	}
+
+
+	void NodeMeshFormat::SetOptimizeBuffers(bool bEnable)
+	{
+		m_pD->bOptimizeBuffers = bEnable;
+	}
+
+}
+
+
diff --git a/Engine/Plugins/Mutable/Source/MutableTools/Private/MuT/NodeMeshGeometryOperation.cpp b/Engine/Plugins/Mutable/Source/MutableTools/Private/MuT/NodeMeshGeometryOperation.cpp
new file mode 100644
--- /dev/null
+++ b/Engine/Plugins/Mutable/Source/MutableTools/Private/MuT/NodeMeshGeometryOperation.cpp
@@ -0,0 +1,90 @@
+// Copyright Epic Games, Inc. All Rights Reserved.
+
+
+#include "MuT/NodeMeshGeometryOperation.h"
+
+#include "Misc/AssertionMacros.h"
+#include "MuT/NodeLayout.h"
+#include "MuT/NodeMeshGeometryOperationPrivate.h"
+#include "MuT/NodeMeshPrivate.h"
+#include "MuT/NodePrivate.h"
+#include "MuT/NodeScalar.h"
+
+
+
+namespace mu
+{
+
+	//---------------------------------------------------------------------------------------------
+	// Static initialisation
+	//---------------------------------------------------------------------------------------------
+	FNodeType NodeMeshGeometryOperation::Private::s_type = FNodeType(Node::EType::MeshGeometryOperation, NodeMesh::GetStaticType() );
+
+
+	//---------------------------------------------------------------------------------------------
+	//!
+	//---------------------------------------------------------------------------------------------
+
+	MUTABLE_IMPLEMENT_NODE( NodeMeshGeometryOperation )
+
+
+	//---------------------------------------------------------------------------------------------
+	// Own Interface
+	//---------------------------------------------------------------------------------------------
+	const NodeScalarPtr& NodeMeshGeometryOperation::GetScalarA() const
+	{
+		return m_pD->m_pScalarA;
+	}
+
+
+	//---------------------------------------------------------------------------------------------
+	void NodeMeshGeometryOperation::SetScalarA(const NodeScalarPtr& pNode)
+	{
+		m_pD->m_pScalarA = pNode;
+	}
+
+
+	//---------------------------------------------------------------------------------------------
+	const NodeScalarPtr& NodeMeshGeometryOperation::GetScalarB() const
+	{
+		return m_pD->m_pScalarB;
+	}
+
+
+	//---------------------------------------------------------------------------------------------
+	void NodeMeshGeometryOperation::SetScalarB(const NodeScalarPtr& pNode)
+	{
+		m_pD->m_pScalarB = pNode;
+	}
+
+
+	//---------------------------------------------------------------------------------------------
+	const NodeMeshPtr& NodeMeshGeometryOperation::GetMeshA() const
+	{
+		return m_pD->m_pMeshA;
+	}
+
+
+	//---------------------------------------------------------------------------------------------
+	void NodeMeshGeometryOperation::SetMeshA(const NodeMeshPtr& pNode)
+	{
+		m_pD->m_pMeshA = pNode;
+	}
+
+
+	//---------------------------------------------------------------------------------------------
+	const NodeMeshPtr& NodeMeshGeometryOperation::GetMeshB() const
+	{
+		return m_pD->m_pMeshB;
+	}
+
+
+	//---------------------------------------------------------------------------------------------
+	void NodeMeshGeometryOperation::SetMeshB(const NodeMeshPtr& pNode)
+	{
+		m_pD->m_pMeshB = pNode;
+	}
+
+}
+
+
diff --git a/Engine/Plugins/Mutable/Source/MutableTools/Private/MuT/NodeMeshGeometryOperationPrivate.h b/Engine/Plugins/Mutable/Source/MutableTools/Private/MuT/NodeMeshGeometryOperationPrivate.h
new file mode 100644
--- /dev/null
+++ b/Engine/Plugins/Mutable/Source/MutableTools/Private/MuT/NodeMeshGeometryOperationPrivate.h
@@ -0,0 +1,26 @@
+// Copyright Epic Games, Inc. All Rights Reserved.
+
+#pragma once
+
+#include "MuT/NodeMeshGeometryOperation.h"
+#include "MuT/NodeScalar.h"
+#include "MuT/NodeMeshPrivate.h"
+
+namespace mu
+{
+
+
+	class NodeMeshGeometryOperation::Private : public NodeMesh::Private
+	{
+	public:
+
+		static FNodeType s_type;
+
+		NodeMeshPtr m_pMeshA;
+		NodeMeshPtr m_pMeshB;
+		NodeScalarPtr m_pScalarA;
+		NodeScalarPtr m_pScalarB;
+
+	};
+
+}
diff --git a/Engine/Plugins/Mutable/Source/MutableTools/Private/MuT/NodeMeshInterpolate.cpp b/Engine/Plugins/Mutable/Source/MutableTools/Private/MuT/NodeMeshInterpolate.cpp
new file mode 100644
--- /dev/null
+++ b/Engine/Plugins/Mutable/Source/MutableTools/Private/MuT/NodeMeshInterpolate.cpp
@@ -0,0 +1,97 @@
+// Copyright Epic Games, Inc. All Rights Reserved.
+
+
+#include "MuT/NodeMeshInterpolate.h"
+
+#include "Misc/AssertionMacros.h"
+#include "MuT/NodeLayout.h"
+#include "MuT/NodeMeshInterpolatePrivate.h"
+#include "MuT/NodeMeshPrivate.h"
+#include "MuT/NodePrivate.h"
+#include "MuT/NodeScalar.h"
+
+
+namespace mu
+{
+
+	//---------------------------------------------------------------------------------------------
+	// Static initialisation
+	//---------------------------------------------------------------------------------------------
+	FNodeType NodeMeshInterpolate::Private::s_type = FNodeType(Node::EType::MeshInterpolate, NodeMesh::GetStaticType() );
+
+
+	//---------------------------------------------------------------------------------------------
+	//!
+	//---------------------------------------------------------------------------------------------
+
+	MUTABLE_IMPLEMENT_NODE( NodeMeshInterpolate )
+
+
+	//---------------------------------------------------------------------------------------------
+	// Own Interface
+	//---------------------------------------------------------------------------------------------
+	NodeScalarPtr NodeMeshInterpolate::GetFactor() const
+	{
+		return m_pD->m_pFactor.get();
+	}
+
+
+	//---------------------------------------------------------------------------------------------
+	void NodeMeshInterpolate::SetFactor( NodeScalarPtr pNode )
+	{
+		m_pD->m_pFactor = pNode;
+	}
+
+
+	//---------------------------------------------------------------------------------------------
+	NodeMeshPtr NodeMeshInterpolate::GetTarget( int t ) const
+	{
+		check( t>=0 && t<(int)m_pD->m_targets.Num() );
+
+		return m_pD->m_targets[t].get();
+	}
+
+
+	//---------------------------------------------------------------------------------------------
+	void NodeMeshInterpolate::SetTarget( int t, NodeMeshPtr pNode )
+	{
+		check( t>=0 && t<(int)m_pD->m_targets.Num() );
+		m_pD->m_targets[t] = pNode;
+	}
+
+
+	//---------------------------------------------------------------------------------------------
+	void NodeMeshInterpolate::SetTargetCount( int t )
+	{
+		m_pD->m_targets.SetNum(t);
+	}
+
+
+	//---------------------------------------------------------------------------------------------
+	int NodeMeshInterpolate::GetTargetCount() const
+	{
+		return m_pD->m_targets.Num();
+	}
+
+
+	//---------------------------------------------------------------------------------------------
+	void NodeMeshInterpolate::SetChannelCount( int t )
+	{
+		m_pD->m_channels.SetNum(t);
+	}
+
+
+	//---------------------------------------------------------------------------------------------
+	void NodeMeshInterpolate::SetChannel( int i, EMeshBufferSemantic semantic, int semanticIndex )
+	{
+		check( i>=0 && i<m_pD->m_channels.Num() );
+
+		if ( i>=0 && i<m_pD->m_channels.Num() )
+		{
+			m_pD->m_channels[i].semantic = semantic;
+			m_pD->m_channels[i].semanticIndex = semanticIndex;
+		}
+	}
+
+
+}
diff --git a/Engine/Plugins/Mutable/Source/MutableTools/Private/MuT/NodeMeshInterpolatePrivate.h b/Engine/Plugins/Mutable/Source/MutableTools/Private/MuT/NodeMeshInterpolatePrivate.h
new file mode 100644
--- /dev/null
+++ b/Engine/Plugins/Mutable/Source/MutableTools/Private/MuT/NodeMeshInterpolatePrivate.h
@@ -0,0 +1,48 @@
+// Copyright Epic Games, Inc. All Rights Reserved.
+
+#pragma once
+
+#include "MuT/NodeMeshInterpolate.h"
+#include "MuT/NodeScalar.h"
+#include "MuT/NodeMeshPrivate.h"
+
+
+namespace mu
+{
+
+
+	class NodeMeshInterpolate::Private : public NodeMesh::Private
+	{
+	public:
+
+		static FNodeType s_type;
+
+		NodeScalarPtr m_pFactor;
+		TArray<NodeMeshPtr> m_targets;
+
+		//!
+		struct CHANNEL
+		{
+			CHANNEL()
+			{
+				semantic = MBS_NONE;
+				semanticIndex = 0;
+			}
+
+			CHANNEL( EMeshBufferSemantic asemantic,
+                     int32_t asemanticIndex )
+			{
+				semantic = asemantic;
+				semanticIndex = asemanticIndex;
+			}
+
+			EMeshBufferSemantic semantic;
+            int32_t semanticIndex;
+		};
+
+		TArray<CHANNEL> m_channels;
+
+	};
+
+
+}
diff --git a/Engine/Plugins/Mutable/Source/MutableTools/Private/MuT/NodeMeshMakeMorph.cpp b/Engine/Plugins/Mutable/Source/MutableTools/Private/MuT/NodeMeshMakeMorph.cpp
new file mode 100644
--- /dev/null
+++ b/Engine/Plugins/Mutable/Source/MutableTools/Private/MuT/NodeMeshMakeMorph.cpp
@@ -0,0 +1,71 @@
+// Copyright Epic Games, Inc. All Rights Reserved.
+
+
+#include "MuT/NodeMeshMakeMorph.h"
+
+#include "Misc/AssertionMacros.h"
+#include "MuT/NodeLayout.h"
+#include "MuT/NodeMeshMakeMorphPrivate.h"
+#include "MuT/NodeMeshPrivate.h"
+#include "MuT/NodePrivate.h"
+
+
+namespace mu
+{
+
+	//---------------------------------------------------------------------------------------------
+	// Static initialisation
+	//---------------------------------------------------------------------------------------------
+    FNodeType NodeMeshMakeMorph::Private::s_type = FNodeType(Node::EType::MeshMakeMorph, NodeMesh::GetStaticType() );
+
+
+	//---------------------------------------------------------------------------------------------
+	//!
+	//---------------------------------------------------------------------------------------------
+
+    MUTABLE_IMPLEMENT_NODE( NodeMeshMakeMorph )
+
+
+	//---------------------------------------------------------------------------------------------
+	// Own Interface
+	//---------------------------------------------------------------------------------------------
+    NodeMeshPtr NodeMeshMakeMorph::GetBase() const
+	{
+		return m_pD->m_pBase;
+	}
+
+
+	//---------------------------------------------------------------------------------------------
+    void NodeMeshMakeMorph::SetBase( NodeMesh* p )
+	{
+		m_pD->m_pBase = p;
+	}
+
+
+	//---------------------------------------------------------------------------------------------
+    NodeMeshPtr NodeMeshMakeMorph::GetTarget() const
+	{
+		return m_pD->m_pTarget;
+	}
+
+
+	//---------------------------------------------------------------------------------------------
+    void NodeMeshMakeMorph::SetTarget( NodeMesh* p )
+	{
+		m_pD->m_pTarget = p;
+	}
+
+	//---------------------------------------------------------------------------------------------
+	void NodeMeshMakeMorph::SetOnlyPositionAndNormal(bool bInOnlyPositionAndNormals)
+	{
+		m_pD->bOnlyPositionAndNormal = bInOnlyPositionAndNormals;
+	}
+
+
+	//---------------------------------------------------------------------------------------------
+	bool NodeMeshMakeMorph::GetOnlyPositionAndNormal() const
+	{
+		return m_pD->bOnlyPositionAndNormal;
+	}
+
+}
diff --git a/Engine/Plugins/Mutable/Source/MutableTools/Private/MuT/NodeMeshMorph.cpp b/Engine/Plugins/Mutable/Source/MutableTools/Private/MuT/NodeMeshMorph.cpp
new file mode 100644
--- /dev/null
+++ b/Engine/Plugins/Mutable/Source/MutableTools/Private/MuT/NodeMeshMorph.cpp
@@ -0,0 +1,99 @@
+// Copyright Epic Games, Inc. All Rights Reserved.
+
+#include "MuT/NodeMeshMorph.h"
+
+#include "Containers/Array.h"
+#include "HAL/PlatformCrt.h"
+#include "Misc/AssertionMacros.h"
+#include "MuT/NodeLayout.h"
+#include "MuT/NodeMeshMorphPrivate.h"
+#include "MuT/NodeMeshPrivate.h"
+#include "MuT/NodePrivate.h"
+#include "MuT/NodeScalar.h"
+
+
+namespace mu
+{
+
+	//---------------------------------------------------------------------------------------------
+	// Static initialisation
+	//---------------------------------------------------------------------------------------------
+	FNodeType NodeMeshMorph::Private::s_type = FNodeType(Node::EType::MeshMorph, NodeMesh::GetStaticType() );
+
+
+	//---------------------------------------------------------------------------------------------
+	//!
+	//---------------------------------------------------------------------------------------------
+
+	MUTABLE_IMPLEMENT_NODE( NodeMeshMorph )
+
+
+	//---------------------------------------------------------------------------------------------
+	// Own Interface
+	//---------------------------------------------------------------------------------------------
+	Ptr<NodeScalar> NodeMeshMorph::GetFactor() const
+	{
+		return m_pD->Factor.get();
+	}
+
+
+	//---------------------------------------------------------------------------------------------
+	void NodeMeshMorph::SetFactor( Ptr<NodeScalar> pNode )
+	{
+		m_pD->Factor = pNode;
+	}
+
+
+	//---------------------------------------------------------------------------------------------
+	Ptr<NodeMesh> NodeMeshMorph::GetBase() const
+	{
+		return m_pD->Base.get();
+	}
+
+
+	//---------------------------------------------------------------------------------------------
+	void NodeMeshMorph::SetBase( Ptr<NodeMesh> pNode )
+	{
+		m_pD->Base = pNode;
+	}
+
+
+	//---------------------------------------------------------------------------------------------
+	NodeMeshPtr NodeMeshMorph::GetMorph() const
+	{
+		return m_pD->Morph.get();
+	}
+
+
+	//---------------------------------------------------------------------------------------------
+	void NodeMeshMorph::SetMorph( Ptr<NodeMesh> pNode )
+	{
+		m_pD->Morph = pNode;
+	}
+
+	//---------------------------------------------------------------------------------------------
+	void NodeMeshMorph::SetReshapeSkeleton(bool bEnable)
+	{
+		m_pD->bReshapeSkeleton = bEnable;
+	}	
+
+	//---------------------------------------------------------------------------------------------
+	void NodeMeshMorph::AddBoneToDeform(const FBoneName& BoneId)
+	{
+		m_pD->BonesToDeform.Emplace(BoneId);
+	}
+
+	void NodeMeshMorph::AddPhysicsBodyToDeform(const FBoneName& BoneId)
+	{
+		m_pD->PhysicsToDeform.Emplace(BoneId);
+	}
+
+	//---------------------------------------------------------------------------------------------
+	void NodeMeshMorph::SetReshapePhysicsVolumes(bool bEnable)
+	{
+		m_pD->bReshapePhysicsVolumes = bEnable;
+	}
+
+}
+
+
