[mutable] Moved the Mutable plugin out of Experimental status into Beta.

diff --git a/Engine/Plugins/Mutable/Source/MutableTools/Private/MuT/NodeMeshReshape.cpp b/Engine/Plugins/Mutable/Source/MutableTools/Private/MuT/NodeMeshReshape.cpp
new file mode 100644
--- /dev/null
+++ b/Engine/Plugins/Mutable/Source/MutableTools/Private/MuT/NodeMeshReshape.cpp
@@ -0,0 +1,128 @@
+// Copyright Epic Games, Inc. All Rights Reserved.
+
+#include "MuT/NodeMeshReshape.h"
+
+#include "Containers/Array.h"
+#include "HAL/PlatformCrt.h"
+#include "Misc/AssertionMacros.h"
+#include "MuT/NodeLayout.h"
+#include "MuT/NodeMeshPrivate.h"
+#include "MuT/NodeMeshReshapePrivate.h"
+#include "MuT/NodePrivate.h"
+
+
+namespace mu
+{
+
+	//---------------------------------------------------------------------------------------------
+	// Static initialisation
+	//---------------------------------------------------------------------------------------------
+	FNodeType NodeMeshReshape::Private::s_type = FNodeType(Node::EType::MeshReshape, NodeMesh::GetStaticType());
+
+
+	//---------------------------------------------------------------------------------------------
+	//!
+	//---------------------------------------------------------------------------------------------
+
+	MUTABLE_IMPLEMENT_NODE(NodeMeshReshape )
+
+
+	//---------------------------------------------------------------------------------------------
+	// Own Interface
+	//---------------------------------------------------------------------------------------------
+	const Ptr<NodeMesh>& NodeMeshReshape::GetBaseMesh() const
+	{
+		return m_pD->BaseMesh;
+	}
+
+
+	//---------------------------------------------------------------------------------------------
+	void NodeMeshReshape::SetBaseMesh(const Ptr<NodeMesh>& pNode)
+	{
+		m_pD->BaseMesh = pNode;
+	}
+
+
+	//---------------------------------------------------------------------------------------------
+	const Ptr<NodeMesh>& NodeMeshReshape::GetBaseShape() const
+	{
+		return m_pD->BaseShape;
+	}
+
+
+	//---------------------------------------------------------------------------------------------
+	void NodeMeshReshape::SetBaseShape(const Ptr<NodeMesh>& pNode)
+	{
+		m_pD->BaseShape = pNode;
+	}
+
+
+	//---------------------------------------------------------------------------------------------
+	const Ptr<NodeMesh>& NodeMeshReshape::GetTargetShape() const
+	{
+		return m_pD->TargetShape;
+	}
+
+
+	//---------------------------------------------------------------------------------------------
+	void NodeMeshReshape::SetTargetShape(const Ptr<NodeMesh>& pNode)
+	{
+		m_pD->TargetShape = pNode;
+	}
+
+	//---------------------------------------------------------------------------------------------
+	void NodeMeshReshape::SetReshapeVertices(bool bEnable)
+	{
+		m_pD->bReshapeVertices = bEnable;
+	}
+
+	//---------------------------------------------------------------------------------------------
+	void NodeMeshReshape::SetRecomputeNormals(bool bEnable)
+	{
+		m_pD->bRecomputeNormals = bEnable;
+	}
+
+
+	//---------------------------------------------------------------------------------------------
+	void NodeMeshReshape::SetApplyLaplacian(bool bEnable)
+	{
+		m_pD->bApplyLaplacian = bEnable;
+	}
+
+
+	//---------------------------------------------------------------------------------------------
+	void NodeMeshReshape::SetReshapeSkeleton(bool bEnable)
+	{
+		m_pD->bReshapeSkeleton = bEnable;
+	}
+	
+
+	//---------------------------------------------------------------------------------------------
+	void NodeMeshReshape::SetColorUsages(EVertexColorUsage R, EVertexColorUsage G, EVertexColorUsage B, EVertexColorUsage A)
+	{	
+		m_pD->ColorRChannelUsage = R;
+		m_pD->ColorGChannelUsage = G;
+		m_pD->ColorBChannelUsage = B;
+		m_pD->ColorAChannelUsage = A;
+	}	
+
+	//---------------------------------------------------------------------------------------------
+	void NodeMeshReshape::AddBoneToDeform(const FBoneName& BoneId)
+	{
+		m_pD->BonesToDeform.Emplace(BoneId);
+	}
+
+	//---------------------------------------------------------------------------------------------
+	void NodeMeshReshape::AddPhysicsBodyToDeform(const FBoneName& BoneId)
+	{
+		m_pD->PhysicsToDeform.Emplace(BoneId);
+	}
+
+	void NodeMeshReshape::SetReshapePhysicsVolumes(bool bEnable)
+	{
+		m_pD->bReshapePhysicsVolumes = bEnable;
+	}
+
+}
+
+
diff --git a/Engine/Plugins/Mutable/Source/MutableTools/Private/MuT/NodeMeshReshapePrivate.h b/Engine/Plugins/Mutable/Source/MutableTools/Private/MuT/NodeMeshReshapePrivate.h
new file mode 100644
--- /dev/null
+++ b/Engine/Plugins/Mutable/Source/MutableTools/Private/MuT/NodeMeshReshapePrivate.h
@@ -0,0 +1,39 @@
+// Copyright Epic Games, Inc. All Rights Reserved.
+
+#pragma once
+
+#include "MuT/NodeMeshReshape.h"
+#include "MuT/NodeScalar.h"
+#include "MuT/NodeMeshPrivate.h"
+#include "MuR/Skeleton.h"
+
+namespace mu
+{
+
+
+	class NodeMeshReshape::Private : public NodeMesh::Private
+	{
+	public:
+
+		static FNodeType s_type;
+
+		Ptr<NodeMesh> BaseMesh;
+		Ptr<NodeMesh> BaseShape;
+		Ptr<NodeMesh> TargetShape;
+		bool bReshapeVertices = true;
+		bool bRecomputeNormals = false;
+		bool bApplyLaplacian = false;
+		bool bReshapeSkeleton = false;
+		bool bReshapePhysicsVolumes = false;
+
+		EVertexColorUsage ColorRChannelUsage = EVertexColorUsage::None;
+		EVertexColorUsage ColorGChannelUsage = EVertexColorUsage::None;
+		EVertexColorUsage ColorBChannelUsage = EVertexColorUsage::None;
+		EVertexColorUsage ColorAChannelUsage = EVertexColorUsage::None;
+
+		TArray<FBoneName> BonesToDeform;
+		TArray<FBoneName> PhysicsToDeform;
+
+	};
+
+}
diff --git a/Engine/Plugins/Mutable/Source/MutableTools/Private/MuT/NodeMeshSwitch.cpp b/Engine/Plugins/Mutable/Source/MutableTools/Private/MuT/NodeMeshSwitch.cpp
new file mode 100644
--- /dev/null
+++ b/Engine/Plugins/Mutable/Source/MutableTools/Private/MuT/NodeMeshSwitch.cpp
@@ -0,0 +1,69 @@
+// Copyright Epic Games, Inc. All Rights Reserved.
+
+#include "MuT/NodeMeshSwitch.h"
+
+#include "Misc/AssertionMacros.h"
+#include "MuT/NodeLayout.h"
+#include "MuT/NodeMeshPrivate.h"
+#include "MuT/NodeMeshSwitchPrivate.h"
+#include "MuT/NodePrivate.h"
+#include "MuT/NodeScalar.h"
+
+
+namespace mu
+{
+
+	//---------------------------------------------------------------------------------------------
+	// Static initialisation
+	//---------------------------------------------------------------------------------------------
+	FNodeType NodeMeshSwitch::Private::s_type = FNodeType(Node::EType::MeshSwitch, NodeMesh::GetStaticType() );
+
+
+	//---------------------------------------------------------------------------------------------
+	//!
+	//---------------------------------------------------------------------------------------------
+
+	MUTABLE_IMPLEMENT_NODE( NodeMeshSwitch )
+
+
+	//---------------------------------------------------------------------------------------------
+	// Own Interface
+	//---------------------------------------------------------------------------------------------
+	NodeScalarPtr NodeMeshSwitch::GetParameter() const
+	{
+		return m_pD->m_pParameter.get();
+	}
+
+
+	//---------------------------------------------------------------------------------------------
+	void NodeMeshSwitch::SetParameter( NodeScalarPtr pNode )
+	{
+		m_pD->m_pParameter = pNode;
+	}
+
+
+	//---------------------------------------------------------------------------------------------
+	void NodeMeshSwitch::SetOptionCount( int t )
+	{
+		m_pD->m_options.SetNum(t);
+	}
+
+
+	//---------------------------------------------------------------------------------------------
+	NodeMeshPtr NodeMeshSwitch::GetOption( int t ) const
+	{
+		check( t>=0 && t<m_pD->m_options.Num() );
+		return m_pD->m_options[t].get();
+	}
+
+
+	//---------------------------------------------------------------------------------------------
+	void NodeMeshSwitch::SetOption( int t, NodeMeshPtr pNode )
+	{
+		check( t>=0 && t<m_pD->m_options.Num() );
+		m_pD->m_options[t] = pNode;
+	}
+
+}
+
+
diff --git a/Engine/Plugins/Mutable/Source/MutableTools/Private/MuT/NodeMeshSwitchPrivate.h b/Engine/Plugins/Mutable/Source/MutableTools/Private/MuT/NodeMeshSwitchPrivate.h
new file mode 100644
--- /dev/null
+++ b/Engine/Plugins/Mutable/Source/MutableTools/Private/MuT/NodeMeshSwitchPrivate.h
@@ -0,0 +1,26 @@
+// Copyright Epic Games, Inc. All Rights Reserved.
+
+#pragma once
+
+#include "MuT/NodeMeshSwitch.h"
+#include "MuT/NodeScalar.h"
+#include "MuT/NodeMeshPrivate.h"
+
+
+namespace mu
+{
+
+
+	class NodeMeshSwitch::Private : public NodeMesh::Private
+	{
+	public:
+
+		static FNodeType s_type;
+
+		NodeScalarPtr m_pParameter;
+		TArray<NodeMeshPtr> m_options;
+
+	};
+
+
+}
diff --git a/Engine/Plugins/Mutable/Source/MutableTools/Private/MuT/NodeMeshTable.cpp b/Engine/Plugins/Mutable/Source/MutableTools/Private/MuT/NodeMeshTable.cpp
new file mode 100644
--- /dev/null
+++ b/Engine/Plugins/Mutable/Source/MutableTools/Private/MuT/NodeMeshTable.cpp
@@ -0,0 +1,88 @@
+// Copyright Epic Games, Inc. All Rights Reserved.
+
+
+#include "MuT/NodeMeshTable.h"
+
+#include "Misc/AssertionMacros.h"
+#include "MuT/NodeLayout.h"
+#include "MuT/NodePrivate.h"
+#include "MuT/Table.h"
+
+
+namespace mu
+{
+
+	//---------------------------------------------------------------------------------------------
+	// Static initialisation
+	//---------------------------------------------------------------------------------------------
+	FNodeType NodeMeshTable::StaticType = FNodeType(Node::EType::MeshTable, NodeMesh::GetStaticType() );
+
+
+	//---------------------------------------------------------------------------------------------
+	// Own Interface
+	//---------------------------------------------------------------------------------------------
+	void NodeMeshTable::SetColumn( const FString& strName )
+	{
+		ColumnName = strName;
+	}
+
+
+	//---------------------------------------------------------------------------------------------
+	void NodeMeshTable::SetParameterName( const FString& strName )
+	{
+		ParameterName = strName;
+	}
+
+
+	//---------------------------------------------------------------------------------------------
+	int NodeMeshTable::GetLayoutCount() const
+	{
+		return Layouts.Num();
+	}
+
+
+	//---------------------------------------------------------------------------------------------
+	void NodeMeshTable::SetLayoutCount( int i )
+	{
+		Layouts.SetNum( i );
+	}
+
+
+	//---------------------------------------------------------------------------------------------
+	Ptr<NodeLayout> NodeMeshTable::GetLayout( int i ) const
+	{
+		Ptr<NodeLayout> pResult;
+
+		if (i >= 0 && i < Layouts.Num())
+		{
+			pResult = Layouts[i];
+		}
+
+		return pResult;
+	}
+
+
+	//---------------------------------------------------------------------------------------------
+	void NodeMeshTable::SetLayout( int i, Ptr<NodeLayout> pLayout )
+	{
+		check( i>=0 && i<GetLayoutCount() );
+		Layouts[i] = pLayout;
+	}
+
+
+	//---------------------------------------------------------------------------------------------
+	void NodeMeshTable::SetNoneOption(bool bAddNoneOption)
+	{
+		bNoneOption = bAddNoneOption;
+	}
+
+
+	//---------------------------------------------------------------------------------------------
+	void NodeMeshTable::SetDefaultRowName(const FString& RowName)
+	{
+		DefaultRowName = RowName;
+	}
+
+}
+
+
diff --git a/Engine/Plugins/Mutable/Source/MutableTools/Private/MuT/NodeMeshTangents.cpp b/Engine/Plugins/Mutable/Source/MutableTools/Private/MuT/NodeMeshTangents.cpp
new file mode 100644
--- /dev/null
+++ b/Engine/Plugins/Mutable/Source/MutableTools/Private/MuT/NodeMeshTangents.cpp
@@ -0,0 +1,44 @@
+// Copyright Epic Games, Inc. All Rights Reserved.
+
+
+#include "MuT/NodeMeshTangents.h"
+
+#include "Misc/AssertionMacros.h"
+#include "MuT/NodeLayout.h"
+#include "MuT/NodeMeshPrivate.h"
+#include "MuT/NodeMeshTangentsPrivate.h"
+#include "MuT/NodePrivate.h"
+
+
+namespace mu
+{
+
+	//---------------------------------------------------------------------------------------------
+	// Static initialisation
+	//---------------------------------------------------------------------------------------------
+	FNodeType NodeMeshTangents::Private::s_type = FNodeType(Node::EType::MeshTangents, NodeMesh::GetStaticType() );
+
+
+	//---------------------------------------------------------------------------------------------
+	//!
+	//---------------------------------------------------------------------------------------------
+
+	MUTABLE_IMPLEMENT_NODE( NodeMeshTangents )
+
+
+	//---------------------------------------------------------------------------------------------
+	// Own Interface
+	//---------------------------------------------------------------------------------------------
+	NodeMeshPtr NodeMeshTangents::GetSource() const
+	{
+		return m_pD->m_pSource;
+	}
+
+
+	//---------------------------------------------------------------------------------------------
+	void NodeMeshTangents::SetSource( NodeMesh* p )
+	{
+		m_pD->m_pSource = p;
+	}
+
+}
diff --git a/Engine/Plugins/Mutable/Source/MutableTools/Private/MuT/NodeMeshTangentsPrivate.h b/Engine/Plugins/Mutable/Source/MutableTools/Private/MuT/NodeMeshTangentsPrivate.h
new file mode 100644
--- /dev/null
+++ b/Engine/Plugins/Mutable/Source/MutableTools/Private/MuT/NodeMeshTangentsPrivate.h
@@ -0,0 +1,24 @@
+// Copyright Epic Games, Inc. All Rights Reserved.
+
+#pragma once
+
+#include "MuT/NodeMeshPrivate.h"
+#include "MuT/NodeMeshTangents.h"
+
+
+namespace mu
+{
+
+
+	class NodeMeshTangents::Private : public NodeMesh::Private
+	{
+	public:
+
+		static FNodeType s_type;
+
+		NodeMeshPtr m_pSource;
+
+	};
+
+
+}
diff --git a/Engine/Plugins/Mutable/Source/MutableTools/Private/MuT/NodeMeshTransform.cpp b/Engine/Plugins/Mutable/Source/MutableTools/Private/MuT/NodeMeshTransform.cpp
new file mode 100644
--- /dev/null
+++ b/Engine/Plugins/Mutable/Source/MutableTools/Private/MuT/NodeMeshTransform.cpp
@@ -0,0 +1,60 @@
+// Copyright Epic Games, Inc. All Rights Reserved.
+
+
+#include "MuT/NodeMeshTransform.h"
+
+#include "HAL/PlatformString.h"
+#include "Misc/AssertionMacros.h"
+#include "MuR/MutableMath.h"
+#include "MuT/NodeLayout.h"
+#include "MuT/NodeMeshPrivate.h"
+#include "MuT/NodeMeshTransformPrivate.h"
+#include "MuT/NodePrivate.h"
+
+
+namespace mu
+{
+
+	//---------------------------------------------------------------------------------------------
+	// Static initialisation
+	//---------------------------------------------------------------------------------------------
+    FNodeType NodeMeshTransform::Private::s_type = FNodeType(Node::EType::MeshTransform, NodeMesh::GetStaticType() );
+
+
+	//---------------------------------------------------------------------------------------------
+	//!
+	//---------------------------------------------------------------------------------------------
+
+    MUTABLE_IMPLEMENT_NODE( NodeMeshTransform )
+
+
+	//---------------------------------------------------------------------------------------------
+	// Own Interface
+	//---------------------------------------------------------------------------------------------
+    NodeMeshPtr NodeMeshTransform::GetSource() const
+	{
+		return m_pD->Source;
+	}
+
+
+    //---------------------------------------------------------------------------------------------
+    void NodeMeshTransform::SetSource( NodeMesh* p )
+    {
+        m_pD->Source = p;
+    }
+
+
+    //---------------------------------------------------------------------------------------------
+    void NodeMeshTransform::SetTransform( const FMatrix44f& Value )
+    {
+        m_pD->Transform = Value;
+    }
+
+
+    //---------------------------------------------------------------------------------------------
+	const FMatrix44f& NodeMeshTransform::GetTransform() const
+    {
+        return m_pD->Transform;
+    }
+
+}
diff --git a/Engine/Plugins/Mutable/Source/MutableTools/Private/MuT/NodeMeshTransformPrivate.h b/Engine/Plugins/Mutable/Source/MutableTools/Private/MuT/NodeMeshTransformPrivate.h
new file mode 100644
--- /dev/null
+++ b/Engine/Plugins/Mutable/Source/MutableTools/Private/MuT/NodeMeshTransformPrivate.h
@@ -0,0 +1,26 @@
+// Copyright Epic Games, Inc. All Rights Reserved.
+
+#pragma once
+
+#include "MuT/NodeMeshPrivate.h"
+#include "MuT/NodeMeshTransform.h"
+#include "MuR/MutableMath.h"
+
+
+namespace mu
+{
+
+
+    class NodeMeshTransform::Private : public NodeMesh::Private
+	{
+	public:
+
+		static FNodeType s_type;
+
+		NodeMeshPtr Source;
+		FMatrix44f Transform;
+
+	};
+
+
+}
diff --git a/Engine/Plugins/Mutable/Source/MutableTools/Private/MuT/NodeMeshVariation.cpp b/Engine/Plugins/Mutable/Source/MutableTools/Private/MuT/NodeMeshVariation.cpp
new file mode 100644
--- /dev/null
+++ b/Engine/Plugins/Mutable/Source/MutableTools/Private/MuT/NodeMeshVariation.cpp
@@ -0,0 +1,75 @@
+// Copyright Epic Games, Inc. All Rights Reserved.
+
+#include "MuT/NodeMeshVariation.h"
+
+#include "Misc/AssertionMacros.h"
+#include "MuT/NodeMeshVariationPrivate.h"
+#include "MuT/NodePrivate.h"
+
+
+namespace mu
+{
+
+    //---------------------------------------------------------------------------------------------
+    // Static initialisation
+    //---------------------------------------------------------------------------------------------
+    FNodeType NodeMeshVariation::Private::s_type = FNodeType(Node::EType::MeshVariation, NodeMesh::GetStaticType() );
+
+
+    //---------------------------------------------------------------------------------------------
+    //---------------------------------------------------------------------------------------------
+    //---------------------------------------------------------------------------------------------
+
+    MUTABLE_IMPLEMENT_NODE( NodeMeshVariation )
+
+
+    //---------------------------------------------------------------------------------------------
+    // Own Interface
+    //---------------------------------------------------------------------------------------------
+    void NodeMeshVariation::SetDefaultMesh( NodeMesh* p ) 
+	{ 
+		m_pD->m_defaultMesh = p; 
+	}
+
+
+    //---------------------------------------------------------------------------------------------
+    int NodeMeshVariation::GetVariationCount() const 
+	{ 
+		return  m_pD->m_variations.Num(); 
+	}
+
+
+    //---------------------------------------------------------------------------------------------
+    void NodeMeshVariation::SetVariationCount( int num )
+    {
+        check( num >= 0 );
+        m_pD->m_variations.SetNum( num );
+    }
+
+    //---------------------------------------------------------------------------------------------
+    void NodeMeshVariation::SetVariationTag( int index, const char* strTag )
+    {
+        check( index >= 0 && index < m_pD->m_variations.Num() );
+        check( strTag );
+
+        if ( strTag )
+        {
+            m_pD->m_variations[index].m_tag = strTag;
+        }
+        else
+        {
+            m_pD->m_variations[index].m_tag = "";
+        }
+    }
+
+
+    //---------------------------------------------------------------------------------------------
+    void NodeMeshVariation::SetVariationMesh( int index, NodeMesh* pNode )
+    {
+        check( index >= 0 && index < m_pD->m_variations.Num() );
+
+        m_pD->m_variations[index].m_mesh = pNode;
+    }
+
+
+} // namespace mu
diff --git a/Engine/Plugins/Mutable/Source/MutableTools/Private/MuT/NodeMeshVariationPrivate.h b/Engine/Plugins/Mutable/Source/MutableTools/Private/MuT/NodeMeshVariationPrivate.h
new file mode 100644
--- /dev/null
+++ b/Engine/Plugins/Mutable/Source/MutableTools/Private/MuT/NodeMeshVariationPrivate.h
@@ -0,0 +1,32 @@
+// Copyright Epic Games, Inc. All Rights Reserved.
+
+#pragma once
+
+#include "MuT/NodePrivate.h"
+#include "MuT/NodeMeshVariation.h"
+
+namespace mu
+{
+
+    //---------------------------------------------------------------------------------------------
+    //!
+    //---------------------------------------------------------------------------------------------
+    class NodeMeshVariation::Private : public Node::Private
+    {
+    public:
+        Private() {}
+
+        static FNodeType s_type;
+
+        Ptr<NodeMesh> m_defaultMesh;
+
+        struct FVariation
+        {
+			Ptr<NodeMesh> m_mesh;
+            FString m_tag;
+        };
+
+        TArray<FVariation> m_variations;
+    };
+
+} // namespace mu
diff --git a/Engine/Plugins/Mutable/Source/MutableTools/Private/MuT/NodeModifierMeshClipMorphPlane.cpp b/Engine/Plugins/Mutable/Source/MutableTools/Private/MuT/NodeModifierMeshClipMorphPlane.cpp
new file mode 100644
--- /dev/null
+++ b/Engine/Plugins/Mutable/Source/MutableTools/Private/MuT/NodeModifierMeshClipMorphPlane.cpp
@@ -0,0 +1,56 @@
+// Copyright Epic Games, Inc. All Rights Reserved.
+
+
+#include "MuT/NodeModifierMeshClipMorphPlane.h"
+
+#include "Misc/AssertionMacros.h"
+#include "MuR/MutableMath.h"
+#include "MuT/NodePrivate.h"
+
+
+namespace mu
+{
+
+	//---------------------------------------------------------------------------------------------
+	// Own Interface
+	//---------------------------------------------------------------------------------------------
+    void NodeModifierMeshClipMorphPlane::SetPlane(FVector3f Center, FVector3f Normal)
+	{
+		Parameters.Origin = Center;
+		Parameters.Normal = Normal;
+	}
+
+
+	//---------------------------------------------------------------------------------------------
+	void NodeModifierMeshClipMorphPlane::SetParams(float dist, float factor)
+	{
+		Parameters.DistanceToPlane = dist;
+		Parameters.LinearityFactor = factor;
+	}
+
+	//---------------------------------------------------------------------------------------------
+	void NodeModifierMeshClipMorphPlane::SetMorphEllipse(float radius1, float radius2, float rotation)
+	{
+		Parameters.Radius1 = radius1;
+		Parameters.Radius2 = radius2;
+		Parameters.Rotation = rotation;
+	}
+
+	//---------------------------------------------------------------------------------------------
+	void NodeModifierMeshClipMorphPlane::SetVertexSelectionBox(float centerX, float centerY, float centerZ, float radiusX, float radiusY, float radiusZ)
+	{
+		Parameters.VertexSelectionType = FClipMorphPlaneParameters::VS_SHAPE;
+		Parameters.SelectionBoxOrigin = FVector3f(centerX, centerY, centerZ);
+		Parameters.SelectionBoxRadius = FVector3f(radiusX, radiusY, radiusZ);
+	}
+
+	//---------------------------------------------------------------------------------------------
+	void NodeModifierMeshClipMorphPlane::SetVertexSelectionBone(const FBoneName& BoneId, float maxEffectRadius)
+	{
+		Parameters.VertexSelectionType = FClipMorphPlaneParameters::VS_BONE_HIERARCHY;
+		Parameters.VertexSelectionBone = BoneId;
+		Parameters.MaxEffectRadius = maxEffectRadius;
+	}
+
+}
+
diff --git a/Engine/Plugins/Mutable/Source/MutableTools/Private/MuT/NodeObject.cpp b/Engine/Plugins/Mutable/Source/MutableTools/Private/MuT/NodeObject.cpp
new file mode 100644
--- /dev/null
+++ b/Engine/Plugins/Mutable/Source/MutableTools/Private/MuT/NodeObject.cpp
@@ -0,0 +1,30 @@
+// Copyright Epic Games, Inc. All Rights Reserved.
+
+#include "MuT/NodeObject.h"
+
+#include "Misc/AssertionMacros.h"
+
+
+namespace mu
+{
+
+	FNodeType s_nodeObjectType = FNodeType(Node::EType::Object, Node::GetStaticType() );
+
+
+	//---------------------------------------------------------------------------------------------
+	const FNodeType* NodeObject::GetType() const
+	{
+		return GetStaticType();
+	}
+
+
+	//---------------------------------------------------------------------------------------------
+	const FNodeType* NodeObject::GetStaticType()
+	{
+		return &s_nodeObjectType;
+	}
+
+
+}
+
+
diff --git a/Engine/Plugins/Mutable/Source/MutableTools/Private/MuT/NodeObjectGroup.cpp b/Engine/Plugins/Mutable/Source/MutableTools/Private/MuT/NodeObjectGroup.cpp
new file mode 100644
--- /dev/null
+++ b/Engine/Plugins/Mutable/Source/MutableTools/Private/MuT/NodeObjectGroup.cpp
@@ -0,0 +1,114 @@
+// Copyright Epic Games, Inc. All Rights Reserved.
+
+#include "MuT/NodeObjectGroup.h"
+
+#include "HAL/PlatformCrt.h"
+#include "Misc/AssertionMacros.h"
+#include "MuT/NodeLayout.h"
+#include "MuT/NodeObjectGroupPrivate.h"
+#include "MuT/NodeObjectPrivate.h"
+#include "MuT/NodePrivate.h"
+
+
+namespace mu
+{
+	MUTABLE_IMPLEMENT_ENUM_SERIALISABLE(NodeObjectGroup::CHILD_SELECTION)
+
+
+	//---------------------------------------------------------------------------------------------
+	// Static initialisation
+	//---------------------------------------------------------------------------------------------
+	FNodeType NodeObjectGroup::Private::s_type = FNodeType(Node::EType::ObjectGroup, NodeObject::GetStaticType() );
+
+
+	//---------------------------------------------------------------------------------------------
+	//!
+	//---------------------------------------------------------------------------------------------
+
+	MUTABLE_IMPLEMENT_NODE( NodeObjectGroup )
+
+
+	//---------------------------------------------------------------------------------------------
+	// Own Interface
+	//---------------------------------------------------------------------------------------------
+	const FString& NodeObjectGroup::GetName() const
+	{
+		return m_pD->Name;
+	}
+
+
+	//---------------------------------------------------------------------------------------------
+	void NodeObjectGroup::SetName( const FString& Name )
+	{
+		m_pD->Name = Name;
+	}
+
+
+	const FString& NodeObjectGroup::GetUid() const
+	{
+		return m_pD->Uid;
+	}
+
+
+	//---------------------------------------------------------------------------------------------
+	void NodeObjectGroup::SetUid( const FString& Uid )
+	{
+		m_pD->Uid = Uid;
+	}
+
+
+	//---------------------------------------------------------------------------------------------
+	NodeObjectGroup::CHILD_SELECTION NodeObjectGroup::GetSelectionType() const
+	{
+		return m_pD->m_type;
+	}
+
+
+	//---------------------------------------------------------------------------------------------
+	void NodeObjectGroup::SetSelectionType( CHILD_SELECTION t )
+	{
+		m_pD->m_type = t;
+	}
+
+
+	//---------------------------------------------------------------------------------------------
+	int NodeObjectGroup::GetChildCount() const
+	{
+		return m_pD->m_children.Num();
+	}
+
+
+	//---------------------------------------------------------------------------------------------
+	void NodeObjectGroup::SetChildCount( int num )
+	{
+		check( num >=0 );
+		m_pD->m_children.SetNum( num );
+	}
+
+
+	//---------------------------------------------------------------------------------------------
+	NodeObjectPtr NodeObjectGroup::GetChild( int index ) const
+	{
+		check( index >=0 && index < m_pD->m_children.Num() );
+
+		return m_pD->m_children[ index ].get();
+	}
+
+
+	//---------------------------------------------------------------------------------------------
+	void NodeObjectGroup::SetChild( int index, NodeObjectPtr pObject )
+	{
+		check( index >=0 && index < m_pD->m_children.Num() );
+
+		m_pD->m_children[ index ] = pObject;
+	}
+
+
+	void NodeObjectGroup::SetDefaultValue(int32 Value)
+	{
+		m_pD->DefaultValue = Value;
+	}
+
+}
+
+
diff --git a/Engine/Plugins/Mutable/Source/MutableTools/Private/MuT/NodeProjector.cpp b/Engine/Plugins/Mutable/Source/MutableTools/Private/MuT/NodeProjector.cpp
new file mode 100644
--- /dev/null
+++ b/Engine/Plugins/Mutable/Source/MutableTools/Private/MuT/NodeProjector.cpp
@@ -0,0 +1,30 @@
+// Copyright Epic Games, Inc. All Rights Reserved.
+
+#include "MuT/NodeProjector.h"
+
+#include "Misc/AssertionMacros.h"
+
+
+namespace mu
+{
+
+	static FNodeType s_nodeProjectorType = FNodeType(Node::EType::Projector, Node::GetStaticType() );
+
+
+	//---------------------------------------------------------------------------------------------
+	const FNodeType* NodeProjector::GetType() const
+	{
+		return GetStaticType();
+	}
+
+
+	//---------------------------------------------------------------------------------------------
+	const FNodeType* NodeProjector::GetStaticType()
+	{
+		return &s_nodeProjectorType;
+	}
+
+
+}
+
+
diff --git a/Engine/Plugins/Mutable/Source/MutableTools/Private/MuT/NodeProjectorConstant.cpp b/Engine/Plugins/Mutable/Source/MutableTools/Private/MuT/NodeProjectorConstant.cpp
new file mode 100644
--- /dev/null
+++ b/Engine/Plugins/Mutable/Source/MutableTools/Private/MuT/NodeProjectorConstant.cpp
@@ -0,0 +1,64 @@
+// Copyright Epic Games, Inc. All Rights Reserved.
+
+
+#include "Misc/AssertionMacros.h"
+#include "MuR/MutableMath.h"
+#include "MuR/Parameters.h"
+#include "MuT/Node.h"
+#include "MuT/NodePrivate.h"
+#include "MuT/NodeProjector.h"
+#include "MuT/NodeProjectorPrivate.h"
+
+
+namespace mu
+{
+
+	//---------------------------------------------------------------------------------------------
+	// Static initialisation
+	//---------------------------------------------------------------------------------------------
+	FNodeType NodeProjectorConstant::Private::s_type = FNodeType(Node::EType::ProjectorConstant, NodeProjector::GetStaticType() );
+
+
+	//---------------------------------------------------------------------------------------------
+	//!
+	//---------------------------------------------------------------------------------------------
+
+	MUTABLE_IMPLEMENT_NODE( NodeProjectorConstant )
+
+
+	//---------------------------------------------------------------------------------------------
+	// Own Interface
+	//---------------------------------------------------------------------------------------------
+    void NodeProjectorConstant::GetValue( PROJECTOR_TYPE* OutType,
+		FVector3f* OutPos,
+		FVector3f* OutDir,
+		FVector3f* OutUp,
+		FVector3f* OutScale,
+		float* OutProjectionAngle ) const
+	{
+        if (OutType) *OutType = m_pD->m_type;
+
+        if (OutPos) *OutPos = m_pD->m_position;
+        if (OutDir) *OutDir = m_pD->m_direction;
+        if (OutUp) *OutUp = m_pD->m_up;
+        if (OutScale) *OutScale = m_pD->m_scale;
+
+        if (OutProjectionAngle) *OutProjectionAngle = m_pD->m_projectionAngle;
+    }
+
+
+	//---------------------------------------------------------------------------------------------
+    void NodeProjectorConstant::SetValue( PROJECTOR_TYPE type, FVector3f pos, FVector3f dir, FVector3f up, FVector3f scale, float projectionAngle )
+	{
+        m_pD->m_type = type;
+		m_pD->m_position = pos;
+		m_pD->m_direction = dir;
+		m_pD->m_up = up;
+        m_pD->m_scale = scale;
+        m_pD->m_projectionAngle = projectionAngle;
+	}
+
+
+}
+
+
diff --git a/Engine/Plugins/Mutable/Source/MutableTools/Private/MuT/NodeProjectorParameter.cpp b/Engine/Plugins/Mutable/Source/MutableTools/Private/MuT/NodeProjectorParameter.cpp
new file mode 100644
--- /dev/null
+++ b/Engine/Plugins/Mutable/Source/MutableTools/Private/MuT/NodeProjectorParameter.cpp
@@ -0,0 +1,90 @@
+// Copyright Epic Games, Inc. All Rights Reserved.
+
+
+#include "Misc/AssertionMacros.h"
+#include "MuR/MutableMath.h"
+#include "MuR/Parameters.h"
+#include "MuR/Ptr.h"
+#include "MuR/RefCounted.h"
+#include "MuT/Node.h"
+#include "MuT/NodePrivate.h"
+#include "MuT/NodeProjector.h"
+#include "MuT/NodeProjectorPrivate.h"
+#include "MuT/NodeRange.h"
+
+
+namespace mu
+{
+
+	//---------------------------------------------------------------------------------------------
+	// Static initialisation
+	//---------------------------------------------------------------------------------------------
+	FNodeType NodeProjectorParameter::Private::s_type = FNodeType(Node::EType::ProjectorParameter, NodeProjector::GetStaticType() );
+
+
+	//---------------------------------------------------------------------------------------------
+	//!
+	//---------------------------------------------------------------------------------------------
+
+	MUTABLE_IMPLEMENT_NODE( NodeProjectorParameter )
+
+
+	//---------------------------------------------------------------------------------------------
+	// Own Interface
+	//---------------------------------------------------------------------------------------------
+	void NodeProjectorParameter::SetName( const FString& strName )
+	{
+		m_pD->m_name = strName;
+	}
+
+
+	const FString& NodeProjectorParameter::GetUid() const
+	{
+		return m_pD->m_uid;
+	}
+
+
+	void NodeProjectorParameter::SetUid( const FString& strUid )
+	{
+		m_pD->m_uid = strUid;
+	}
+
+
+	//---------------------------------------------------------------------------------------------
+    void NodeProjectorParameter::SetDefaultValue( PROJECTOR_TYPE type,
+		FVector3f pos,
+		FVector3f dir,
+		FVector3f up,
+		FVector3f scale,
+		float projectionAngle )
+	{
+        m_pD->m_type = type;
+        m_pD->m_position = pos;
+		m_pD->m_direction = dir;
+		m_pD->m_up = up;
+        m_pD->m_scale = scale;
+        m_pD->m_projectionAngle = projectionAngle;
+    }
+
+
+    //---------------------------------------------------------------------------------------------
+    void NodeProjectorParameter::SetRangeCount( int i )
+    {
+        check(i>=0);
+        m_pD->m_ranges.SetNum(i);
+    }
+
+
+    //---------------------------------------------------------------------------------------------
+    void NodeProjectorParameter::SetRange( int i, NodeRangePtr pRange )
+    {
+        check( i>=0 && i<m_pD->m_ranges.Num() );
+        if ( i>=0 && i<m_pD->m_ranges.Num() )
+        {
+            m_pD->m_ranges[i] = pRange;
+        }
+    }
+
+}
+
+
diff --git a/Engine/Plugins/Mutable/Source/MutableTools/Private/MuT/NodeRange.cpp b/Engine/Plugins/Mutable/Source/MutableTools/Private/MuT/NodeRange.cpp
new file mode 100644
--- /dev/null
+++ b/Engine/Plugins/Mutable/Source/MutableTools/Private/MuT/NodeRange.cpp
@@ -0,0 +1,33 @@
+// Copyright Epic Games, Inc. All Rights Reserved.
+
+
+#include "MuT/NodeRange.h"
+
+#include "Misc/AssertionMacros.h"
+
+namespace mu
+{
+
+	//---------------------------------------------------------------------------------------------
+	// Static initialisation
+	//---------------------------------------------------------------------------------------------
+    static FNodeType s_nodeRangeType = 	FNodeType(Node::EType::Range, Node::GetStaticType() );
+
+
+	//---------------------------------------------------------------------------------------------
+    const FNodeType* NodeRange::GetType() const
+	{
+		return GetStaticType();
+	}
+
+
+	//---------------------------------------------------------------------------------------------
+    const FNodeType* NodeRange::GetStaticType()
+	{
+        return &s_nodeRangeType;
+	}
+
+
+}
+
+
diff --git a/Engine/Plugins/Mutable/Source/MutableTools/Private/MuT/NodeRangeFromScalar.cpp b/Engine/Plugins/Mutable/Source/MutableTools/Private/MuT/NodeRangeFromScalar.cpp
new file mode 100644
--- /dev/null
+++ b/Engine/Plugins/Mutable/Source/MutableTools/Private/MuT/NodeRangeFromScalar.cpp
@@ -0,0 +1,59 @@
+// Copyright Epic Games, Inc. All Rights Reserved.
+
+
+#include "MuT/NodeRangeFromScalar.h"
+
+#include "Misc/AssertionMacros.h"
+#include "MuR/RefCounted.h"
+#include "MuT/NodePrivate.h"
+#include "MuT/NodeRangeFromScalarPrivate.h"
+#include "MuT/NodeScalar.h"
+
+
+namespace mu
+{
+
+	//---------------------------------------------------------------------------------------------
+	// Static initialisation
+	//---------------------------------------------------------------------------------------------
+    FNodeType NodeRangeFromScalar::Private::s_type = FNodeType(Node::EType::RangeFromScalar, NodeRange::GetStaticType() );
+
+
+	//---------------------------------------------------------------------------------------------
+	//!
+	//---------------------------------------------------------------------------------------------
+
+    MUTABLE_IMPLEMENT_NODE( NodeRangeFromScalar )
+
+
+	//---------------------------------------------------------------------------------------------
+	// Own Interface
+	//---------------------------------------------------------------------------------------------
+	Ptr<NodeScalar> NodeRangeFromScalar::GetSize() const
+	{
+        return m_pD->m_pSize.get();
+	}
+
+
+	//---------------------------------------------------------------------------------------------
+    void NodeRangeFromScalar::SetSize( const Ptr<NodeScalar>& pNode )
+	{
+        m_pD->m_pSize = pNode;
+	}
+
+
+	//---------------------------------------------------------------------------------------------
+    const FString& NodeRangeFromScalar::GetName() const
+	{
+        return m_pD->m_name;
+	}
+
+
+	//---------------------------------------------------------------------------------------------
+    void NodeRangeFromScalar::SetName( const FString& strName )
+	{
+        m_pD->m_name = strName;
+	}
+
+}
+
diff --git a/Engine/Plugins/Mutable/Source/MutableTools/Private/MuT/NodeRangeFromScalarPrivate.h b/Engine/Plugins/Mutable/Source/MutableTools/Private/MuT/NodeRangeFromScalarPrivate.h
new file mode 100644
--- /dev/null
+++ b/Engine/Plugins/Mutable/Source/MutableTools/Private/MuT/NodeRangeFromScalarPrivate.h
@@ -0,0 +1,26 @@
+// Copyright Epic Games, Inc. All Rights Reserved.
+
+#pragma once
+
+#include "MuT/NodeRangePrivate.h"
+#include "MuT/NodeRangeFromScalar.h"
+#include "MuT/NodeRange.h"
+#include "MuT/NodeScalar.h"
+
+
+namespace mu
+{
+
+
+    class NodeRangeFromScalar::Private : public NodeRange::Private
+	{
+	public:
+
+		static FNodeType s_type;
+
+        NodeScalarPtr m_pSize;
+        FString m_name;
+	};
+
+
+}
diff --git a/Engine/Plugins/Mutable/Source/MutableTools/Private/MuT/NodeRangePrivate.h b/Engine/Plugins/Mutable/Source/MutableTools/Private/MuT/NodeRangePrivate.h
new file mode 100644
--- /dev/null
+++ b/Engine/Plugins/Mutable/Source/MutableTools/Private/MuT/NodeRangePrivate.h
@@ -0,0 +1,17 @@
+// Copyright Epic Games, Inc. All Rights Reserved.
+
+#pragma once
+
+#include "MuT/NodePrivate.h"
+#include "MuT/NodeRange.h"
+
+
+namespace mu
+{
+
+
+    class NodeRange::Private : public Node::Private
+	{
+	};
+
+}
diff --git a/Engine/Plugins/Mutable/Source/MutableTools/Private/MuT/NodeScalar.cpp b/Engine/Plugins/Mutable/Source/MutableTools/Private/MuT/NodeScalar.cpp
new file mode 100644
--- /dev/null
+++ b/Engine/Plugins/Mutable/Source/MutableTools/Private/MuT/NodeScalar.cpp
@@ -0,0 +1,31 @@
+// Copyright Epic Games, Inc. All Rights Reserved.
+
+
+#include "MuT/NodeScalar.h"
+
+#include "Misc/AssertionMacros.h"
+
+
+namespace mu
+{
+
+	static FNodeType s_nodeScalarType = FNodeType(Node::EType::Scalar, Node::GetStaticType() );
+
+
+	//---------------------------------------------------------------------------------------------
+	const FNodeType* NodeScalar::GetType() const
+	{
+		return GetStaticType();
+	}
+
+
+	//---------------------------------------------------------------------------------------------
+	const FNodeType* NodeScalar::GetStaticType()
+	{
+		return &s_nodeScalarType;
+	}
+
+
+}
+
+
diff --git a/Engine/Plugins/Mutable/Source/MutableTools/Private/MuT/NodeScalarArithmeticOperation.cpp b/Engine/Plugins/Mutable/Source/MutableTools/Private/MuT/NodeScalarArithmeticOperation.cpp
new file mode 100644
--- /dev/null
+++ b/Engine/Plugins/Mutable/Source/MutableTools/Private/MuT/NodeScalarArithmeticOperation.cpp
@@ -0,0 +1,84 @@
+// Copyright Epic Games, Inc. All Rights Reserved.
+
+
+#include "MuT/NodeScalarArithmeticOperation.h"
+
+#include "Misc/AssertionMacros.h"
+#include "MuT/NodePrivate.h"
+#include "MuT/NodeScalar.h"
+#include "MuT/NodeScalarArithmeticOperationPrivate.h"
+#include "MuR/Serialisation.h"
+
+
+namespace mu
+{
+    MUTABLE_IMPLEMENT_ENUM_SERIALISABLE(NodeScalarArithmeticOperation::OPERATION)
+
+	
+	//---------------------------------------------------------------------------------------------
+	// Static initialisation
+	//---------------------------------------------------------------------------------------------
+    FNodeType NodeScalarArithmeticOperation::Private::s_type = FNodeType(Node::EType::ScalarArithmeticOperation, NodeScalar::GetStaticType() );
+
+
+	//---------------------------------------------------------------------------------------------
+	//!
+	//---------------------------------------------------------------------------------------------
+
+    MUTABLE_IMPLEMENT_NODE( NodeScalarArithmeticOperation )
+
+
+    //---------------------------------------------------------------------------------------------
+    const char* NodeScalarArithmeticOperation::s_opTypeName[] =
+    {
+        "Add",
+        "Subtract",
+        "Multiply",
+        "Divide"
+    };
+
+
+	//---------------------------------------------------------------------------------------------
+	// Own Interface
+	//---------------------------------------------------------------------------------------------
+    NodeScalarArithmeticOperation::OPERATION NodeScalarArithmeticOperation::GetOperation() const
+	{
+		return m_pD->m_operation;
+	}
+
+
+	//---------------------------------------------------------------------------------------------
+    void NodeScalarArithmeticOperation::SetOperation(NodeScalarArithmeticOperation::OPERATION o)
+	{
+		m_pD->m_operation = o;
+	}
+
+
+    NodeScalarPtr NodeScalarArithmeticOperation::GetA() const
+	{
+		return m_pD->m_pA.get();
+	}
+
+
+	//---------------------------------------------------------------------------------------------
+    void NodeScalarArithmeticOperation::SetA(NodeScalarPtr pNode)
+	{
+		m_pD->m_pA = pNode;
+	}
+
+
+	//---------------------------------------------------------------------------------------------
+    NodeScalarPtr NodeScalarArithmeticOperation::GetB() const
+	{
+		return m_pD->m_pB.get();
+	}
+
+
+	//---------------------------------------------------------------------------------------------
+    void NodeScalarArithmeticOperation::SetB( NodeScalarPtr pNode )
+	{
+		m_pD->m_pB = pNode;
+	}
+
+}
+
diff --git a/Engine/Plugins/Mutable/Source/MutableTools/Private/MuT/NodeScalarArithmeticOperationPrivate.h b/Engine/Plugins/Mutable/Source/MutableTools/Private/MuT/NodeScalarArithmeticOperationPrivate.h
new file mode 100644
--- /dev/null
+++ b/Engine/Plugins/Mutable/Source/MutableTools/Private/MuT/NodeScalarArithmeticOperationPrivate.h
@@ -0,0 +1,25 @@
+// Copyright Epic Games, Inc. All Rights Reserved.
+
+#pragma once
+
+#include "MuT/NodeScalarPrivate.h"
+#include "MuT/NodeScalarArithmeticOperation.h"
+
+namespace mu
+{
+    MUTABLE_DEFINE_ENUM_SERIALISABLE(NodeScalarArithmeticOperation::OPERATION)
+
+    class NodeScalarArithmeticOperation::Private : public NodeScalar::Private
+	{
+	public:
+
+		static FNodeType s_type;
+
+        NodeScalarArithmeticOperation::OPERATION m_operation;
+        NodeScalarPtr m_pA;
+        NodeScalarPtr m_pB;
+
+	};
+
+
+}
diff --git a/Engine/Plugins/Mutable/Source/MutableTools/Private/MuT/NodeScalarConstant.cpp b/Engine/Plugins/Mutable/Source/MutableTools/Private/MuT/NodeScalarConstant.cpp
new file mode 100644
--- /dev/null
+++ b/Engine/Plugins/Mutable/Source/MutableTools/Private/MuT/NodeScalarConstant.cpp
@@ -0,0 +1,43 @@
+// Copyright Epic Games, Inc. All Rights Reserved.
+
+
+#include "MuT/NodeScalarConstant.h"
+
+#include "Misc/AssertionMacros.h"
+#include "MuT/NodePrivate.h"
+#include "MuT/NodeScalarConstantPrivate.h"
+
+
+namespace mu
+{
+
+	//---------------------------------------------------------------------------------------------
+	// Static initialisation
+	//---------------------------------------------------------------------------------------------
+	FNodeType NodeScalarConstant::Private::s_type = FNodeType(Node::EType::ScalarConstant, NodeScalar::GetStaticType() );
+
+
+	//---------------------------------------------------------------------------------------------
+	//!
+	//---------------------------------------------------------------------------------------------
+
+	MUTABLE_IMPLEMENT_NODE( NodeScalarConstant )
+
+
+	//---------------------------------------------------------------------------------------------
+	// Own Interface
+	//---------------------------------------------------------------------------------------------
+	float NodeScalarConstant::GetValue() const
+	{
+		return m_pD->m_value;
+	}
+
+	//---------------------------------------------------------------------------------------------
+	void NodeScalarConstant::SetValue( float v )
+	{
+		m_pD->m_value = v;
+	}
+
+
+}
+
diff --git a/Engine/Plugins/Mutable/Source/MutableTools/Private/MuT/NodeScalarConstantPrivate.h b/Engine/Plugins/Mutable/Source/MutableTools/Private/MuT/NodeScalarConstantPrivate.h
new file mode 100644
--- /dev/null
+++ b/Engine/Plugins/Mutable/Source/MutableTools/Private/MuT/NodeScalarConstantPrivate.h
@@ -0,0 +1,23 @@
+// Copyright Epic Games, Inc. All Rights Reserved.
+
+#pragma once
+
+#include "MuT/NodeScalarPrivate.h"
+
+#include "MuT/NodeScalarConstant.h"
+#include "MuT/NodeImage.h"
+
+namespace mu
+{
+
+	class NodeScalarConstant::Private : public NodeScalar::Private
+	{
+	public:
+
+		static FNodeType s_type;
+
+		float m_value;
+	};
+
+
+}
diff --git a/Engine/Plugins/Mutable/Source/MutableTools/Private/MuT/NodeScalarCurve.cpp b/Engine/Plugins/Mutable/Source/MutableTools/Private/MuT/NodeScalarCurve.cpp
new file mode 100644
--- /dev/null
+++ b/Engine/Plugins/Mutable/Source/MutableTools/Private/MuT/NodeScalarCurve.cpp
@@ -0,0 +1,21 @@
+// Copyright Epic Games, Inc. All Rights Reserved.
+
+
+#include "MuT/NodeScalarCurve.h"
+
+#include "Misc/AssertionMacros.h"
+#include "MuR/ParametersPrivate.h"
+#include "MuT/NodePrivate.h"
+#include "MuT/NodeScalar.h"
+#include "MuT/NodeScalarCurvePrivate.h"
+
+
+namespace mu
+{
+
+    FNodeType NodeScalarCurve::Private::s_type = FNodeType(Node::EType::ScalarCurve, NodeScalar::GetStaticType() );
+
+    MUTABLE_IMPLEMENT_NODE( NodeScalarCurve )
+
+}
+
diff --git a/Engine/Plugins/Mutable/Source/MutableTools/Private/MuT/NodeScalarEnumParameter.cpp b/Engine/Plugins/Mutable/Source/MutableTools/Private/MuT/NodeScalarEnumParameter.cpp
new file mode 100644
--- /dev/null
+++ b/Engine/Plugins/Mutable/Source/MutableTools/Private/MuT/NodeScalarEnumParameter.cpp
@@ -0,0 +1,111 @@
+// Copyright Epic Games, Inc. All Rights Reserved.
+
+
+#include "MuT/NodeScalarEnumParameter.h"
+
+#include "Misc/AssertionMacros.h"
+#include "MuT/NodePrivate.h"
+#include "MuT/NodeScalarEnumParameterPrivate.h"
+
+
+namespace mu
+{
+
+	//---------------------------------------------------------------------------------------------
+	// Static initialisation
+	//---------------------------------------------------------------------------------------------
+	FNodeType NodeScalarEnumParameter::Private::s_type = FNodeType(Node::EType::ScalarEnumParameter, NodeScalar::GetStaticType() );
+
+
+	//---------------------------------------------------------------------------------------------
+	//!
+	//---------------------------------------------------------------------------------------------
+
+	MUTABLE_IMPLEMENT_NODE( NodeScalarEnumParameter )
+
+
+	//---------------------------------------------------------------------------------------------
+	const FString& NodeScalarEnumParameter::GetName() const
+	{
+		return m_pD->m_name;
+	}
+
+
+	//---------------------------------------------------------------------------------------------
+	void NodeScalarEnumParameter::SetName( const FString& strName )
+	{
+		m_pD->m_name = strName;
+	}
+
+
+	const FString& NodeScalarEnumParameter::GetUid() const
+	{
+		return m_pD->m_uid;
+	}
+
+
+	//---------------------------------------------------------------------------------------------
+	void NodeScalarEnumParameter::SetUid( const FString& strUid )
+	{
+		m_pD->m_uid = strUid;
+	}
+
+
+	//---------------------------------------------------------------------------------------------
+	int NodeScalarEnumParameter::GetDefaultValueIndex() const
+	{
+		return m_pD->m_defaultValue;
+	}
+
+
+	//---------------------------------------------------------------------------------------------
+	void NodeScalarEnumParameter::SetDefaultValueIndex( int i )
+	{
+		m_pD->m_defaultValue = i;
+	}
+
+
+	//---------------------------------------------------------------------------------------------
+	void NodeScalarEnumParameter::SetValueCount( int i )
+	{
+		m_pD->m_options.SetNum(i);
+	}
+
+
+	//---------------------------------------------------------------------------------------------
+	int NodeScalarEnumParameter::GetValueCount() const
+	{
+		return (int)m_pD->m_options.Num();
+	}
+
+
+	//---------------------------------------------------------------------------------------------
+	void NodeScalarEnumParameter::SetValue( int i, float value, const FString& strName )
+	{
+		check( i>=0 && i<(int)m_pD->m_options.Num() );
+		m_pD->m_options[i].name = strName;
+		m_pD->m_options[i].value = value;
+	}
+
+
+    //---------------------------------------------------------------------------------------------
+    void NodeScalarEnumParameter::SetRangeCount( int i )
+    {
+        check(i>=0);
+        m_pD->m_ranges.SetNum(i);
+    }
+
+
+    //---------------------------------------------------------------------------------------------
+    void NodeScalarEnumParameter::SetRange( int i, NodeRangePtr pRange )
+    {
+        check( i>=0 && i<int(m_pD->m_ranges.Num()) );
+        if ( i>=0 && i<int(m_pD->m_ranges.Num()) )
+        {
+            m_pD->m_ranges[i] = pRange;
+        }
+    }
+
+}
+
+
diff --git a/Engine/Plugins/Mutable/Source/MutableTools/Private/MuT/NodeScalarParameter.cpp b/Engine/Plugins/Mutable/Source/MutableTools/Private/MuT/NodeScalarParameter.cpp
new file mode 100644
--- /dev/null
+++ b/Engine/Plugins/Mutable/Source/MutableTools/Private/MuT/NodeScalarParameter.cpp
@@ -0,0 +1,71 @@
+// Copyright Epic Games, Inc. All Rights Reserved.
+
+
+#include "MuT/NodeScalarParameter.h"
+
+#include "Misc/AssertionMacros.h"
+#include "MuT/NodeImage.h"
+#include "MuT/NodePrivate.h"
+#include "MuT/NodeRange.h"
+#include "MuT/NodeScalarParameterPrivate.h"
+
+
+namespace mu
+{
+
+	//---------------------------------------------------------------------------------------------
+	// Static initialisation
+	//---------------------------------------------------------------------------------------------
+	FNodeType NodeScalarParameter::Private::s_type = FNodeType(Node::EType::ScalarParameter, NodeScalar::GetStaticType() );
+
+
+	//---------------------------------------------------------------------------------------------
+	//!
+	//---------------------------------------------------------------------------------------------
+
+	MUTABLE_IMPLEMENT_NODE( NodeScalarParameter )
+
+
+	//---------------------------------------------------------------------------------------------
+	void NodeScalarParameter::SetName( const FString& Name )
+	{
+		m_pD->m_name = Name;
+	}
+
+
+	//---------------------------------------------------------------------------------------------
+	void NodeScalarParameter::SetUid( const FString& Uid )
+	{
+		m_pD->m_uid = Uid;
+	}
+
+
+	//---------------------------------------------------------------------------------------------
+	void NodeScalarParameter::SetDefaultValue( float v )
+	{
+		m_pD->m_defaultValue = v;
+	}
+
+
+    //---------------------------------------------------------------------------------------------
+    void NodeScalarParameter::SetRangeCount( int i )
+    {
+        check(i>=0);
+        m_pD->m_ranges.SetNum(i);
+    }
+
+
+    //---------------------------------------------------------------------------------------------
+    void NodeScalarParameter::SetRange( int i, NodeRangePtr pRange )
+    {
+        check( i>=0 && i<int(m_pD->m_ranges.Num()) );
+        if ( i>=0 && i<int(m_pD->m_ranges.Num()) )
+        {
+            m_pD->m_ranges[i] = pRange;
+        }
+    }
+
+
+}
+
+
diff --git a/Engine/Plugins/Mutable/Source/MutableTools/Private/MuT/NodeScalarTable.cpp b/Engine/Plugins/Mutable/Source/MutableTools/Private/MuT/NodeScalarTable.cpp
new file mode 100644
--- /dev/null
+++ b/Engine/Plugins/Mutable/Source/MutableTools/Private/MuT/NodeScalarTable.cpp
@@ -0,0 +1,51 @@
+// Copyright Epic Games, Inc. All Rights Reserved.
+
+
+#include "MuT/NodeScalarTable.h"
+
+#include "Misc/AssertionMacros.h"
+#include "MuT/NodePrivate.h"
+#include "MuT/Table.h"
+
+
+namespace mu
+{
+
+
+	//---------------------------------------------------------------------------------------------
+	// Static initialisation
+	//---------------------------------------------------------------------------------------------
+	FNodeType NodeScalarTable::StaticType = FNodeType(Node::EType::ScalarTable, NodeScalar::GetStaticType() );
+
+
+
+	//---------------------------------------------------------------------------------------------
+	// Own Interface
+	//---------------------------------------------------------------------------------------------
+	void NodeScalarTable::SetColumn( const FString& strName )
+	{
+		ColumnName = strName;
+	}
+
+
+	//---------------------------------------------------------------------------------------------
+	void NodeScalarTable::SetParameterName( const FString& strName )
+	{
+		ParameterName = strName;
+	}
+
+
+	//---------------------------------------------------------------------------------------------
+	void NodeScalarTable::SetNoneOption(bool bAddNoneOption)
+	{
+		bNoneOption = bAddNoneOption;
+	}
+
+	//---------------------------------------------------------------------------------------------
+	void NodeScalarTable::SetDefaultRowName(const FString& RowName)
+	{
+		DefaultRowName = RowName;
+	}
+}
+
+
diff --git a/Engine/Plugins/Mutable/Source/MutableTools/Private/MuT/NodeScalarVariation.cpp b/Engine/Plugins/Mutable/Source/MutableTools/Private/MuT/NodeScalarVariation.cpp
new file mode 100644
--- /dev/null
+++ b/Engine/Plugins/Mutable/Source/MutableTools/Private/MuT/NodeScalarVariation.cpp
@@ -0,0 +1,59 @@
+// Copyright Epic Games, Inc. All Rights Reserved.
+
+
+#include "MuT/NodeScalarVariation.h"
+
+#include "Misc/AssertionMacros.h"
+#include "MuT/NodePrivate.h"
+#include "MuT/NodeScalarVariationPrivate.h"
+
+
+namespace mu
+{
+
+    //---------------------------------------------------------------------------------------------
+    // Static initialisation
+    //---------------------------------------------------------------------------------------------
+    FNodeType NodeScalarVariation::Private::s_type = FNodeType(Node::EType::ScalarVariation, NodeScalar::GetStaticType() );
+
+
+    //---------------------------------------------------------------------------------------------
+    //---------------------------------------------------------------------------------------------
+    //---------------------------------------------------------------------------------------------
+
+    MUTABLE_IMPLEMENT_NODE( NodeScalarVariation )
+
+
+    //---------------------------------------------------------------------------------------------
+    // Own Interface
+    //---------------------------------------------------------------------------------------------
+    void NodeScalarVariation::SetDefaultScalar( NodeScalar* p ) 
+	{ 
+		m_pD->m_defaultScalar = p; 
+	}
+
+
+	void NodeScalarVariation::SetVariationCount( int32 num )
+    {
+        check( num >= 0 );
+        m_pD->m_variations.SetNum( num );
+    }
+
+
+	void NodeScalarVariation::SetVariationTag( int32 index, const FString& Tag )
+    {
+        check( index >= 0 && index < m_pD->m_variations.Num() );
+
+		m_pD->m_variations[index].m_tag = Tag;
+    }
+
+
+    void NodeScalarVariation::SetVariationScalar( int32 index, NodeScalar* pNode )
+    {
+        check( index >= 0 && index < m_pD->m_variations.Num() );
+
+        m_pD->m_variations[index].m_scalar = pNode;
+    }
+
+
+} // namespace mu
diff --git a/Engine/Plugins/Mutable/Source/MutableTools/Private/MuT/NodeScalarVariationPrivate.h b/Engine/Plugins/Mutable/Source/MutableTools/Private/MuT/NodeScalarVariationPrivate.h
new file mode 100644
--- /dev/null
+++ b/Engine/Plugins/Mutable/Source/MutableTools/Private/MuT/NodeScalarVariationPrivate.h
@@ -0,0 +1,32 @@
+// Copyright Epic Games, Inc. All Rights Reserved.
+
+#pragma once
+
+#include "MuT/NodePrivate.h"
+#include "MuT/NodeScalarVariation.h"
+
+namespace mu
+{
+
+    //---------------------------------------------------------------------------------------------
+    //!
+    //---------------------------------------------------------------------------------------------
+    class NodeScalarVariation::Private : public Node::Private
+    {
+    public:
+
+        static FNodeType s_type;
+
+        NodeScalarPtr m_defaultScalar;
+
+        struct FVariation
+        {
+            NodeScalarPtr m_scalar;
+            FString m_tag;
+        };
+
+        TArray<FVariation> m_variations;
+
+    };
+
+} // namespace mu
diff --git a/Engine/Plugins/Mutable/Source/MutableTools/Private/MuT/NodeString.cpp b/Engine/Plugins/Mutable/Source/MutableTools/Private/MuT/NodeString.cpp
new file mode 100644
--- /dev/null
+++ b/Engine/Plugins/Mutable/Source/MutableTools/Private/MuT/NodeString.cpp
@@ -0,0 +1,34 @@
+// Copyright Epic Games, Inc. All Rights Reserved.
+
+
+#include "MuT/NodeString.h"
+
+#include "Misc/AssertionMacros.h"
+
+
+namespace mu
+{
+
+	//---------------------------------------------------------------------------------------------
+	// Static initialisation
+	//---------------------------------------------------------------------------------------------
+	static FNodeType s_nodeStringType = FNodeType(Node::EType::String, Node::GetStaticType() );
+
+
+	//---------------------------------------------------------------------------------------------
+	const FNodeType* NodeString::GetType() const
+	{
+		return GetStaticType();
+	}
+
+
+	//---------------------------------------------------------------------------------------------
+	const FNodeType* NodeString::GetStaticType()
+	{
+		return &s_nodeStringType;
+	}
+
+
+}
+
+
diff --git a/Engine/Plugins/Mutable/Source/MutableTools/Private/MuT/NodeStringConstant.cpp b/Engine/Plugins/Mutable/Source/MutableTools/Private/MuT/NodeStringConstant.cpp
new file mode 100644
--- /dev/null
+++ b/Engine/Plugins/Mutable/Source/MutableTools/Private/MuT/NodeStringConstant.cpp
@@ -0,0 +1,37 @@
+// Copyright Epic Games, Inc. All Rights Reserved.
+
+
+#include "MuT/NodeStringConstant.h"
+
+#include "Misc/AssertionMacros.h"
+#include "MuT/NodePrivate.h"
+#include "MuT/NodeStringConstantPrivate.h"
+
+
+namespace mu
+{
+
+	//---------------------------------------------------------------------------------------------
+	// Static initialisation
+	//---------------------------------------------------------------------------------------------
+	FNodeType NodeStringConstant::Private::s_type = FNodeType(Node::EType::StringConstant, NodeString::GetStaticType() );
+
+
+	//---------------------------------------------------------------------------------------------
+	//!
+	//---------------------------------------------------------------------------------------------
+
+	MUTABLE_IMPLEMENT_NODE( NodeStringConstant )
+
+
+	//---------------------------------------------------------------------------------------------
+	// Own Interface
+	//---------------------------------------------------------------------------------------------
+	void NodeStringConstant::SetValue( const FString& v )
+	{
+		m_pD->m_value = v;
+	}
+
+
+}
+
diff --git a/Engine/Plugins/Mutable/Source/MutableTools/Private/MuT/NodeStringConstantPrivate.h b/Engine/Plugins/Mutable/Source/MutableTools/Private/MuT/NodeStringConstantPrivate.h
new file mode 100644
--- /dev/null
+++ b/Engine/Plugins/Mutable/Source/MutableTools/Private/MuT/NodeStringConstantPrivate.h
@@ -0,0 +1,25 @@
+// Copyright Epic Games, Inc. All Rights Reserved.
+
+#pragma once
+
+#include "MuT/NodeStringPrivate.h"
+
+#include "MuT/NodeStringConstant.h"
+#include "MuT/NodeImage.h"
+
+namespace mu
+{
+
+
+	class NodeStringConstant::Private : public NodeString::Private
+	{
+	public:
+
+		static FNodeType s_type;
+
+		FString m_value;
+
+	};
+
+
+}
diff --git a/Engine/Plugins/Mutable/Source/MutableTools/Private/MuT/NodeStringParameter.cpp b/Engine/Plugins/Mutable/Source/MutableTools/Private/MuT/NodeStringParameter.cpp
new file mode 100644
--- /dev/null
+++ b/Engine/Plugins/Mutable/Source/MutableTools/Private/MuT/NodeStringParameter.cpp
@@ -0,0 +1,71 @@
+// Copyright Epic Games, Inc. All Rights Reserved.
+
+
+#include "MuT/NodeStringParameter.h"
+
+#include "Misc/AssertionMacros.h"
+#include "MuT/NodeImage.h"
+#include "MuT/NodePrivate.h"
+#include "MuT/NodeRange.h"
+#include "MuT/NodeStringParameterPrivate.h"
+
+
+namespace mu
+{
+
+	//---------------------------------------------------------------------------------------------
+	// Static initialisation
+	//---------------------------------------------------------------------------------------------
+	FNodeType NodeStringParameter::Private::s_type = FNodeType(Node::EType::StringParameter, NodeString::GetStaticType() );
+
+
+	//---------------------------------------------------------------------------------------------
+	//!
+	//---------------------------------------------------------------------------------------------
+
+	MUTABLE_IMPLEMENT_NODE( NodeStringParameter )
+
+
+	//---------------------------------------------------------------------------------------------
+	void NodeStringParameter::SetName( const FString& Name )
+	{
+		m_pD->m_name = Name;
+	}
+
+
+	//---------------------------------------------------------------------------------------------
+	void NodeStringParameter::SetUid( const FString& Uid )
+	{
+		m_pD->m_uid = Uid;
+	}
+
+
+	//---------------------------------------------------------------------------------------------
+	void NodeStringParameter::SetDefaultValue( const FString& v )
+	{
+		m_pD->m_defaultValue = v;
+	}
+
+
+    //---------------------------------------------------------------------------------------------
+    void NodeStringParameter::SetRangeCount( int i )
+    {
+        check(i>=0);
+        m_pD->m_ranges.SetNum(i);
+    }
+
+
+    //---------------------------------------------------------------------------------------------
+    void NodeStringParameter::SetRange( int i, NodeRangePtr pRange )
+    {
+        check( i>=0 && i<int(m_pD->m_ranges.Num()) );
+        if ( i>=0 && i<int(m_pD->m_ranges.Num()) )
+        {
+            m_pD->m_ranges[i] = pRange;
+        }
+    }
+
+
+}
+
+
diff --git a/Engine/Plugins/Mutable/Source/MutableTools/Private/MuT/NodeStringParameterPrivate.h b/Engine/Plugins/Mutable/Source/MutableTools/Private/MuT/NodeStringParameterPrivate.h
new file mode 100644
--- /dev/null
+++ b/Engine/Plugins/Mutable/Source/MutableTools/Private/MuT/NodeStringParameterPrivate.h
@@ -0,0 +1,31 @@
+// Copyright Epic Games, Inc. All Rights Reserved.
+
+#pragma once
+
+#include "MuT/NodeStringPrivate.h"
+
+#include "MuT/NodeStringParameter.h"
+#include "MuT/NodeImage.h"
+#include "MuT/NodeRange.h"
+
+
+namespace mu
+{
+
+	class NodeStringParameter::Private : public NodeString::Private
+	{
+	public:
+
+		static FNodeType s_type;
+
+		FString m_defaultValue;
+		FString m_name;
+		FString m_uid;
+
+        TArray<Ptr<NodeImage>> m_additionalImages;
+
+		TArray<Ptr<NodeRange>> m_ranges;
+
+	};
+
+}
diff --git a/Engine/Plugins/Mutable/Source/MutableTools/Private/MuT/NodeStringPrivate.h b/Engine/Plugins/Mutable/Source/MutableTools/Private/MuT/NodeStringPrivate.h
new file mode 100644
--- /dev/null
+++ b/Engine/Plugins/Mutable/Source/MutableTools/Private/MuT/NodeStringPrivate.h
@@ -0,0 +1,20 @@
+// Copyright Epic Games, Inc. All Rights Reserved.
+
+#pragma once
+
+#include "MuT/NodePrivate.h"
+
+#include "MuT/NodeString.h"
+#include "MuR/ParametersPrivate.h"
+
+
+namespace mu
+{
+
+
+	class NodeString::Private : public Node::Private
+	{
+	};
+
+
+}
diff --git a/Engine/Plugins/Mutable/Source/MutableTools/Private/MuT/TypeInfo.cpp b/Engine/Plugins/Mutable/Source/MutableTools/Private/MuT/TypeInfo.cpp
new file mode 100644
--- /dev/null
+++ b/Engine/Plugins/Mutable/Source/MutableTools/Private/MuT/TypeInfo.cpp
@@ -0,0 +1,152 @@
+// Copyright Epic Games, Inc. All Rights Reserved.
+
+#include "MuT/TypeInfo.h"
+
+#include "MuR/ParametersPrivate.h"
+#include "MuR/Types.h"
+
+namespace mu
+{
+
+	const char* TypeInfo::s_imageFormatName[size_t(EImageFormat::IF_COUNT)] =
+	{
+		"None",
+		"RGB U888",
+		"RGBA U8888",
+		"L U8",
+
+		"PVRTC2",
+		"PVRTC4",
+		"ETC1",
+		"ETC2",
+
+		"L U8 RLE",
+		"RGB U888 RLE",
+		"RGBA U8888 RLE",
+		"L U1 RLE",
+
+		"BC1",
+		"BC2",
+		"BC3",
+		"BC4",
+		"BC5",
+		"BC6",
+		"BC7",
+
+		"BGRA U8888",
+
+		"ASTC_4x4_RGB_LDR",
+		"ASTC_4x4_RGBA_LDR",
+		"ASTC_4x4_RG_LDR",
+
+		"ASTC_8x8_RGB_LDR",
+		"ASTC_8x8_RGBA_LDR",
+		"ASTC_8x8_RG_LDR",
+		"ASTC_12x12_RGB_LDR",
+		"ASTC_12x12_RGBA_LDR",
+		"ASTC_12x12_RG_LDR",
+
+		"ASTC_6x6_RGB_LDR",
+		"ASTC_6x6_RGBA_LDR",
+		"ASTC_6x6_RG_LDR",
+		"ASTC_10x10_RGB_LDR",
+		"ASTC_10x10_RGBA_LDR",
+		"ASTC_10x10_RG_LDR",
+
+	};
+
+	static_assert(sizeof(TypeInfo::s_imageFormatName) / sizeof(void*) == int32(EImageFormat::IF_COUNT));
+
+
+	const char* TypeInfo::s_meshBufferSemanticName[] =
+	{
+		"None",
+
+		"VertexIndex",
+
+		"Position",
+		"Normal",
+		"Tangent",
+		"Binormal",
+		"Tex Coords",
+		"Colour",
+		"Bone Weights",
+		"Bone Indices",
+
+		"Layout Block",
+
+		"Chart",
+
+		"Other",
+
+		"Tangent Space Sign",
+
+		"TriangleIndex",
+		"BarycentricCoords",
+		"Distance",
+
+		"AltSkinWeight"
+	};
+
+	static_assert(sizeof(TypeInfo::s_meshBufferSemanticName) / sizeof(void*) == int32(MBS_COUNT));
+
+	const char* TypeInfo::s_meshBufferFormatName[MBF_COUNT] =
+	{
+		"None",
+		"Float16",
+		"Float32",
+
+		"UInt8",
+		"UInt16",
+		"UInt32",
+		"Int8",
+		"Int16",
+		"Int32",
+
+		"NUInt8",
+		"NUInt16",
+		"NUInt32",
+		"NInt8",
+		"NInt16",
+		"NInt32",
+
+		"PackedDir8",
+		"PackedDir8_WTangentSign",
+		"PackedDirS8",
+		"PackedDirS8_WTangentSign",
+
+		"Float64",
+
+		"UInt64",
+		"Int64",
+		"NUint64",
+		"NInt64",
+	};
+
+	static_assert(sizeof(TypeInfo::s_meshBufferFormatName) / sizeof(void*) == int32(MBF_COUNT));
+
+	const char* TypeInfo::s_projectorTypeName[static_cast<uint32>(PROJECTOR_TYPE::COUNT)] =
+	{
+		"Planar",
+		"Cylindrical",
+		"Wrapping",
+	};
+	static_assert(sizeof(TypeInfo::s_projectorTypeName) / sizeof(void*) == int32(PROJECTOR_TYPE::COUNT));
+
+	const char* TypeInfo::s_blendModeName[uint32(EBlendType::_BT_COUNT)] =
+	{
+		"None",
+		"SoftLight",
+		"HardLight",
+		"Burn",
+		"Dodge",
+		"Screen",
+		"Overlay",
+		"Blend",
+		"Multiply",
+		"Lighten",
+		"NormalCombine",
+	};
+	static_assert(sizeof(TypeInfo::s_blendModeName) / sizeof(void*) == int32(EBlendType::_BT_COUNT));
+
+}
diff --git a/Engine/Plugins/Mutable/Source/MutableTools/Public/MuT/Compiler.h b/Engine/Plugins/Mutable/Source/MutableTools/Public/MuT/Compiler.h
new file mode 100644
--- /dev/null
+++ b/Engine/Plugins/Mutable/Source/MutableTools/Public/MuT/Compiler.h
@@ -0,0 +1,239 @@
+// Copyright Epic Games, Inc. All Rights Reserved.
+
+#pragma once
+
+#include "MuR/Ptr.h"
+#include "MuR/RefCounted.h"
+#include "MuR/Image.h"
+#include "MuR/System.h"
+#include "Templates/SharedPointer.h"
+#include "HAL/PlatformMath.h"
+#include "Tasks/Task.h"
+
+namespace mu
+{
+    // Forward declarations
+    class Compiler;
+    typedef Ptr<Compiler> CompilerPtr;
+    typedef Ptr<const Compiler> CompilerPtrConst;
+
+    class Model;
+    using ModelPtr=Ptr<Model>;
+    using ModelPtrConst=Ptr<const Model>;
+
+    class Transform;
+    using TransformPtr=Ptr<Transform>;
+    using TransformPtrConst=Ptr<const Transform>;
+
+    class ErrorLog;
+    using ErrorLogPtr=Ptr<ErrorLog>;
+    using ErrorLogPtrConst=Ptr<const ErrorLog>;
+
+    class Node;
+    class NodeTransformedObject;
+
+	typedef TFunction<UE::Tasks::FTask(int32, TSharedPtr<Ptr<Image>>, bool)> FReferencedResourceFunc;
+
+    //! \brief Options used to compile the models with a compiler.
+    class MUTABLETOOLS_API CompilerOptions : public RefCounted
+    {
+    public:
+
+        //! Create new settings with the default values as explained in each method below.
+        CompilerOptions();
+
+        //!
+        void SetLogEnabled( bool bEnabled );
+
+        //!
+        void SetOptimisationEnabled( bool bEnabled );
+
+        //!
+        void SetConstReductionEnabled( bool bEnabled );
+
+        //!
+        void SetOptimisationMaxIteration( int32 MaxIterations );
+
+        //!
+        void SetIgnoreStates( bool bIgnore );
+
+		/** Enable concurrent compilation. It's faster, but uses more CPU and memory. */
+		void SetUseConcurrency(bool bEnabled);
+
+        //! If enabled, the disk will be used as temporary memory. This will make the compilation
+        //! process very slow, but will be able to compile very large models.
+        void SetUseDiskCache( bool bEnabled );
+
+		//! Set the quality for the image compression algorithms. The level value is used internally
+		//! with System::SetImagecompressionQuality
+		void SetImageCompressionQuality(int32 Quality);
+
+		/** Set the image tiling strategy :
+		 * If 0 (default) there is no tiling. Otherwise, images will be generated in tiles of the given size or less, and assembled afterwards as a final step.
+		 */
+		void SetImageTiling(int32 Tiling);
+
+        //! 
+        void SetDataPackingStrategy( int32 MinTextureResidentMipCount, uint64 EmbeddedDataBytesLimit, uint64 PackagedDataBytesLimit);
+
+		/** If enabled it will make sure that the object is compile to generate smaller mips of the images. */
+		void SetEnableProgressiveImages(bool bEnabled);
+
+		/** Set an optional pixel conversion function that will be called before any pixel format conversion. */
+		void SetImagePixelFormatOverride(const FImageOperator::FImagePixelFormatFunc&);
+
+		/** */
+		void SetReferencedResourceCallback(const FReferencedResourceFunc&);
+
+        //! Different data packing strategies
+        enum class TextureLayoutStrategy : uint8
+        {
+            //! Pack texture layouts without changing any scale
+            Pack,
+
+            //! Do not touch mesh or image texture layouts
+            None,
+
+            //! Helper value, not really a strategy.
+            Count
+        };
+
+        //! 
+        static const char* GetTextureLayoutStrategyName( TextureLayoutStrategy s );
+
+		/** Output some stats about the complete compilation to the log. */
+		void LogStats() const;
+
+        //-----------------------------------------------------------------------------------------
+        // Interface pattern
+        //-----------------------------------------------------------------------------------------
+        class Private;
+        Private* GetPrivate() const;
+
+        CompilerOptions( const CompilerOptions& ) = delete;
+        CompilerOptions( CompilerOptions&& ) = delete;
+        CompilerOptions& operator=(const CompilerOptions&) = delete;
+        CompilerOptions& operator=(CompilerOptions&&) = delete;
+
+    protected:
+
+        //! Forbidden. Manage with the Ptr<> template.
+        ~CompilerOptions() override;
+
+    private:
+
+        Private* m_pD;
+
+    };
+
+
+    //!
+    struct FStateOptimizationOptions
+    {
+		uint8 FirstLOD = 0;
+		uint8 NumExtraLODsToBuildAfterFirstLOD = 0;
+		bool bOnlyFirstLOD = false;
+		ETextureCompressionStrategy TextureCompressionStrategy = ETextureCompressionStrategy::None;
+
+        void Serialise( OutputArchive& arch ) const
+        {
+            const int32_t ver = 4;
+            arch << ver;
+
+			arch << FirstLOD;
+			arch << bOnlyFirstLOD;
+            arch << TextureCompressionStrategy;
+			arch << NumExtraLODsToBuildAfterFirstLOD;
+        }
+
+
+        void Unserialise( InputArchive& arch )
+        {
+            int32_t ver = 0;
+            arch >> ver;
+			check(ver <= 4);
+
+			if (ver >= 2)
+			{
+				arch >> FirstLOD;
+			}
+			else
+			{
+				FirstLOD = 0;
+			}
+
+            arch >> bOnlyFirstLOD;
+
+			if (ver >= 4)
+			{
+				arch >> TextureCompressionStrategy;
+			}
+			else
+			{
+				bool bAvoidRuntimeCompression;
+				arch >> bAvoidRuntimeCompression;
+				TextureCompressionStrategy = bAvoidRuntimeCompression ? ETextureCompressionStrategy::DontCompressRuntime : ETextureCompressionStrategy::None;
+			}
+
+			if (ver == 3)
+			{
+				int32 OldNumExtraLODsToBuildAfterFirstLOD;
+				arch >> OldNumExtraLODsToBuildAfterFirstLOD;
+				NumExtraLODsToBuildAfterFirstLOD = OldNumExtraLODsToBuildAfterFirstLOD;
+			}
+			else if (ver >= 4)
+			{
+				arch >> NumExtraLODsToBuildAfterFirstLOD;
+			}
+			else
+			{
+				NumExtraLODsToBuildAfterFirstLOD = 0;
+			}
+
+		}
+    };
+
+
+	//! Information about an object state in the source data
+	struct FObjectState
+	{
+		//! Name used to identify the state from the code and user interface.
+		FString Name;
+
+		//! GPU Optimisation options
+		FStateOptimizationOptions Optimisation;
+
+		//! List of names of the runtime parameters in this state
+		TArray<FString> RuntimeParams;
+	};
+
+	
+    //! This utility class compiles two types of expressions: models and transforms.
+    //! Model expressions are compiled into run-time usable objects.
+    //! Transform expressions are compiled into Transform objects than can be applied to Model
+    //! expressions.
+    class MUTABLETOOLS_API Compiler : public RefCounted
+    {
+    public:
+
+        Compiler( Ptr<CompilerOptions> Options=nullptr );
+
+        //! Compile the expression into a run-time model.
+		TSharedPtr<Model> Compile( const Ptr<Node>& pNode );
+
+        //! Return the log of messages of all the compile operations executed so far.
+        ErrorLogPtrConst GetLog() const;
+
+    protected:
+
+        //! Forbidden. Manage with the Ptr<> template.
+        ~Compiler();
+
+    private:
+
+        class Private;
+        Private* m_pD;
+
+    };
+
+}
diff --git a/Engine/Plugins/Mutable/Source/MutableTools/Public/MuT/ErrorLog.h b/Engine/Plugins/Mutable/Source/MutableTools/Public/MuT/ErrorLog.h
new file mode 100644
--- /dev/null
+++ b/Engine/Plugins/Mutable/Source/MutableTools/Public/MuT/ErrorLog.h
@@ -0,0 +1,105 @@
+// Copyright Epic Games, Inc. All Rights Reserved.
+
+#pragma once
+
+#include "HAL/PlatformCrt.h"
+#include "MuR/Ptr.h"
+#include "MuR/RefCounted.h"
+
+
+namespace mu
+{
+
+	// Forward declarations
+	class ErrorLog;
+	typedef Ptr<ErrorLog> ErrorLogPtr;
+	typedef Ptr<const ErrorLog> ErrorLogPtrConst;
+
+
+	//! Types of message stored in the log
+	//! \ingroup tools
+	typedef enum
+	{
+        ELMT_NONE=0,
+		ELMT_ERROR,
+		ELMT_WARNING,
+		ELMT_INFO
+	} ErrorLogMessageType;
+
+	//! Categories of message stored in the log for the purpose of limiiting duplication of non-identical messages
+	//! \ingroup tools
+	typedef enum
+	{
+		ELMSB_ALL = 0,
+		ELMSB_UNKNOWN_TAG
+	} ErrorLogMessageSpamBin;
+
+	struct ErrorLogMessageAttachedDataView 
+	{
+		const float* m_unassignedUVs = nullptr;
+		size_t m_unassignedUVsSize = 0;
+	};
+
+	//! Class used to store the error, warning and information messages from several processes
+	//! performed by the tools library, like model transformation or compilation.
+	//! \ingroup tools
+	class MUTABLETOOLS_API ErrorLog : public RefCounted
+	{
+	public:
+
+		//-----------------------------------------------------------------------------------------
+		// Life cycle
+		//-----------------------------------------------------------------------------------------
+		ErrorLog();
+
+		//-----------------------------------------------------------------------------------------
+		// Own interface
+		//-----------------------------------------------------------------------------------------
+
+		//! Get the number of messages.
+		//! If a message type is provided, return the number of that type of message only.
+		int GetMessageCount( ErrorLogMessageType type = ELMT_NONE ) const;
+
+		//! Get the text of a message.
+		//! \param index index of the message from 0 to GetMessageCount(ELMT_NONE)-1
+		const FString& GetMessageText( int index ) const;
+
+		//! Get the opaque context of a message.
+		//! \param index index of the message from 0 to GetMessageCount(ELMT_NONE)-1
+		const void* GetMessageContext( int index ) const;
+
+		//!
+		ErrorLogMessageType GetMessageType( int index ) const;
+
+		//!
+		ErrorLogMessageSpamBin GetMessageSpamBin(int index) const;
+
+		//! Get the attached data of a message.
+		//! \param index index of message data from 0 to GetMessageCount(ELMT_NONE)-1
+		ErrorLogMessageAttachedDataView GetMessageAttachedData( int index ) const;
+
+		//!
+		void Log() const;
+
+		//!
+		void Merge( const ErrorLog* pOther );
+
+		//-----------------------------------------------------------------------------------------
+		// Interface pattern
+		//-----------------------------------------------------------------------------------------
+		class Private;
+		Private* GetPrivate() const;
+
+	protected:
+
+		//! Forbidden. Manage with the Ptr<> template.
+		~ErrorLog();
+
+	private:
+
+		Private* m_pD;
+
+	};
+	
+	MUTABLETOOLS_API extern const TCHAR* s_opNames[];
+}
diff --git a/Engine/Plugins/Mutable/Source/MutableTools/Public/MuT/NodeBool.h b/Engine/Plugins/Mutable/Source/MutableTools/Public/MuT/NodeBool.h
new file mode 100644
--- /dev/null
+++ b/Engine/Plugins/Mutable/Source/MutableTools/Public/MuT/NodeBool.h
@@ -0,0 +1,228 @@
+// Copyright Epic Games, Inc. All Rights Reserved.
+
+#pragma once
+
+#include "HAL/Platform.h"
+#include "MuR/Ptr.h"
+#include "MuR/RefCounted.h"
+#include "MuT/Node.h"
+
+
+namespace mu
+{
+	class NodeRange;
+
+	//---------------------------------------------------------------------------------------------
+    //! %Base class of any node that outputs a Bool value.
+	//! \ingroup model
+	//---------------------------------------------------------------------------------------------
+	class MUTABLETOOLS_API NodeBool : public Node
+	{
+	public:
+
+		//-----------------------------------------------------------------------------------------
+		// Node Interface
+		//-----------------------------------------------------------------------------------------
+
+        const FNodeType* GetType() const override;
+		static const FNodeType* GetStaticType();
+
+
+		//-----------------------------------------------------------------------------------------
+		// Interface pattern
+		//-----------------------------------------------------------------------------------------
+		class Private;
+
+	protected:
+
+		//! Forbidden. Manage with the Ptr<> template.
+		inline ~NodeBool() {}
+	};
+
+
+	//---------------------------------------------------------------------------------------------
+	//! Node returning a Bool constant value.
+	//! \ingroup model
+	//---------------------------------------------------------------------------------------------
+	class MUTABLETOOLS_API NodeBoolConstant : public NodeBool
+	{
+	public:
+
+		NodeBoolConstant();
+
+
+		//-----------------------------------------------------------------------------------------
+		// Node Interface
+		//-----------------------------------------------------------------------------------------
+
+        const FNodeType* GetType() const override;
+		static const FNodeType* GetStaticType();
+
+		//-----------------------------------------------------------------------------------------
+		// Own Interface
+		//-----------------------------------------------------------------------------------------
+
+		//! Get the value to be returned by the node.
+		bool GetValue() const;
+
+		//! Set the value to be returned by the node.
+		void SetValue( bool v );
+
+		//-----------------------------------------------------------------------------------------
+		// Interface pattern
+		//-----------------------------------------------------------------------------------------
+		class Private;
+		Private* GetPrivate() const;
+
+	protected:
+
+		//! Forbidden.
+		~NodeBoolConstant();
+
+	private:
+
+		Private* m_pD;
+
+	};
+
+
+	//---------------------------------------------------------------------------------------------
+	//! Node that defines a Bool model parameter.
+	//! \ingroup model
+	//---------------------------------------------------------------------------------------------
+	class MUTABLETOOLS_API NodeBoolParameter : public NodeBool
+	{
+	public:
+
+		NodeBoolParameter();
+
+		//-----------------------------------------------------------------------------------------
+		// Node Interface
+		//-----------------------------------------------------------------------------------------
+
+        const FNodeType* GetType() const override;
+		static const FNodeType* GetStaticType();
+
+		//-----------------------------------------------------------------------------------------
+		// Own Interface
+		//-----------------------------------------------------------------------------------------
+
+		//! Set the name of the parameter. It will be exposed in the final compiled data.
+		void SetName( const FString& );
+
+		//! Get the default value of the parameter.
+		void SetDefaultValue( bool v );
+
+        //! Set the number of ranges (dimensions) for this parameter.
+        //! By default a parameter has 0 ranges, meaning it only has one value.
+        void SetRangeCount( int i );
+        void SetRange( int i, Ptr<NodeRange> );
+
+		//-----------------------------------------------------------------------------------------
+		// Interface pattern
+		//-----------------------------------------------------------------------------------------
+		class Private;
+		Private* GetPrivate() const;
+
+	protected:
+
+		//! Forbidden. 
+		~NodeBoolParameter();
+
+	private:
+
+		Private* m_pD;
+
+	};
+
+
+	//---------------------------------------------------------------------------------------------
+	//! Node that returns the oposite of the input value.
+	//! \ingroup model
+	//---------------------------------------------------------------------------------------------
+	class MUTABLETOOLS_API NodeBoolNot : public NodeBool
+	{
+	public:
+
+		NodeBoolNot();
+
+		//-----------------------------------------------------------------------------------------
+		// Node Interface
+		//-----------------------------------------------------------------------------------------
+
+        const FNodeType* GetType() const override;
+		static const FNodeType* GetStaticType();
+
+		//-----------------------------------------------------------------------------------------
+		// Own Interface
+		//-----------------------------------------------------------------------------------------
+		
+		//! Input
+		Ptr<NodeBool> GetInput() const;
+		void SetInput( Ptr<NodeBool> );
+
+		//-----------------------------------------------------------------------------------------
+		// Interface pattern
+		//-----------------------------------------------------------------------------------------
+		class Private;
+		Private* GetPrivate() const;
+
+	protected:
+
+		//! Forbidden. 
+		~NodeBoolNot();
+
+	private:
+
+		Private* m_pD;
+
+	};
+
+
+	//---------------------------------------------------------------------------------------------
+	//! 
+	//! \ingroup model
+	//---------------------------------------------------------------------------------------------
+	class MUTABLETOOLS_API NodeBoolAnd : public NodeBool
+	{
+	public:
+
+		NodeBoolAnd();
+
+		//-----------------------------------------------------------------------------------------
+		// Node Interface
+		//-----------------------------------------------------------------------------------------
+
+        const FNodeType* GetType() const override;
+		static const FNodeType* GetStaticType();
+
+		//-----------------------------------------------------------------------------------------
+		// Own Interface
+		//-----------------------------------------------------------------------------------------
+		
+		//! Inputs
+		Ptr<NodeBool> GetA() const;
+		void SetA(Ptr<NodeBool>);
+
+		Ptr<NodeBool> GetB() const;
+		void SetB(Ptr<NodeBool>);
+
+		//-----------------------------------------------------------------------------------------
+		// Interface pattern
+		//-----------------------------------------------------------------------------------------
+		class Private;
+		Private* GetPrivate() const;
+
+	protected:
+
+		//! Forbidden. 
+		~NodeBoolAnd();
+
+	private:
+
+		Private* m_pD;
+
+	};
+
+
+}
diff --git a/Engine/Plugins/Mutable/Source/MutableTools/Public/MuT/NodeColour.h b/Engine/Plugins/Mutable/Source/MutableTools/Public/MuT/NodeColour.h
new file mode 100644
--- /dev/null
+++ b/Engine/Plugins/Mutable/Source/MutableTools/Public/MuT/NodeColour.h
@@ -0,0 +1,36 @@
+// Copyright Epic Games, Inc. All Rights Reserved.
+
+#pragma once
+
+#include "HAL/Platform.h"
+#include "MuR/Ptr.h"
+#include "MuR/RefCounted.h"
+#include "MuT/Node.h"
+
+
+namespace mu
+{
+
+    /** Base class of any node that outputs a colour.
+	*/
+    class MUTABLETOOLS_API NodeColour : public Node
+    {
+    public:
+
+        // Node Interface
+		virtual const FNodeType* GetType() const override { return &StaticType; }
+		static const FNodeType* GetStaticType() { return &StaticType; }
+
+    protected:
+
+        //! Forbidden. Manage with the Ptr<> template.
+        inline ~NodeColour() {}
+
+	private:
+
+		static FNodeType StaticType;
+
+    };
+
+}
+
diff --git a/Engine/Plugins/Mutable/Source/MutableTools/Public/MuT/NodeColourArithmeticOperation.h b/Engine/Plugins/Mutable/Source/MutableTools/Public/MuT/NodeColourArithmeticOperation.h
new file mode 100644
--- /dev/null
+++ b/Engine/Plugins/Mutable/Source/MutableTools/Public/MuT/NodeColourArithmeticOperation.h
@@ -0,0 +1,50 @@
+// Copyright Epic Games, Inc. All Rights Reserved.
+
+#pragma once
+
+#include "MuR/Ptr.h"
+#include "MuR/RefCounted.h"
+#include "MuT/NodeColour.h"
+
+
+namespace mu
+{
+
+    /** Perform a per - component arithmetic operation between two colours. 
+	*/
+	class MUTABLETOOLS_API NodeColourArithmeticOperation : public NodeColour
+	{
+	public:
+
+		/** Possible arithmetic operations. */
+		enum class EOperation
+		{
+			Add,
+			Subtract,
+			Multiply,
+			Divide
+		};
+
+		EOperation Operation;
+		Ptr<NodeColour> A;
+		Ptr<NodeColour> B;
+
+	public:
+
+		/** Node type hierarchy data. */
+		virtual const FNodeType* GetType() const { return &StaticType; }
+		static const FNodeType* GetStaticType() { return &StaticType; }
+
+	protected:
+
+		//! Forbidden. Manage with the Ptr<> template.
+		~NodeColourArithmeticOperation() {}
+
+	private:
+
+		static FNodeType StaticType;
+
+	};
+
+	
+}
diff --git a/Engine/Plugins/Mutable/Source/MutableTools/Public/MuT/NodeColourVariation.h b/Engine/Plugins/Mutable/Source/MutableTools/Public/MuT/NodeColourVariation.h
new file mode 100644
--- /dev/null
+++ b/Engine/Plugins/Mutable/Source/MutableTools/Public/MuT/NodeColourVariation.h
@@ -0,0 +1,47 @@
+// Copyright Epic Games, Inc. All Rights Reserved.
+
+#pragma once
+
+#include "MuR/Ptr.h"
+#include "MuR/RefCounted.h"
+#include "MuT/Node.h"
+#include "MuT/NodeColour.h"
+
+
+namespace mu
+{
+
+	/** Select different color subgraphs based on active tags. */
+    class MUTABLETOOLS_API NodeColourVariation : public NodeColour
+    {
+	public:
+
+		Ptr<NodeColour> DefaultColour;
+
+		struct FVariation
+		{
+			Ptr<NodeColour> Colour;
+			FString Tag;
+		};
+
+		TArray<FVariation> Variations;
+
+    public:
+
+		// Node Interface
+		virtual const FNodeType* GetType() const override { return &StaticType; }
+		static const FNodeType* GetStaticType() { return &StaticType; }
+
+    protected:
+
+        //! Forbidden. Manage with the Ptr<> template.
+		~NodeColourVariation() {}
+
+	private:
+
+		static FNodeType StaticType;
+	
+	};
+
+
+} // namespace mu
diff --git a/Engine/Plugins/Mutable/Source/MutableTools/Public/MuT/NodeComponentEdit.h b/Engine/Plugins/Mutable/Source/MutableTools/Public/MuT/NodeComponentEdit.h
new file mode 100644
--- /dev/null
+++ b/Engine/Plugins/Mutable/Source/MutableTools/Public/MuT/NodeComponentEdit.h
@@ -0,0 +1,56 @@
+// Copyright Epic Games, Inc. All Rights Reserved.
+
+#pragma once
+
+#include "MuR/Ptr.h"
+#include "MuR/RefCounted.h"
+#include "MuT/Node.h"
+#include "MuT/NodeComponent.h"
+
+
+namespace mu
+{
+
+	//! This node modifies a node of the parent object of the object that this node belongs to.
+	//! It allows to extend, cut and morph the parent component's meshes.
+	//! It also allows to patch the parent component's textures.
+	class MUTABLETOOLS_API NodeComponentEdit : public NodeComponent
+	{
+	public:
+
+		//-----------------------------------------------------------------------------------------
+        // Node interface
+		//-----------------------------------------------------------------------------------------
+
+		virtual const FNodeType* GetType() const override { return GetStaticType(); }
+		static const FNodeType* GetStaticType() { return &StaticType; }
+
+        //-----------------------------------------------------------------------------------------
+        // NodeComponent interface
+        //-----------------------------------------------------------------------------------------
+		virtual const class NodeComponentNew* GetParentComponentNew() const override;
+
+		//-----------------------------------------------------------------------------------------
+        // Own interface
+		//-----------------------------------------------------------------------------------------
+
+		//! Set the parent component to modify. It should be a node from the parent objects of this node's object.
+        void SetParent( NodeComponent* );
+        NodeComponent* GetParent() const;
+
+	protected:
+
+		//! Forbidden. Manage with the Ptr<> template.
+		~NodeComponentEdit() {}
+
+	private:
+
+		static FNodeType StaticType;
+
+		Ptr<NodeComponent> m_pParent;
+
+	};
+
+
+
+}
diff --git a/Engine/Plugins/Mutable/Source/MutableTools/Public/MuT/NodeExtensionDataConstant.h b/Engine/Plugins/Mutable/Source/MutableTools/Public/MuT/NodeExtensionDataConstant.h
new file mode 100644
--- /dev/null
+++ b/Engine/Plugins/Mutable/Source/MutableTools/Public/MuT/NodeExtensionDataConstant.h
@@ -0,0 +1,49 @@
+// Copyright Epic Games, Inc. All Rights Reserved.
+
+#pragma once
+
+#include "MuR/ExtensionData.h"
+#include "MuR/Ptr.h"
+#include "MuR/RefCounted.h"
+#include "MuT/NodeExtensionData.h"
+
+
+namespace mu
+{
+
+	//! Node that outputs a constant ExtensionData
+	//! \ingroup model
+	class MUTABLETOOLS_API NodeExtensionDataConstant : public NodeExtensionData
+	{
+	public:
+
+		Ptr<const ExtensionData> Value;
+
+	public:
+
+		//-----------------------------------------------------------------------------------------
+		// Node Interface
+		//-----------------------------------------------------------------------------------------
+		virtual const FNodeType* GetType() const override { return GetStaticType(); }
+		static const FNodeType* GetStaticType() { return &StaticType; }
+
+		//-----------------------------------------------------------------------------------------
+		// Own Interface
+		//-----------------------------------------------------------------------------------------
+
+		/** Deprecated. Access Value attribute directly. */
+		void SetValue(const Ptr<const ExtensionData>& In) { Value = In; }
+
+	protected:
+
+		//! Forbidden. Manage with the Ptr<> template.
+		~NodeExtensionDataConstant() {}
+
+	private:
+
+		static FNodeType StaticType;
+
+	};
+
+
+}
diff --git a/Engine/Plugins/Mutable/Source/MutableTools/Public/MuT/NodeExtensionDataSwitch.h b/Engine/Plugins/Mutable/Source/MutableTools/Public/MuT/NodeExtensionDataSwitch.h
new file mode 100644
--- /dev/null
+++ b/Engine/Plugins/Mutable/Source/MutableTools/Public/MuT/NodeExtensionDataSwitch.h
@@ -0,0 +1,51 @@
+// Copyright Epic Games, Inc. All Rights Reserved.
+
+#pragma once
+
+#include "MuR/ExtensionData.h"
+#include "MuR/Ptr.h"
+#include "MuR/RefCounted.h"
+#include "MuT/NodeExtensionData.h"
+#include "MuT/NodeScalar.h"
+
+namespace mu
+{
+
+	class MUTABLETOOLS_API NodeExtensionDataSwitch : public NodeExtensionData
+	{
+	public:
+
+		Ptr<NodeScalar> Parameter;
+		TArray<Ptr<NodeExtensionData>> Options;
+
+	public:
+
+		//-----------------------------------------------------------------------------------------
+		// Node Interface
+		//-----------------------------------------------------------------------------------------
+		virtual const FNodeType* GetType() const override { return GetStaticType(); }
+		static const FNodeType* GetStaticType() { return &StaticType; }
+
+		//-----------------------------------------------------------------------------------------
+		// Own Interface
+		//-----------------------------------------------------------------------------------------
+		Ptr<NodeScalar> GetParameter() const;
+		void SetParameter(Ptr<NodeScalar> InParameter);
+
+		void SetOptionCount(int);
+
+		Ptr<NodeExtensionData> GetOption(int32 t) const;
+		void SetOption(int32 t, Ptr<NodeExtensionData>);
+
+
+	protected:
+
+		//! Forbidden. Manage with the Ptr<> template.
+		~NodeExtensionDataSwitch() {}
+
+	private:
+
+		static FNodeType StaticType;
+
+	};
+}
\ No newline at end of file
diff --git a/Engine/Plugins/Mutable/Source/MutableTools/Public/MuT/NodeExtensionDataVariation.h b/Engine/Plugins/Mutable/Source/MutableTools/Public/MuT/NodeExtensionDataVariation.h
new file mode 100644
--- /dev/null
+++ b/Engine/Plugins/Mutable/Source/MutableTools/Public/MuT/NodeExtensionDataVariation.h
@@ -0,0 +1,58 @@
+// Copyright Epic Games, Inc. All Rights Reserved.
+
+#pragma once
+
+#include "MuR/ExtensionData.h"
+#include "MuR/Ptr.h"
+#include "MuR/RefCounted.h"
+#include "MuT/NodeExtensionData.h"
+
+namespace mu
+{
+
+	class MUTABLETOOLS_API NodeExtensionDataVariation : public NodeExtensionData
+	{
+	public:
+
+		Ptr<NodeExtensionData> DefaultValue;
+
+		struct FVariation
+		{
+			Ptr<NodeExtensionData> Value;
+			FString Tag;
+		};
+
+		TArray<FVariation> Variations;
+
+	public:
+
+		//-----------------------------------------------------------------------------------------
+		// Node Interface
+		//-----------------------------------------------------------------------------------------
+		virtual const FNodeType* GetType() const override { return GetStaticType(); }
+		static const FNodeType* GetStaticType() { return &StaticType; }
+
+		//-----------------------------------------------------------------------------------------
+		// Own Interface
+		//-----------------------------------------------------------------------------------------
+		void SetDefaultValue(Ptr<NodeExtensionData> InValue);
+
+		void SetVariationCount(int32 InCount);
+		int GetVariationCount() const;
+
+		void SetVariationTag(int32 InIndex, const FString& InTag);
+
+		void SetVariationValue(int32 InIndex, Ptr<NodeExtensionData> InValue);
+
+
+	protected:
+
+		//! Forbidden. Manage with the Ptr<> template.
+		~NodeExtensionDataVariation() {}
+
+	private:
+
+		static FNodeType StaticType;
+
+	};
+}
\ No newline at end of file
diff --git a/Engine/Plugins/Mutable/Source/MutableTools/Public/MuT/NodeImage.h b/Engine/Plugins/Mutable/Source/MutableTools/Public/MuT/NodeImage.h
new file mode 100644
--- /dev/null
+++ b/Engine/Plugins/Mutable/Source/MutableTools/Public/MuT/NodeImage.h
@@ -0,0 +1,112 @@
+// Copyright Epic Games, Inc. All Rights Reserved.
+
+#pragma once
+
+#include "HAL/Platform.h"
+#include "MuR/Ptr.h"
+#include "MuR/RefCounted.h"
+#include "MuT/Node.h"
+
+namespace mu
+{
+
+	// Forward definitions
+	class NodeImage;
+	typedef Ptr<NodeImage> NodeImagePtr;
+	typedef Ptr<const NodeImage> NodeImagePtrConst;
+
+	/** Data related to the a source image that is necessary to classify the final image fragments
+	* that are derived from this source.
+	*/
+	struct FSourceDataDescriptor
+	{
+		enum ESpecialValues
+		{
+			EInvalid = -2,
+			ENeutral = -1,
+		};
+
+		/** Number of mips in the source texture that are considered high resolution.
+		* If -1 it means the entire descriptor is neutral.
+		* If -2 it means the descriptor is invalid (results of an operation that shouldn't happen).
+		*/
+		int32 SourceHighResMips = ENeutral;
+
+		/** Source tags that mark this data and prevent it from mixing with other data at compile time. */
+		TArray<FString> Tags;
+
+		inline bool operator==(const FSourceDataDescriptor& Other) const
+		{
+			return SourceHighResMips == Other.SourceHighResMips && Tags == Other.Tags;
+		}
+
+		inline bool IsInvalid() const
+		{
+			return SourceHighResMips == EInvalid;
+		}
+
+		inline bool IsNeutral() const
+		{
+			return SourceHighResMips == ENeutral;
+		}
+
+		void CombineWith(const FSourceDataDescriptor& Other)
+		{
+			if (IsInvalid() || Other.IsInvalid())
+			{
+				SourceHighResMips = EInvalid;
+				Tags.Empty();
+				return;
+			}
+
+			if (Other.IsNeutral())
+			{
+				return;
+			}
+
+			if (IsNeutral())
+			{
+				SourceHighResMips = Other.SourceHighResMips;
+				Tags = Other.Tags;
+				return;
+			}
+
+			if (!(*this==Other))
+			{
+				SourceHighResMips = EInvalid;
+				Tags.Empty();
+				return;
+			}
+
+			return;
+		}
+	};
+
+
+    /** Base class of any node that outputs an image. */
+	class MUTABLETOOLS_API NodeImage : public Node
+	{
+	public:
+
+		//-----------------------------------------------------------------------------------------
+		// Node Interface
+		//-----------------------------------------------------------------------------------------
+
+		const FNodeType* GetType() const override;
+		static const FNodeType* GetStaticType();
+
+		//-----------------------------------------------------------------------------------------
+		// Interface pattern
+		//-----------------------------------------------------------------------------------------
+		class Private;
+
+	protected:
+
+		//! Forbidden. Manage with the Ptr<> template.
+		inline ~NodeImage() {}
+
+	};
+
+
+
+}
diff --git a/Engine/Plugins/Mutable/Source/MutableTools/Public/MuT/NodeImageBinarise.h b/Engine/Plugins/Mutable/Source/MutableTools/Public/MuT/NodeImageBinarise.h
new file mode 100644
--- /dev/null
+++ b/Engine/Plugins/Mutable/Source/MutableTools/Public/MuT/NodeImageBinarise.h
@@ -0,0 +1,69 @@
+// Copyright Epic Games, Inc. All Rights Reserved.
+
+#pragma once
+
+#include "MuR/Ptr.h"
+#include "MuR/RefCounted.h"
+#include "MuT/Node.h"
+#include "MuT/NodeImage.h"
+
+
+namespace mu
+{
+
+	// Forward definitions
+	class NodeScalar;
+	typedef Ptr<NodeScalar> NodeScalarPtr;
+	typedef Ptr<const NodeScalar> NodeScalarPtrConst;
+
+	class NodeImageBinarise;
+	typedef Ptr<NodeImageBinarise> NodeImageBinarisePtr;
+	typedef Ptr<const NodeImageBinarise> NodeImageBinarisePtrConst;
+
+
+	//! Node that multiplies the colors of an image, channel by channel.
+	//! \ingroup model
+	class MUTABLETOOLS_API NodeImageBinarise : public NodeImage
+	{
+	public:
+
+		NodeImageBinarise();
+
+		//-----------------------------------------------------------------------------------------
+		// Node Interface
+		//-----------------------------------------------------------------------------------------
+
+        const FNodeType* GetType() const override;
+		static const FNodeType* GetStaticType();
+
+		//-----------------------------------------------------------------------------------------
+		// Own Interface
+		//-----------------------------------------------------------------------------------------
+
+		//! Base image to multiply.
+		NodeImagePtr GetBase() const;
+		void SetBase( NodeImagePtr );
+
+		//! Colour to multiply by the image.
+		NodeScalarPtr GetThreshold() const;
+		void SetThreshold( NodeScalarPtr );
+
+		//-----------------------------------------------------------------------------------------
+		// Interface pattern
+		//-----------------------------------------------------------------------------------------
+		class Private;
+		Private* GetPrivate() const;
+
+	protected:
+
+		//! Forbidden. Manage with the Ptr<> template.
+		~NodeImageBinarise();
+
+	private:
+
+		Private* m_pD;
+
+	};
+
+
+}
