[mutable] Moved the Mutable plugin out of Experimental status into Beta.

diff --git a/Engine/Plugins/Mutable/Source/MutableTools/Public/MuT/NodeImageColourMap.h b/Engine/Plugins/Mutable/Source/MutableTools/Public/MuT/NodeImageColourMap.h
new file mode 100644
--- /dev/null
+++ b/Engine/Plugins/Mutable/Source/MutableTools/Public/MuT/NodeImageColourMap.h
@@ -0,0 +1,72 @@
+// Copyright Epic Games, Inc. All Rights Reserved.
+
+#pragma once
+
+#include "MuR/Ptr.h"
+#include "MuR/RefCounted.h"
+#include "MuT/Node.h"
+#include "MuT/NodeImage.h"
+
+
+namespace mu
+{
+
+	// Forward definitions
+	class NodeImageColourMap;
+	typedef Ptr<NodeImageColourMap> NodeImageColourMapPtr;
+	typedef Ptr<const NodeImageColourMap> NodeImageColourMapPtrConst;
+
+
+	//! This node changes the colours of a selectd part of the image, applying a colour map from
+	//! conteined in another image.
+	//! \ingroup model
+	class MUTABLETOOLS_API NodeImageColourMap : public NodeImage
+	{
+	public:
+
+		NodeImageColourMap();
+
+
+		//-----------------------------------------------------------------------------------------
+		// Node Interface
+		//-----------------------------------------------------------------------------------------
+
+        const FNodeType* GetType() const override;
+		static const FNodeType* GetStaticType();
+
+		//-----------------------------------------------------------------------------------------
+		// Own Interface
+		//-----------------------------------------------------------------------------------------
+
+		//! Get the node generating the base image that will have the other one blended on top.
+		NodeImagePtr GetBase() const;
+		void SetBase( NodeImagePtr );
+
+		//! Get the node generating the mask image controlling the weight of the blend effect.
+		//! \todo: make it optional
+		NodeImagePtr GetMask() const;
+		void SetMask( NodeImagePtr );
+
+		//! Get the node generating the image to blend on the base.
+		NodeImagePtr GetMap() const;
+		void SetMap( NodeImagePtr );
+
+		//-----------------------------------------------------------------------------------------
+		// Interface pattern
+		//-----------------------------------------------------------------------------------------
+		class Private;
+		Private* GetPrivate() const;
+
+	protected:
+
+		//! Forbidden. Manage with the Ptr<> template.
+		~NodeImageColourMap();
+
+	private:
+
+		Private* m_pD;
+
+	};
+
+
+}
diff --git a/Engine/Plugins/Mutable/Source/MutableTools/Public/MuT/NodeImageConditional.h b/Engine/Plugins/Mutable/Source/MutableTools/Public/MuT/NodeImageConditional.h
new file mode 100644
--- /dev/null
+++ b/Engine/Plugins/Mutable/Source/MutableTools/Public/MuT/NodeImageConditional.h
@@ -0,0 +1,70 @@
+// Copyright Epic Games, Inc. All Rights Reserved.
+
+#pragma once
+
+#include "MuR/Ptr.h"
+#include "MuR/RefCounted.h"
+#include "MuT/Node.h"
+#include "MuT/NodeImage.h"
+
+
+namespace mu
+{
+
+	// Forward definitions
+    class NodeBool;
+    class NodeImageConditional;
+    typedef Ptr<NodeImageConditional> NodeImageConditionalPtr;
+    typedef Ptr<const NodeImageConditional> NodeImageConditionalPtrConst;
+
+
+	//! This node selects an output image from a set of input images based on a parameter.
+	//! \ingroup model
+    class MUTABLETOOLS_API NodeImageConditional : public NodeImage
+	{
+	public:
+
+        NodeImageConditional();
+
+		//-----------------------------------------------------------------------------------------
+		// Node Interface
+		//-----------------------------------------------------------------------------------------
+
+        const FNodeType* GetType() const override;
+		static const FNodeType* GetStaticType();
+
+		//-----------------------------------------------------------------------------------------
+		// Own Interface
+		//-----------------------------------------------------------------------------------------
+
+		//! Get the node generating the parameter used to select the option.
+        NodeBool* GetParameter() const;
+        void SetParameter( NodeBool* );
+
+        //! Get the node generating the option image when the parameters is true.
+        NodeImage* GetOptionTrue() const;
+        void SetOptionTrue( NodeImage* );
+
+        NodeImage* GetOptionFalse() const;
+        void SetOptionFalse( NodeImage* );
+
+
+		//-----------------------------------------------------------------------------------------
+		// Interface pattern
+		//-----------------------------------------------------------------------------------------
+		class Private;
+		Private* GetPrivate() const;
+
+	protected:
+
+		//! Forbidden. Manage with the Ptr<> template.
+        ~NodeImageConditional();
+
+	private:
+
+		Private* m_pD;
+
+	};
+
+
+}
diff --git a/Engine/Plugins/Mutable/Source/MutableTools/Public/MuT/NodeImageConstant.h b/Engine/Plugins/Mutable/Source/MutableTools/Public/MuT/NodeImageConstant.h
new file mode 100644
--- /dev/null
+++ b/Engine/Plugins/Mutable/Source/MutableTools/Public/MuT/NodeImageConstant.h
@@ -0,0 +1,58 @@
+// Copyright Epic Games, Inc. All Rights Reserved.
+
+#pragma once
+
+#include "MuR/Image.h"
+#include "MuR/Ptr.h"
+#include "MuR/RefCounted.h"
+#include "MuT/Node.h"
+#include "MuT/NodeImage.h"
+
+
+namespace mu
+{
+
+	/** Node that outputs a constant image.
+	* This node also supports "image references".
+	*/
+	class MUTABLETOOLS_API NodeImageConstant : public NodeImage
+	{
+	public:
+
+		/** */
+		Ptr<ResourceProxy<Image>> Proxy;
+		
+		/** */
+		FSourceDataDescriptor SourceDataDescriptor;
+
+	public:
+
+		//-----------------------------------------------------------------------------------------
+		// Node Interface
+		//-----------------------------------------------------------------------------------------
+
+		virtual const FNodeType* GetType() const override { return GetStaticType(); }
+		static const FNodeType* GetStaticType() { return &StaticType; }
+
+		//-----------------------------------------------------------------------------------------
+		// Own Interface
+		//-----------------------------------------------------------------------------------------
+
+        //! Set the image to be output by this node
+        void SetValue( const Image* );
+
+        //! Set the image proxy that will provide the image for this node when necessary
+        void SetValue( Ptr<ResourceProxy<Image>> pImage );
+
+	protected:
+
+		//! Forbidden. Manage with the Ptr<> template.
+		~NodeImageConstant() {}
+
+	private:
+
+		static FNodeType StaticType;
+
+	};
+
+}
diff --git a/Engine/Plugins/Mutable/Source/MutableTools/Public/MuT/NodeImageFormat.h b/Engine/Plugins/Mutable/Source/MutableTools/Public/MuT/NodeImageFormat.h
new file mode 100644
--- /dev/null
+++ b/Engine/Plugins/Mutable/Source/MutableTools/Public/MuT/NodeImageFormat.h
@@ -0,0 +1,70 @@
+// Copyright Epic Games, Inc. All Rights Reserved.
+
+#pragma once
+
+#include "MuR/Image.h"
+#include "MuR/Ptr.h"
+#include "MuR/RefCounted.h"
+#include "MuT/Node.h"
+#include "MuT/NodeImage.h"
+
+
+namespace mu
+{
+
+	// Forward definitions
+	class NodeImageFormat;
+	typedef Ptr<NodeImageFormat> NodeImageFormatPtr;
+	typedef Ptr<const NodeImageFormat> NodeImageFormatPtrConst;
+
+
+	//! Node that composes a new image by gathering pixel data from channels in other images.
+	//! \ingroup model
+	class MUTABLETOOLS_API NodeImageFormat : public NodeImage
+	{
+	public:
+
+		NodeImageFormat();
+
+
+		//-----------------------------------------------------------------------------------------
+		// Node Interface
+		//-----------------------------------------------------------------------------------------
+
+		const FNodeType* GetType() const override;
+		static const FNodeType* GetStaticType();
+
+		//-----------------------------------------------------------------------------------------
+		// Own Interface
+		//-----------------------------------------------------------------------------------------
+
+		//!
+		EImageFormat GetFormat() const;
+        void SetFormat(EImageFormat format, EImageFormat formatIfAlpha = EImageFormat::IF_NONE );
+
+		//!
+		NodeImagePtr GetSource() const;
+
+		void SetSource( NodeImagePtr );
+
+
+		//-----------------------------------------------------------------------------------------
+		// Interface pattern
+		//-----------------------------------------------------------------------------------------
+		class Private;
+		Private* GetPrivate() const;
+
+	protected:
+
+		//! Forbidden. Manage with the Ptr<> template.
+		~NodeImageFormat();
+
+	private:
+
+		Private* m_pD;
+
+	};
+
+
+}
+
diff --git a/Engine/Plugins/Mutable/Source/MutableTools/Public/MuT/NodeImageGradient.h b/Engine/Plugins/Mutable/Source/MutableTools/Public/MuT/NodeImageGradient.h
new file mode 100644
--- /dev/null
+++ b/Engine/Plugins/Mutable/Source/MutableTools/Public/MuT/NodeImageGradient.h
@@ -0,0 +1,75 @@
+// Copyright Epic Games, Inc. All Rights Reserved.
+
+#pragma once
+
+#include "MuR/Ptr.h"
+#include "MuR/RefCounted.h"
+#include "MuT/Node.h"
+#include "MuT/NodeImage.h"
+
+
+namespace mu
+{
+
+	// Forward definitions
+	class NodeImageGradient;
+	typedef Ptr<NodeImageGradient> NodeImageGradientPtr;
+	typedef Ptr<const NodeImageGradient> NodeImageGradientPtrConst;
+
+	class NodeColour;
+	typedef Ptr<NodeColour> NodeColourPtr;
+	typedef Ptr<const NodeColour> NodeColourPtrConst;
+
+
+	//! This node generats an horizontal linear gradient between two colours in a new image.
+	//! \ingroup model
+	class MUTABLETOOLS_API NodeImageGradient : public NodeImage
+	{
+	public:
+
+		NodeImageGradient();
+
+
+		//-----------------------------------------------------------------------------------------
+		// Node Interface
+		//-----------------------------------------------------------------------------------------
+
+        const FNodeType* GetType() const override;
+		static const FNodeType* GetStaticType();
+
+		//-----------------------------------------------------------------------------------------
+		// Own Interface
+		//-----------------------------------------------------------------------------------------
+
+		//! First colour of the gradient.
+		NodeColourPtr GetColour0() const;
+		void SetColour0( NodeColourPtr );
+
+		//! Second colour of the gradient.
+		NodeColourPtr GetColour1() const;
+		void SetColour1( NodeColourPtr );
+
+		//! Generated image size.
+		int GetSizeX() const;
+		int GetSizeY() const;
+		void SetSize( int x, int y );
+
+		//-----------------------------------------------------------------------------------------
+		// Interface pattern
+		//-----------------------------------------------------------------------------------------
+		class Private;
+		Private* GetPrivate() const;
+
+	protected:
+
+		//! Forbidden. Manage with the Ptr<> template.
+		~NodeImageGradient();
+
+	private:
+
+		Private* m_pD;
+
+	};
+
+
+}
diff --git a/Engine/Plugins/Mutable/Source/MutableTools/Public/MuT/NodeImageInterpolate.h b/Engine/Plugins/Mutable/Source/MutableTools/Public/MuT/NodeImageInterpolate.h
new file mode 100644
--- /dev/null
+++ b/Engine/Plugins/Mutable/Source/MutableTools/Public/MuT/NodeImageInterpolate.h
@@ -0,0 +1,81 @@
+// Copyright Epic Games, Inc. All Rights Reserved.
+
+#pragma once
+
+#include "MuR/Ptr.h"
+#include "MuR/RefCounted.h"
+#include "MuT/Node.h"
+#include "MuT/NodeImage.h"
+
+
+namespace mu
+{
+
+	// Forward definitions
+	class NodeScalar;
+	typedef Ptr<NodeScalar> NodeScalarPtr;
+	typedef Ptr<const NodeScalar> NodeScalarPtrConst;
+
+	class NodeImageInterpolate;
+	typedef Ptr<NodeImageInterpolate> NodeImageInterpolatePtr;
+	typedef Ptr<const NodeImageInterpolate> NodeImageInterpolatePtrConst;
+
+
+	//! Node that interpolates linearly between several images based on a weight.
+	//! The different input images are uniformly distributed in the 0 to 1 range:
+	//! - if two images are set, the first one is 0.0 and the second one is 1.0
+	//! - if three images are set, the first one is 0.0 the second one is 0.5 and the third one 1.0
+	//! regardless of the input slots used: If B and C are set, but not A, B will be at weight 0
+	//! and C will be at weight 1
+	//! \ingroup model
+	class MUTABLETOOLS_API NodeImageInterpolate : public NodeImage
+	{
+	public:
+
+		NodeImageInterpolate();
+
+
+		//-----------------------------------------------------------------------------------------
+		// Node Interface
+		//-----------------------------------------------------------------------------------------
+
+        const FNodeType* GetType() const override;
+		static const FNodeType* GetStaticType();
+
+		//-----------------------------------------------------------------------------------------
+		// Own Interface
+		//-----------------------------------------------------------------------------------------
+
+		//! Get the node generating the weight of the interpolation.
+		NodeScalarPtr GetFactor() const;
+		void SetFactor( NodeScalarPtr );
+
+		//! Set the number of target images. It will keep the currently set targets and initialise
+		//! the new ones as null.
+		void SetTargetCount( int );
+		int GetTargetCount() const;
+
+		//! Get the node generating the t-th target image of the interpolation.
+		NodeImagePtr GetTarget( int t ) const;
+		void SetTarget( int t, NodeImagePtr );
+
+
+		//-----------------------------------------------------------------------------------------
+		// Interface pattern
+		//-----------------------------------------------------------------------------------------
+		class Private;
+		Private* GetPrivate() const;
+
+	protected:
+
+		//! Forbidden. Manage with the Ptr<> template.
+		~NodeImageInterpolate();
+
+	private:
+
+		Private* m_pD;
+
+	};
+
+
+}
diff --git a/Engine/Plugins/Mutable/Source/MutableTools/Public/MuT/NodeImageInvert.h b/Engine/Plugins/Mutable/Source/MutableTools/Public/MuT/NodeImageInvert.h
new file mode 100644
--- /dev/null
+++ b/Engine/Plugins/Mutable/Source/MutableTools/Public/MuT/NodeImageInvert.h
@@ -0,0 +1,57 @@
+// Copyright Epic Games, Inc. All Rights Reserved.
+
+#pragma once
+
+#include "MuR/Ptr.h"
+#include "MuR/RefCounted.h"
+#include "MuT/Node.h"
+#include "MuT/NodeImage.h"
+
+
+namespace mu
+{
+
+	// Forward definitions
+	class NodeImageInvert;
+	typedef Ptr<NodeImageInvert> NodeImageInvertPtr;
+	typedef Ptr<const NodeImageInvert> NodeImageInvertPtrConst;
+
+	//! Node that inverts the colors of an image, channel by channel
+	class MUTABLETOOLS_API NodeImageInvert : public NodeImage
+	{
+	public:
+
+		NodeImageInvert();
+
+
+		//-----------------------------------------------------------------------------------------
+		// Node Interface
+		//-----------------------------------------------------------------------------------------
+
+		const FNodeType* GetType() const override;
+		static const FNodeType* GetStaticType();
+
+		//-----------------------------------------------------------------------------------------
+		// Own Interface
+		//-----------------------------------------------------------------------------------------
+
+		//! Base image to invert.
+		NodeImagePtr GetBase() const;
+		void SetBase(NodeImagePtr);
+
+		//-----------------------------------------------------------------------------------------
+		// Interface pattern
+		//-----------------------------------------------------------------------------------------
+		class Private;
+		Private* GetPrivate() const;
+
+	protected:
+
+		//! Forbidden. Manage with the Ptr<> template
+		~NodeImageInvert();
+
+	private:
+
+		Private* m_pD;
+	};
+}
diff --git a/Engine/Plugins/Mutable/Source/MutableTools/Public/MuT/NodeImageLayer.h b/Engine/Plugins/Mutable/Source/MutableTools/Public/MuT/NodeImageLayer.h
new file mode 100644
--- /dev/null
+++ b/Engine/Plugins/Mutable/Source/MutableTools/Public/MuT/NodeImageLayer.h
@@ -0,0 +1,74 @@
+// Copyright Epic Games, Inc. All Rights Reserved.
+
+#pragma once
+
+#include "HAL/Platform.h"
+#include "MuR/Image.h"
+#include "MuR/Ptr.h"
+#include "MuR/RefCounted.h"
+#include "MuT/Node.h"
+#include "MuT/NodeImage.h"
+
+namespace mu
+{
+
+	// Forward definitions
+	class NodeImageLayer;
+	typedef Ptr<NodeImageLayer> NodeImageLayerPtr;
+	typedef Ptr<const NodeImageLayer> NodeImageLayerPtrConst;
+
+
+	//! This node applies a layer blending effect on a base image using a mask and a blended image.
+	//! \ingroup model
+	class MUTABLETOOLS_API NodeImageLayer : public NodeImage
+	{
+	public:
+
+		NodeImageLayer();
+
+		//-----------------------------------------------------------------------------------------
+		// Node Interface
+		//-----------------------------------------------------------------------------------------
+
+        const FNodeType* GetType() const override;
+		static const FNodeType* GetStaticType();
+
+		//-----------------------------------------------------------------------------------------
+		// Own Interface
+		//-----------------------------------------------------------------------------------------
+
+		//! Get the node generating the base image that will have the blendeing effect applied.
+		NodeImagePtr GetBase() const;
+		void SetBase( NodeImagePtr );
+
+		//! Get the node generating the mask image controlling the weight of the effect.
+		NodeImagePtr GetMask() const;
+		void SetMask( NodeImagePtr );
+
+		//! Get the image blended on top of the base.
+		NodeImagePtr GetBlended() const;
+		void SetBlended( NodeImagePtr );
+
+		//!
+		EBlendType GetBlendType() const;
+		void SetBlendType(EBlendType);
+
+		//-----------------------------------------------------------------------------------------
+		// Interface pattern
+		//-----------------------------------------------------------------------------------------
+		class Private;
+		Private* GetPrivate() const;
+
+	protected:
+
+		//! Forbidden. Manage with the Ptr<> template.
+		~NodeImageLayer();
+
+	private:
+
+		Private* m_pD;
+
+	};
+
+
+}
diff --git a/Engine/Plugins/Mutable/Source/MutableTools/Public/MuT/NodeImageLayerColour.h b/Engine/Plugins/Mutable/Source/MutableTools/Public/MuT/NodeImageLayerColour.h
new file mode 100644
--- /dev/null
+++ b/Engine/Plugins/Mutable/Source/MutableTools/Public/MuT/NodeImageLayerColour.h
@@ -0,0 +1,79 @@
+// Copyright Epic Games, Inc. All Rights Reserved.
+
+#pragma once
+
+#include "MuR/Image.h"
+#include "MuR/Ptr.h"
+#include "MuR/RefCounted.h"
+#include "MuT/Node.h"
+#include "MuT/NodeImage.h"
+
+namespace mu
+{
+
+	// Forward definitions
+	class NodeColour;
+	typedef Ptr<NodeColour> NodeColourPtr;
+	typedef Ptr<const NodeColour> NodeColourPtrConst;
+
+	class NodeImageLayerColour;
+	typedef Ptr<NodeImageLayerColour> NodeImageLayerColourPtr;
+	typedef Ptr<const NodeImageLayerColour> NodeImageLayerColourPtrConst;
+
+
+	//! This node applies a layer blending effect on a base image using a mask and a colour.
+	//! \ingroup model
+	class MUTABLETOOLS_API NodeImageLayerColour : public NodeImage
+	{
+	public:
+
+		NodeImageLayerColour();
+
+		//-----------------------------------------------------------------------------------------
+		// Node Interface
+		//-----------------------------------------------------------------------------------------
+
+        const FNodeType* GetType() const override;
+		static const FNodeType* GetStaticType();
+
+		//-----------------------------------------------------------------------------------------
+		// Own Interface
+		//-----------------------------------------------------------------------------------------
+
+		//! Get the node generating the base image that will have the blending effect applied.
+		NodeImagePtr GetBase() const;
+		void SetBase( NodeImagePtr );
+
+		//! Get the node generating the mask image controlling the weight of the effect.
+		NodeImagePtr GetMask() const;
+		void SetMask( NodeImagePtr );
+
+		//! Get the node generating the color to blend on the base.
+		NodeColourPtr GetColour() const;
+		void SetColour( NodeColourPtr );
+
+		//!
+		EBlendType GetBlendType() const;
+
+		//!
+		void SetBlendType(EBlendType);
+
+		//-----------------------------------------------------------------------------------------
+		// Interface pattern
+		//-----------------------------------------------------------------------------------------
+		class Private;
+		Private* GetPrivate() const;
+
+	protected:
+
+		//! Forbidden. Manage with the Ptr<> template.
+		~NodeImageLayerColour();
+
+	private:
+
+		Private* m_pD;
+
+	};
+
+
+}
diff --git a/Engine/Plugins/Mutable/Source/MutableTools/Public/MuT/NodeImageLuminance.h b/Engine/Plugins/Mutable/Source/MutableTools/Public/MuT/NodeImageLuminance.h
new file mode 100644
--- /dev/null
+++ b/Engine/Plugins/Mutable/Source/MutableTools/Public/MuT/NodeImageLuminance.h
@@ -0,0 +1,61 @@
+// Copyright Epic Games, Inc. All Rights Reserved.
+
+#pragma once
+
+#include "MuR/Ptr.h"
+#include "MuR/RefCounted.h"
+#include "MuT/Node.h"
+#include "MuT/NodeImage.h"
+
+
+namespace mu
+{
+
+	// Forward definitions
+	class NodeImageLuminance;
+	typedef Ptr<NodeImageLuminance> NodeImageLuminancePtr;
+	typedef Ptr<const NodeImageLuminance> NodeImageLuminancePtrConst;
+
+
+	//! Calculate the luminance of an image into a new single-channel image..
+	//! \ingroup model
+	class MUTABLETOOLS_API NodeImageLuminance : public NodeImage
+	{
+	public:
+
+		NodeImageLuminance();
+
+		//-----------------------------------------------------------------------------------------
+		// Node Interface
+		//-----------------------------------------------------------------------------------------
+
+        const FNodeType* GetType() const override;
+		static const FNodeType* GetStaticType();
+
+		//-----------------------------------------------------------------------------------------
+		// Own Interface
+		//-----------------------------------------------------------------------------------------
+
+		//! Get the node generating the source image calculate the luminance from.
+		NodeImagePtr GetSource() const;
+		void SetSource( NodeImagePtr );
+
+		//-----------------------------------------------------------------------------------------
+		// Interface pattern
+		//-----------------------------------------------------------------------------------------
+		class Private;
+		Private* GetPrivate() const;
+
+	protected:
+
+		//! Forbidden. Manage with the Ptr<> template.
+		~NodeImageLuminance();
+
+	private:
+
+		Private* m_pD;
+
+	};
+
+
+}
diff --git a/Engine/Plugins/Mutable/Source/MutableTools/Public/MuT/NodeImageMipmap.h b/Engine/Plugins/Mutable/Source/MutableTools/Public/MuT/NodeImageMipmap.h
new file mode 100644
--- /dev/null
+++ b/Engine/Plugins/Mutable/Source/MutableTools/Public/MuT/NodeImageMipmap.h
@@ -0,0 +1,68 @@
+// Copyright Epic Games, Inc. All Rights Reserved.
+
+#pragma once
+
+#include "MuR/Image.h"
+#include "MuR/Ptr.h"
+#include "MuR/RefCounted.h"
+#include "MuT/Node.h"
+#include "MuT/NodeImage.h"
+
+
+namespace mu
+{
+
+	// Forward definitions
+	class NodeScalar;
+	typedef Ptr<NodeScalar> NodeScalarPtr;
+	typedef Ptr<const NodeScalar> NodeScalarPtrConst;
+
+    class NodeImageMipmap;
+    typedef Ptr<NodeImageMipmap> NodeImageMipmapPtr;
+    typedef Ptr<const NodeImageMipmap> NodeImageMipmapPtrConst;
+
+
+	//! Generate mimaps for an image.
+	//! \ingroup model
+    class MUTABLETOOLS_API NodeImageMipmap : public NodeImage
+	{
+	public:
+
+        NodeImageMipmap();
+
+		//-----------------------------------------------------------------------------------------
+		// Node Interface
+		//-----------------------------------------------------------------------------------------
+
+        const FNodeType* GetType() const override;
+		static const FNodeType* GetStaticType();
+
+		//-----------------------------------------------------------------------------------------
+		// Own Interface
+		//-----------------------------------------------------------------------------------------
+
+        //! Get the node generating the source image to be mipmapped.
+		NodeImagePtr GetSource() const;
+		void SetSource( NodeImagePtr );
+
+		void SetMipmapGenerationSettings(EMipmapFilterType FilterType, EAddressMode AddressMode);
+
+		//-----------------------------------------------------------------------------------------
+		// Interface pattern
+		//-----------------------------------------------------------------------------------------
+		class Private;
+		Private* GetPrivate() const;
+
+	protected:
+
+		//! Forbidden. Manage with the Ptr<> template.
+        ~NodeImageMipmap();
+
+	private:
+
+		Private* m_pD;
+
+	};
+
+
+}
diff --git a/Engine/Plugins/Mutable/Source/MutableTools/Public/MuT/NodeImageMultiLayer.h b/Engine/Plugins/Mutable/Source/MutableTools/Public/MuT/NodeImageMultiLayer.h
new file mode 100644
--- /dev/null
+++ b/Engine/Plugins/Mutable/Source/MutableTools/Public/MuT/NodeImageMultiLayer.h
@@ -0,0 +1,80 @@
+// Copyright Epic Games, Inc. All Rights Reserved.
+
+#pragma once
+
+#include "MuR/Image.h"
+#include "MuR/Ptr.h"
+#include "MuR/RefCounted.h"
+#include "MuT/Node.h"
+#include "MuT/NodeImage.h"
+#include "MuT/NodeRange.h"
+
+
+namespace mu
+{
+
+	// Forward definitions
+    class NodeImageMultiLayer;
+    typedef Ptr<NodeImageMultiLayer> NodeImageMultiLayerPtr;
+    typedef Ptr<const NodeImageMultiLayer> NodeImageMultiLayerPtrConst;
+
+
+    //! This node applies any numbre of layer blending effect on a base image using a mask and a
+    //! blended image. The number of layers depends on a scalar input of this node.
+	//! \ingroup model
+    class MUTABLETOOLS_API NodeImageMultiLayer : public NodeImage
+	{
+	public:
+
+        NodeImageMultiLayer();
+
+		//-----------------------------------------------------------------------------------------
+		// Node Interface
+		//-----------------------------------------------------------------------------------------
+
+        const FNodeType* GetType() const override;
+		static const FNodeType* GetStaticType();
+
+		//-----------------------------------------------------------------------------------------
+		// Own Interface
+		//-----------------------------------------------------------------------------------------
+
+		//! Get the node generating the base image that will have the blendeing effect applied.
+		NodeImagePtr GetBase() const;
+		void SetBase( NodeImagePtr );
+
+		//! Get the node generating the mask image controlling the weight of the effect.
+		NodeImagePtr GetMask() const;
+		void SetMask( NodeImagePtr );
+
+        //! Get the image blended on top of the base.
+        NodeImagePtr GetBlended() const;
+        void SetBlended( NodeImagePtr );
+
+        //! Get the number of layers to apply.
+        NodeRangePtr GetRange() const;
+        void SetRange( NodeRangePtr );
+
+		//!
+		EBlendType GetBlendType() const;
+        void SetBlendType(EBlendType);
+
+		//-----------------------------------------------------------------------------------------
+		// Interface pattern
+		//-----------------------------------------------------------------------------------------
+		class Private;
+		Private* GetPrivate() const;
+
+	protected:
+
+		//! Forbidden. Manage with the Ptr<> template.
+        ~NodeImageMultiLayer();
+
+	private:
+
+		Private* m_pD;
+
+	};
+
+
+}
diff --git a/Engine/Plugins/Mutable/Source/MutableTools/Public/MuT/NodeImageNormalComposite.h b/Engine/Plugins/Mutable/Source/MutableTools/Public/MuT/NodeImageNormalComposite.h
new file mode 100644
--- /dev/null
+++ b/Engine/Plugins/Mutable/Source/MutableTools/Public/MuT/NodeImageNormalComposite.h
@@ -0,0 +1,65 @@
+// Copyright Epic Games, Inc. All Rights Reserved.
+
+#pragma once
+
+#include "MuR/Image.h"
+#include "MuR/Ptr.h"
+#include "MuR/RefCounted.h"
+#include "MuT/Node.h"
+#include "MuT/NodeImage.h"
+
+namespace mu
+{
+
+	// Forward definitions
+	class NodeImageNormalComposite;
+	typedef Ptr<NodeImageNormalComposite> NodeImageNormalCompositePtr;
+	typedef Ptr<const NodeImageNormalComposite> NodeImageNormalCompositePtrConst;
+
+	//! Node that inverts the colors of an image, channel by channel
+	class MUTABLETOOLS_API NodeImageNormalComposite : public NodeImage
+	{
+	public:
+
+		NodeImageNormalComposite();
+
+		//-----------------------------------------------------------------------------------------
+		// Node Interface
+		//-----------------------------------------------------------------------------------------
+
+		const FNodeType* GetType() const override;
+		static const FNodeType* GetStaticType();
+
+		//-----------------------------------------------------------------------------------------
+		// Own Interface
+		//-----------------------------------------------------------------------------------------
+
+		NodeImagePtr GetBase() const;
+		void SetBase(NodeImagePtr);
+
+        NodeImagePtr GetNormal() const;
+        void SetNormal(NodeImagePtr);
+
+		float GetPower() const;
+		void SetPower(float); 
+
+        ECompositeImageMode GetMode() const;
+        void SetMode(ECompositeImageMode);
+
+		//-----------------------------------------------------------------------------------------
+		// Interface pattern
+		//-----------------------------------------------------------------------------------------
+		class Private;
+		Private* GetPrivate() const;
+
+	protected:
+
+		//! Forbidden. Manage with the Ptr<> template
+		~NodeImageNormalComposite();
+
+	private:
+
+		Private* m_pD;
+	};
+
+}
diff --git a/Engine/Plugins/Mutable/Source/MutableTools/Public/MuT/NodeImageParameter.h b/Engine/Plugins/Mutable/Source/MutableTools/Public/MuT/NodeImageParameter.h
new file mode 100644
--- /dev/null
+++ b/Engine/Plugins/Mutable/Source/MutableTools/Public/MuT/NodeImageParameter.h
@@ -0,0 +1,66 @@
+// Copyright Epic Games, Inc. All Rights Reserved.
+
+#pragma once
+
+#include "MuR/Ptr.h"
+#include "MuR/RefCounted.h"
+#include "MuT/Node.h"
+#include "MuT/NodeImage.h"
+
+
+namespace mu
+{
+	class NodeRange;
+
+    //! Node that defines a Image model parameter.
+	//! \ingroup model
+    class MUTABLETOOLS_API NodeImageParameter : public NodeImage
+	{
+	public:
+
+        NodeImageParameter();
+
+		//-----------------------------------------------------------------------------------------
+		// Node Interface
+		//-----------------------------------------------------------------------------------------
+
+		const FNodeType* GetType() const override;
+		static const FNodeType* GetStaticType();
+
+		//-----------------------------------------------------------------------------------------
+		// Own Interface
+		//-----------------------------------------------------------------------------------------
+
+		//! Set the name of the parameter.
+		void SetName(const FString&);
+
+		//! Set the uid of the parameter.
+		void SetUid(const FString&);
+
+    	//! Get the default value of the parameter.
+    	void SetDefaultValue(FName Value);
+    	
+		//! Set the number of ranges (dimensions) for this parameter.
+		//! By default a parameter has 0 ranges, meaning it only has one value.
+		void SetRangeCount(int Index);
+		void SetRange(int Index, Ptr<NodeRange> Range);
+
+		//-----------------------------------------------------------------------------------------
+		// Interface pattern
+		//-----------------------------------------------------------------------------------------
+		class Private;
+		Private* GetPrivate() const;
+
+	protected:
+
+		//! Forbidden. Manage with the Ptr<> template.
+        ~NodeImageParameter();
+
+	private:
+
+		Private* m_pD;
+
+	};
+
+
+}
diff --git a/Engine/Plugins/Mutable/Source/MutableTools/Public/MuT/NodeImagePlainColour.h b/Engine/Plugins/Mutable/Source/MutableTools/Public/MuT/NodeImagePlainColour.h
new file mode 100644
--- /dev/null
+++ b/Engine/Plugins/Mutable/Source/MutableTools/Public/MuT/NodeImagePlainColour.h
@@ -0,0 +1,69 @@
+// Copyright Epic Games, Inc. All Rights Reserved.
+
+#pragma once
+
+#include "MuR/Ptr.h"
+#include "MuR/RefCounted.h"
+#include "MuT/Node.h"
+#include "MuT/NodeImage.h"
+
+
+namespace mu
+{
+
+	// Forward definitions
+	class NodeColour;
+	typedef Ptr<NodeColour> NodeColourPtr;
+	typedef Ptr<const NodeColour> NodeColourPtrConst;
+
+	class NodeImagePlainColour;
+	typedef Ptr<NodeImagePlainColour> NodeImagePlainColourPtr;
+	typedef Ptr<const NodeImagePlainColour> NodeImagePlainColourPtrConst;
+
+
+	//! Node that multiplies the colors of an image, channel by channel.
+	//! \ingroup model
+	class MUTABLETOOLS_API NodeImagePlainColour : public NodeImage
+	{
+	public:
+
+		NodeImagePlainColour();
+
+
+		//-----------------------------------------------------------------------------------------
+		// Node Interface
+		//-----------------------------------------------------------------------------------------
+
+        const FNodeType* GetType() const override;
+		static const FNodeType* GetStaticType();
+
+		//-----------------------------------------------------------------------------------------
+		// Own Interface
+		//-----------------------------------------------------------------------------------------
+
+		//! Colour of the image.
+		NodeColourPtr GetColour() const;
+		void SetColour( NodeColourPtr );
+
+		//! New size or relative factor
+		void SetSize( int x, int y );
+
+		//-----------------------------------------------------------------------------------------
+		// Interface pattern
+		//-----------------------------------------------------------------------------------------
+		class Private;
+		Private* GetPrivate() const;
+
+	protected:
+
+		//! Forbidden. Manage with the Ptr<> template.
+		~NodeImagePlainColour();
+
+	private:
+
+		Private* m_pD;
+
+	};
+
+
+}
diff --git a/Engine/Plugins/Mutable/Source/MutableTools/Public/MuT/NodeImageProject.h b/Engine/Plugins/Mutable/Source/MutableTools/Public/MuT/NodeImageProject.h
new file mode 100644
--- /dev/null
+++ b/Engine/Plugins/Mutable/Source/MutableTools/Public/MuT/NodeImageProject.h
@@ -0,0 +1,98 @@
+// Copyright Epic Games, Inc. All Rights Reserved.
+
+
+#pragma once
+
+#include "MuR/Ptr.h"
+#include "MuR/RefCounted.h"
+#include "MuT/Node.h"
+#include "MuT/NodeImage.h"
+#include "MuR/Image.h"
+
+namespace mu
+{
+
+	// Forward definitions
+	class NodeScalar;
+	class NodeMesh;
+	class NodeProjector;
+	class NodeImageProject;
+
+	//!
+	//! \ingroup model
+	class MUTABLETOOLS_API NodeImageProject : public NodeImage
+	{
+	public:
+
+		NodeImageProject();
+
+		//-----------------------------------------------------------------------------------------
+		// Node Interface
+		//-----------------------------------------------------------------------------------------
+
+        const FNodeType* GetType() const override;
+		static const FNodeType* GetStaticType();
+
+		//-----------------------------------------------------------------------------------------
+		// Own Interface
+		//-----------------------------------------------------------------------------------------
+
+		//! Get the node generating the mesh to be used for the projection.
+		void SetMesh( Ptr<NodeMesh> );
+
+		//! Get the node generating the projector to be used for the projection.
+		void SetProjector( Ptr<NodeProjector> );
+
+		//! Get the node generating a mask to filter out the projected images. It is optional.
+		void SetTargetMask( Ptr<NodeImage> );
+
+		/** Eanble or disable the additional operations to correct texture UV seam artifact correction. Default is enabled.*/
+		void SetEnableSeamCorrection(bool bEnabled);
+
+		//! Set the angle-based fading behaviour for this projector. It must be set separately for 
+		//! RGB and Alpha channels of the images to project. By default it is enabled for all channels.
+		//! Disabling the angle fade will still use the AngleFadeEnd below to clamp the projection,
+		//! but it will ignore the start angle, so it will not apply a fading gradient.
+		void SetAngleFadeChannels( bool bFadeRGB, bool bFadeA );
+
+        //! Set the node generating the fading start angle. Only relevant if fading is enabled with 
+		void SetAngleFadeStart( Ptr<NodeScalar> );
+
+		//! Set the node generating the fading end angle
+		void SetAngleFadeEnd( Ptr<NodeScalar> );
+
+		//! Set sampling method.
+		void SetSamplingMethod(ESamplingMethod SamplingMethod);
+		
+		//! Set min filter method
+		void SetMinFilterMethod(EMinFilterMethod MinFilterMethod);
+
+        //! Get the node generating the image to project.
+        void SetImage( Ptr<NodeImage> );
+
+        //! UV layout of the mesh to use for the generated image. Defaults to 0.
+        void SetLayout( uint8  );
+
+		//! Set the size of the image to generate with the projection. If set to 0 (default) a size
+		//! that matches how this node is used will try to be guessed. 
+		void SetImageSize( const FUintVector2& size );
+
+		//-----------------------------------------------------------------------------------------
+		// Interface pattern
+		//-----------------------------------------------------------------------------------------
+		class Private;
+		Private* GetPrivate() const;
+
+	protected:
+
+		//! Forbidden. Manage with the Ptr<> template.
+		~NodeImageProject();
+
+	private:
+
+		Private* m_pD;
+
+	};
+
+
+}
diff --git a/Engine/Plugins/Mutable/Source/MutableTools/Public/MuT/NodeImageResize.h b/Engine/Plugins/Mutable/Source/MutableTools/Public/MuT/NodeImageResize.h
new file mode 100644
--- /dev/null
+++ b/Engine/Plugins/Mutable/Source/MutableTools/Public/MuT/NodeImageResize.h
@@ -0,0 +1,72 @@
+// Copyright Epic Games, Inc. All Rights Reserved.
+
+#pragma once
+
+#include "MuR/Ptr.h"
+#include "MuR/RefCounted.h"
+#include "MuT/Node.h"
+#include "MuT/NodeImage.h"
+
+
+namespace mu
+{
+
+	// Forward definitions
+	class NodeScalar;
+	typedef Ptr<NodeScalar> NodeScalarPtr;
+	typedef Ptr<const NodeScalar> NodeScalarPtrConst;
+
+	class NodeImageResize;
+	typedef Ptr<NodeImageResize> NodeImageResizePtr;
+	typedef Ptr<const NodeImageResize> NodeImageResizePtrConst;
+
+
+	//! Node that multiplies the colors of an image, channel by channel.
+	//! \ingroup model
+	class MUTABLETOOLS_API NodeImageResize : public NodeImage
+	{
+	public:
+
+		NodeImageResize();
+
+
+		//-----------------------------------------------------------------------------------------
+		// Node Interface
+		//-----------------------------------------------------------------------------------------
+
+        const FNodeType* GetType() const override;
+		static const FNodeType* GetStaticType();
+
+		//-----------------------------------------------------------------------------------------
+		// Own Interface
+		//-----------------------------------------------------------------------------------------
+
+		//! Base image to resize.
+		NodeImagePtr GetBase() const;
+		void SetBase( NodeImagePtr );
+
+		//! Is the size a relative factor or an absolute size?
+		void SetRelative( bool );
+
+		//! New size or relative factor
+		void SetSize( float x, float y );
+
+		//-----------------------------------------------------------------------------------------
+		// Interface pattern
+		//-----------------------------------------------------------------------------------------
+		class Private;
+		Private* GetPrivate() const;
+
+	protected:
+
+		//! Forbidden. Manage with the Ptr<> template.
+		~NodeImageResize();
+
+	private:
+
+		Private* m_pD;
+
+	};
+
+
+}
diff --git a/Engine/Plugins/Mutable/Source/MutableTools/Public/MuT/NodeImageSaturate.h b/Engine/Plugins/Mutable/Source/MutableTools/Public/MuT/NodeImageSaturate.h
new file mode 100644
--- /dev/null
+++ b/Engine/Plugins/Mutable/Source/MutableTools/Public/MuT/NodeImageSaturate.h
@@ -0,0 +1,73 @@
+// Copyright Epic Games, Inc. All Rights Reserved.
+
+#pragma once
+
+#include "MuR/Ptr.h"
+#include "MuR/RefCounted.h"
+#include "MuT/Node.h"
+#include "MuT/NodeImage.h"
+
+
+namespace mu
+{
+
+	// Forward definitions
+	class NodeScalar;
+	typedef Ptr<NodeScalar> NodeScalarPtr;
+	typedef Ptr<const NodeScalar> NodeScalarPtrConst;
+
+	class NodeImageSaturate;
+	typedef Ptr<NodeImageSaturate> NodeImageSaturatePtr;
+	typedef Ptr<const NodeImageSaturate> NodeImageSaturatePtrConst;
+
+
+	//! Change the saturation of an image. This node can be used to increase the saturation with a
+	//! factor bigger than 1, or to decrease it or desaturate it completely with a factor smaller
+	//! than 1 or 0.
+	class MUTABLETOOLS_API NodeImageSaturate : public NodeImage
+	{
+	public:
+
+		NodeImageSaturate();
+
+		//-----------------------------------------------------------------------------------------
+		// Node Interface
+		//-----------------------------------------------------------------------------------------
+
+        const FNodeType* GetType() const override;
+		static const FNodeType* GetStaticType();
+
+		//-----------------------------------------------------------------------------------------
+		// Own Interface
+		//-----------------------------------------------------------------------------------------
+
+		//! Get the node generating the saturation factor.
+		//! A value of 0 completely desaturates. 1 will leave the same saturation and bigger than 1
+		//! will increase it.
+		NodeScalarPtr GetFactor() const;
+		void SetFactor( NodeScalarPtr );
+
+		//! Get the node generating the source image to be [de]saturated.
+		NodeImagePtr GetSource() const;
+		void SetSource( NodeImagePtr );
+
+		//-----------------------------------------------------------------------------------------
+		// Interface pattern
+		//-----------------------------------------------------------------------------------------
+		class Private;
+
+		Private* GetPrivate() const;
+
+	protected:
+
+		//! Forbidden. Manage with the Ptr<> template.
+		~NodeImageSaturate();
+
+	private:
+
+		Private* m_pD;
+
+	};
+
+
+}
diff --git a/Engine/Plugins/Mutable/Source/MutableTools/Public/MuT/NodeImageSwitch.h b/Engine/Plugins/Mutable/Source/MutableTools/Public/MuT/NodeImageSwitch.h
new file mode 100644
--- /dev/null
+++ b/Engine/Plugins/Mutable/Source/MutableTools/Public/MuT/NodeImageSwitch.h
@@ -0,0 +1,74 @@
+// Copyright Epic Games, Inc. All Rights Reserved.
+
+#pragma once
+
+#include "MuR/Ptr.h"
+#include "MuR/RefCounted.h"
+#include "MuT/Node.h"
+#include "MuT/NodeImage.h"
+
+
+namespace mu
+{
+
+	// Forward definitions
+	class NodeScalar;
+	typedef Ptr<NodeScalar> NodeScalarPtr;
+	typedef Ptr<const NodeScalar> NodeScalarPtrConst;
+
+	class NodeImageSwitch;
+	typedef Ptr<NodeImageSwitch> NodeImageSwitchPtr;
+	typedef Ptr<const NodeImageSwitch> NodeImageSwitchPtrConst;
+
+
+	//! This node selects an output image from a set of input images based on a parameter.
+	class MUTABLETOOLS_API NodeImageSwitch : public NodeImage
+	{
+	public:
+
+		NodeImageSwitch();
+
+		//-----------------------------------------------------------------------------------------
+		// Node Interface
+		//-----------------------------------------------------------------------------------------
+
+        const FNodeType* GetType() const override;
+		static const FNodeType* GetStaticType();
+
+		//-----------------------------------------------------------------------------------------
+		// Own Interface
+		//-----------------------------------------------------------------------------------------
+
+		//! Get the node generating the parameter used to select the option.
+		NodeScalarPtr GetParameter() const;
+		void SetParameter( NodeScalarPtr );
+
+		//! Set the number of option images. It will keep the currently set targets and initialise
+		//! the new ones as null.
+		void SetOptionCount( int );
+		int GetOptionCount() const;
+
+		//! Get the node generating the t-th option image.
+		NodeImagePtr GetOption( int t ) const;
+		void SetOption( int t, NodeImagePtr );
+
+
+		//-----------------------------------------------------------------------------------------
+		// Interface pattern
+		//-----------------------------------------------------------------------------------------
+		class Private;
+		Private* GetPrivate() const;
+
+	protected:
+
+		//! Forbidden. Manage with the Ptr<> template.
+		~NodeImageSwitch();
+
+	private:
+
+		Private* m_pD;
+
+	};
+
+
+}
diff --git a/Engine/Plugins/Mutable/Source/MutableTools/Public/MuT/NodeImageSwizzle.h b/Engine/Plugins/Mutable/Source/MutableTools/Public/MuT/NodeImageSwizzle.h
new file mode 100644
--- /dev/null
+++ b/Engine/Plugins/Mutable/Source/MutableTools/Public/MuT/NodeImageSwizzle.h
@@ -0,0 +1,69 @@
+// Copyright Epic Games, Inc. All Rights Reserved.
+
+#pragma once
+
+#include "MuR/Image.h"
+#include "MuR/Ptr.h"
+#include "MuR/RefCounted.h"
+#include "MuT/Node.h"
+#include "MuT/NodeImage.h"
+
+
+namespace mu
+{
+
+	// Forward definitions
+	class NodeImageSwizzle;
+	typedef Ptr<NodeImageSwizzle> NodeImageSwizzlePtr;
+	typedef Ptr<const NodeImageSwizzle> NodeImageSwizzlePtrConst;
+
+
+	//! Node that composes a new image by gathering pixel data from channels in other images.
+	//! \ingroup model
+	class MUTABLETOOLS_API NodeImageSwizzle : public NodeImage
+	{
+	public:
+
+		NodeImageSwizzle();
+
+		//-----------------------------------------------------------------------------------------
+		// Node Interface
+		//-----------------------------------------------------------------------------------------
+
+        const FNodeType* GetType() const override;
+		static const FNodeType* GetStaticType();
+
+		//-----------------------------------------------------------------------------------------
+		// Own Interface
+		//-----------------------------------------------------------------------------------------
+
+		//!
+		EImageFormat GetFormat() const;
+		void SetFormat(EImageFormat);
+
+		//!
+		NodeImagePtr GetSource( int ) const;
+		int GetSourceChannel( int ) const;
+
+		void SetSource( int, NodeImagePtr );
+		void SetSourceChannel( int, int );
+
+		//-----------------------------------------------------------------------------------------
+		// Interface pattern
+		//-----------------------------------------------------------------------------------------
+		class Private;
+		Private* GetPrivate() const;
+
+	protected:
+
+		//! Forbidden. Manage with the Ptr<> template.
+		~NodeImageSwizzle();
+
+	private:
+
+		Private* m_pD;
+
+	};
+
+
+}
diff --git a/Engine/Plugins/Mutable/Source/MutableTools/Public/MuT/NodeImageTransform.h b/Engine/Plugins/Mutable/Source/MutableTools/Public/MuT/NodeImageTransform.h
new file mode 100644
--- /dev/null
+++ b/Engine/Plugins/Mutable/Source/MutableTools/Public/MuT/NodeImageTransform.h
@@ -0,0 +1,88 @@
+// Copyright Epic Games, Inc. All Rights Reserved.
+
+#pragma once
+
+#include "MuR/Ptr.h"
+#include "MuR/RefCounted.h"
+#include "MuT/Node.h"
+#include "MuT/NodeImage.h"
+
+namespace mu
+{
+
+	// Forward definitions
+	class NodeScalar;
+	typedef Ptr<NodeScalar> NodeScalarPtr;
+	typedef Ptr<const NodeScalar> NodeScalarPtrConst;
+
+	class NodeImageTransform;
+	typedef Ptr<NodeImageTransform> NodeImageTransformPtr;
+	typedef Ptr<const NodeImageTransform> NodeImageTransformPtrConst;
+
+	enum class EAddressMode;
+
+	//! Node that multiplies the colors of an image, channel by channel.
+	//! \ingroup model
+	class MUTABLETOOLS_API NodeImageTransform : public NodeImage
+	{
+	public:
+
+		NodeImageTransform();
+
+		//-----------------------------------------------------------------------------------------
+		// Node Interface
+		//-----------------------------------------------------------------------------------------
+
+        const FNodeType* GetType() const override;
+		static const FNodeType* GetStaticType();
+
+		//-----------------------------------------------------------------------------------------
+		// Own Interface
+		//-----------------------------------------------------------------------------------------
+
+		//! Base image to multiply.
+		NodeImagePtr GetBase() const;
+		void SetBase( NodeImagePtr );
+	
+		NodeScalarPtr GetOffsetX() const;
+		NodeScalarPtr GetOffsetY() const;
+		NodeScalarPtr GetScaleX() const;
+		NodeScalarPtr GetScaleY() const;
+		NodeScalarPtr GetRotation() const;
+		void SetOffsetX( NodeScalarPtr pNode );
+		void SetOffsetY( NodeScalarPtr pNode );
+		void SetScaleX( NodeScalarPtr pNode );
+		void SetScaleY( NodeScalarPtr pNode );
+		void SetRotation( NodeScalarPtr pNode );
+
+		EAddressMode GetAddressMode() const;
+		void SetAddressMode(EAddressMode AddressMode);
+
+		uint16 GetSizeX() const;
+		void SetSizeX(uint16 SizeX); 
+
+		uint16 GetSizeY() const;
+		void SetSizeY(uint16 SizeY); 
+
+		bool GetKeepAspectRatio() const;
+		void SetKeepAspectRatio(bool bKeepAspectRatio);
+
+		//-----------------------------------------------------------------------------------------
+		// Interface pattern
+		//-----------------------------------------------------------------------------------------
+		class Private;
+		Private* GetPrivate() const;
+
+	protected:
+
+		//! Forbidden. Manage with the Ptr<> template.
+		~NodeImageTransform();
+
+	private:
+
+		Private* m_pD;
+
+	};
+
+
+}
diff --git a/Engine/Plugins/Mutable/Source/MutableTools/Public/MuT/NodeImageVariation.h b/Engine/Plugins/Mutable/Source/MutableTools/Public/MuT/NodeImageVariation.h
new file mode 100644
--- /dev/null
+++ b/Engine/Plugins/Mutable/Source/MutableTools/Public/MuT/NodeImageVariation.h
@@ -0,0 +1,75 @@
+// Copyright Epic Games, Inc. All Rights Reserved.
+
+#pragma once
+
+#include "MuR/Ptr.h"
+#include "MuR/RefCounted.h"
+#include "MuT/Node.h"
+#include "MuT/NodeImage.h"
+
+
+namespace mu
+{
+
+    // Forward definitions
+    class NodeImageVariation;
+    typedef Ptr<NodeImageVariation> NodeImageVariationPtr;
+    typedef Ptr<const NodeImageVariation> NodeImageVariationPtrConst;
+
+
+    //!
+    //! \ingroup model
+    class MUTABLETOOLS_API NodeImageVariation : public NodeImage
+    {
+    public:
+        //-----------------------------------------------------------------------------------------
+        // Life cycle
+        //-----------------------------------------------------------------------------------------
+
+        NodeImageVariation();
+
+        //-----------------------------------------------------------------------------------------
+        // Node Interface
+        //-----------------------------------------------------------------------------------------
+
+        const FNodeType* GetType() const override;
+        static const FNodeType* GetStaticType();
+
+        //-----------------------------------------------------------------------------------------
+        // Own Interface
+        //-----------------------------------------------------------------------------------------
+
+        //!
+        void SetDefaultImage( NodeImage* Image );
+
+        //! Set the number of tags to consider in this variation
+        void SetVariationCount( int count );
+
+        //!
+        int GetVariationCount() const;
+
+        //! Set the tag or state name that will enable a specific vartiation
+        void SetVariationTag( int index, const FString& Tag );
+
+        //!
+        void SetVariationImage( int index, NodeImage* Image );
+
+        //!}
+
+
+        //-----------------------------------------------------------------------------------------
+        // Interface pattern
+        //-----------------------------------------------------------------------------------------
+        class Private;
+        Private* GetPrivate() const;
+
+    protected:
+        //! Forbidden. Manage with the Ptr<> template.
+        ~NodeImageVariation();
+
+    private:
+        Private* m_pD;
+    };
+
+
+} // namespace mu
diff --git a/Engine/Plugins/Mutable/Source/MutableTools/Public/MuT/NodeMesh.h b/Engine/Plugins/Mutable/Source/MutableTools/Public/MuT/NodeMesh.h
new file mode 100644
--- /dev/null
+++ b/Engine/Plugins/Mutable/Source/MutableTools/Public/MuT/NodeMesh.h
@@ -0,0 +1,72 @@
+// Copyright Epic Games, Inc. All Rights Reserved.
+
+#pragma once
+
+#include "HAL/Platform.h"
+#include "MuR/Ptr.h"
+#include "MuR/RefCounted.h"
+#include "MuT/Node.h"
+#include "Containers/Array.h"
+#include "Containers/Map.h"
+
+namespace mu
+{
+
+	// Forward definitions
+	class Mesh;
+	class NodeMesh;
+	typedef Ptr<NodeMesh> NodeMeshPtr;
+	typedef Ptr<const NodeMesh> NodeMeshPtrConst;
+
+
+    /** Base class of any node that outputs a mesh. */
+	class MUTABLETOOLS_API NodeMesh : public Node
+	{
+	public:
+
+		// Node Interface
+        const FNodeType* GetType() const override;
+		static const FNodeType* GetStaticType();
+
+		// Interface pattern
+		class Private;
+
+	protected:
+
+		// Forbidden. Manage with the Ptr<> template.
+		inline ~NodeMesh() {}
+	};
+
+
+	/** Helper structs for mesh utility methods below. */
+	struct FTriangleInfo
+	{
+		/** Vertex indices in the original mesh. */
+		uint32 Indices[3];
+
+		/** Vertex indices in the collapsed vertex list of the mesh. */
+		uint32 CollapsedIndices[3];
+
+		/** Optional data with layout block indices. */
+		uint16 BlockIndices[3];
+
+		/** Optional data with a flag indicating the UVs have changed during layout for this trioangle. */
+		bool bUVsFixed = false;
+	};
+
+
+	/** Fill an array with the indices of all triangles belonging to the same UV island as InFirstTriangle. 
+	*/
+	extern void GetUVIsland(TArray<FTriangleInfo>& InTriangles,
+		const uint32 InFirstTriangle,
+		TArray<uint32>& OutTriangleIndices,
+		const TArray<FVector2f>& InUVs,
+		const TMultiMap<int32, uint32>& InVertexToTriangleMap);
+
+
+	/** Create a map from vertices into vertices, collapsing vertices that have the same position.
+	*/
+	extern void MeshCreateCollapsedVertexMap(const mu::Mesh* Mesh, TArray<int32>& CollapsedVertices);
+
+}
+
diff --git a/Engine/Plugins/Mutable/Source/MutableTools/Public/MuT/NodeMeshApplyPose.h b/Engine/Plugins/Mutable/Source/MutableTools/Public/MuT/NodeMeshApplyPose.h
new file mode 100644
--- /dev/null
+++ b/Engine/Plugins/Mutable/Source/MutableTools/Public/MuT/NodeMeshApplyPose.h
@@ -0,0 +1,66 @@
+// Copyright Epic Games, Inc. All Rights Reserved.
+
+#pragma once
+
+#include "MuR/Ptr.h"
+#include "MuR/RefCounted.h"
+#include "MuT/Node.h"
+#include "MuT/NodeMesh.h"
+
+
+namespace mu
+{
+
+	// Forward definitions
+    class NodeMeshApplyPose;
+    using NodeMeshApplyPosePtr = Ptr<NodeMeshApplyPose>;
+    using NodeMeshApplyPosePtrConst = Ptr<const NodeMeshApplyPose>;
+
+
+    //! Node that applies a pose to a mesh, baking it into the vertex data
+    class MUTABLETOOLS_API NodeMeshApplyPose : public NodeMesh
+	{
+	public:
+
+        NodeMeshApplyPose();
+
+
+		//-----------------------------------------------------------------------------------------
+		// Node Interface
+		//-----------------------------------------------------------------------------------------
+
+		const FNodeType* GetType() const override;
+		static const FNodeType* GetStaticType();
+
+		//-----------------------------------------------------------------------------------------
+		// Own Interface
+		//-----------------------------------------------------------------------------------------
+
+		//! Get the nodes generating the base mesh to be morphed
+		NodeMeshPtr GetBase() const;
+		void SetBase( NodeMeshPtr );
+
+        //! Get the nodes generating the pose to apply. A pose is represented with a mesh object
+        //! with no geometry: only bone matrices and names are relevant.
+        NodeMeshPtr GetPose() const;
+        void SetPose( NodeMeshPtr );
+
+		//-----------------------------------------------------------------------------------------
+		// Interface pattern
+		//-----------------------------------------------------------------------------------------
+		class Private;
+		Private* GetPrivate() const;
+
+	protected:
+
+		//! Forbidden. Manage with the Ptr<> template.
+        ~NodeMeshApplyPose();
+
+	private:
+
+		Private* m_pD;
+
+	};
+
+
+}
diff --git a/Engine/Plugins/Mutable/Source/MutableTools/Public/MuT/NodeMeshClipDeform.h b/Engine/Plugins/Mutable/Source/MutableTools/Public/MuT/NodeMeshClipDeform.h
new file mode 100644
--- /dev/null
+++ b/Engine/Plugins/Mutable/Source/MutableTools/Public/MuT/NodeMeshClipDeform.h
@@ -0,0 +1,40 @@
+// Copyright Epic Games, Inc. All Rights Reserved.
+
+#pragma once
+
+#include "MuR/Ptr.h"
+#include "MuT/Node.h"
+#include "MuT/NodeImage.h"
+#include "MuT/NodeMesh.h"
+
+
+namespace mu
+{
+
+	//! 
+	class MUTABLETOOLS_API NodeMeshClipDeform : public NodeMesh
+	{
+	public:
+
+		Ptr<NodeMesh> BaseMesh;
+		Ptr<NodeMesh> ClipShape;
+		Ptr<NodeImage> ShapeWeights;
+
+	public:
+
+		// Node interface
+		virtual const FNodeType* GetType() const override { return &StaticType; }
+		static const FNodeType* GetStaticType() { return &StaticType; }
+
+	protected:
+
+		//! Forbidden. Manage with the Ptr<> template.
+		inline ~NodeMeshClipDeform() {}
+
+	private:
+
+		static FNodeType StaticType;
+
+	};
+
+}
diff --git a/Engine/Plugins/Mutable/Source/MutableTools/Public/MuT/NodeMeshClipWithMesh.h b/Engine/Plugins/Mutable/Source/MutableTools/Public/MuT/NodeMeshClipWithMesh.h
new file mode 100644
--- /dev/null
+++ b/Engine/Plugins/Mutable/Source/MutableTools/Public/MuT/NodeMeshClipWithMesh.h
@@ -0,0 +1,41 @@
+// Copyright Epic Games, Inc. All Rights Reserved.
+
+#pragma once
+
+#include "MuR/Ptr.h"
+#include "MuR/RefCounted.h"
+#include "MuT/Node.h"
+#include "MuT/NodeMesh.h"
+
+
+namespace mu
+{
+
+    //! This node applies a geometric transform represented by a 4x4 matrix to a mesh
+	//! \ingroup model
+    class MUTABLETOOLS_API NodeMeshClipWithMesh : public NodeMesh
+	{
+	public:
+
+		Ptr<NodeMesh> Source;
+		Ptr<NodeMesh> ClipMesh;
+
+	public:
+
+		// Node interface
+		virtual const FNodeType* GetType() const override { return &StaticType; }
+		static const FNodeType* GetStaticType() { return &StaticType; }
+
+	protected:
+
+		//! Forbidden. Manage with the Ptr<> template.
+		inline ~NodeMeshClipWithMesh() {}
+
+	private:
+
+		static FNodeType StaticType;
+
+	};
+
+
+}
diff --git a/Engine/Plugins/Mutable/Source/MutableTools/Public/MuT/NodeMeshConstant.h b/Engine/Plugins/Mutable/Source/MutableTools/Public/MuT/NodeMeshConstant.h
new file mode 100644
--- /dev/null
+++ b/Engine/Plugins/Mutable/Source/MutableTools/Public/MuT/NodeMeshConstant.h
@@ -0,0 +1,79 @@
+// Copyright Epic Games, Inc. All Rights Reserved.
+
+#pragma once
+
+#include "MuR/Ptr.h"
+#include "MuR/RefCounted.h"
+#include "MuT/Node.h"
+#include "MuT/NodeMesh.h"
+
+
+namespace mu
+{
+
+	// Forward definitions
+	class NodeMeshConstant;
+	typedef Ptr<NodeMeshConstant> NodeMeshConstantPtr;
+	typedef Ptr<const NodeMeshConstant> NodeMeshConstantPtrConst;
+
+	class Mesh;
+	class NodeLayout;
+
+	//! Node that outputs a constant mesh.
+	//! It allows to define the layouts for the texture channels of the constant mesh
+	class MUTABLETOOLS_API NodeMeshConstant : public NodeMesh
+	{
+	public:
+
+		NodeMeshConstant();
+
+		//-----------------------------------------------------------------------------------------
+		// Node Interface
+		//-----------------------------------------------------------------------------------------
+
+		const FNodeType* GetType() const override;
+		static const FNodeType* GetStaticType();
+
+		//-----------------------------------------------------------------------------------------
+		// Own Interface
+		//-----------------------------------------------------------------------------------------
+
+		//! Get the constant mesh that will be returned.
+		Ptr<Mesh> GetValue() const;
+
+		//! Set the constant mesh that will be returned.
+		void SetValue(Ptr<Mesh>);
+
+		/** */
+		void AddMorph(const FString& Name, Ptr<Mesh>);
+
+		/** */
+		Ptr<Mesh> FindMorph(const FString& Name) const;
+
+		//! Get the number of layouts defined in this mesh.
+		int32 GetLayoutCount() const;
+		void SetLayoutCount( int32 );
+
+		//! Get the node defining a layout of the returned mesh.
+		Ptr<NodeLayout> GetLayout( int32 index ) const;
+		void SetLayout( int32 index, Ptr<NodeLayout>);
+
+		//-----------------------------------------------------------------------------------------
+		// Interface pattern
+		//-----------------------------------------------------------------------------------------
+		class Private;
+		Private* GetPrivate() const;
+
+	protected:
+
+		//! Forbidden. Manage with the Ptr<> template.
+		~NodeMeshConstant();
+
+	private:
+
+		Private* m_pD;
+
+	};
+
+
+}
diff --git a/Engine/Plugins/Mutable/Source/MutableTools/Public/MuT/NodeMeshFormat.h b/Engine/Plugins/Mutable/Source/MutableTools/Public/MuT/NodeMeshFormat.h
new file mode 100644
--- /dev/null
+++ b/Engine/Plugins/Mutable/Source/MutableTools/Public/MuT/NodeMeshFormat.h
@@ -0,0 +1,74 @@
+// Copyright Epic Games, Inc. All Rights Reserved.
+
+#pragma once
+
+#include "MuR/Ptr.h"
+#include "MuR/RefCounted.h"
+#include "MuT/Node.h"
+#include "MuT/NodeMesh.h"
+
+namespace mu
+{
+
+	// Forward definitions
+	class FMeshBufferSet;
+
+	/** This node can change the buffer formats of a mesh vertices, indices and faces. */
+	class MUTABLETOOLS_API NodeMeshFormat : public NodeMesh
+	{
+	public:
+
+		NodeMeshFormat();
+
+		//-----------------------------------------------------------------------------------------
+		// Node Interface
+		//-----------------------------------------------------------------------------------------
+
+		const FNodeType* GetType() const override;
+		static const FNodeType* GetStaticType();
+
+		//-----------------------------------------------------------------------------------------
+		// Own Interface
+		//-----------------------------------------------------------------------------------------
+
+		/** Mesh to be reformatted. */
+		Ptr<NodeMesh> GetSource() const;
+		void SetSource( NodeMesh* );
+
+		/** Access the MeshBufferSet that defines the new format for the mesh vertices.These
+		* buffers don't really contain any data (they have 0 elements) but they define the
+		* structure. If this is null, the vertex buffers will not be changed.
+		*/
+		FMeshBufferSet& GetVertexBuffers();
+
+		/** Access the MeshBufferSet that defines the new format for the mesh indices. These
+		* buffers don't really contain any data (they have 0 elements) but they define the
+		* structure. If this is null, the vertex buffers will not be changed.
+		*/
+		FMeshBufferSet& GetIndexBuffers();
+
+		/** Optimize the buffers is possible. This may change the target format to reduce the number
+		* of channels or even the type if possible to minimize size.
+		* By default it is disabled.
+		*/
+		void SetOptimizeBuffers(bool bEnable);
+
+		//-----------------------------------------------------------------------------------------
+		// Interface pattern
+		//-----------------------------------------------------------------------------------------
+		class Private;
+		Private* GetPrivate() const;
+
+	protected:
+
+		//! Forbidden. Manage with the Ptr<> template.
+		~NodeMeshFormat();
+
+	private:
+
+		Private* m_pD;
+
+	};
+
+
+}
diff --git a/Engine/Plugins/Mutable/Source/MutableTools/Public/MuT/NodeMeshGeometryOperation.h b/Engine/Plugins/Mutable/Source/MutableTools/Public/MuT/NodeMeshGeometryOperation.h
new file mode 100644
--- /dev/null
+++ b/Engine/Plugins/Mutable/Source/MutableTools/Public/MuT/NodeMeshGeometryOperation.h
@@ -0,0 +1,71 @@
+// Copyright Epic Games, Inc. All Rights Reserved.
+
+#pragma once
+
+#include "MuR/Ptr.h"
+#include "MuR/RefCounted.h"
+#include "MuT/Node.h"
+#include "MuT/NodeMesh.h"
+
+
+namespace mu
+{
+
+	// Forward definitions
+	class NodeScalar;
+	typedef Ptr<NodeScalar> NodeScalarPtr;
+	typedef Ptr<const NodeScalar> NodeScalarPtrConst;
+
+	class NodeMeshGeometryOperation;
+	typedef Ptr<NodeMeshGeometryOperation> NodeMeshGeometryOperationPtr;
+	typedef Ptr<const NodeMeshGeometryOperation> NodeMeshGeometryOperationPtrConst;
+
+
+	//! Node that morphs a base mesh with one or two weighted targets from a sequence.
+	class MUTABLETOOLS_API NodeMeshGeometryOperation : public NodeMesh
+	{
+	public:
+
+		NodeMeshGeometryOperation();
+
+		//-----------------------------------------------------------------------------------------
+		// Node Interface
+		//-----------------------------------------------------------------------------------------
+
+		const FNodeType* GetType() const override;
+		static const FNodeType* GetStaticType();
+
+		//-----------------------------------------------------------------------------------------
+		// Own Interface
+		//-----------------------------------------------------------------------------------------
+
+		const NodeMeshPtr& GetMeshA() const;
+		void SetMeshA( const NodeMeshPtr& );
+
+		const NodeMeshPtr& GetMeshB() const;
+		void SetMeshB(const NodeMeshPtr&);
+
+		const NodeScalarPtr& GetScalarA() const;
+		void SetScalarA(const NodeScalarPtr&);
+
+		const NodeScalarPtr& GetScalarB() const;
+		void SetScalarB(const NodeScalarPtr&);
+
+        //-----------------------------------------------------------------------------------------
+		// Interface pattern
+		//-----------------------------------------------------------------------------------------
+		class Private;
+		Private* GetPrivate() const;
+
+	protected:
+
+		//! Forbidden. Manage with the Ptr<> template.
+		~NodeMeshGeometryOperation();
+
+	private:
+
+		Private* m_pD;
+
+	};
+
+}
diff --git a/Engine/Plugins/Mutable/Source/MutableTools/Public/MuT/NodeMeshInterpolate.h b/Engine/Plugins/Mutable/Source/MutableTools/Public/MuT/NodeMeshInterpolate.h
new file mode 100644
--- /dev/null
+++ b/Engine/Plugins/Mutable/Source/MutableTools/Public/MuT/NodeMeshInterpolate.h
@@ -0,0 +1,88 @@
+// Copyright Epic Games, Inc. All Rights Reserved.
+
+#pragma once
+
+#include "MuR/MeshBufferSet.h"
+#include "MuR/Ptr.h"
+#include "MuR/RefCounted.h"
+#include "MuT/Node.h"
+#include "MuT/NodeMesh.h"
+
+
+namespace mu
+{
+
+	// Forward definitions
+	class NodeScalar;
+	typedef Ptr<NodeScalar> NodeScalarPtr;
+	typedef Ptr<const NodeScalar> NodeScalarPtrConst;
+
+	class NodeMeshInterpolate;
+	typedef Ptr<NodeMeshInterpolate> NodeMeshInterpolatePtr;
+	typedef Ptr<const NodeMeshInterpolate> NodeMeshInterpolatePtrConst;
+
+
+	//! Node that morphs (interpolates linearly) between several meshes based on a weight.
+	//! The different input images are uniformly distributed in the 0 to 1 range:
+	//! - if two meshes are set, the first one is 0.0 and the second one is 1.0
+	//! - if three meshes are set, the first one is 0.0 the second one is 0.5 and the third one 1.0
+	//! regardless of the input slots used: If B and C are set, but not A, B will be at weight 0
+	//! and C will be at weight 1
+	class MUTABLETOOLS_API NodeMeshInterpolate : public NodeMesh
+	{
+	public:
+
+		NodeMeshInterpolate();
+
+
+		//-----------------------------------------------------------------------------------------
+		// Node Interface
+		//-----------------------------------------------------------------------------------------
+
+		const FNodeType* GetType() const override;
+		static const FNodeType* GetStaticType();
+
+		//-----------------------------------------------------------------------------------------
+		// Own Interface
+		//-----------------------------------------------------------------------------------------
+
+		//! Get the node generating the weight of the interpolation.
+		NodeScalarPtr GetFactor() const;
+		void SetFactor( NodeScalarPtr );
+
+		//! Set the number of target meshes. It will keep the currently set targets and initialise
+		//! the new ones as 0.
+		void SetTargetCount( int );
+		int GetTargetCount() const;
+
+		//! Get the node generating the several mesh targets of the interpolation.
+		NodeMeshPtr GetTarget( int t ) const;
+		void SetTarget( int t, NodeMeshPtr );
+
+		//! Set the number of channels of the mesh that will be interpolated. It will keep the
+		//! currently set targets and initialise the new ones as null. The channels that are not
+		//! interpolated will be copied from the first target.
+		void SetChannelCount( int );
+
+		//! Set a channel to interpolate from the mesh.
+		void SetChannel( int index, EMeshBufferSemantic semantic, int semanticIndex );
+
+		//-----------------------------------------------------------------------------------------
+		// Interface pattern
+		//-----------------------------------------------------------------------------------------
+		class Private;
+		Private* GetPrivate() const;
+
+	protected:
+
+		//! Forbidden. Manage with the Ptr<> template.
+		~NodeMeshInterpolate();
+
+	private:
+
+		Private* m_pD;
+
+	};
+
+
+}
diff --git a/Engine/Plugins/Mutable/Source/MutableTools/Public/MuT/NodeMeshMakeMorph.h b/Engine/Plugins/Mutable/Source/MutableTools/Public/MuT/NodeMeshMakeMorph.h
new file mode 100644
--- /dev/null
+++ b/Engine/Plugins/Mutable/Source/MutableTools/Public/MuT/NodeMeshMakeMorph.h
@@ -0,0 +1,67 @@
+// Copyright Epic Games, Inc. All Rights Reserved.
+
+#pragma once
+
+#include "MuR/Ptr.h"
+#include "MuR/RefCounted.h"
+#include "MuT/Node.h"
+#include "MuT/NodeMesh.h"
+
+
+namespace mu
+{
+
+	// Forward definitions
+    class NodeMeshMakeMorph;
+    typedef Ptr<NodeMeshMakeMorph> NodeMeshMakeMorphPtr;
+    typedef Ptr<const NodeMeshMakeMorph> NodeMeshMakeMorphPtrConst;
+
+
+	//! This node removes from a mesh A all the faces that are also part of a mesh B.
+    class MUTABLETOOLS_API NodeMeshMakeMorph : public NodeMesh
+	{
+	public:
+
+        NodeMeshMakeMorph();
+
+		//-----------------------------------------------------------------------------------------
+		// Node Interface
+		//-----------------------------------------------------------------------------------------
+
+		const FNodeType* GetType() const override;
+		static const FNodeType* GetStaticType();
+
+		//-----------------------------------------------------------------------------------------
+		// Own Interface
+		//-----------------------------------------------------------------------------------------
+
+		//!
+		NodeMeshPtr GetBase() const;
+		void SetBase( NodeMesh* p );
+
+		//!
+		NodeMeshPtr GetTarget() const;
+		void SetTarget( NodeMesh* p );
+
+		bool GetOnlyPositionAndNormal() const;
+		void SetOnlyPositionAndNormal(bool bOnlyPositionAndNormal);
+
+		//-----------------------------------------------------------------------------------------
+		// Interface pattern
+		//-----------------------------------------------------------------------------------------
+		class Private;
+		Private* GetPrivate() const;
+
+	protected:
+
+		//! Forbidden. Manage with the Ptr<> template.
+        ~NodeMeshMakeMorph();
+
+	private:
+
+		Private* m_pD;
+
+	};
+
+
+}
diff --git a/Engine/Plugins/Mutable/Source/MutableTools/Public/MuT/NodeMeshMorph.h b/Engine/Plugins/Mutable/Source/MutableTools/Public/MuT/NodeMeshMorph.h
new file mode 100644
--- /dev/null
+++ b/Engine/Plugins/Mutable/Source/MutableTools/Public/MuT/NodeMeshMorph.h
@@ -0,0 +1,78 @@
+// Copyright Epic Games, Inc. All Rights Reserved.
+
+#pragma once
+
+#include "MuR/Ptr.h"
+#include "MuR/RefCounted.h"
+#include "MuT/Node.h"
+#include "MuT/NodeMesh.h"
+
+
+namespace mu
+{
+
+	// Forward definitions
+	class NodeScalar;
+	struct FBoneName;
+
+	/** Node that morphs a base mesh with one weighted target. */
+	class MUTABLETOOLS_API NodeMeshMorph : public NodeMesh
+	{
+	public:
+
+		NodeMeshMorph();
+
+		//-----------------------------------------------------------------------------------------
+		// Node Interface
+		//-----------------------------------------------------------------------------------------		
+
+		const FNodeType* GetType() const override;
+		static const FNodeType* GetStaticType();
+
+		//-----------------------------------------------------------------------------------------
+		// Own Interface
+		//-----------------------------------------------------------------------------------------
+
+		/** Get the node generating the weight used to select(and combine) the morphs to apply. */
+		Ptr<NodeScalar> GetFactor() const;
+		void SetFactor( Ptr<NodeScalar> );
+
+		/** Get the nodes generating the base mesh to be morphed. */
+		Ptr<NodeMesh> GetBase() const;
+		void SetBase( Ptr<NodeMesh> );
+
+		/** Get the nodes generating the several morph targets. */
+		Ptr<NodeMesh> GetMorph() const;
+		void SetMorph( Ptr<NodeMesh> );
+
+		/** Also deform the mesh skeleton. Disabled by default. */
+		void SetReshapeSkeleton(bool);	
+
+		/** Deform Mesh Physics Volumes */
+		void SetReshapePhysicsVolumes(bool);
+
+		/** Sets the number of bones that will be deform */
+		void AddBoneToDeform(const FBoneName& BoneName);	
+
+		/** Add a Physics Body to deform */
+		void AddPhysicsBodyToDeform(const FBoneName& BoneName);
+
+        //-----------------------------------------------------------------------------------------
+		// Interface pattern
+		//-----------------------------------------------------------------------------------------
+		class Private;
+		Private* GetPrivate() const;
+
+	protected:
+
+		//! Forbidden. Manage with the Ptr<> template.
+		~NodeMeshMorph();
+
+	private:
+
+		Private* m_pD;
+
+	};
+
+
+}
diff --git a/Engine/Plugins/Mutable/Source/MutableTools/Public/MuT/NodeMeshReshape.h b/Engine/Plugins/Mutable/Source/MutableTools/Public/MuT/NodeMeshReshape.h
new file mode 100644
--- /dev/null
+++ b/Engine/Plugins/Mutable/Source/MutableTools/Public/MuT/NodeMeshReshape.h
@@ -0,0 +1,83 @@
+// Copyright Epic Games, Inc. All Rights Reserved.
+
+#pragma once
+
+#include "MuR/Ptr.h"
+#include "MuT/Node.h"
+#include "MuT/NodeMesh.h"
+#include "MuR/Mesh.h"
+
+namespace mu
+{
+	// Forward declarations
+	struct FBoneName;
+
+	//! Node that morphs a base mesh with one or two weighted targets from a sequence.
+	class MUTABLETOOLS_API NodeMeshReshape : public NodeMesh
+	{
+	public:
+
+		NodeMeshReshape();
+
+		//-----------------------------------------------------------------------------------------
+		// Node Interface
+		//-----------------------------------------------------------------------------------------
+
+		const FNodeType* GetType() const override;
+		static const FNodeType* GetStaticType();
+
+		//-----------------------------------------------------------------------------------------
+		// Own Interface
+		//-----------------------------------------------------------------------------------------
+
+		const NodeMeshPtr& GetBaseMesh() const;
+		void SetBaseMesh( const NodeMeshPtr& );
+
+		const NodeMeshPtr& GetBaseShape() const;
+		void SetBaseShape(const NodeMeshPtr&);
+
+		const NodeMeshPtr& GetTargetShape() const;
+		void SetTargetShape(const NodeMeshPtr&);
+
+		void SetReshapeVertices(bool);
+
+		/** Recompute normals after reshaping. This will ignore the reshaped normals. Disabled by default. */
+		void SetRecomputeNormals(bool);
+
+
+		/** Also deform the mesh skeleton. Disabled by default. */
+		void SetReshapeSkeleton(bool);
+	
+		/** Apply Laplacian smoothing to the reshaped mesh.  */
+		void SetApplyLaplacian(bool);
+		
+		/** Set vertex color channel usages for Reshape operations. */
+		void SetColorUsages(EVertexColorUsage R, EVertexColorUsage G, EVertexColorUsage B, EVertexColorUsage A);
+
+		/** Deform Mesh Physics Volumes */
+		void SetReshapePhysicsVolumes(bool);
+
+		/** Sets the number of bones that will be deform */
+		void AddBoneToDeform(const FBoneName& BoneId);
+	
+		/** Add a Physics Body to deform */
+		void AddPhysicsBodyToDeform(const FBoneName& BoneId);
+        
+		//-----------------------------------------------------------------------------------------
+		// Interface pattern
+		//-----------------------------------------------------------------------------------------
+		class Private;
+		Private* GetPrivate() const;
+
+	protected:
+
+		//! Forbidden. Manage with the Ptr<> template.
+		~NodeMeshReshape();
+
+	private:
+
+		Private* m_pD;
+
+	};
+
+}
diff --git a/Engine/Plugins/Mutable/Source/MutableTools/Public/MuT/NodeMeshSwitch.h b/Engine/Plugins/Mutable/Source/MutableTools/Public/MuT/NodeMeshSwitch.h
new file mode 100644
--- /dev/null
+++ b/Engine/Plugins/Mutable/Source/MutableTools/Public/MuT/NodeMeshSwitch.h
@@ -0,0 +1,73 @@
+// Copyright Epic Games, Inc. All Rights Reserved.
+
+#pragma once
+
+#include "MuR/Ptr.h"
+#include "MuR/RefCounted.h"
+#include "MuT/Node.h"
+#include "MuT/NodeMesh.h"
+
+
+namespace mu
+{
+
+	// Forward definitions
+	class NodeScalar;
+	typedef Ptr<NodeScalar> NodeScalarPtr;
+	typedef Ptr<const NodeScalar> NodeScalarPtrConst;
+
+	class NodeMeshSwitch;
+	typedef Ptr<NodeMeshSwitch> NodeMeshSwitchPtr;
+	typedef Ptr<const NodeMeshSwitch> NodeMeshSwitchPtrConst;
+
+
+	//! This node selects an output Mesh from a set of input Meshs based on a parameter.
+	class MUTABLETOOLS_API NodeMeshSwitch : public NodeMesh
+	{
+	public:
+
+		NodeMeshSwitch();
+
+		//-----------------------------------------------------------------------------------------
+		// Node Interface
+		//-----------------------------------------------------------------------------------------
+
+		const FNodeType* GetType() const override;
+		static const FNodeType* GetStaticType();
+
+		//-----------------------------------------------------------------------------------------
+		// Own Interface
+		//-----------------------------------------------------------------------------------------
+
+		//! Get the node generating the parameter used to select the option.
+		NodeScalarPtr GetParameter() const;
+		void SetParameter( NodeScalarPtr );
+
+		//! Set the number of option Meshs. It will keep the currently set targets and initialise
+		//! the new ones as null.
+		void SetOptionCount( int );
+
+		//! Get the node generating the t-th option Mesh.
+		NodeMeshPtr GetOption( int t ) const;
+		void SetOption( int t, NodeMeshPtr );
+
+
+		//-----------------------------------------------------------------------------------------
+		// Interface pattern
+		//-----------------------------------------------------------------------------------------
+		class Private;
+		Private* GetPrivate() const;
+
+	protected:
+
+		//! Forbidden. Manage with the Ptr<> template.
+		~NodeMeshSwitch();
+
+	private:
+
+		Private* m_pD;
+
+	};
+
+
+}
diff --git a/Engine/Plugins/Mutable/Source/MutableTools/Public/MuT/NodeMeshTable.h b/Engine/Plugins/Mutable/Source/MutableTools/Public/MuT/NodeMeshTable.h
new file mode 100644
--- /dev/null
+++ b/Engine/Plugins/Mutable/Source/MutableTools/Public/MuT/NodeMeshTable.h
@@ -0,0 +1,78 @@
+// Copyright Epic Games, Inc. All Rights Reserved.
+
+#pragma once
+
+#include "MuR/Ptr.h"
+#include "MuR/RefCounted.h"
+#include "MuT/Node.h"
+#include "MuT/NodeMesh.h"
+#include "MuT/NodeLayout.h"
+#include "MuT/Table.h"
+
+
+namespace mu
+{
+
+	// Forward definitions
+	class NodeMeshTable;
+	typedef Ptr<NodeMeshTable> NodeMeshTablePtr;
+	typedef Ptr<const NodeMeshTable> NodeMeshTablePtrConst;
+
+
+	//! This node provides the meshes stored in the column of a table.
+	class MUTABLETOOLS_API NodeMeshTable : public NodeMesh
+	{
+	public:
+
+		FString ParameterName;
+		Ptr<Table> Table;
+		FString ColumnName;
+		bool bNoneOption = false;
+		FString DefaultRowName;
+
+		TArray<Ptr<NodeLayout>> Layouts;
+
+	public:
+
+		// Node interface
+		virtual const FNodeType* GetType() const override { return GetStaticType(); }
+		static const FNodeType* GetStaticType() { return &StaticType; }
+
+		//-----------------------------------------------------------------------------------------
+		// Own Interface
+		//-----------------------------------------------------------------------------------------
+
+		//! Set the name of the implicit table parameter.
+		void SetParameterName( const FString& strName );
+
+		//!
+		void SetColumn( const FString& strName );
+
+		//! Get the number of layouts defined in the meshes on this column.
+		int GetLayoutCount() const;
+		void SetLayoutCount( int );
+
+		//! Get the node defining a layout of the meshes on this column.
+		Ptr<NodeLayout> GetLayout( int index ) const;
+		void SetLayout( int index, Ptr<NodeLayout>);
+
+		//! Adds the "None" option to the parameter that represents this table column
+		void SetNoneOption(bool bAddOption);
+
+		//! Set the row name to be used as default value
+		void SetDefaultRowName(const FString& RowName);
+
+
+	protected:
+
+		//! Forbidden. Manage with the Ptr<> template.
+		~NodeMeshTable() {}
+
+	private:
+
+		static FNodeType StaticType;
+
+	};
+
+
+}
diff --git a/Engine/Plugins/Mutable/Source/MutableTools/Public/MuT/NodeMeshTangents.h b/Engine/Plugins/Mutable/Source/MutableTools/Public/MuT/NodeMeshTangents.h
new file mode 100644
--- /dev/null
+++ b/Engine/Plugins/Mutable/Source/MutableTools/Public/MuT/NodeMeshTangents.h
@@ -0,0 +1,60 @@
+// Copyright Epic Games, Inc. All Rights Reserved.
+
+#pragma once
+
+#include "MuR/Ptr.h"
+#include "MuR/RefCounted.h"
+#include "MuT/Node.h"
+#include "MuT/NodeMesh.h"
+
+
+namespace mu
+{
+
+	// Forward definitions
+	class NodeMeshTangents;
+	typedef Ptr<NodeMeshTangents> NodeMeshTangentsPtr;
+	typedef Ptr<const NodeMeshTangents> NodeMeshTangentsPtrConst;
+
+
+	//! This node rebuilds the tangents and binormals of the source mesh.
+	class MUTABLETOOLS_API NodeMeshTangents : public NodeMesh
+	{
+	public:
+
+		NodeMeshTangents();
+
+		//-----------------------------------------------------------------------------------------
+		// Node Interface
+		//-----------------------------------------------------------------------------------------
+
+        const FNodeType* GetType() const override;
+		static const FNodeType* GetStaticType();
+
+		//-----------------------------------------------------------------------------------------
+		// Own Interface
+		//-----------------------------------------------------------------------------------------
+
+		//!
+		NodeMeshPtr GetSource() const;
+		void SetSource( NodeMesh* p );
+
+		//-----------------------------------------------------------------------------------------
+		// Interface pattern
+		//-----------------------------------------------------------------------------------------
+		class Private;
+		Private* GetPrivate() const;
+
+	protected:
+
+		//! Forbidden. Manage with the Ptr<> template.
+		~NodeMeshTangents();
+
+	private:
+
+		Private* m_pD;
+
+	};
+
+
+}
diff --git a/Engine/Plugins/Mutable/Source/MutableTools/Public/MuT/NodeMeshTransform.h b/Engine/Plugins/Mutable/Source/MutableTools/Public/MuT/NodeMeshTransform.h
new file mode 100644
--- /dev/null
+++ b/Engine/Plugins/Mutable/Source/MutableTools/Public/MuT/NodeMeshTransform.h
@@ -0,0 +1,63 @@
+// Copyright Epic Games, Inc. All Rights Reserved.
+
+#pragma once
+
+#include "MuR/Ptr.h"
+#include "MuR/RefCounted.h"
+#include "MuT/Node.h"
+#include "MuT/NodeMesh.h"
+#include "Math/MathFwd.h"
+
+namespace mu
+{
+
+	// Forward definitions
+    class NodeMeshTransform;
+    typedef Ptr<NodeMeshTransform> NodeMeshTransformPtr;
+    typedef Ptr<const NodeMeshTransform> NodeMeshTransformPtrConst;
+
+
+    //! This node applies a geometric transform represented by a 4x4 matrix to a mesh
+    class MUTABLETOOLS_API NodeMeshTransform : public NodeMesh
+	{
+	public:
+
+        NodeMeshTransform();
+
+		//-----------------------------------------------------------------------------------------
+		// Node Interface
+		//-----------------------------------------------------------------------------------------
+
+        const FNodeType* GetType() const override;
+		static const FNodeType* GetStaticType();
+
+		//-----------------------------------------------------------------------------------------
+		// Own Interface
+		//-----------------------------------------------------------------------------------------
+
+		//! Source mesh to be re-formatted
+        NodeMeshPtr GetSource() const;
+		void SetSource( NodeMesh* );
+
+        void SetTransform(const FMatrix44f&);
+		const FMatrix44f& GetTransform() const;
+
+		//-----------------------------------------------------------------------------------------
+		// Interface pattern
+		//-----------------------------------------------------------------------------------------
+		class Private;
+		Private* GetPrivate() const;
+
+	protected:
+
+		//! Forbidden. Manage with the Ptr<> template.
+        ~NodeMeshTransform();
+
+	private:
+
+		Private* m_pD;
+
+	};
+
+
+}
diff --git a/Engine/Plugins/Mutable/Source/MutableTools/Public/MuT/NodeMeshVariation.h b/Engine/Plugins/Mutable/Source/MutableTools/Public/MuT/NodeMeshVariation.h
new file mode 100644
--- /dev/null
+++ b/Engine/Plugins/Mutable/Source/MutableTools/Public/MuT/NodeMeshVariation.h
@@ -0,0 +1,71 @@
+// Copyright Epic Games, Inc. All Rights Reserved.
+
+#pragma once
+
+#include "MuR/Ptr.h"
+#include "MuR/RefCounted.h"
+#include "MuT/Node.h"
+#include "MuT/NodeMesh.h"
+
+
+namespace mu
+{
+
+    // Forward definitions
+    class NodeMeshVariation;
+    typedef Ptr<NodeMeshVariation> NodeMeshVariationPtr;
+    typedef Ptr<const NodeMeshVariation> NodeMeshVariationPtrConst;
+
+    //!
+    //! \ingroup model
+    class MUTABLETOOLS_API NodeMeshVariation : public NodeMesh
+    {
+    public:
+
+        NodeMeshVariation();
+
+        //-----------------------------------------------------------------------------------------
+        // Node Interface
+        //-----------------------------------------------------------------------------------------
+
+        const FNodeType* GetType() const override;
+        static const FNodeType* GetStaticType();
+
+        //-----------------------------------------------------------------------------------------
+        // Own Interface
+        //-----------------------------------------------------------------------------------------
+
+        //!
+        void SetDefaultMesh( NodeMesh* Mesh );
+
+        //! Set the number of tags to consider in this variation
+        void SetVariationCount( int count );
+
+        //!
+        int GetVariationCount() const;
+
+        //! Set the tag or state name that will enable a specific vartiation
+        void SetVariationTag( int index, const char* strTag );
+
+        //!
+        void SetVariationMesh( int index, NodeMesh* Mesh );
+
+        //!}
+
+
+        //-----------------------------------------------------------------------------------------
+        // Interface pattern
+        //-----------------------------------------------------------------------------------------
+        class Private;
+        Private* GetPrivate() const;
+
+    protected:
+        //! Forbidden. Manage with the Ptr<> template.
+        ~NodeMeshVariation();
+
+    private:
+        Private* m_pD;
+    };
+
+
+} // namespace mu
diff --git a/Engine/Plugins/Mutable/Source/MutableTools/Public/MuT/NodeModifierMeshClipMorphPlane.h b/Engine/Plugins/Mutable/Source/MutableTools/Public/MuT/NodeModifierMeshClipMorphPlane.h
new file mode 100644
--- /dev/null
+++ b/Engine/Plugins/Mutable/Source/MutableTools/Public/MuT/NodeModifierMeshClipMorphPlane.h
@@ -0,0 +1,100 @@
+// Copyright Epic Games, Inc. All Rights Reserved.
+
+#pragma once
+
+#include "MuR/Ptr.h"
+#include "MuR/RefCounted.h"
+#include "MuR/Skeleton.h"
+#include "MuT/Node.h"
+#include "MuT/NodeModifier.h"
+
+
+namespace mu
+{
+
+	/** */
+
+	struct FBoneName;
+
+	struct FClipMorphPlaneParameters
+	{
+		// Morph field parameters
+
+		 //! Distance to the plane of last affected vertex
+		float DistanceToPlane = 0;
+
+		//! "Linearity" factor of the influence.
+		float LinearityFactor = 0;
+
+		// Ellipse location
+		FVector3f Origin = FVector3f(0, 0, 0);
+		FVector3f Normal = FVector3f(0, 0, 0);
+		float Radius1 = 0, Radius2 = 0, Rotation = 0;
+
+		//! Typed of vertex selection
+		typedef enum
+		{
+			//! All vertices, so no extra info is needed
+			VS_ALL,
+
+			//! Select vertices inside a shape
+			VS_SHAPE,
+
+			//! Select all vertices affected by any bone in a sub hierarchy
+			VS_BONE_HIERARCHY,
+		} EVertexSelection;
+
+		// Vertex selection box
+		uint8 VertexSelectionType = VS_ALL;
+		FVector3f SelectionBoxOrigin = FVector3f(0, 0, 0);
+		FVector3f SelectionBoxRadius = FVector3f(0, 0, 0);
+		FBoneName VertexSelectionBone;
+
+		// Max distance a vertex can have to the bone in order to be affected. A negative value
+		// means no limit.
+		float MaxEffectRadius = -1.0f;
+	};
+
+
+	/** */
+	class MUTABLETOOLS_API NodeModifierMeshClipMorphPlane : public NodeModifier
+	{
+	public:
+
+		FClipMorphPlaneParameters Parameters;
+
+	public:
+
+		// Node interface
+		virtual const FNodeType* GetType() const override { return &StaticType; }
+		static const FNodeType* GetStaticType() { return &StaticType; }
+
+		//-----------------------------------------------------------------------------------------
+        // Own interface
+		//-----------------------------------------------------------------------------------------
+
+		void SetPlane(FVector3f Center, FVector3f Normal);
+		void SetParams(float dist, float factor);
+		void SetMorphEllipse(float radius1, float radius2, float rotation);
+
+		//! Define an axis-aligned box that will select the vertices to be morphed.
+		//! Only one of Box or Bone Hierarchy can be used (the last one set)
+		void SetVertexSelectionBox(float centerX, float centerY, float centerZ, float radiusX, float radiusY, float radiusZ);
+
+		//! Define the root bone of the subhierarchy of the mesh that will be affected.
+		//! Only one of Box or Bone Hierarchy can be used (the last one set)
+		void SetVertexSelectionBone(const FBoneName& BoneId, float maxEffectRadius);
+
+	protected:
+
+		//! Forbidden. Manage with the Ptr<> template.
+		~NodeModifierMeshClipMorphPlane() {}
+
+	private:
+
+		static FNodeType StaticType;
+
+	};
+
+
+}
diff --git a/Engine/Plugins/Mutable/Source/MutableTools/Public/MuT/NodeModifierMeshClipWithMesh.h b/Engine/Plugins/Mutable/Source/MutableTools/Public/MuT/NodeModifierMeshClipWithMesh.h
new file mode 100644
--- /dev/null
+++ b/Engine/Plugins/Mutable/Source/MutableTools/Public/MuT/NodeModifierMeshClipWithMesh.h
@@ -0,0 +1,39 @@
+// Copyright Epic Games, Inc. All Rights Reserved.
+
+#pragma once
+
+#include "MuR/Ptr.h"
+#include "MuR/RefCounted.h"
+#include "MuT/Node.h"
+#include "MuT/NodeModifier.h"
+#include "MuT/NodeMesh.h"
+
+namespace mu
+{
+
+	/** */
+	class MUTABLETOOLS_API NodeModifierMeshClipWithMesh : public NodeModifier
+	{
+	public:
+
+		//! 
+		Ptr<NodeMesh> ClipMesh;
+
+	public:
+
+		// Node interface
+		virtual const FNodeType* GetType() const override { return &StaticType; }
+		static const FNodeType* GetStaticType() { return &StaticType; }
+
+	protected:
+
+		//! Forbidden. Manage with the Ptr<> template.
+		inline ~NodeModifierMeshClipWithMesh() {}
+
+	private:
+
+		static FNodeType StaticType;
+
+	};
+
+}
diff --git a/Engine/Plugins/Mutable/Source/MutableTools/Public/MuT/NodeModifierMeshClipWithUVMask.h b/Engine/Plugins/Mutable/Source/MutableTools/Public/MuT/NodeModifierMeshClipWithUVMask.h
new file mode 100644
--- /dev/null
+++ b/Engine/Plugins/Mutable/Source/MutableTools/Public/MuT/NodeModifierMeshClipWithUVMask.h
@@ -0,0 +1,48 @@
+// Copyright Epic Games, Inc. All Rights Reserved.
+
+#pragma once
+
+#include "MuR/Ptr.h"
+#include "MuR/RefCounted.h"
+#include "MuT/Node.h"
+#include "MuT/NodeModifier.h"
+#include "MuT/NodeImage.h"
+#include "MuT/NodeLayout.h"
+
+
+namespace mu
+{
+
+	/** This node can clip part of a mesh using UV space-data (an image or a list of layout blocks). */
+	class MUTABLETOOLS_API NodeModifierMeshClipWithUVMask : public NodeModifier
+	{
+	public:
+
+		/** Image with the regions to remove. It will be interpreted as a bitmap. */
+		Ptr<NodeImage> ClipMask;
+
+		/** If the image above is null, clipping may hapopen with layout blocks. */
+		Ptr<NodeLayout> ClipLayout;
+
+		/** Layout index of the UVs to use inthe source mesh to ben clipped with the mask. */
+		uint8 LayoutIndex = 0;
+
+	public:
+
+		// Node interface
+		virtual const FNodeType* GetType() const override { return &StaticType; }
+		static const FNodeType* GetStaticType() { return &StaticType; }
+
+	protected:
+
+		//! Forbidden. Manage with the Ptr<> template.
+		inline ~NodeModifierMeshClipWithUVMask() {}
+
+	private:
+
+		static FNodeType StaticType;
+
+	};
+
+
+}
diff --git a/Engine/Plugins/Mutable/Source/MutableTools/Public/MuT/NodeModifierSurfaceEdit.h b/Engine/Plugins/Mutable/Source/MutableTools/Public/MuT/NodeModifierSurfaceEdit.h
new file mode 100644
--- /dev/null
+++ b/Engine/Plugins/Mutable/Source/MutableTools/Public/MuT/NodeModifierSurfaceEdit.h
@@ -0,0 +1,89 @@
+// Copyright Epic Games, Inc. All Rights Reserved.
+
+#pragma once
+
+#include "MuT/NodeModifier.h"
+#include "MuT/NodeMesh.h"
+#include "MuT/NodeImage.h"
+#include "MuT/NodeScalar.h"
+#include "MuR/ImageTypes.h"
+#include "MuR/RefCounted.h"
+#include "MuR/Ptr.h"
+
+
+namespace mu
+{
+
+	/** This node modifies a surface node.
+    * It allows to extend, cut and morph the parent Surface's meshes.
+    * It also allows to patch the parent Surface's textures.
+	*/
+    class MUTABLETOOLS_API NodeModifierSurfaceEdit : public NodeModifier
+	{
+	public:
+
+		/** Mesh to remove from the modified surface. */
+		Ptr<NodeMesh> MeshRemove;
+
+		/** Mesh to add to the modified surface. */
+		Ptr<NodeMesh> MeshAdd;
+
+		/** Name of the morph to apply to the modified surface if it has it. */
+		FString MeshMorph;
+
+		/** Factor of the morph to apply. */
+		Ptr<NodeScalar> MorphFactor;
+
+		/** Data for every modified texture. */
+		struct FTexture
+		{
+			/** Image to add if extgending. */
+			Ptr<NodeImage> Extend;
+
+			/** Image to blend if patching. */
+			Ptr<NodeImage> PatchImage;
+
+			/** Optional mask controlling the blending area. */
+			Ptr<NodeImage> PatchMask;
+
+			/** Rects in the parent layout homogeneous UV space to patch. */
+			TArray<FBox2f> PatchBlocks;
+
+			/** */
+			EBlendType PatchBlendType = EBlendType::BT_BLEND;
+
+			/** Patch alpha channel as well? */
+			bool bPatchApplyToAlpha = false;
+
+		};
+
+		/** Textures to modify. */
+		TArray<FTexture> Textures;
+
+	public:
+
+		NodeModifierSurfaceEdit() 
+		{
+			// This modifier needs to be applied at the end of the operations.
+			bApplyBeforeNormalOperations = false;
+		}
+
+		// Node interface
+		virtual const FNodeType* GetType() const override { return &StaticType; }
+		static const FNodeType* GetStaticType() { return &StaticType; }
+
+	protected:
+
+		//! Forbidden. Manage with the Ptr<> template.
+		~NodeModifierSurfaceEdit() {}
+
+	private:
+
+		static FNodeType StaticType;
+
+	};
+
+
+
+}
+
diff --git a/Engine/Plugins/Mutable/Source/MutableTools/Public/MuT/NodeObject.h b/Engine/Plugins/Mutable/Source/MutableTools/Public/MuT/NodeObject.h
new file mode 100644
--- /dev/null
+++ b/Engine/Plugins/Mutable/Source/MutableTools/Public/MuT/NodeObject.h
@@ -0,0 +1,62 @@
+// Copyright Epic Games, Inc. All Rights Reserved.
+
+#pragma once
+
+#include "HAL/Platform.h"
+#include "MuR/Ptr.h"
+#include "MuR/RefCounted.h"
+#include "MuT/Node.h"
+
+
+namespace mu
+{
+
+	// Forward definitions
+	class NodeObject;
+	typedef Ptr<NodeObject> NodeObjectPtr;
+	typedef Ptr<const NodeObject> NodeObjectPtrConst;
+
+
+    //! %Base class for any node that outputs an object.
+	//! \ingroup model
+	class MUTABLETOOLS_API NodeObject : public Node
+	{
+	public:
+
+		//-----------------------------------------------------------------------------------------
+		// Node Interface
+		//-----------------------------------------------------------------------------------------
+
+		const FNodeType* GetType() const override;
+		static const FNodeType* GetStaticType();
+
+		//-----------------------------------------------------------------------------------------
+		// Own interface
+		//-----------------------------------------------------------------------------------------
+
+		//! Get the name of the object.
+		virtual const FString& GetName() const = 0;
+
+		//! Set the name of the object.
+		virtual void SetName( const FString& ) = 0;
+
+		//! Get the uid of the object.
+		virtual const FString& GetUid() const = 0;
+
+		//! Set the uid of the object.
+		virtual void SetUid( const FString& ) = 0;
+
+		//-----------------------------------------------------------------------------------------
+		// Interface pattern
+		//-----------------------------------------------------------------------------------------
+		class Private;
+
+	protected:
+
+		//! Forbidden. Manage with the Ptr<> template.
+		inline ~NodeObject() {}
+
+	};
+
+
+}
diff --git a/Engine/Plugins/Mutable/Source/MutableTools/Public/MuT/NodeObjectGroup.h b/Engine/Plugins/Mutable/Source/MutableTools/Public/MuT/NodeObjectGroup.h
new file mode 100644
--- /dev/null
+++ b/Engine/Plugins/Mutable/Source/MutableTools/Public/MuT/NodeObjectGroup.h
@@ -0,0 +1,110 @@
+// Copyright Epic Games, Inc. All Rights Reserved.
+
+#pragma once
+
+#include "MuR/Ptr.h"
+#include "MuR/RefCounted.h"
+#include "MuT/Node.h"
+#include "MuT/NodeObject.h"
+
+
+namespace mu
+{
+
+	// Forward definitions
+	class NodeObjectGroup;
+	typedef Ptr<NodeObjectGroup> NodeObjectGroupPtr;
+	typedef Ptr<const NodeObjectGroup>NodeObjectGroupPtrConst;
+
+
+	//! Node that creates a group of objects and describes how they are selected.
+	//! \ingroup model
+	class MUTABLETOOLS_API NodeObjectGroup : public NodeObject
+	{
+	public:
+
+		NodeObjectGroup();
+
+
+		//-----------------------------------------------------------------------------------------
+		// Node Interface
+		//-----------------------------------------------------------------------------------------
+        const FNodeType* GetType() const override;
+		static const FNodeType* GetStaticType();
+
+		//-----------------------------------------------------------------------------------------
+		// NodeObject Interface
+		//-----------------------------------------------------------------------------------------
+        virtual const FString& GetName() const override;
+		virtual void SetName( const FString&) override;
+		virtual const FString& GetUid() const override;
+		virtual void SetUid( const FString&) override;
+
+		//-----------------------------------------------------------------------------------------
+		// Own Interface
+		//-----------------------------------------------------------------------------------------
+
+		//! Typed of child selection
+		typedef enum
+		{
+			//! All objects in the group will always be enabled, and no parameter will be generated.
+			CS_ALWAYS_ALL,
+
+			//! Only one children may be selected, but it is allowed to have none.
+			//! An enumeration parameter will be generated but it may have a null value
+			CS_ONE_OR_NONE,
+
+			//! One and only one children has to be selected at all times
+			//! An enumeraation parameter will be generated and it cannot be null
+			CS_ALWAYS_ONE,
+
+			//! Each child in the group can be enabled or disabled individually.
+			//! A boolean parameter will be generated for every child.
+			CS_TOGGLE_EACH
+
+		} CHILD_SELECTION;
+
+		//! Get child selection type
+		CHILD_SELECTION GetSelectionType() const;
+
+		//! Set the child selection type
+		void SetSelectionType( CHILD_SELECTION );
+
+		//! Get the number of child objects
+		int GetChildCount() const;
+
+		//! Set the number of child objects
+		void SetChildCount( int );
+
+		//! Get a child object node
+		NodeObjectPtr GetChild( int index ) const;
+
+		//! Set a child object node
+		void SetChild( int index, NodeObjectPtr );
+
+		//! Set default value for CS_ONE_OR_NONE or CS_ALWAYS_ONE groups
+		//! the value is the index of the child option in the group
+		//! -1 is the value for the None option.
+		//! 0 is the first child whether or not the NONE option is present.
+		void SetDefaultValue(int32 Value);
+
+		//-----------------------------------------------------------------------------------------
+		// Interface pattern
+		//-----------------------------------------------------------------------------------------
+		class Private;
+		Private* GetPrivate() const;
+
+	protected:
+
+		//! Forbidden. Manage with the Ptr<> template.
+		~NodeObjectGroup();
+
+	private:
+
+		Private* m_pD;
+
+	};
+
+
+
+}
