[mutable] Moved the Mutable plugin out of Experimental status into Beta.

diff --git a/Engine/Plugins/Mutable/Source/MutableTools/Public/MuT/NodeObjectNew.h b/Engine/Plugins/Mutable/Source/MutableTools/Public/MuT/NodeObjectNew.h
new file mode 100644
--- /dev/null
+++ b/Engine/Plugins/Mutable/Source/MutableTools/Public/MuT/NodeObjectNew.h
@@ -0,0 +1,110 @@
+// Copyright Epic Games, Inc. All Rights Reserved.
+
+#pragma once
+
+#include "MuR/Ptr.h"
+#include "MuR/RefCounted.h"
+#include "MuR/System.h"
+#include "MuT/Node.h"
+#include "MuT/NodeObject.h"
+#include "MuT/NodeComponent.h"
+#include "MuT/NodeExtensionData.h"
+#include "MuT/Compiler.h"
+
+
+namespace mu
+{
+
+	//! Node that creates a new object by setting its levels-of-detail and its children.
+	class MUTABLETOOLS_API NodeObjectNew : public NodeObject
+	{
+	public:
+
+		/** Name of the object. */
+		FString Name;
+
+		/** Externally provided id for the object. */
+		FString Uid;
+
+		/** Components defined in the object. */
+		TArray<Ptr<NodeComponent>> Components;
+
+		/** Modifiers defined in the object. */
+		TArray<Ptr<NodeModifier>> Modifiers;
+
+		/** Children objects. */
+		TArray<Ptr<NodeObject>> Children;
+
+		/** Extension data attached to this object. */
+		struct FNamedExtensionDataNode
+		{
+			Ptr<NodeExtensionData> Node;
+			FString Name;
+		};
+		TArray<FNamedExtensionDataNode> ExtensionDataNodes;
+
+		/** States defined in this object. */
+		TArray<FObjectState> States;
+
+
+	public:
+
+		//-----------------------------------------------------------------------------------------
+		// Node Interface
+		//-----------------------------------------------------------------------------------------
+		virtual const FNodeType* GetType() const override { return GetStaticType(); }
+		static const FNodeType* GetStaticType() { return &StaticType; }
+
+		//-----------------------------------------------------------------------------------------
+		// NodeObject Interface
+		//-----------------------------------------------------------------------------------------
+        virtual const FString& GetName() const override;
+		virtual void SetName( const FString& ) override;
+		virtual const FString& GetUid() const override;
+		virtual void SetUid( const FString& ) override;
+
+		//-----------------------------------------------------------------------------------------
+		// Own Interface
+		//-----------------------------------------------------------------------------------------
+
+		//! Set the number of states that the model can be in.
+		int32 GetStateCount() const;
+		void SetStateCount( int32 c );
+
+		//! Set the name of a state
+		void SetStateName( int32 s, const FString& n );
+
+		//! See if a state has a parameter as runtime.
+		bool HasStateParam( int32 s, const FString& param ) const;
+
+		//! Add a runtime parameter to the state.
+		void AddStateParam( int32 s, const FString& param );
+
+		//! Remove a runtime parameter from a state.
+		void RemoveStateParam( int32 s, const FString& param );
+
+        //! Set the optimisation properties of a state
+		void SetStateProperties(int32 StateIndex,
+			ETextureCompressionStrategy TextureCompressionStrategy,
+			bool bOnlyFirstLOD,
+			uint8 FirstLOD,
+			uint8 NumExtraLODsToBuildAfterFirstLOD);
+
+		//! Connect a node that produces ExtensionData to be added to the final Instance, and provide a name to associate with the data
+		void AddExtensionDataNode(Ptr<NodeExtensionData> Node, const FString& Name);
+
+	protected:
+
+		//! Forbidden. Manage with the Ptr<> template.
+		~NodeObjectNew() {}
+
+	private:
+
+		static FNodeType StaticType;
+
+
+	};
+
+
+
+}
diff --git a/Engine/Plugins/Mutable/Source/MutableTools/Public/MuT/NodeProjector.h b/Engine/Plugins/Mutable/Source/MutableTools/Public/MuT/NodeProjector.h
new file mode 100644
--- /dev/null
+++ b/Engine/Plugins/Mutable/Source/MutableTools/Public/MuT/NodeProjector.h
@@ -0,0 +1,162 @@
+// Copyright Epic Games, Inc. All Rights Reserved.
+
+#pragma once
+
+#include "HAL/Platform.h"
+#include "MuR/Parameters.h"
+#include "MuR/Ptr.h"
+#include "MuR/RefCounted.h"
+#include "MuT/Node.h"
+
+
+namespace mu
+{
+
+	class NodeProjector;
+    using NodeProjectorPtr = Ptr<NodeProjector>;
+    using NodeProjectorPtrConst = Ptr<const NodeProjector>;
+
+	class NodeProjectorConstant;
+    using NodeProjectorConstantPtr = Ptr<NodeProjectorConstant>;
+    using NodeProjectorConstantPtrConst = Ptr<const NodeProjectorConstant>;
+
+	class NodeProjectorParameter;
+    using NodeProjectorParameterPtr = Ptr<NodeProjectorParameter>;
+    using NodeProjectorParameterPtrConst = Ptr<const NodeProjectorParameter>;
+
+    class NodeRange;
+    using NodeRangePtr = Ptr<NodeRange>;
+    using NodeRangePtrConst = Ptr<const NodeRange>;
+
+
+    //! Base class of any node that outputs a Projector.
+	class MUTABLETOOLS_API NodeProjector : public Node
+	{
+	public:
+
+		//-----------------------------------------------------------------------------------------
+		// Node Interface
+		//-----------------------------------------------------------------------------------------
+
+        const FNodeType* GetType() const override;
+		static const FNodeType* GetStaticType();
+
+
+		//-----------------------------------------------------------------------------------------
+		// Interface pattern
+		//-----------------------------------------------------------------------------------------
+		class Private;
+
+	protected:
+
+		//! Forbidden. Manage with the Ptr<> template.
+		inline ~NodeProjector() {}
+
+	};
+
+
+	//! This node outputs a predefined Projector value.
+	class MUTABLETOOLS_API NodeProjectorConstant : public NodeProjector
+	{
+	public:
+
+		NodeProjectorConstant();
+
+		//-----------------------------------------------------------------------------------------
+		// Node Interface
+		//-----------------------------------------------------------------------------------------
+
+        const FNodeType* GetType() const override;
+		static const FNodeType* GetStaticType();
+
+		//-----------------------------------------------------------------------------------------
+		// Own Interface
+		//-----------------------------------------------------------------------------------------
+
+		//! Get the value that this node returns
+        void GetValue( PROJECTOR_TYPE* OutType,
+			FVector3f* OutPos,
+			FVector3f* OutDir,
+			FVector3f* OutUp,
+			FVector3f* OutScaleU,
+			float* OutProjectionAngle ) const;
+
+		//! Set the value to be returned by this node
+        void SetValue( PROJECTOR_TYPE type, FVector3f pos, FVector3f dir, FVector3f up, FVector3f scale, float projectionAngle);
+
+		//-----------------------------------------------------------------------------------------
+		// Interface pattern
+		//-----------------------------------------------------------------------------------------
+		class Private;
+		Private* GetPrivate() const;
+
+	protected:
+
+		//! Forbidden. Manage with the Ptr<> template.
+		~NodeProjectorConstant();
+
+	private:
+
+		Private* m_pD;
+
+	};
+
+
+	//! Node that defines a Projector model parameter.
+	class MUTABLETOOLS_API NodeProjectorParameter : public NodeProjector
+	{
+	public:
+
+		NodeProjectorParameter();
+
+		//-----------------------------------------------------------------------------------------
+		// Node Interface
+		//-----------------------------------------------------------------------------------------
+
+        const FNodeType* GetType() const override;
+		static const FNodeType* GetStaticType();
+
+		//-----------------------------------------------------------------------------------------
+		// Own Interface
+		//-----------------------------------------------------------------------------------------
+
+		//! Set the name of the parameter.
+		void SetName( const FString& );
+
+
+		//! Get the uid of the parameter. It will be exposed in the final compiled data.
+		const FString& GetUid() const;
+		void SetUid( const FString& );
+
+		//! Set the default value of the parameter.
+        void SetDefaultValue( PROJECTOR_TYPE type,
+			FVector3f pos,
+			FVector3f dir,
+			FVector3f up,
+			FVector3f scale,
+			float projectionAngle );
+
+        //! Set the number of ranges (dimensions) for this parameter.
+        //! By default a parameter has 0 ranges, meaning it only has one value.
+        void SetRangeCount( int i );
+        void SetRange( int i, NodeRangePtr pRange );
+
+		//-----------------------------------------------------------------------------------------
+		// Interface pattern
+		//-----------------------------------------------------------------------------------------
+		class Private;
+		Private* GetPrivate() const;
+
+	protected:
+
+		//! Forbidden. Manage with the Ptr<> template.
+		~NodeProjectorParameter();
+
+	private:
+
+		Private* m_pD;
+
+	};
+
+
+}
diff --git a/Engine/Plugins/Mutable/Source/MutableTools/Public/MuT/NodeRange.h b/Engine/Plugins/Mutable/Source/MutableTools/Public/MuT/NodeRange.h
new file mode 100644
--- /dev/null
+++ b/Engine/Plugins/Mutable/Source/MutableTools/Public/MuT/NodeRange.h
@@ -0,0 +1,47 @@
+// Copyright Epic Games, Inc. All Rights Reserved.
+
+#pragma once
+
+#include "HAL/Platform.h"
+#include "MuR/Ptr.h"
+#include "MuR/RefCounted.h"
+#include "MuT/Node.h"
+
+
+namespace mu
+{
+
+    // Forward definitions
+    class NodeRange;
+    typedef Ptr<NodeRange> NodeRangePtr;
+    typedef Ptr<const NodeRange> NodeRangePtrConst;
+
+
+    //! %Base class of any node that outputs a range.
+    class MUTABLETOOLS_API NodeRange : public Node
+	{
+	public:
+
+		//-----------------------------------------------------------------------------------------
+		// Node Interface
+		//-----------------------------------------------------------------------------------------
+
+		const FNodeType* GetType() const override;
+		static const FNodeType* GetStaticType();
+
+
+		//-----------------------------------------------------------------------------------------
+		// Interface pattern
+		//-----------------------------------------------------------------------------------------
+		class Private;
+
+	protected:
+
+		//! Forbidden. Manage with the Ptr<> template.
+        inline ~NodeRange() {}
+
+	};
+
+
+}
+
diff --git a/Engine/Plugins/Mutable/Source/MutableTools/Public/MuT/NodeRangeFromScalar.h b/Engine/Plugins/Mutable/Source/MutableTools/Public/MuT/NodeRangeFromScalar.h
new file mode 100644
--- /dev/null
+++ b/Engine/Plugins/Mutable/Source/MutableTools/Public/MuT/NodeRangeFromScalar.h
@@ -0,0 +1,64 @@
+// Copyright Epic Games, Inc. All Rights Reserved.
+
+#pragma once
+
+#include "MuR/Ptr.h"
+#include "MuT/Node.h"
+#include "MuT/NodeRange.h"
+
+namespace mu
+{
+
+	// Forward definitions
+    class NodeScalar;
+
+
+    class MUTABLETOOLS_API NodeRangeFromScalar : public NodeRange
+	{
+	public:
+
+		//-----------------------------------------------------------------------------------------
+		// Life cycle
+		//-----------------------------------------------------------------------------------------
+
+        NodeRangeFromScalar();
+
+
+		//-----------------------------------------------------------------------------------------
+		// Node Interface
+		//-----------------------------------------------------------------------------------------
+
+		const FNodeType* GetType() const override;
+		static const FNodeType* GetStaticType();
+
+		//-----------------------------------------------------------------------------------------
+		// Own Interface
+		//-----------------------------------------------------------------------------------------
+
+        //!
+        Ptr<NodeScalar> GetSize() const;
+        void SetSize( const Ptr<NodeScalar>& );
+
+        //!
+		const FString& GetName() const;
+		void SetName(const FString& Name);
+
+		//-----------------------------------------------------------------------------------------
+		// Interface pattern
+		//-----------------------------------------------------------------------------------------
+		class Private;
+		Private* GetPrivate() const;
+
+	protected:
+
+		//! Forbidden. Manage with the Ptr<> template.
+        ~NodeRangeFromScalar();
+
+	private:
+
+		Private* m_pD;
+
+	};
+
+
+}
diff --git a/Engine/Plugins/Mutable/Source/MutableTools/Public/MuT/NodeScalar.h b/Engine/Plugins/Mutable/Source/MutableTools/Public/MuT/NodeScalar.h
new file mode 100644
--- /dev/null
+++ b/Engine/Plugins/Mutable/Source/MutableTools/Public/MuT/NodeScalar.h
@@ -0,0 +1,46 @@
+// Copyright Epic Games, Inc. All Rights Reserved.
+
+#pragma once
+
+#include "HAL/Platform.h"
+#include "MuR/Ptr.h"
+#include "MuR/RefCounted.h"
+#include "MuT/Node.h"
+
+
+namespace mu
+{
+
+	// Forward definitions
+	class NodeScalar;
+	typedef Ptr<NodeScalar> NodeScalarPtr;
+	typedef Ptr<const NodeScalar> NodeScalarPtrConst;
+
+
+    //! %Base class of any node that outputs a scalar value.
+	class MUTABLETOOLS_API NodeScalar : public Node
+	{
+	public:
+
+		//-----------------------------------------------------------------------------------------
+		// Node Interface
+		//-----------------------------------------------------------------------------------------
+
+		const FNodeType* GetType() const override;
+		static const FNodeType* GetStaticType();
+
+
+		//-----------------------------------------------------------------------------------------
+		// Interface pattern
+		//-----------------------------------------------------------------------------------------
+		class Private;
+
+	protected:
+
+		//! Forbidden. Manage with the Ptr<> template.
+		inline ~NodeScalar() {}
+
+	};
+
+
+}
diff --git a/Engine/Plugins/Mutable/Source/MutableTools/Public/MuT/NodeScalarArithmeticOperation.h b/Engine/Plugins/Mutable/Source/MutableTools/Public/MuT/NodeScalarArithmeticOperation.h
new file mode 100644
--- /dev/null
+++ b/Engine/Plugins/Mutable/Source/MutableTools/Public/MuT/NodeScalarArithmeticOperation.h
@@ -0,0 +1,86 @@
+// Copyright Epic Games, Inc. All Rights Reserved.
+
+#pragma once
+
+#include "MuR/Ptr.h"
+#include "MuR/RefCounted.h"
+#include "MuT/Node.h"
+#include "MuT/NodeScalar.h"
+
+
+namespace mu
+{
+
+	// Forward definitions
+    class NodeScalarArithmeticOperation;
+    typedef Ptr<NodeScalarArithmeticOperation> NodeScalarArithmeticOperationPtr;
+    typedef Ptr<const NodeScalarArithmeticOperation> NodeScalarArithmeticOperationPtrConst;
+
+
+    //! Perform an arithmetic operation between two scalars.
+	//! \ingroup model
+    class MUTABLETOOLS_API NodeScalarArithmeticOperation : public NodeScalar
+	{
+	public:
+
+        NodeScalarArithmeticOperation();
+
+
+		//-----------------------------------------------------------------------------------------
+		// Node Interface
+		//-----------------------------------------------------------------------------------------
+
+        const FNodeType* GetType() const override;
+		static const FNodeType* GetStaticType();
+
+
+		//-----------------------------------------------------------------------------------------
+		// Own Interface
+		//-----------------------------------------------------------------------------------------
+
+		//! Possible arithmetic operations
+		typedef enum
+		{
+            AO_ADD = 0,
+			AO_SUBTRACT,
+			AO_MULTIPLY,
+            AO_DIVIDE,
+            _AO_COUNT
+		} OPERATION;
+
+        static const char* s_opTypeName[ _AO_COUNT ];
+
+		//! Get child selection type
+		OPERATION GetOperation() const;
+
+		//! Set the child selection type
+		void SetOperation(OPERATION);
+
+
+		//! Get the first operand
+        NodeScalarPtr GetA() const;
+        void SetA(NodeScalarPtr);
+
+		// Get the second operand
+        NodeScalarPtr GetB() const;
+        void SetB(NodeScalarPtr);
+
+		//-----------------------------------------------------------------------------------------
+		// Interface pattern
+		//-----------------------------------------------------------------------------------------
+		class Private;
+		Private* GetPrivate() const;
+
+	protected:
+
+		//! Forbidden. Manage with the Ptr<> template.
+        ~NodeScalarArithmeticOperation();
+
+	private:
+
+		Private* m_pD;
+
+	};
+
+
+}
diff --git a/Engine/Plugins/Mutable/Source/MutableTools/Public/MuT/NodeScalarConstant.h b/Engine/Plugins/Mutable/Source/MutableTools/Public/MuT/NodeScalarConstant.h
new file mode 100644
--- /dev/null
+++ b/Engine/Plugins/Mutable/Source/MutableTools/Public/MuT/NodeScalarConstant.h
@@ -0,0 +1,65 @@
+// Copyright Epic Games, Inc. All Rights Reserved.
+
+#pragma once
+
+#include "MuR/Ptr.h"
+#include "MuR/RefCounted.h"
+#include "MuT/Node.h"
+#include "MuT/NodeScalar.h"
+
+
+namespace mu
+{
+
+	// Forward definitions
+	class NodeScalarConstant;
+	typedef Ptr<NodeScalarConstant> NodeScalarConstantPtr;
+	typedef Ptr<const NodeScalarConstant> NodeScalarConstantPtrConst;
+
+
+	//! Node returning a scalar constant value.
+	//! \ingroup model
+	class MUTABLETOOLS_API NodeScalarConstant : public NodeScalar
+	{
+	public:
+
+		NodeScalarConstant();
+
+
+		//-----------------------------------------------------------------------------------------
+		// Node Interface
+		//-----------------------------------------------------------------------------------------
+
+        const FNodeType* GetType() const override;
+		static const FNodeType* GetStaticType();
+
+		//-----------------------------------------------------------------------------------------
+		// Own Interface
+		//-----------------------------------------------------------------------------------------
+
+		//! Get the value to be returned by the node.
+		float GetValue() const;
+
+		//! Set the value to be returned by the node.
+		void SetValue( float v );
+
+		//-----------------------------------------------------------------------------------------
+		// Interface pattern
+		//-----------------------------------------------------------------------------------------
+		class Private;
+		Private* GetPrivate() const;
+
+	protected:
+
+		//! Forbidden. Manage with the Ptr<> template.
+		~NodeScalarConstant();
+
+	private:
+
+		Private* m_pD;
+
+	};
+
+
+
+}
diff --git a/Engine/Plugins/Mutable/Source/MutableTools/Public/MuT/NodeScalarCurve.h b/Engine/Plugins/Mutable/Source/MutableTools/Public/MuT/NodeScalarCurve.h
new file mode 100644
--- /dev/null
+++ b/Engine/Plugins/Mutable/Source/MutableTools/Public/MuT/NodeScalarCurve.h
@@ -0,0 +1,56 @@
+// Copyright Epic Games, Inc. All Rights Reserved.
+
+#pragma once
+
+#include "MuR/Ptr.h"
+#include "MuR/RefCounted.h"
+#include "MuT/Node.h"
+#include "MuT/NodeScalar.h"
+
+#include "Curves/RichCurve.h"
+
+namespace mu
+{
+
+    /** This node makes a new scalar value transforming another scalar value with a curve. */
+    class MUTABLETOOLS_API NodeScalarCurve : public NodeScalar
+	{
+	public:
+
+		FRichCurve Curve;
+		Ptr<NodeScalar> CurveSampleValue;
+
+	public:
+
+        NodeScalarCurve();
+
+		//-----------------------------------------------------------------------------------------
+		// Node Interface
+		//-----------------------------------------------------------------------------------------
+
+        const FNodeType* GetType() const override;
+		static const FNodeType* GetStaticType();
+
+		//-----------------------------------------------------------------------------------------
+		// Own Interface
+		//-----------------------------------------------------------------------------------------
+
+		//-----------------------------------------------------------------------------------------
+		// Interface pattern
+		//-----------------------------------------------------------------------------------------
+		class Private;
+		Private* GetPrivate() const;
+
+	protected:
+
+		//! Forbidden. Manage with the Ptr<> template.
+        ~NodeScalarCurve();
+
+	private:
+
+		Private* m_pD;
+
+	};
+
+
+}
diff --git a/Engine/Plugins/Mutable/Source/MutableTools/Public/MuT/NodeScalarEnumParameter.h b/Engine/Plugins/Mutable/Source/MutableTools/Public/MuT/NodeScalarEnumParameter.h
new file mode 100644
--- /dev/null
+++ b/Engine/Plugins/Mutable/Source/MutableTools/Public/MuT/NodeScalarEnumParameter.h
@@ -0,0 +1,90 @@
+// Copyright Epic Games, Inc. All Rights Reserved.
+
+#pragma once
+
+#include "MuR/Ptr.h"
+#include "MuR/RefCounted.h"
+#include "MuT/Node.h"
+#include "MuT/NodeRange.h"
+#include "MuT/NodeScalar.h"
+
+
+namespace mu
+{
+
+	// Forward definitions
+	class NodeImage;
+    using NodeImagePtr = Ptr<NodeImage>;
+    using NodeImagePtrConst = Ptr<const NodeImage>;
+
+	class NodeScalarEnumParameter;
+    using NodeScalarEnumParameterPtr = Ptr<NodeScalarEnumParameter>;
+    using NodeScalarEnumParameterPtrConst = Ptr<const NodeScalarEnumParameter>;
+
+    class NodeRange;
+    using NodeRangePtr = Ptr<NodeRange>;
+    using NodeRangePtrConst = Ptr<const NodeRange>;
+
+
+	//! Node that defines a scalar model parameter to be selected from a set of named values.
+	class MUTABLETOOLS_API NodeScalarEnumParameter : public NodeScalar
+	{
+	public:
+
+		NodeScalarEnumParameter();
+
+		//-----------------------------------------------------------------------------------------
+		// Node Interface
+		//-----------------------------------------------------------------------------------------
+        const FNodeType* GetType() const override;
+		static const FNodeType* GetStaticType();
+
+		//-----------------------------------------------------------------------------------------
+		// Own Interface
+		//-----------------------------------------------------------------------------------------
+
+		//! Get the name of the parameter. It will be exposed in the final compiled data.
+		const FString& GetName() const;
+		void SetName( const FString&);
+
+		//! Get the uid of the parameter. It will be exposed in the final compiled data.
+		const FString& GetUid() const;
+		void SetUid( const FString&);
+
+		//! Get the index of the default value of the parameter.
+		int GetDefaultValueIndex() const;
+		void SetDefaultValueIndex( int v );
+
+		//! Set the number of possible values.
+		void SetValueCount( int i );
+
+		//! Get the number of possible values.
+		int GetValueCount() const;
+
+		//! Set the data of one of the possible values of the parameter.
+		void SetValue( int i, float value, const FString& strName );
+
+        //! Set the number of ranges (dimensions) for this parameter.
+        //! By default a parameter has 0 ranges, meaning it only has one value.
+        void SetRangeCount( int i );
+        void SetRange( int i, NodeRangePtr pRange );
+
+		//-----------------------------------------------------------------------------------------
+		// Interface pattern
+		//-----------------------------------------------------------------------------------------
+		class Private;
+		Private* GetPrivate() const;
+
+	protected:
+
+		//! Forbidden. Manage with the Ptr<> template.
+		~NodeScalarEnumParameter();
+
+	private:
+
+		Private* m_pD;
+
+	};
+
+
+}
diff --git a/Engine/Plugins/Mutable/Source/MutableTools/Public/MuT/NodeScalarParameter.h b/Engine/Plugins/Mutable/Source/MutableTools/Public/MuT/NodeScalarParameter.h
new file mode 100644
--- /dev/null
+++ b/Engine/Plugins/Mutable/Source/MutableTools/Public/MuT/NodeScalarParameter.h
@@ -0,0 +1,83 @@
+// Copyright Epic Games, Inc. All Rights Reserved.
+
+#pragma once
+
+#include "MuR/Parameters.h"
+#include "MuR/Ptr.h"
+#include "MuR/RefCounted.h"
+#include "MuT/Node.h"
+#include "MuT/NodeScalar.h"
+
+
+namespace mu
+{
+
+	// Forward definitions
+    class NodeImage;
+    using NodeImagePtr = Ptr<NodeImage>;
+    using NodeImagePtrConst = Ptr<const NodeImage>;
+
+    class NodeRange;
+    using NodeRangePtr = Ptr<NodeRange>;
+    using NodeRangePtrConst = Ptr<const NodeRange>;
+
+    class NodeScalarParameter;
+    using NodeScalarParameterPtr = Ptr<NodeScalarParameter>;
+    using NodeScalarParameterPtrConst = Ptr<const NodeScalarParameter>;
+
+
+	//! Node that defines a scalar model parameter.
+	class MUTABLETOOLS_API NodeScalarParameter : public NodeScalar
+	{
+	public:
+
+		//-----------------------------------------------------------------------------------------
+		// Life cycle
+		//-----------------------------------------------------------------------------------------
+
+		NodeScalarParameter();
+
+		//-----------------------------------------------------------------------------------------
+		// Node Interface
+		//-----------------------------------------------------------------------------------------
+
+        const FNodeType* GetType() const override;
+		static const FNodeType* GetStaticType();
+
+		//-----------------------------------------------------------------------------------------
+		// Own Interface
+		//-----------------------------------------------------------------------------------------
+
+		//! Get the name of the parameter. It will be exposed in the final compiled data.
+		void SetName( const FString& );
+
+		//! Get the uid of the parameter. It will be exposed in the final compiled data.
+		void SetUid( const FString&);
+
+		//! Get the default value of the parameter.
+		void SetDefaultValue( float v );
+
+        //! Set the number of ranges (dimensions) for this parameter.
+        //! By default a parameter has 0 ranges, meaning it only has one value.
+        void SetRangeCount( int i );
+        void SetRange( int i, NodeRangePtr pRange );
+
+		//-----------------------------------------------------------------------------------------
+		// Interface pattern
+		//-----------------------------------------------------------------------------------------
+		class Private;
+		Private* GetPrivate() const;
+
+	protected:
+
+		//! Forbidden. Manage with the Ptr<> template.
+		~NodeScalarParameter();
+
+	private:
+
+		Private* m_pD;
+
+	};
+
+
+}
diff --git a/Engine/Plugins/Mutable/Source/MutableTools/Public/MuT/NodeScalarTable.h b/Engine/Plugins/Mutable/Source/MutableTools/Public/MuT/NodeScalarTable.h
new file mode 100644
--- /dev/null
+++ b/Engine/Plugins/Mutable/Source/MutableTools/Public/MuT/NodeScalarTable.h
@@ -0,0 +1,67 @@
+// Copyright Epic Games, Inc. All Rights Reserved.
+
+#pragma once
+
+#include "MuR/Ptr.h"
+#include "MuR/RefCounted.h"
+#include "MuT/Node.h"
+#include "MuT/NodeScalar.h"
+#include "MuT/Table.h"
+
+
+namespace mu
+{
+
+	// Forward definitions
+	class NodeScalarTable;
+	typedef Ptr<NodeScalarTable> NodeScalarTablePtr;
+	typedef Ptr<const NodeScalarTable> NodeScalarTablePtrConst;
+
+
+	//! This node provides the meshes stored in the column of a table.
+	class MUTABLETOOLS_API NodeScalarTable : public NodeScalar
+	{
+	public:
+
+		FString ParameterName;
+		Ptr<Table> Table;
+		FString ColumnName;
+		bool bNoneOption = false;
+		FString DefaultRowName;
+
+	public:
+
+		// Node interface
+		virtual const FNodeType* GetType() const override { return GetStaticType(); }
+		static const FNodeType* GetStaticType() { return &StaticType; }
+
+		//-----------------------------------------------------------------------------------------
+		// Own Interface
+		//-----------------------------------------------------------------------------------------
+
+		//! Set the name of the implicit table parameter.
+		void SetParameterName( const FString& strName );
+
+		//!
+		void SetColumn( const FString& strName );
+
+		//! Adds the "None" option to the parameter that represents this table column
+		void SetNoneOption(bool bAddOption);
+
+		//! Set the row name to be used as default value
+		void SetDefaultRowName(const FString& RowName);
+
+
+	protected:
+
+		//! Forbidden. Manage with the Ptr<> template.
+		~NodeScalarTable() {}
+
+	private:
+
+		static FNodeType StaticType;
+
+	};
+
+
+}
diff --git a/Engine/Plugins/Mutable/Source/MutableTools/Public/MuT/NodeScalarVariation.h b/Engine/Plugins/Mutable/Source/MutableTools/Public/MuT/NodeScalarVariation.h
new file mode 100644
--- /dev/null
+++ b/Engine/Plugins/Mutable/Source/MutableTools/Public/MuT/NodeScalarVariation.h
@@ -0,0 +1,67 @@
+// Copyright Epic Games, Inc. All Rights Reserved.
+
+#pragma once
+
+#include "MuR/Ptr.h"
+#include "MuR/RefCounted.h"
+#include "MuT/Node.h"
+#include "MuT/NodeScalar.h"
+
+
+namespace mu
+{
+
+    // Forward definitions
+    class NodeScalarVariation;
+    typedef Ptr<NodeScalarVariation> NodeScalarVariationPtr;
+    typedef Ptr<const NodeScalarVariation> NodeScalarVariationPtrConst;
+
+
+     class MUTABLETOOLS_API NodeScalarVariation : public NodeScalar
+    {
+    public:
+ 
+        NodeScalarVariation();
+
+        //-----------------------------------------------------------------------------------------
+        // Node Interface
+        //-----------------------------------------------------------------------------------------
+
+        const FNodeType* GetType() const override;
+        static const FNodeType* GetStaticType();
+
+        //-----------------------------------------------------------------------------------------
+        // Own Interface
+        //-----------------------------------------------------------------------------------------
+
+        //!
+        void SetDefaultScalar( NodeScalar* Scalar );
+
+        //! Set the number of tags to consider in this variation
+        void SetVariationCount( int32 count );
+
+        //! Set the tag or state name that will enable a specific vartiation
+        void SetVariationTag( int32 index, const FString& strTag );
+
+        //!
+        void SetVariationScalar( int32 index, NodeScalar* Scalar );
+
+        //!}
+
+
+        //-----------------------------------------------------------------------------------------
+        // Interface pattern
+        //-----------------------------------------------------------------------------------------
+        class Private;
+        Private* GetPrivate() const;
+
+    protected:
+        //! Forbidden. Manage with the Ptr<> template.
+        ~NodeScalarVariation();
+
+    private:
+        Private* m_pD;
+    };
+
+
+} // namespace mu
diff --git a/Engine/Plugins/Mutable/Source/MutableTools/Public/MuT/NodeString.h b/Engine/Plugins/Mutable/Source/MutableTools/Public/MuT/NodeString.h
new file mode 100644
--- /dev/null
+++ b/Engine/Plugins/Mutable/Source/MutableTools/Public/MuT/NodeString.h
@@ -0,0 +1,46 @@
+// Copyright Epic Games, Inc. All Rights Reserved.
+
+#pragma once
+
+#include "HAL/Platform.h"
+#include "MuR/Ptr.h"
+#include "MuR/RefCounted.h"
+#include "MuT/Node.h"
+
+
+namespace mu
+{
+
+	// Forward definitions
+	class NodeString;
+	typedef Ptr<NodeString> NodeStringPtr;
+	typedef Ptr<const NodeString> NodeStringPtrConst;
+
+
+    //! %Base class of any node that outputs a string value.
+	class MUTABLETOOLS_API NodeString : public Node
+	{
+	public:
+
+		//-----------------------------------------------------------------------------------------
+		// Node Interface
+		//-----------------------------------------------------------------------------------------
+
+		const FNodeType* GetType() const override;
+		static const FNodeType* GetStaticType();
+
+
+		//-----------------------------------------------------------------------------------------
+		// Interface pattern
+		//-----------------------------------------------------------------------------------------
+		class Private;
+
+	protected:
+
+		//! Forbidden. Manage with the Ptr<> template.
+		inline ~NodeString() {}
+
+	};
+
+
+}
diff --git a/Engine/Plugins/Mutable/Source/MutableTools/Public/MuT/NodeStringConstant.h b/Engine/Plugins/Mutable/Source/MutableTools/Public/MuT/NodeStringConstant.h
new file mode 100644
--- /dev/null
+++ b/Engine/Plugins/Mutable/Source/MutableTools/Public/MuT/NodeStringConstant.h
@@ -0,0 +1,61 @@
+// Copyright Epic Games, Inc. All Rights Reserved.
+
+#pragma once
+
+#include "MuR/Ptr.h"
+#include "MuR/RefCounted.h"
+#include "MuT/Node.h"
+#include "MuT/NodeString.h"
+
+
+namespace mu
+{
+
+	// Forward definitions
+	class NodeStringConstant;
+	typedef Ptr<NodeStringConstant> NodeStringConstantPtr;
+	typedef Ptr<const NodeStringConstant> NodeStringConstantPtrConst;
+
+
+	//! Node returning a string constant value.
+	//! \ingroup model
+	class MUTABLETOOLS_API NodeStringConstant : public NodeString
+	{
+	public:
+
+		NodeStringConstant();
+
+
+		//-----------------------------------------------------------------------------------------
+		// Node Interface
+		//-----------------------------------------------------------------------------------------
+        const FNodeType* GetType() const override;
+		static const FNodeType* GetStaticType();
+
+		//-----------------------------------------------------------------------------------------
+		// Own Interface
+		//-----------------------------------------------------------------------------------------
+
+		//! Set the value to be returned by the node.
+		void SetValue( const FString& v );
+
+		//-----------------------------------------------------------------------------------------
+		// Interface pattern
+		//-----------------------------------------------------------------------------------------
+		class Private;
+		Private* GetPrivate() const;
+
+	protected:
+
+		//! Forbidden. Manage with the Ptr<> template.
+		~NodeStringConstant();
+
+	private:
+
+		Private* m_pD;
+
+	};
+
+
+
+}
diff --git a/Engine/Plugins/Mutable/Source/MutableTools/Public/MuT/NodeStringParameter.h b/Engine/Plugins/Mutable/Source/MutableTools/Public/MuT/NodeStringParameter.h
new file mode 100644
--- /dev/null
+++ b/Engine/Plugins/Mutable/Source/MutableTools/Public/MuT/NodeStringParameter.h
@@ -0,0 +1,79 @@
+// Copyright Epic Games, Inc. All Rights Reserved.
+
+#pragma once
+
+#include "MuR/Parameters.h"
+#include "MuR/Ptr.h"
+#include "MuR/RefCounted.h"
+#include "MuT/Node.h"
+#include "MuT/NodeString.h"
+
+
+namespace mu
+{
+
+	// Forward definitions
+    class NodeImage;
+    using NodeImagePtr = Ptr<NodeImage>;
+    using NodeImagePtrConst = Ptr<const NodeImage>;
+
+    class NodeRange;
+    using NodeRangePtr = Ptr<NodeRange>;
+    using NodeRangePtrConst = Ptr<const NodeRange>;
+
+    class NodeStringParameter;
+    using NodeStringParameterPtr = Ptr<NodeStringParameter>;
+    using NodeStringParameterPtrConst = Ptr<const NodeStringParameter>;
+
+
+	//! Node that defines a string model parameter.
+	class MUTABLETOOLS_API NodeStringParameter : public NodeString
+	{
+	public:
+
+		NodeStringParameter();
+
+		//-----------------------------------------------------------------------------------------
+		// Node Interface
+		//-----------------------------------------------------------------------------------------
+
+        const FNodeType* GetType() const override;
+		static const FNodeType* GetStaticType();
+
+		//-----------------------------------------------------------------------------------------
+		// Own Interface
+		//-----------------------------------------------------------------------------------------
+
+		//! Set the name of the parameter. It will be exposed in the final compiled data.
+		void SetName( const FString& );
+
+		//! Set the uid of the parameter. It will be exposed in the final compiled data.
+		void SetUid( const FString& );
+
+		//! Set the default value of the parameter.
+		void SetDefaultValue( const FString& );
+
+        //! Set the number of ranges (dimensions) for this parameter.
+        //! By default a parameter has 0 ranges, meaning it only has one value.
+        void SetRangeCount( int i );
+        void SetRange( int i, NodeRangePtr pRange );
+
+		//-----------------------------------------------------------------------------------------
+		// Interface pattern
+		//-----------------------------------------------------------------------------------------
+		class Private;
+		Private* GetPrivate() const;
+
+	protected:
+
+		//! Forbidden. Manage with the Ptr<> template.
+		~NodeStringParameter();
+
+	private:
+
+		Private* m_pD;
+
+	};
+
+
+}
diff --git a/Engine/Plugins/Mutable/Source/MutableTools/Public/MuT/TypeInfo.h b/Engine/Plugins/Mutable/Source/MutableTools/Public/MuT/TypeInfo.h
new file mode 100644
--- /dev/null
+++ b/Engine/Plugins/Mutable/Source/MutableTools/Public/MuT/TypeInfo.h
@@ -0,0 +1,33 @@
+// Copyright Epic Games, Inc. All Rights Reserved.
+
+#pragma once
+
+#include "HAL/PlatformMath.h"
+#include "MuR/Image.h"
+#include "MuR/MeshBufferSet.h"
+#include "MuR/Parameters.h"
+#include "MuR/ParametersPrivate.h"
+
+// This header file provides additional information about data types defined in the runtime, that
+// can be useful only for the tools. For instance, it provides strings for some enumeration values.
+
+namespace mu
+{
+
+	class MUTABLETOOLS_API TypeInfo
+	{
+	public:
+
+		static const char* s_imageFormatName[size_t(EImageFormat::IF_COUNT)];
+
+		static const char* s_meshBufferSemanticName[ MBS_COUNT ];
+
+		static const char* s_meshBufferFormatName[ MBF_COUNT ];
+
+		static const char* s_blendModeName[uint32(EBlendType::_BT_COUNT) ];
+
+		static const char* s_projectorTypeName [ static_cast<uint32>(mu::PROJECTOR_TYPE::COUNT) ];
+
+	};
+
+}
