[Mutable] Fixed modifiers not working properly

Reverted hack that takes that, if not found, it takes uses the firtst component.

#rnx

diff --git a/Engine/Plugins/Mutable/Source/CustomizableObjectEditor/Private/MuCOE/GenerateMutableSource/GenerateMutableSource.cpp b/Engine/Plugins/Mutable/Source/CustomizableObjectEditor/Private/MuCOE/GenerateMutableSource/GenerateMutableSource.cpp
--- a/Engine/Plugins/Mutable/Source/CustomizableObjectEditor/Private/MuCOE/GenerateMutableSource/GenerateMutableSource.cpp
+++ b/Engine/Plugins/Mutable/Source/CustomizableObjectEditor/Private/MuCOE/GenerateMutableSource/GenerateMutableSource.cpp
@@ -445,7 +445,20 @@ mu::FBoneName FMutableGraphGenerationContext::GetBoneUnique(const FName& InBoneN
 
 FMutableComponentInfo* FMutableGraphGenerationContext::GetCurrentComponentInfo()
 {
-	return ComponentInfos.FindByPredicate([this](const FMutableComponentInfo& Component) { return Component.ComponentName == CurrentMeshComponent; });
+	FMutableComponentInfo* CurrentComponentInfo = ComponentInfos.FindByPredicate(
+		[this](const FMutableComponentInfo& Component) { return Component.ComponentName == CurrentMeshComponent; });
+	
+	// Temp workaround to the problem of modifiers that generate meshes (like ExtendMeshSection) not having a known component at generation time.
+	// TODO: Actually detect all possible components and change this query in GenerationContext state to accomodate a set of "current components" instead of one.
+	//check(CurrentComponentInfo);
+
+	if (!CurrentComponentInfo)
+	{
+		check(!ComponentInfos.IsEmpty());
+		CurrentComponentInfo = &ComponentInfos[0];
+	}
+
+	return CurrentComponentInfo;
 }
 
 
