[mutable] Fix broken extend and edit modifiers with LODs.
#rnx

diff --git a/Engine/Plugins/Mutable/Source/CustomizableObjectEditor/Private/MuCOE/GenerateMutableSource/GenerateMutableSource.cpp b/Engine/Plugins/Mutable/Source/CustomizableObjectEditor/Private/MuCOE/GenerateMutableSource/GenerateMutableSource.cpp
--- a/Engine/Plugins/Mutable/Source/CustomizableObjectEditor/Private/MuCOE/GenerateMutableSource/GenerateMutableSource.cpp
+++ b/Engine/Plugins/Mutable/Source/CustomizableObjectEditor/Private/MuCOE/GenerateMutableSource/GenerateMutableSource.cpp
@@ -930,14 +930,16 @@ mu::Ptr<mu::NodeObject> GenerateMutableSource(const UEdGraphPin * Pin, FMutableG
 					continue;
 				}
 
-				// Set it to "None" to indicate we don't care about component id.
-				FName OldCurrentMeshComponent = GenerationContext.CurrentMeshComponent;
-				GenerationContext.CurrentMeshComponent = FName();
-
+				// Temp hack
+				if (GenerationContext.ComponentInfos.Num())
+				{
+					GenerationContext.CurrentMeshComponent = GenerationContext.ComponentInfos[0].ComponentName;
+				}
+				
 				mu::Ptr<mu::NodeModifier> ModifierNode = GenerateMutableSourceModifier(ChildNodePin, GenerationContext);
 				ObjectNode->Modifiers.Add(ModifierNode);
 
-				GenerationContext.CurrentMeshComponent = OldCurrentMeshComponent;
+				GenerationContext.CurrentMeshComponent = FName();
 			}
 		}
 		
diff --git a/Engine/Plugins/Mutable/Source/CustomizableObjectEditor/Private/MuCOE/GenerateMutableSource/GenerateMutableSourceComponent.cpp b/Engine/Plugins/Mutable/Source/CustomizableObjectEditor/Private/MuCOE/GenerateMutableSource/GenerateMutableSourceComponent.cpp
--- a/Engine/Plugins/Mutable/Source/CustomizableObjectEditor/Private/MuCOE/GenerateMutableSource/GenerateMutableSourceComponent.cpp
+++ b/Engine/Plugins/Mutable/Source/CustomizableObjectEditor/Private/MuCOE/GenerateMutableSource/GenerateMutableSourceComponent.cpp
@@ -94,44 +94,20 @@ void GenerateMutableSourceComponentMesh(FMutableGraphGenerationContext& Generati
 				// Modifiers are shared for all components and are processed per LOD and not component.
 				if (Cast<UCustomizableObjectNodeModifierBase>(ChildNodePin->GetOwningNode()))
 				{
+					FString Msg = FString::Printf(TEXT("The object has legacy modifier connections that cannot be generated. Their connections should be updated."));
+					GenerationContext.Compiler->CompilerLog(FText::FromString(Msg), &TypedComponentMesh, EMessageSeverity::Warning);
 					continue;
 				}
 				
 				mu::Ptr<mu::NodeSurface> SurfaceNode = GenerateMutableSourceSurface(ChildNodePin, GenerationContext);
 				LODNode->Surfaces.Add(SurfaceNode);
 			}
-
-			// Process legacy modifiers.
-			for (UEdGraphPin* const ChildNodePin : ConnectedLODPins)
-			{
-				if (!Cast<UCustomizableObjectNodeModifierBase>(ChildNodePin->GetOwningNode()))
-				{
-					continue;
-				}
-
-				if (NodeObject)
-				{
-					// Warn about legacy modifier connection
-					FString Msg = FString::Printf(TEXT("The object has legacy modifier connections (to material pins?) that should be updated."));
-					GenerationContext.Compiler->CompilerLog(FText::FromString(Msg), &TypedComponentMesh, EMessageSeverity::Warning);
-
-					// Set it to "None" to indicate we don't care about component id.
-					FName OldCurrentMeshComponent = GenerationContext.CurrentMeshComponent;
-					GenerationContext.CurrentMeshComponent = FName();
-
-					mu::Ptr<mu::NodeModifier> ModifierNode = GenerateMutableSourceModifier(ChildNodePin, GenerationContext);
-					NodeObject->Modifiers.AddUnique(ModifierNode);
-
-					GenerationContext.CurrentMeshComponent = OldCurrentMeshComponent;
-				}
-				else
-				{
-					FString Msg = FString::Printf(TEXT("The object has legacy modifier connections that cannot be generated. Their connections should be updated."));
-					GenerationContext.Compiler->CompilerLog(FText::FromString(Msg), &TypedComponentMesh, EMessageSeverity::Warning);
-				}
-			}
 		}
 	}
+
+	// Clear the context state for LODs
+	GenerationContext.CurrentLOD = 0;
+	GenerationContext.FromLOD = 0;
 }
 
 
diff --git a/Engine/Plugins/Mutable/Source/CustomizableObjectEditor/Private/MuCOE/GenerateMutableSource/GenerateMutableSourceModifier.cpp b/Engine/Plugins/Mutable/Source/CustomizableObjectEditor/Private/MuCOE/GenerateMutableSource/GenerateMutableSourceModifier.cpp
--- a/Engine/Plugins/Mutable/Source/CustomizableObjectEditor/Private/MuCOE/GenerateMutableSource/GenerateMutableSourceModifier.cpp
+++ b/Engine/Plugins/Mutable/Source/CustomizableObjectEditor/Private/MuCOE/GenerateMutableSource/GenerateMutableSourceModifier.cpp
@@ -77,10 +77,7 @@ mu::Ptr<mu::NodeModifier> GenerateMutableSourceModifier(const UEdGraphPin * Pin,
 		ClipNode->SetVertexSelectionBone(GenerationContext.GetBoneUnique(TypedNodeClip->BoneName), TypedNodeClip->MaxEffectRadius);
 
 		ClipNode->MultipleTagsPolicy = TypedNodeClip->MultipleTagPolicy;
-		for (const FString& Tag : TypedNodeClip->RequiredTags)
-		{
-			 ClipNode->RequiredTags.Add(Tag);
-		}
+		ClipNode->RequiredTags = TypedNodeClip->RequiredTags;
 
 		GenerationContext.MeshGenerationFlags.Pop();
 	}
@@ -129,10 +126,7 @@ mu::Ptr<mu::NodeModifier> GenerateMutableSourceModifier(const UEdGraphPin * Pin,
 		}
 	
 		ClipNode->MultipleTagsPolicy = TypedNodeClipDeform->MultipleTagPolicy;
-		for (const FString& Tag : TypedNodeClipDeform->RequiredTags)
-		{
-			ClipNode->RequiredTags.Add(Tag);
-		}		
+		ClipNode->RequiredTags = TypedNodeClipDeform->RequiredTags;
 
 		GenerationContext.MeshGenerationFlags.Pop();
 	}
@@ -202,10 +196,7 @@ mu::Ptr<mu::NodeModifier> GenerateMutableSourceModifier(const UEdGraphPin * Pin,
 		}
 
 		ClipNode->MultipleTagsPolicy = TypedNodeClipMesh->MultipleTagPolicy;
-		for (const FString& Tag : TypedNodeClipMesh->RequiredTags)
-		{
-			 ClipNode->RequiredTags.Add(Tag);
-		}
+		ClipNode->RequiredTags = TypedNodeClipMesh->RequiredTags;
 
 		if (TypedNodeClipMesh->CustomizableObjectToClipWith != nullptr)
 		{
@@ -259,10 +250,7 @@ mu::Ptr<mu::NodeModifier> GenerateMutableSourceModifier(const UEdGraphPin * Pin,
 		ClipNode->LayoutIndex = TypedNodeClipUVMask->UVChannelForMask;
 
 		ClipNode->MultipleTagsPolicy = TypedNodeClipUVMask->MultipleTagPolicy;
-		for (const FString& Tag : TypedNodeClipUVMask->RequiredTags)
-		{
-			ClipNode->RequiredTags.Add(Tag);
-		}
+		ClipNode->RequiredTags = TypedNodeClipUVMask->RequiredTags;
 
 		GenerationContext.MeshGenerationFlags.Pop();
 	}
@@ -278,13 +266,17 @@ mu::Ptr<mu::NodeModifier> GenerateMutableSourceModifier(const UEdGraphPin * Pin,
 		// TODO: This was used in the non-modifier version for group projectors. It may affect the "drop projection from LOD" feature.
 		const int32 LOD = Node->IsAffectedByLOD() ? GenerationContext.CurrentLOD : 0;
 
-		[&] // Using a lambda so control flow is easier to manage.
+		SurfNode->MultipleTagsPolicy = TypedNodeExt->MultipleTagPolicy;
+		SurfNode->RequiredTags = TypedNodeExt->RequiredTags;
+
+		// Is this enough? Should we try to narrow down with potential mesh sections modified by this?
+		int32 LODCount = GenerationContext.NumLODsInRoot;
+		SurfNode->LODs.SetNum(LODCount);
+
+		for (int32 LODIndex=0; LODIndex<LODCount; ++LODIndex)
 		{
-			SurfNode->MultipleTagsPolicy = TypedNodeExt->MultipleTagPolicy;
-			for (const FString& Tag : TypedNodeExt->RequiredTags)
-			{
-				SurfNode->RequiredTags.Add(Tag);
-			}
+			GenerationContext.FromLOD = 0;
+			GenerationContext.CurrentLOD = LODIndex;
 
 			mu::Ptr<mu::NodeMesh> AddMeshNode;
 			FMutableGraphMeshGenerationData MeshData;
@@ -335,11 +327,11 @@ mu::Ptr<mu::NodeModifier> GenerateMutableSourceModifier(const UEdGraphPin * Pin,
 
 				MeshFormat->SetSource(MeshPtr.get());
 
-				SurfNode->MeshAdd = MeshFormat;
+				SurfNode->LODs[LODIndex].MeshAdd = MeshFormat;
 			}
 
 			const int32 NumImages = TypedNodeExt->GetNumParameters(EMaterialParameterType::Texture);
-			SurfNode->Textures.SetNum(NumImages);
+			SurfNode->LODs[LODIndex].Textures.SetNum(NumImages);
 			for (int32 ImageIndex = 0; ImageIndex < NumImages; ++ImageIndex)
 			{
 				mu::NodeImagePtr ImageNode;
@@ -394,18 +386,17 @@ mu::Ptr<mu::NodeModifier> GenerateMutableSourceModifier(const UEdGraphPin * Pin,
 
 				if (ImageNode)
 				{
-					SurfNode->Textures[ImageIndex].Extend = ImageNode;
+					SurfNode->LODs[LODIndex].Textures[ImageIndex].Extend = ImageNode;
 				}
-
 			}
+		}
 
-			for (const FString& Tag : TypedNodeExt->Tags)
-			{
-				SurfNode->EnableTags.Add(Tag);
-			}
-		}();
+		SurfNode->EnableTags = TypedNodeExt->Tags;
 
 		GenerationContext.MeshGenerationFlags.Pop();
+		GenerationContext.FromLOD = 0;
+		GenerationContext.CurrentLOD = 0;
+
 	}
 
 	else if (const UCustomizableObjectNodeModifierRemoveMesh* TypedNodeRem = Cast<UCustomizableObjectNodeModifierRemoveMesh>(Node))
@@ -419,22 +410,28 @@ mu::Ptr<mu::NodeModifier> GenerateMutableSourceModifier(const UEdGraphPin * Pin,
 		Result = SurfNode;
 
 		SurfNode->MultipleTagsPolicy = TypedNodeRem->MultipleTagPolicy;
-		for (const FString& Tag : TypedNodeRem->RequiredTags)
-		{
-			SurfNode->RequiredTags.Add(Tag);
-		}
-
-		mu::Ptr<mu::NodeMesh> RemoveMeshNode;
+		SurfNode->RequiredTags = TypedNodeRem->RequiredTags;
 
 		if (const UEdGraphPin* ConnectedPin = FollowInputPin(*TypedNodeRem->RemoveMeshPin()))
 		{
-			FMutableGraphMeshGenerationData DummyMeshData;
-			RemoveMeshNode = GenerateMutableSourceMesh(ConnectedPin, GenerationContext, DummyMeshData, false, true);
-		}
+			// Is this enough? Should we try to narrow down with potential mesh sections modified by this?
+			int32 LODCount = GenerationContext.NumLODsInRoot;
+			SurfNode->LODs.SetNum(LODCount);
 
-		SurfNode->MeshRemove = RemoveMeshNode;
+			for (int32 LODIndex = 0; LODIndex < LODCount; ++LODIndex)
+			{
+				GenerationContext.FromLOD = 0;
+				GenerationContext.CurrentLOD = LODIndex;
+
+				FMutableGraphMeshGenerationData DummyMeshData;
+				mu::Ptr<mu::NodeMesh> RemoveMeshNode = GenerateMutableSourceMesh(ConnectedPin, GenerationContext, DummyMeshData, false, true);
+				SurfNode->LODs[LODIndex].MeshRemove = RemoveMeshNode;
+			}
+		}
 
 		GenerationContext.MeshGenerationFlags.Pop();
+		GenerationContext.FromLOD = 0;
+		GenerationContext.CurrentLOD = 0;
 	}
 
 	else if (const UCustomizableObjectNodeModifierRemoveMeshBlocks* TypedNodeRemBlocks = Cast<UCustomizableObjectNodeModifierRemoveMeshBlocks>(Node))
@@ -448,10 +445,7 @@ mu::Ptr<mu::NodeModifier> GenerateMutableSourceModifier(const UEdGraphPin * Pin,
 		Result = ClipNode;
 
 		ClipNode->MultipleTagsPolicy = TypedNodeRemBlocks->MultipleTagPolicy;
-		for (const FString& Tag : TypedNodeRemBlocks->RequiredTags)
-		{
-			ClipNode->RequiredTags.Add(Tag);
-		}
+		ClipNode->RequiredTags = TypedNodeRemBlocks->RequiredTags;
 
 		bool bWasEmpty = false;
 		mu::Ptr<mu::NodeLayout> SourceLayout = CreateMutableLayoutNode(GenerationContext, TypedNodeRemBlocks->Layout, true, bWasEmpty);
@@ -472,58 +466,68 @@ mu::Ptr<mu::NodeModifier> GenerateMutableSourceModifier(const UEdGraphPin * Pin,
 		Result = SurfNode;
 
 		SurfNode->MultipleTagsPolicy = TypedNodeEdit->MultipleTagPolicy;
-		for (const FString& Tag : TypedNodeEdit->RequiredTags)
-		{
-			SurfNode->RequiredTags.Add(Tag);
-		}
+		SurfNode->RequiredTags = TypedNodeEdit->RequiredTags;
+
+		// Is this enough? Should we try to narrow down with potential mesh sections modified by this?
+		int32 LODCount = GenerationContext.NumLODsInRoot;
+		SurfNode->LODs.SetNum(LODCount);
 
-		const int32 NumImages = TypedNodeEdit->GetNumParameters(EMaterialParameterType::Texture);
-		SurfNode->Textures.SetNum(NumImages);
-		for (int32 ImageIndex = 0; ImageIndex < NumImages; ++ImageIndex)
+		for (int32 LODIndex = 0; LODIndex < LODCount; ++LODIndex)
 		{
-			const FNodeMaterialParameterId ImageId = TypedNodeEdit->GetParameterId(EMaterialParameterType::Texture, ImageIndex);
+			GenerationContext.FromLOD = 0;
+			GenerationContext.CurrentLOD = LODIndex;
+
 
-			if (TypedNodeEdit->UsesImage(ImageId))
+			const int32 NumImages = TypedNodeEdit->GetNumParameters(EMaterialParameterType::Texture);
+			SurfNode->LODs[LODIndex].Textures.SetNum(NumImages);
+			for (int32 ImageIndex = 0; ImageIndex < NumImages; ++ImageIndex)
 			{
-				// TODO
-				//check(ParentMaterialNode->IsImageMutableMode(ImageIndex)); // Ensured at graph time. If it fails, something is wrong.
+				const FNodeMaterialParameterId ImageId = TypedNodeEdit->GetParameterId(EMaterialParameterType::Texture, ImageIndex);
 
-				const UEdGraphPin* ConnectedImagePin = FollowInputPin(*TypedNodeEdit->GetUsedImagePin(ImageId));
+				if (TypedNodeEdit->UsesImage(ImageId))
+				{
+					// TODO
+					//check(ParentMaterialNode->IsImageMutableMode(ImageIndex)); // Ensured at graph time. If it fails, something is wrong.
 
-				mu::NodeModifierSurfaceEdit::FTexture& ImagePatch = SurfNode->Textures[ImageIndex];
+					const UEdGraphPin* ConnectedImagePin = FollowInputPin(*TypedNodeEdit->GetUsedImagePin(ImageId));
 
-				// \todo: expose these two options?
-				ImagePatch.PatchBlendType = mu::EBlendType::BT_BLEND;
-				ImagePatch.bPatchApplyToAlpha = true;
+					mu::NodeModifierSurfaceEdit::FTexture& ImagePatch = SurfNode->LODs[LODIndex].Textures[ImageIndex];
 
-				// ReferenceTextureSize is used to limit the size of textures contributing to the final image.
-				const int32 ReferenceTextureSize = 0; //TODO GetBaseTextureSize(GenerationContext, ParentMaterialNode, ImageIndex);
+					// \todo: expose these two options?
+					ImagePatch.PatchBlendType = mu::EBlendType::BT_BLEND;
+					ImagePatch.bPatchApplyToAlpha = true;
 
-				ImagePatch.PatchImage = GenerateMutableSourceImage(ConnectedImagePin, GenerationContext, ReferenceTextureSize);
+					// ReferenceTextureSize is used to limit the size of textures contributing to the final image.
+					const int32 ReferenceTextureSize = 0; //TODO GetBaseTextureSize(GenerationContext, ParentMaterialNode, ImageIndex);
 
-				const UEdGraphPin* ImageMaskPin = TypedNodeEdit->GetUsedImageMaskPin(ImageId);
-				check(ImageMaskPin); // Ensured when reconstructing EditMaterial nodes. If it fails, something is wrong.
+					ImagePatch.PatchImage = GenerateMutableSourceImage(ConnectedImagePin, GenerationContext, ReferenceTextureSize);
 
-				if (const UEdGraphPin* ConnectedMaskPin = FollowInputPin(*ImageMaskPin))
-				{
-					ImagePatch.PatchMask = GenerateMutableSourceImage(ConnectedMaskPin, GenerationContext, ReferenceTextureSize);
-				}
+					const UEdGraphPin* ImageMaskPin = TypedNodeEdit->GetUsedImageMaskPin(ImageId);
+					check(ImageMaskPin); // Ensured when reconstructing EditMaterial nodes. If it fails, something is wrong.
 
-				// Add the blocks to patch
-				FIntPoint GridSize = TypedNodeEdit->Layout->GetGridSize();
-				FVector2f GridSizeF = FVector2f(GridSize);
-				ImagePatch.PatchBlocks.Reserve(TypedNodeEdit->Layout->Blocks.Num());
-				for (const FCustomizableObjectLayoutBlock& LayoutBlock : TypedNodeEdit->Layout->Blocks)
-				{
-					FBox2f Rect;
-					Rect.Min = FVector2f(LayoutBlock.Min) / GridSizeF;
-					Rect.Max = FVector2f(LayoutBlock.Max) / GridSizeF;
-					ImagePatch.PatchBlocks.Add(Rect);
+					if (const UEdGraphPin* ConnectedMaskPin = FollowInputPin(*ImageMaskPin))
+					{
+						ImagePatch.PatchMask = GenerateMutableSourceImage(ConnectedMaskPin, GenerationContext, ReferenceTextureSize);
+					}
+
+					// Add the blocks to patch
+					FIntPoint GridSize = TypedNodeEdit->Layout->GetGridSize();
+					FVector2f GridSizeF = FVector2f(GridSize);
+					ImagePatch.PatchBlocks.Reserve(TypedNodeEdit->Layout->Blocks.Num());
+					for (const FCustomizableObjectLayoutBlock& LayoutBlock : TypedNodeEdit->Layout->Blocks)
+					{
+						FBox2f Rect;
+						Rect.Min = FVector2f(LayoutBlock.Min) / GridSizeF;
+						Rect.Max = FVector2f(LayoutBlock.Max) / GridSizeF;
+						ImagePatch.PatchBlocks.Add(Rect);
+					}
 				}
 			}
 		}
 
 		GenerationContext.MeshGenerationFlags.Pop();
+		GenerationContext.FromLOD = 0;
+		GenerationContext.CurrentLOD = 0;
 	}
 
 	else if (const UCustomizableObjectNodeModifierMorphMeshSection* TypedNodeMorph = Cast<UCustomizableObjectNodeModifierMorphMeshSection>(Node))
@@ -540,10 +544,7 @@ mu::Ptr<mu::NodeModifier> GenerateMutableSourceModifier(const UEdGraphPin * Pin,
 		SurfNode->bApplyBeforeNormalOperations = true;
 
 		SurfNode->MultipleTagsPolicy = TypedNodeMorph->MultipleTagPolicy;
-		for (const FString& Tag : TypedNodeMorph->RequiredTags)
-		{
-			SurfNode->RequiredTags.Add(Tag);
-		}
+		SurfNode->RequiredTags = TypedNodeMorph->RequiredTags;
 
 		SurfNode->MeshMorph = TypedNodeMorph->MorphTargetName;
 
diff --git a/Engine/Plugins/Mutable/Source/CustomizableObjectEditor/Private/MuCOE/SMutableGraphViewer.cpp b/Engine/Plugins/Mutable/Source/CustomizableObjectEditor/Private/MuCOE/SMutableGraphViewer.cpp
--- a/Engine/Plugins/Mutable/Source/CustomizableObjectEditor/Private/MuCOE/SMutableGraphViewer.cpp
+++ b/Engine/Plugins/Mutable/Source/CustomizableObjectEditor/Private/MuCOE/SMutableGraphViewer.cpp
@@ -368,16 +368,21 @@ void SMutableGraphViewer::GetChildrenForInfo(TSharedPtr<FMutableGraphTreeElement
 	else if (ParentNode->GetType() == mu::NodeModifierSurfaceEdit::GetStaticType())
 	{
 		mu::NodeModifierSurfaceEdit* SurfaceEdit = StaticCast<mu::NodeModifierSurfaceEdit*>(ParentNode);
-		AddChildFunc(SurfaceEdit->MeshAdd.get(), TEXT("MESH_ADD"));
-		AddChildFunc(SurfaceEdit->MeshRemove.get(), TEXT("MESH_REMOVE"));
-		AddChildFunc(SurfaceEdit->MorphFactor.get(), FString::Printf(TEXT("MORPH_FACTOR [%s]"), *SurfaceEdit->MeshMorph) );
 
-		for (int32 l = 0; l < SurfaceEdit->Textures.Num(); ++l)
+		AddChildFunc(SurfaceEdit->MorphFactor.get(), FString::Printf(TEXT("MORPH_FACTOR [%s]"), *SurfaceEdit->MeshMorph));
+
+		for (int32 LODIndex = 0; LODIndex < SurfaceEdit->LODs.Num(); ++LODIndex)
 		{
-			AddChildFunc(SurfaceEdit->Textures[l].Extend.get(), FString::Printf(TEXT("EXTEND [%d]"), l));
-			AddChildFunc(SurfaceEdit->Textures[l].PatchImage.get(), FString::Printf(TEXT("PATCH IMAGE [%d]"), l));
-			AddChildFunc(SurfaceEdit->Textures[l].PatchMask.get(), FString::Printf(TEXT("PATCH MASK [%d]"), l));
+			AddChildFunc(SurfaceEdit->LODs[LODIndex].MeshAdd.get(), FString::Printf(TEXT("LOD%d MESH_ADD"), LODIndex));
+			AddChildFunc(SurfaceEdit->LODs[LODIndex].MeshRemove.get(), FString::Printf(TEXT("LOD%d MESH_REMOVE"), LODIndex));
+
+			for (int32 l = 0; l < SurfaceEdit->LODs[LODIndex].Textures.Num(); ++l)
+			{
+				AddChildFunc(SurfaceEdit->LODs[LODIndex].Textures[l].Extend.get(), FString::Printf(TEXT("LOD%d EXTEND [%d]"), LODIndex, l));
+				AddChildFunc(SurfaceEdit->LODs[LODIndex].Textures[l].PatchImage.get(), FString::Printf(TEXT("LOD%d PATCH IMAGE [%d]"), LODIndex, l));
+				AddChildFunc(SurfaceEdit->LODs[LODIndex].Textures[l].PatchMask.get(), FString::Printf(TEXT("LOD%d PATCH MASK [%d]"), LODIndex, l));
 
+			}
 		}
 	}
 
diff --git a/Engine/Plugins/Mutable/Source/MutableTools/Private/MuT/ASTOpMeshMorph.cpp b/Engine/Plugins/Mutable/Source/MutableTools/Private/MuT/ASTOpMeshMorph.cpp
--- a/Engine/Plugins/Mutable/Source/MutableTools/Private/MuT/ASTOpMeshMorph.cpp
+++ b/Engine/Plugins/Mutable/Source/MutableTools/Private/MuT/ASTOpMeshMorph.cpp
@@ -129,7 +129,7 @@ namespace mu
 		}
 
 		// If not optimized yet
-		if (!NewOp)
+		if (!NewOp && Target)
 		{
 			// Target optimizations
 			OP_TYPE MorphType = Target.child()->GetOpType();
diff --git a/Engine/Plugins/Mutable/Source/MutableTools/Private/MuT/CodeGenerator.cpp b/Engine/Plugins/Mutable/Source/MutableTools/Private/MuT/CodeGenerator.cpp
--- a/Engine/Plugins/Mutable/Source/MutableTools/Private/MuT/CodeGenerator.cpp
+++ b/Engine/Plugins/Mutable/Source/MutableTools/Private/MuT/CodeGenerator.cpp
@@ -1941,8 +1941,8 @@ namespace mu
 		// Apply mesh modifiers
 		TArray<FirstPassGenerator::FModifier> Modifiers;
 
-		int32 currentLOD = CurrentParents.Last().Lod;
-		GetModifiersFor(Options.ActiveTags, currentLOD, bModifiersForBeforeOperations, Modifiers);
+		int32 CurrentLOD = CurrentParents.Last().Lod;
+		GetModifiersFor(Options.ActiveTags, CurrentLOD, bModifiersForBeforeOperations, Modifiers);
 
 		Ptr<ASTOp> PreModifiersMesh = LastMeshOp;
 
@@ -1958,8 +1958,11 @@ namespace mu
 
 				BaseMeshResult.ExtraMeshLayouts.Emplace();
 
-				if (Ptr<NodeMesh> pAdd = Edit->MeshAdd)
+				bool bAffectsCurrentLOD = Edit->LODs.IsValidIndex(CurrentLOD);
+				if (bAffectsCurrentLOD && Edit->LODs[CurrentLOD].MeshAdd)
 				{
+					Ptr<NodeMesh> pAdd = Edit->LODs[CurrentLOD].MeshAdd;
+
 					// Store the data necessary to apply modifiers for the pre-normal operations stage.
 					ActiveTags.Push(Edit->EnableTags);
 
@@ -2026,12 +2029,16 @@ namespace mu
 			{
 				const NodeModifierSurfaceEdit* Edit = static_cast<const NodeModifierSurfaceEdit*>(m.Node);
 
+				bool bAffectsCurrentLOD = Edit->LODs.IsValidIndex(CurrentLOD);
+
 				// Apply mesh removes from child objects "edit surface" nodes.
 				// "Removes" need to come after "Adds" because some removes may refer to added meshes,
 				// and not the base.
 				// \TODO: Apply base removes first, and then "added meshes" removes here. It may have lower memory footprint during generation.
-				if (Ptr<NodeMesh> pRemove = Edit->MeshRemove)
+				if (bAffectsCurrentLOD && Edit->LODs[CurrentLOD].MeshRemove)
 				{
+					Ptr<NodeMesh> pRemove = Edit->LODs[CurrentLOD].MeshRemove;
+
 					FMeshGenerationResult removeResults;
 					FMeshGenerationOptions RemoveMeshOptions;
 					RemoveMeshOptions.bLayouts = false;
@@ -2407,8 +2414,8 @@ namespace mu
 		// Apply mesh modifiers
 		TArray<FirstPassGenerator::FModifier> Modifiers;
 
-		int32 currentLOD = CurrentParents.Last().Lod;
-		GetModifiersFor(Options.ActiveTags, currentLOD, bModifiersForBeforeOperations, Modifiers);
+		int32 CurrentLOD = CurrentParents.Last().Lod;
+		GetModifiersFor(Options.ActiveTags, CurrentLOD, bModifiersForBeforeOperations, Modifiers);
 
 		ActiveTags.Push({});
 
@@ -2419,13 +2426,14 @@ namespace mu
 			{
 				const NodeModifierSurfaceEdit* Edit = static_cast<const NodeModifierSurfaceEdit*>(m.Node);
 
+				bool bAffectsCurrentLOD = Edit->LODs.IsValidIndex(CurrentLOD);
 
-				if (ImageIndex >= Edit->Textures.Num())
+				if (!bAffectsCurrentLOD || ImageIndex >= Edit->LODs[CurrentLOD].Textures.Num())
 				{
 					continue;
 				}
 
-				const NodeModifierSurfaceEdit::FTexture& Patch = Edit->Textures[ImageIndex];
+				const NodeModifierSurfaceEdit::FTexture& Patch = Edit->LODs[CurrentLOD].Textures[ImageIndex];
 				if (Patch.PatchImage.get())
 				{
 					// Does the current block need to be patched? Find out by building a mask.
@@ -2477,8 +2485,8 @@ namespace mu
 		// Apply mesh modifiers
 		TArray<FirstPassGenerator::FModifier> Modifiers;
 
-		int32 currentLOD = CurrentParents.Last().Lod;
-		GetModifiersFor(Options.ActiveTags, currentLOD, bModifiersForBeforeOperations, Modifiers);
+		int32 CurrentLOD = CurrentParents.Last().Lod;
+		GetModifiersFor(Options.ActiveTags, CurrentLOD, bModifiersForBeforeOperations, Modifiers);
 
 		ActiveTags.Push({});
 
@@ -2493,12 +2501,18 @@ namespace mu
 				int32 ThisEditIndex = EditIndex;
 				++EditIndex;
 
-				if (ImageIndex >= Edit->Textures.Num())
+				bool bAffectsCurrentLOD = Edit->LODs.IsValidIndex(CurrentLOD);
+				if (!bAffectsCurrentLOD)
+				{
+					continue;
+				}
+
+				if (ImageIndex >= Edit->LODs[CurrentLOD].Textures.Num())
 				{
 					continue;
 				}
 
-				Ptr<NodeImage> pExtend = Edit->Textures[ImageIndex].Extend;
+				Ptr<NodeImage> pExtend = Edit->LODs[CurrentLOD].Textures[ImageIndex].Extend;
 				if (!pExtend)
 				{
 					continue;
diff --git a/Engine/Plugins/Mutable/Source/MutableTools/Public/MuT/NodeModifierSurfaceEdit.h b/Engine/Plugins/Mutable/Source/MutableTools/Public/MuT/NodeModifierSurfaceEdit.h
--- a/Engine/Plugins/Mutable/Source/MutableTools/Public/MuT/NodeModifierSurfaceEdit.h
+++ b/Engine/Plugins/Mutable/Source/MutableTools/Public/MuT/NodeModifierSurfaceEdit.h
@@ -22,18 +22,6 @@ namespace mu
 	{
 	public:
 
-		/** Mesh to remove from the modified surface. */
-		Ptr<NodeMesh> MeshRemove;
-
-		/** Mesh to add to the modified surface. */
-		Ptr<NodeMesh> MeshAdd;
-
-		/** Name of the morph to apply to the modified surface if it has it. */
-		FString MeshMorph;
-
-		/** Factor of the morph to apply. */
-		Ptr<NodeScalar> MorphFactor;
-
 		/** Data for every modified texture. */
 		struct FTexture
 		{
@@ -49,16 +37,32 @@ namespace mu
 			/** Rects in the parent layout homogeneous UV space to patch. */
 			TArray<FBox2f> PatchBlocks;
 
-			/** */
+			/** Type of patching operation. */
 			EBlendType PatchBlendType = EBlendType::BT_BLEND;
 
 			/** Patch alpha channel as well? */
 			bool bPatchApplyToAlpha = false;
+		};
 
+		struct FLOD
+		{
+			/** Mesh to remove from the modified surface. */
+			Ptr<NodeMesh> MeshRemove;
+
+			/** Mesh to add to the modified surface. */
+			Ptr<NodeMesh> MeshAdd;
+
+			/** Textures to modify. */
+			TArray<FTexture> Textures;
 		};
 
-		/** Textures to modify. */
-		TArray<FTexture> Textures;
+		TArray<FLOD> LODs;
+
+		/** Name of the morph to apply to the modified surface if it has it. */
+		FString MeshMorph;
+
+		/** Factor of the morph to apply. */
+		Ptr<NodeScalar> MorphFactor;
 
 	public:
 
