[Mutable] Show component name in Passthrough Mesh node title

#rnx

diff --git a/Engine/Plugins/Mutable/Source/CustomizableObjectEditor/Private/MuCOE/Nodes/CustomizableObjectNodeComponentPassthroughMesh.cpp b/Engine/Plugins/Mutable/Source/CustomizableObjectEditor/Private/MuCOE/Nodes/CustomizableObjectNodeComponentPassthroughMesh.cpp
--- a/Engine/Plugins/Mutable/Source/CustomizableObjectEditor/Private/MuCOE/Nodes/CustomizableObjectNodeComponentPassthroughMesh.cpp
+++ b/Engine/Plugins/Mutable/Source/CustomizableObjectEditor/Private/MuCOE/Nodes/CustomizableObjectNodeComponentPassthroughMesh.cpp
@@ -133,17 +133,16 @@ void UCustomizableObjectNodeComponentPassthroughMesh::AllocateDefaultPins(UCusto
 
 FText UCustomizableObjectNodeComponentPassthroughMesh::GetNodeTitle(ENodeTitleType::Type TitleType) const
 {
-	USkeletalMesh* SkeletalMesh = Cast<USkeletalMesh>(Mesh.TryLoad());
-	if (SkeletalMesh)
+	if (TitleType == ENodeTitleType::ListView)
 	{
-		FFormatNamedArguments Args;
-		Args.Add(TEXT("MeshName"), FText::FromString(SkeletalMesh->GetName()));
-
-		return FText::Format(LOCTEXT("ComponentPassthroughMesh_Title", "{MeshName}\n Passthrough Mesh Component"), Args);
+		return LOCTEXT("ComponentPassthroughMesh", "Passthrough Mesh Component");
 	}
 	else
 	{
-		return LOCTEXT("ComponentPassthroughMesh", "Passthrough Mesh Component");
+		FFormatNamedArguments Args;
+		Args.Add(TEXT("ComponentName"), FText::FromName(ComponentName));
+
+		return FText::Format(LOCTEXT("ComponentPassthroughMesh_Title", "{ComponentName}\n Passthrough Mesh Component"), Args);
 	}
 }
 
