[Mutable] Fix NodeObject compatibility crash.

#rnx

diff --git a/Engine/Plugins/Mutable/Source/CustomizableObjectEditor/Private/MuCOE/Nodes/CustomizableObjectNodeObject.cpp b/Engine/Plugins/Mutable/Source/CustomizableObjectEditor/Private/MuCOE/Nodes/CustomizableObjectNodeObject.cpp
--- a/Engine/Plugins/Mutable/Source/CustomizableObjectEditor/Private/MuCOE/Nodes/CustomizableObjectNodeObject.cpp
+++ b/Engine/Plugins/Mutable/Source/CustomizableObjectEditor/Private/MuCOE/Nodes/CustomizableObjectNodeObject.cpp
@@ -284,7 +284,11 @@ void UCustomizableObjectNodeObject::BackwardsCompatibleFixup(int32 CustomizableO
 		for (int32 LODIndex = 0; LODIndex < NumLODs_DEPRECATED; ++LODIndex)
 		{
 			UEdGraphPin* OldLODPin = FindPin(FString::Printf(TEXT("%s%d "), TEXT("LOD "), LODIndex));
-
+			if (!OldLODPin)
+			{
+				continue;
+			}
+			
 			TArray<UEdGraphPin*> CopyLinkedPins = OldLODPin->LinkedTo;
 			for (UEdGraphPin* LinkedPin : CopyLinkedPins) // Import/Exports/Reroute not supported.
 			{
