[Mutable] Fix Material Overrides not updating.

Error noticeable when using Mesh Cache.

#rnx

diff --git a/Engine/Plugins/Mutable/Source/CustomizableObject/Private/MuCO/CustomizableObjectInstance.cpp b/Engine/Plugins/Mutable/Source/CustomizableObject/Private/MuCO/CustomizableObjectInstance.cpp
--- a/Engine/Plugins/Mutable/Source/CustomizableObject/Private/MuCO/CustomizableObjectInstance.cpp
+++ b/Engine/Plugins/Mutable/Source/CustomizableObject/Private/MuCO/CustomizableObjectInstance.cpp
@@ -250,7 +250,7 @@ FCustomizableInstanceComponentData* UCustomizableInstancePrivate::GetComponentDa
 	}
 
 	const int32 ComponentIndex = Object->GetPrivate()->GetModelResources().ComponentNames.IndexOfByKey(ComponentName);
-	if (ComponentIndex != INDEX_NONE)
+	if (ComponentIndex == INDEX_NONE)
 	{
 		return nullptr;
 	}
