[mutable] Fix double ensure in error fallback.
#rnx

diff --git a/Engine/Plugins/Mutable/Source/CustomizableObject/Private/MuCO/CustomizableObjectInstance.cpp b/Engine/Plugins/Mutable/Source/CustomizableObject/Private/MuCO/CustomizableObjectInstance.cpp
--- a/Engine/Plugins/Mutable/Source/CustomizableObject/Private/MuCO/CustomizableObjectInstance.cpp
+++ b/Engine/Plugins/Mutable/Source/CustomizableObject/Private/MuCO/CustomizableObjectInstance.cpp
@@ -7579,22 +7579,11 @@ void UCustomizableInstancePrivate::RegenerateImportedModels()
 					// Fallback index, may shift materials around sections
 					if (SkeletalMesh->GetMaterials().IsValidIndex(RenderSection.MaterialIndex))
 					{
-						ImportedSection.MaterialIndex = LODMaterialMap[RenderSection.MaterialIndex];
+						ImportedSection.MaterialIndex = RenderSection.MaterialIndex;
 					}
 					else
 					{
-						// The material should have been in the LODMaterialMap
-						ensureMsgf(false, TEXT("Unexpected material index in UCustomizableInstancePrivate::RegenerateImportedModel"));
-
-						// Fallback index, may shift materials around sections
-						if (SkeletalMesh->GetMaterials().IsValidIndex(RenderSection.MaterialIndex))
-						{
-							ImportedSection.MaterialIndex = RenderSection.MaterialIndex;
-						}
-						else
-						{
-							ImportedSection.MaterialIndex = 0;
-						}
+						ImportedSection.MaterialIndex = 0;
 					}
 				}
 
