[mutable] Fix recursion generation bug with modifiers.
#rnx

diff --git a/Engine/Plugins/Mutable/Source/MutableTools/Private/MuT/CodeGenerator.cpp b/Engine/Plugins/Mutable/Source/MutableTools/Private/MuT/CodeGenerator.cpp
--- a/Engine/Plugins/Mutable/Source/MutableTools/Private/MuT/CodeGenerator.cpp
+++ b/Engine/Plugins/Mutable/Source/MutableTools/Private/MuT/CodeGenerator.cpp
@@ -1952,6 +1952,12 @@ namespace mu
 		int32 EditIndex = 0;
 		for (const FirstPassGenerator::FModifier& m : Modifiers)
 		{
+			if (ModifiersToIgnore.Contains(m))
+			{
+				// Prevent recursion.
+				continue;
+			}
+
 			if (m.Node->GetType() == NodeModifierSurfaceEdit::GetStaticType())
 			{
 				const NodeModifierSurfaceEdit* Edit = static_cast<const NodeModifierSurfaceEdit*>(m.Node);
@@ -1983,7 +1989,9 @@ namespace mu
 					// Apply the modifier for the post-normal operations stage to the added mesh
 					FMeshGenerationOptions ModifierOptions(Options);
 					ModifierOptions.ActiveTags = Edit->EnableTags;
+					ModifiersToIgnore.Push(m);
 					Ptr<ASTOp> AddedMeshOp = ApplyMeshModifiers(ModifierOptions, AddResults, SharedMeshResults, bModifiersForBeforeOperations, ErrorContext, nullptr);
+					ModifiersToIgnore.Pop();
 
 					FMeshGenerationResult::FExtraLayouts data;
 					data.GeneratedLayouts = AddResults.GeneratedLayouts;
diff --git a/Engine/Plugins/Mutable/Source/MutableTools/Private/MuT/CodeGenerator.h b/Engine/Plugins/Mutable/Source/MutableTools/Private/MuT/CodeGenerator.h
--- a/Engine/Plugins/Mutable/Source/MutableTools/Private/MuT/CodeGenerator.h
+++ b/Engine/Plugins/Mutable/Source/MutableTools/Private/MuT/CodeGenerator.h
@@ -814,6 +814,9 @@ namespace mu
 		// Get the modifiers that have to be applied to elements with a specific tag.
 		void GetModifiersFor(const TArray<FString>& SurfaceTags, int32 LOD, bool bModifiersForBeforeOperations, TArray<FirstPassGenerator::FModifier>& OutModifiers);
 
+		// Used to avoid recursion when generating modifiers.
+		TArray<FirstPassGenerator::FModifier> ModifiersToIgnore;
+
 		// Apply the required mesh modifiers to the given operation.
 		Ptr<ASTOp> ApplyMeshModifiers(const FMeshGenerationOptions&, FMeshGenerationResult& BaseResults, 
 			const FMeshGenerationResult* SharedSurfaceResults, 
diff --git a/Engine/Plugins/Mutable/Source/MutableTools/Private/MuT/CodeGenerator_FirstPass.h b/Engine/Plugins/Mutable/Source/MutableTools/Private/MuT/CodeGenerator_FirstPass.h
--- a/Engine/Plugins/Mutable/Source/MutableTools/Private/MuT/CodeGenerator_FirstPass.h
+++ b/Engine/Plugins/Mutable/Source/MutableTools/Private/MuT/CodeGenerator_FirstPass.h
@@ -160,6 +160,8 @@ namespace mu
 
             // This is filled in CodeGenerator_SecondPass.
             StateCondition StateCondition;
+						
+			bool operator==(const FModifier&) const = default;
         };
 		TArray<FModifier> Modifiers;
 
