[Mutable] Fix Instance drag and drop to a level now working

- ComponentName now is visible.

#rnx

diff --git a/Engine/Plugins/Mutable/Source/CustomizableObject/Private/MuCO/CustomizableSkeletalComponent.cpp b/Engine/Plugins/Mutable/Source/CustomizableObject/Private/MuCO/CustomizableSkeletalComponent.cpp
--- a/Engine/Plugins/Mutable/Source/CustomizableObject/Private/MuCO/CustomizableSkeletalComponent.cpp
+++ b/Engine/Plugins/Mutable/Source/CustomizableObject/Private/MuCO/CustomizableSkeletalComponent.cpp
@@ -32,6 +32,26 @@ void UCustomizableSkeletalComponent::CreateCustomizableObjectInstanceUsage()
 }
 
 
+#if WITH_EDITOR
+void UCustomizableSkeletalComponent::PostEditChangeProperty(FPropertyChangedEvent& PropertyChangedEvent)
+{
+	Super::PostEditChangeProperty(PropertyChangedEvent);
+
+	if (!PropertyChangedEvent.MemberProperty)
+	{
+		return;
+	}
+	
+	const FName PropertyName = PropertyChangedEvent.MemberProperty->GetFName();
+
+	if (PropertyName == GET_MEMBER_NAME_CHECKED(UCustomizableSkeletalComponent, ComponentName))
+	{
+		SetComponentName(ComponentName);
+	}
+}
+#endif
+
+
 void UCustomizableSkeletalComponent::Callbacks() const
 {
 	if (CustomizableObjectInstanceUsage)
diff --git a/Engine/Plugins/Mutable/Source/CustomizableObject/Public/MuCO/CustomizableSkeletalComponent.h b/Engine/Plugins/Mutable/Source/CustomizableObject/Public/MuCO/CustomizableSkeletalComponent.h
--- a/Engine/Plugins/Mutable/Source/CustomizableObject/Public/MuCO/CustomizableSkeletalComponent.h
+++ b/Engine/Plugins/Mutable/Source/CustomizableObject/Public/MuCO/CustomizableSkeletalComponent.h
@@ -25,8 +25,9 @@ public:
 	int32 ComponentIndex;
 
 private:
-	/** Only used if the ComponentIndex is INDEX_NONE. */
-	UPROPERTY()
+	/** Only used if the ComponentIndex is INDEX_NONE.
+	 *  Editing this property will set ComponentIndex to INDEX_NONE. */
+	UPROPERTY(EditAnywhere, Category = CustomizableSkeletalMesh)
 	FName ComponentName;
 
 public:
@@ -38,6 +39,13 @@ public:
 
 	FCustomizableSkeletalComponentUpdatedDelegate UpdatedDelegate;
 
+	// UObject interface
+#if WITH_EDITOR
+	virtual void PostEditChangeProperty(FPropertyChangedEvent& PropertyChangedEvent) override;
+#endif
+	
+	// Own interface
+	
 	/** Common end point of all updates. Even those which failed. */
 	void Callbacks() const;
 	
diff --git a/Engine/Plugins/Mutable/Source/CustomizableObjectEditor/Private/MuCOE/CustomizableObjectInstanceFactory.cpp b/Engine/Plugins/Mutable/Source/CustomizableObjectEditor/Private/MuCOE/CustomizableObjectInstanceFactory.cpp
--- a/Engine/Plugins/Mutable/Source/CustomizableObjectEditor/Private/MuCOE/CustomizableObjectInstanceFactory.cpp
+++ b/Engine/Plugins/Mutable/Source/CustomizableObjectEditor/Private/MuCOE/CustomizableObjectInstanceFactory.cpp
@@ -29,44 +29,61 @@ UCustomizableObjectInstanceFactory::UCustomizableObjectInstanceFactory(const FOb
 void UCustomizableObjectInstanceFactory::PostSpawnActor(UObject* Asset, AActor* NewActor)
 {
 	Super::PostSpawnActor(Asset, NewActor);
-	UCustomizableObjectInstance* COInstance = Cast<UCustomizableObjectInstance>(Asset);
+	
+	UCustomizableObjectInstance* Instance = Cast<UCustomizableObjectInstance>(Asset);
+	if (!Instance)
+	{
+		return;
+	}
+
+	UCustomizableObject* Object = Instance->GetCustomizableObject();
+	if (!Object)
+	{
+		return;
+	}
+
 	ACustomizableSkeletalMeshActor* NewCSMActor = CastChecked<ACustomizableSkeletalMeshActor>(NewActor);
-	if (NewCSMActor && COInstance)
-	{	
-		int32 NumComponents = COInstance->GetNumComponents();
+	if (!NewCSMActor)
+	{
+		return;
+	}
+	
+	for (int32 ComponentIndex = 0; ComponentIndex < Instance->GetNumComponents(); ++ComponentIndex)
+	{
+		const FName& ComponentName = Object->GetComponentName(ComponentIndex);
+		
+		USkeletalMesh* SkeletalMesh = Instance->GetComponentMeshSkeletalMesh(ComponentName);
+		if (!SkeletalMesh)
+		{
+			SkeletalMesh = Object->GetComponentMeshReferenceSkeletalMesh(ComponentName);
+		}
+		
+		if (ComponentIndex > 0 && ComponentIndex > NewCSMActor->GetNumComponents() - 1)
+		{
+			NewCSMActor->AttachNewComponent();
+		}
 
-		for (int32 ComponentIndex = 0; ComponentIndex < NumComponents; ++ComponentIndex)
+		if (USkeletalMeshComponent* SkeletalMeshComp = NewCSMActor->GetSkeletalMeshComponentAt(ComponentIndex))
 		{
-			USkeletalMesh* SkeletalMesh = GetSkeletalMeshFromAsset(Asset, ComponentIndex);
+			SkeletalMeshComp->UnregisterComponent();
+			UE_MUTABLE_SETSKINNEDASSET(SkeletalMeshComp, SkeletalMesh);
 
-			if (ComponentIndex > 0 && ComponentIndex > NewCSMActor->GetNumComponents() - 1)
+			if (ComponentIndex == 0 && NewCSMActor->GetWorld()->IsGameWorld())
 			{
-				NewCSMActor->AttachNewComponent();
+				NewCSMActor->ReplicatedMesh = SkeletalMesh;
 			}
 
-			if (USkeletalMeshComponent* SkeletalMeshComp = NewCSMActor->GetSkeletalMeshComponentAt(ComponentIndex))
+			if (UCustomizableSkeletalComponent* CustomSkeletalComp = NewCSMActor->GetCustomizableSkeletalComponent(ComponentIndex))
 			{
-				SkeletalMeshComp->UnregisterComponent();
-				UE_MUTABLE_SETSKINNEDASSET(SkeletalMeshComp, SkeletalMesh);
-
-				if (ComponentIndex == 0 && NewCSMActor->GetWorld()->IsGameWorld())
-				{
-					NewCSMActor->ReplicatedMesh = SkeletalMesh;
-				}
-
-				if (UCustomizableSkeletalComponent* CustomSkeletalComp = NewCSMActor->GetCustomizableSkeletalComponent(ComponentIndex))
-				{
-					CustomSkeletalComp->UnregisterComponent();
-					CustomSkeletalComp->CustomizableObjectInstance = COInstance;
-					CustomSkeletalComp->ComponentIndex = ComponentIndex;
-					//CustomSkeletalComp->ComponentName = GetComponentName(Asset,ComponentIndex);
-					CustomSkeletalComp->SetSkeletalMesh(SkeletalMesh);
-					CustomSkeletalComp->UpdateSkeletalMeshAsync();
-					CustomSkeletalComp->RegisterComponent();
-				}
-
-				SkeletalMeshComp->RegisterComponent();
+				CustomSkeletalComp->UnregisterComponent();
+				CustomSkeletalComp->CustomizableObjectInstance = Instance;
+				CustomSkeletalComp->SetComponentName(ComponentName);
+				CustomSkeletalComp->SetSkeletalMesh(SkeletalMesh);
+				CustomSkeletalComp->UpdateSkeletalMeshAsync();
+				CustomSkeletalComp->RegisterComponent();
 			}
+
+			SkeletalMeshComp->RegisterComponent();
 		}
 	}
 }
@@ -137,47 +154,6 @@ bool UCustomizableObjectInstanceFactory::CanCreateActorFrom(const FAssetData& As
 }
 
 
-USkeletalMesh* UCustomizableObjectInstanceFactory::GetSkeletalMeshFromAsset(UObject* Asset, int32 ComponentIndex) const
-{
-    UCustomizableObjectInstance* CustomizableObjectInstance = Cast<UCustomizableObjectInstance>(Asset);
-    if (!CustomizableObjectInstance)
-    {
-    	return nullptr;
-    }
-
-	UCustomizableObject* CustomizableObject = CustomizableObjectInstance->GetCustomizableObject();
-	if (!CustomizableObject)
-	{
-		return nullptr;
-	}
-	
-	const FName& ComponentName = CustomizableObject->GetComponentName(ComponentIndex);
-
-	if (USkeletalMesh* SkeletalMesh = CustomizableObjectInstance->GetComponentMeshSkeletalMesh(ComponentName))
-	{
-		return SkeletalMesh;
-	}
-    		
-	return CustomizableObject->GetComponentMeshReferenceSkeletalMesh(ComponentName);
-}
-
-
-//FName UCustomizableObjectInstanceFactory::GetComponentName(UObject* Asset, int32 ComponentIndex) const
-//{
-//	UCustomizableObjectInstance* CustomizableObjectInstance = Cast<UCustomizableObjectInstance>(Asset);
-//
-//	if (CustomizableObjectInstance != nullptr)
-//	{
-//		if (UCustomizableObject* CustomizableObject = CustomizableObjectInstance->GetCustomizableObject())
-//		{
-//			return CustomizableObject->GetPrivate()->Components[ComponentIndex].Name;
-//		}
-//	}
-//
-//	return FName();
-//}
-
-
 FQuat UCustomizableObjectInstanceFactory::AlignObjectToSurfaceNormal(const FVector& InSurfaceNormal, const FQuat& ActorRotation) const
 {
     // Meshes align the Z (up) axis with the surface normal
diff --git a/Engine/Plugins/Mutable/Source/CustomizableObjectEditor/Public/MuCOE/CustomizableObjectInstanceFactory.h b/Engine/Plugins/Mutable/Source/CustomizableObjectEditor/Public/MuCOE/CustomizableObjectInstanceFactory.h
--- a/Engine/Plugins/Mutable/Source/CustomizableObjectEditor/Public/MuCOE/CustomizableObjectInstanceFactory.h
+++ b/Engine/Plugins/Mutable/Source/CustomizableObjectEditor/Public/MuCOE/CustomizableObjectInstanceFactory.h
@@ -23,6 +23,5 @@ protected:
 	bool CanCreateActorFrom(const FAssetData& AssetData, FText& OutErrorMsg) override;
 	//~ End UActorFactory Interface
 
-	USkeletalMesh* GetSkeletalMeshFromAsset(UObject* Asset, int32 ComponentIndex = 0) const;
 	FQuat AlignObjectToSurfaceNormal(const FVector& InSurfaceNormal, const FQuat& ActorRotation) const;
 };
