[mutable] Fixed a broken path that included the Experimental folder, obsolete since the move to Beta. It fixes the broken Bake button icon.

#rnx
#tests Editor

diff --git a/Engine/Plugins/Mutable/Source/CustomizableObjectEditor/Private/MuCOE/CustomizableObjectEditorStyle.cpp b/Engine/Plugins/Mutable/Source/CustomizableObjectEditor/Private/MuCOE/CustomizableObjectEditorStyle.cpp
--- a/Engine/Plugins/Mutable/Source/CustomizableObjectEditor/Private/MuCOE/CustomizableObjectEditorStyle.cpp
+++ b/Engine/Plugins/Mutable/Source/CustomizableObjectEditor/Private/MuCOE/CustomizableObjectEditorStyle.cpp
@@ -72,7 +72,7 @@ TSharedRef<FSlateStyleSet> FCustomizableObjectEditorStyle::Create()
 
 	TSharedRef< FSlateStyleSet > Style = MakeShareable(new FSlateStyleSet("CustomizableObjectEditorStyle"));
 	Style->SetCoreContentRoot(FPaths::EngineContentDir());
-	Style->SetContentRoot(FPaths::EnginePluginsDir() / TEXT("Experimental/Mutable/Content"));
+	Style->SetContentRoot(FPaths::EnginePluginsDir() / TEXT("Mutable/Content"));
 	
 	Style->Set("CustomizableObjectEditor.Save", new IMAGE_BRUSH_SVG(Style, "Editor/Slate/Starship/Common/SaveCurrent", Icon40x40));
 	Style->Set("CustomizableObjectEditor.Save.Small", new IMAGE_BRUSH_SVG(Style, "Editor/Slate/Starship/Common/SaveCurrent", Icon20x20));
