[mutable] Fix conversion of legacy data after the mesh section pin and node rename for switching, variations, and import/export.
- Make sure the previous conversion is fixed, by not renaming the pin names and changing friendly names instead.
- Make sure already manually converted data also works, by renaming back the pins.
Additionally, disable unused node GeometryOperation.
#rnx

diff --git a/Engine/Plugins/Mutable/Source/CustomizableObject/Internal/MuCO/CustomizableObjectCustomVersion.h b/Engine/Plugins/Mutable/Source/CustomizableObject/Internal/MuCO/CustomizableObjectCustomVersion.h
--- a/Engine/Plugins/Mutable/Source/CustomizableObject/Internal/MuCO/CustomizableObjectCustomVersion.h
+++ b/Engine/Plugins/Mutable/Source/CustomizableObject/Internal/MuCO/CustomizableObjectCustomVersion.h
@@ -154,6 +154,8 @@ struct CUSTOMIZABLEOBJECT_API FCustomizableObjectCustomVersion
 
 		ComponentsArray,
 
+		FixMaterialPinsRename,
+
 		// -----<new versions can be added above this line>-------------------------------------------------
 		VersionPlusOne,
 		LatestVersion = VersionPlusOne - 1
diff --git a/Engine/Plugins/Mutable/Source/CustomizableObjectEditor/Private/MuCOE/EdGraphSchema_CustomizableObject.cpp b/Engine/Plugins/Mutable/Source/CustomizableObjectEditor/Private/MuCOE/EdGraphSchema_CustomizableObject.cpp
--- a/Engine/Plugins/Mutable/Source/CustomizableObjectEditor/Private/MuCOE/EdGraphSchema_CustomizableObject.cpp
+++ b/Engine/Plugins/Mutable/Source/CustomizableObjectEditor/Private/MuCOE/EdGraphSchema_CustomizableObject.cpp
@@ -433,7 +433,7 @@ void UEdGraphSchema_CustomizableObject::GetGraphContextActions(FGraphContextMenu
 			ContextMenuBuilder.CreateTemplateNode<UCustomizableObjectNodeMeshVariation>(),
 			ContextMenuBuilder.CreateTemplateNode<UCustomizableObjectNodeMeshMorphStackDefinition>(),
 			ContextMenuBuilder.CreateTemplateNode<UCustomizableObjectNodeMeshMorphStackApplication>(),
-			ContextMenuBuilder.CreateTemplateNode<UCustomizableObjectNodeMeshGeometryOperation>(),
+			// ContextMenuBuilder.CreateTemplateNode<UCustomizableObjectNodeMeshGeometryOperation>(),
 			ContextMenuBuilder.CreateTemplateNode<UCustomizableObjectNodeMeshReshape>(),
 		};
 
@@ -1348,7 +1348,7 @@ FText UEdGraphSchema_CustomizableObject::GetPinCategoryName(const FName& PinCate
 	}
 	else if (PinCategory == UEdGraphSchema_CustomizableObject::PC_Material)
 	{
-		return LOCTEXT("MeshSection_Pin_Category", "Mesh Section");
+		return LOCTEXT("MeshSection_Pin_Category", "Material");
 	}
 	else if (PinCategory == UEdGraphSchema_CustomizableObject::PC_Modifier)
 	{
@@ -1427,4 +1427,17 @@ FText UEdGraphSchema_CustomizableObject::GetPinCategoryName(const FName& PinCate
 }
 
 
+FText UEdGraphSchema_CustomizableObject::GetPinCategoryFriendlyName(const FName& PinCategory)
+{
+	if (PinCategory == UEdGraphSchema_CustomizableObject::PC_Material)
+	{
+		return LOCTEXT("MeshSection_Pin_Category_Friendly", "Mesh Section");
+	}
+	else
+	{
+		return GetPinCategoryName(PinCategory);
+	}
+}
+
+
 #undef LOCTEXT_NAMESPACE
diff --git a/Engine/Plugins/Mutable/Source/CustomizableObjectEditor/Private/MuCOE/Nodes/CustomizableObjectNodeExposePin.cpp b/Engine/Plugins/Mutable/Source/CustomizableObjectEditor/Private/MuCOE/Nodes/CustomizableObjectNodeExposePin.cpp
--- a/Engine/Plugins/Mutable/Source/CustomizableObjectEditor/Private/MuCOE/Nodes/CustomizableObjectNodeExposePin.cpp
+++ b/Engine/Plugins/Mutable/Source/CustomizableObjectEditor/Private/MuCOE/Nodes/CustomizableObjectNodeExposePin.cpp
@@ -29,7 +29,7 @@ void UCustomizableObjectNodeExposePin::AllocateDefaultPins(UCustomizableObjectNo
 
 FText UCustomizableObjectNodeExposePin::GetNodeTitle(ENodeTitleType::Type TitleType) const
 {
-	const FText PinCategoryName = UEdGraphSchema_CustomizableObject::GetPinCategoryName(PinType);
+	const FText PinCategoryName = UEdGraphSchema_CustomizableObject::GetPinCategoryFriendlyName(PinType);
 
 	if (TitleType == ENodeTitleType::ListView)
 	{
diff --git a/Engine/Plugins/Mutable/Source/CustomizableObjectEditor/Private/MuCOE/Nodes/CustomizableObjectNodeExternalPin.cpp b/Engine/Plugins/Mutable/Source/CustomizableObjectEditor/Private/MuCOE/Nodes/CustomizableObjectNodeExternalPin.cpp
--- a/Engine/Plugins/Mutable/Source/CustomizableObjectEditor/Private/MuCOE/Nodes/CustomizableObjectNodeExternalPin.cpp
+++ b/Engine/Plugins/Mutable/Source/CustomizableObjectEditor/Private/MuCOE/Nodes/CustomizableObjectNodeExternalPin.cpp
@@ -115,7 +115,7 @@ void UCustomizableObjectNodeExternalPin::AllocateDefaultPins(UCustomizableObject
 
 FText UCustomizableObjectNodeExternalPin::GetNodeTitle(ENodeTitleType::Type TitleType) const
 {
-	const FText PinTypeName = UEdGraphSchema_CustomizableObject::GetPinCategoryName(PinType);
+	const FText PinTypeName = UEdGraphSchema_CustomizableObject::GetPinCategoryFriendlyName(PinType);
 
 	if (TitleType == ENodeTitleType::ListView || !ExternalObject)
 	{
diff --git a/Engine/Plugins/Mutable/Source/CustomizableObjectEditor/Private/MuCOE/Nodes/CustomizableObjectNodeMaterialSwitch.cpp b/Engine/Plugins/Mutable/Source/CustomizableObjectEditor/Private/MuCOE/Nodes/CustomizableObjectNodeMaterialSwitch.cpp
--- a/Engine/Plugins/Mutable/Source/CustomizableObjectEditor/Private/MuCOE/Nodes/CustomizableObjectNodeMaterialSwitch.cpp
+++ b/Engine/Plugins/Mutable/Source/CustomizableObjectEditor/Private/MuCOE/Nodes/CustomizableObjectNodeMaterialSwitch.cpp
@@ -3,6 +3,7 @@
 #include "MuCOE/Nodes/CustomizableObjectNodeMaterialSwitch.h"
 
 #include "MuCOE/EdGraphSchema_CustomizableObject.h"
+#include "MuCO/CustomizableObjectCustomVersion.h"
 
 #define LOCTEXT_NAMESPACE "CustomizableObjectEditor"
 
@@ -12,6 +13,64 @@ FName UCustomizableObjectNodeMaterialSwitch::GetCategory() const
 	return UEdGraphSchema_CustomizableObject::PC_Material;
 }
 
+void UCustomizableObjectNodeMaterialSwitch::BackwardsCompatibleFixup(int32 CustomizableObjectCustomVersion)
+{
+	Super::BackwardsCompatibleFixup(CustomizableObjectCustomVersion);
+
+	if (CustomizableObjectCustomVersion == FCustomizableObjectCustomVersion::FixMaterialPinsRename)
+	{
+		const UEdGraphSchema_CustomizableObject* Schema = GetDefault<UEdGraphSchema_CustomizableObject>();
+
+		// if there are pins that were manually fixed and re-saved after the FCustomizableObjectCustomVersion::MaterialPinsRename
+		// conversion, convert them to the new fixed version.
+
+		// In the new fix we only change the friendy name, but the actual name remains "material..."
+
+		// Output
+		UEdGraphPin* OldPin = FindPin(TEXT("Mesh Section"),EEdGraphPinDirection::EGPD_Output);
+		if (OldPin)
+		{
+			UEdGraphPin* NewPin = FindPin(TEXT("Material"), EEdGraphPinDirection::EGPD_Output);
+			if (!NewPin)
+			{
+				NewPin = CustomCreatePin(EGPD_Output, Schema->PC_Material, FName("Material"));
+			}
+
+			FGuid PinId = NewPin->PinId;
+			NewPin->CopyPersistentDataFromOldPin(*OldPin);
+			NewPin->PinId = PinId;
+			NewPin->bHidden = OldPin->bHidden;
+
+			OutputPinReference = NewPin;
+
+			CustomRemovePin(*OldPin);
+		}
+
+		// Inputs
+		int32 InputCount = ReloadingElementsNames.Num();
+		for (int32 InputIndex=0; InputIndex<InputCount; ++InputIndex)
+		{
+			OldPin = FindPin(FString::Printf(TEXT("Mesh Section %d "),InputIndex), EEdGraphPinDirection::EGPD_Input);
+			if (OldPin)
+			{
+				FString NewName = FString::Printf(TEXT("Material %d "), InputIndex);
+				UEdGraphPin* NewPin = FindPin(NewName, EEdGraphPinDirection::EGPD_Input);
+				if (!NewPin)
+				{
+					NewPin = CustomCreatePin(EEdGraphPinDirection::EGPD_Input, Schema->PC_Material, FName(*NewName));
+				}
+
+				FGuid PinId = NewPin->PinId;
+				NewPin->CopyPersistentDataFromOldPin(*OldPin);
+				NewPin->PinId = PinId;
+				NewPin->bHidden = OldPin->bHidden;
+
+				CustomRemovePin(*OldPin);
+			}
+		}
+	}
+}
+
 
 #undef LOCTEXT_NAMESPACE
 
diff --git a/Engine/Plugins/Mutable/Source/CustomizableObjectEditor/Private/MuCOE/Nodes/CustomizableObjectNodeSwitchBase.cpp b/Engine/Plugins/Mutable/Source/CustomizableObjectEditor/Private/MuCOE/Nodes/CustomizableObjectNodeSwitchBase.cpp
--- a/Engine/Plugins/Mutable/Source/CustomizableObjectEditor/Private/MuCOE/Nodes/CustomizableObjectNodeSwitchBase.cpp
+++ b/Engine/Plugins/Mutable/Source/CustomizableObjectEditor/Private/MuCOE/Nodes/CustomizableObjectNodeSwitchBase.cpp
@@ -69,6 +69,7 @@ void UCustomizableObjectNodeSwitchBase::AllocateDefaultPins(UCustomizableObjectN
 
 	UEdGraphPin* OutputPin = CustomCreatePin(EGPD_Output, GetCategory(), FName(UEdGraphSchema_CustomizableObject::GetPinCategoryName(GetCategory()).ToString()));
 	OutputPin->bDefaultValueIsIgnored = true;
+	OutputPin->PinFriendlyName = UEdGraphSchema_CustomizableObject::GetPinCategoryFriendlyName(GetCategory());
 
 	OutputPinReference = FEdGraphPinReference(OutputPin);
 
@@ -103,7 +104,7 @@ void UCustomizableObjectNodeSwitchBase::Serialize(FArchive& Ar)
 
 FText UCustomizableObjectNodeSwitchBase::GetNodeTitle(ENodeTitleType::Type TitleType) const
 {
-	return FText::Format(LOCTEXT("Switch_Title", "{0} Switch"), UEdGraphSchema_CustomizableObject::GetPinCategoryName(GetCategory()));
+	return FText::Format(LOCTEXT("Switch_Title", "{0} Switch"), UEdGraphSchema_CustomizableObject::GetPinCategoryFriendlyName(GetCategory()));
 }
 
 
diff --git a/Engine/Plugins/Mutable/Source/CustomizableObjectEditor/Private/MuCOE/Nodes/CustomizableObjectNodeVariation.cpp b/Engine/Plugins/Mutable/Source/CustomizableObjectEditor/Private/MuCOE/Nodes/CustomizableObjectNodeVariation.cpp
--- a/Engine/Plugins/Mutable/Source/CustomizableObjectEditor/Private/MuCOE/Nodes/CustomizableObjectNodeVariation.cpp
+++ b/Engine/Plugins/Mutable/Source/CustomizableObjectEditor/Private/MuCOE/Nodes/CustomizableObjectNodeVariation.cpp
@@ -30,7 +30,8 @@ void UCustomizableObjectNodeVariation::AllocateDefaultPins(UCustomizableObjectNo
 
 	{
 		const FName PinName = FName(UEdGraphSchema_CustomizableObject::GetPinCategoryName(GetCategory()).ToString());
-		CustomCreatePin(EGPD_Output, Category, PinName);
+		UEdGraphPin* Pin = CustomCreatePin(EGPD_Output, Category, PinName);
+		Pin->PinFriendlyName = UEdGraphSchema_CustomizableObject::GetPinCategoryFriendlyName(GetCategory());
 	}
 	
 	VariationsPins.SetNum(VariationsData.Num());
@@ -84,7 +85,7 @@ UEdGraphPin* UCustomizableObjectNodeVariation::VariationPin(int32 Index) const
 
 FText UCustomizableObjectNodeVariation::GetNodeTitle(ENodeTitleType::Type TitleType) const
 {
-	return FText::Format(LOCTEXT("Variation_Node_Title", "{0} Variation"), UEdGraphSchema_CustomizableObject::GetPinCategoryName(GetCategory()));
+	return FText::Format(LOCTEXT("Variation_Node_Title", "{0} Variation"), UEdGraphSchema_CustomizableObject::GetPinCategoryFriendlyName(GetCategory()));
 }
 
 
@@ -97,7 +98,7 @@ FLinearColor UCustomizableObjectNodeVariation::GetNodeTitleColor() const
 
 FText UCustomizableObjectNodeVariation::GetTooltipText() const
 {
-	return FText::Format(LOCTEXT("Variation_Tooltip", "Select a {0} depending on what tags are active."), UEdGraphSchema_CustomizableObject::GetPinCategoryName(GetCategory()));
+	return FText::Format(LOCTEXT("Variation_Tooltip", "Select a {0} depending on what tags are active."), UEdGraphSchema_CustomizableObject::GetPinCategoryFriendlyName(GetCategory()));
 }
 
 
diff --git a/Engine/Plugins/Mutable/Source/CustomizableObjectEditor/Private/MuCOE/PinViewer/SPinViewerListRow.cpp b/Engine/Plugins/Mutable/Source/CustomizableObjectEditor/Private/MuCOE/PinViewer/SPinViewerListRow.cpp
--- a/Engine/Plugins/Mutable/Source/CustomizableObjectEditor/Private/MuCOE/PinViewer/SPinViewerListRow.cpp
+++ b/Engine/Plugins/Mutable/Source/CustomizableObjectEditor/Private/MuCOE/PinViewer/SPinViewerListRow.cpp
@@ -52,7 +52,7 @@ TSharedRef<SWidget> SPinViewerListRow::GenerateWidgetForColumn(const FName& Colu
 		return SNew(SVerticalBox)
 			+ SVerticalBox::Slot().AutoHeight().Padding(4).VAlign(VAlign_Center).HAlign(HAlign_Center)
 			[
-				SNew(STextBlock).Text(UEdGraphSchema_CustomizableObject::GetPinCategoryName(Pin->PinType.PinCategory))
+				SNew(STextBlock).Text(UEdGraphSchema_CustomizableObject::GetPinCategoryFriendlyName(Pin->PinType.PinCategory))
 			];
 	}
 	else if (ColumnName == SPinViewer::COLUMN_VISIBILITY)
diff --git a/Engine/Plugins/Mutable/Source/CustomizableObjectEditor/Public/MuCOE/EdGraphSchema_CustomizableObject.h b/Engine/Plugins/Mutable/Source/CustomizableObjectEditor/Public/MuCOE/EdGraphSchema_CustomizableObject.h
--- a/Engine/Plugins/Mutable/Source/CustomizableObjectEditor/Public/MuCOE/EdGraphSchema_CustomizableObject.h
+++ b/Engine/Plugins/Mutable/Source/CustomizableObjectEditor/Public/MuCOE/EdGraphSchema_CustomizableObject.h
@@ -167,6 +167,8 @@ public:
 	
 	/** Given a pin category name, get its friendly name (user readable). */
 	static FText GetPinCategoryName(const FName& PinCategory);
+	static FText GetPinCategoryFriendlyName(const FName& PinCategory);
+
 };
 
 
diff --git a/Engine/Plugins/Mutable/Source/CustomizableObjectEditor/Public/MuCOE/Nodes/CustomizableObjectNodeMaterialSwitch.h b/Engine/Plugins/Mutable/Source/CustomizableObjectEditor/Public/MuCOE/Nodes/CustomizableObjectNodeMaterialSwitch.h
--- a/Engine/Plugins/Mutable/Source/CustomizableObjectEditor/Public/MuCOE/Nodes/CustomizableObjectNodeMaterialSwitch.h
+++ b/Engine/Plugins/Mutable/Source/CustomizableObjectEditor/Public/MuCOE/Nodes/CustomizableObjectNodeMaterialSwitch.h
@@ -12,7 +12,10 @@ class CUSTOMIZABLEOBJECTEDITOR_API UCustomizableObjectNodeMaterialSwitch : publi
 {
 public:
 	GENERATED_BODY()
-	
+
+	// UCustomizableObjectNode interface
+	virtual void BackwardsCompatibleFixup(int32 CustomizableObjectCustomVersion) override;
+
 	// UCustomizableObjectNodeSwitchBase interface
 	virtual FName GetCategory() const override;
 };
