[Mutable] Hide layout options based on the active layout packaging strategy.
Update NodeAnimationPose properties category.

#rnx

diff --git a/Engine/Plugins/Mutable/Source/CustomizableObjectEditor/Private/MuCOE/SCustomizableObjectLayoutEditor.cpp b/Engine/Plugins/Mutable/Source/CustomizableObjectEditor/Private/MuCOE/SCustomizableObjectLayoutEditor.cpp
--- a/Engine/Plugins/Mutable/Source/CustomizableObjectEditor/Private/MuCOE/SCustomizableObjectLayoutEditor.cpp
+++ b/Engine/Plugins/Mutable/Source/CustomizableObjectEditor/Private/MuCOE/SCustomizableObjectLayoutEditor.cpp
@@ -249,6 +249,7 @@ void FCustomizableObjectLayoutEditorDetailsBuilder::CustomizeDetails(IDetailLayo
 	if (bShowGridSize)
 	{
 		LayoutOptionsGroup->AddWidgetRow()
+			.Visibility(TAttribute<EVisibility>(LayoutEditor.Get(), &SCustomizableObjectLayoutEditor::GridSizeVisibility))
 			.NameContent()
 			[
 				SNew(STextBlock)
@@ -268,6 +269,7 @@ void FCustomizableObjectLayoutEditorDetailsBuilder::CustomizeDetails(IDetailLayo
 	if (bShowMaxGridSize)
 	{
 		LayoutOptionsGroup->AddWidgetRow()
+			.Visibility(TAttribute<EVisibility>(LayoutEditor.Get(), &SCustomizableObjectLayoutEditor::FixedStrategyOptionsVisibility))
 			.NameContent()
 			[
 				SNew(STextBlock)
@@ -318,7 +320,7 @@ void FCustomizableObjectLayoutEditorDetailsBuilder::CustomizeDetails(IDetailLayo
 		// Warning selector group widget
 		IDetailGroup* IgnoreWarningsGroup = &LayoutCategory.AddGroup(TEXT("LayoutEditor_IgnoreWarningsGroup"), LOCTEXT("LayoutEditor_IgnoreWarningsGroup", "Ignore Unassigned Vertives Warning group"), false, true);
 		IgnoreWarningsGroup->HeaderRow()
-			.Visibility(TAttribute<EVisibility>(LayoutEditor.Get(), &SCustomizableObjectLayoutEditor::LayoutOptionsVisibility))
+			.Visibility(TAttribute<EVisibility>(LayoutEditor.Get(), &SCustomizableObjectLayoutEditor::WarningOptionsVisibility))
 			.NameContent()
 			[
 				SNew(STextBlock)
@@ -340,7 +342,7 @@ void FCustomizableObjectLayoutEditorDetailsBuilder::CustomizeDetails(IDetailLayo
 
 		// LOD selector widget
 		IgnoreWarningsGroup->AddWidgetRow()
-			.Visibility(TAttribute<EVisibility>(LayoutEditor.Get(), &SCustomizableObjectLayoutEditor::LayoutOptionsVisibility))
+			.Visibility(TAttribute<EVisibility>(LayoutEditor.Get(), &SCustomizableObjectLayoutEditor::WarningOptionsVisibility))
 			.NameContent()
 			[
 				SAssignNew(LayoutEditor->LODSelectorTextWidget, STextBlock)
@@ -1473,6 +1475,12 @@ void SCustomizableObjectLayoutEditor::FillLayoutComboBoxOptions()
 }
 
 
+EVisibility SCustomizableObjectLayoutEditor::GridSizeVisibility() const
+{
+	return CurrentLayout && CurrentLayout->PackingStrategy != ECustomizableObjectTextureLayoutPackingStrategy::Overlay ? EVisibility::Visible : EVisibility::Collapsed;
+}
+
+
 EVisibility SCustomizableObjectLayoutEditor::LayoutOptionsVisibility() const
 {
 	return CurrentLayout ? EVisibility::Visible : EVisibility::Collapsed;
@@ -1500,5 +1508,11 @@ EVisibility SCustomizableObjectLayoutEditor::FixedStrategyOptionsVisibility() co
 	return  CurrentLayout && CurrentLayout->PackingStrategy == ECustomizableObjectTextureLayoutPackingStrategy::Fixed ? EVisibility::Visible : EVisibility::Collapsed;
 }
 
+
+EVisibility SCustomizableObjectLayoutEditor::WarningOptionsVisibility() const
+{
+	return CurrentLayout && CurrentLayout->PackingStrategy != ECustomizableObjectTextureLayoutPackingStrategy::Overlay ? EVisibility::Visible : EVisibility::Collapsed;
+}
+
 #undef LOCTEXT_NAMESPACE
 
diff --git a/Engine/Plugins/Mutable/Source/CustomizableObjectEditor/Private/MuCOE/SCustomizableObjectLayoutEditor.h b/Engine/Plugins/Mutable/Source/CustomizableObjectEditor/Private/MuCOE/SCustomizableObjectLayoutEditor.h
--- a/Engine/Plugins/Mutable/Source/CustomizableObjectEditor/Private/MuCOE/SCustomizableObjectLayoutEditor.h
+++ b/Engine/Plugins/Mutable/Source/CustomizableObjectEditor/Private/MuCOE/SCustomizableObjectLayoutEditor.h
@@ -241,8 +241,10 @@ private:
 	FText GetLayoutPackagingStrategyName() const;
 	FText GetLayoutPackagingStrategyToolTip() const;
 
+	EVisibility GridSizeVisibility() const;
 	EVisibility LayoutOptionsVisibility() const;
 	EVisibility AutoBlocksStrategyVisibility() const;
 	EVisibility AutoBlocksMergeStrategyVisibility() const;
 	EVisibility FixedStrategyOptionsVisibility() const;
+	EVisibility WarningOptionsVisibility() const;
 };
diff --git a/Engine/Plugins/Mutable/Source/CustomizableObjectEditor/Public/MuCOE/Nodes/CustomizableObjectNodeAnimationPose.h b/Engine/Plugins/Mutable/Source/CustomizableObjectEditor/Public/MuCOE/Nodes/CustomizableObjectNodeAnimationPose.h
--- a/Engine/Plugins/Mutable/Source/CustomizableObjectEditor/Public/MuCOE/Nodes/CustomizableObjectNodeAnimationPose.h
+++ b/Engine/Plugins/Mutable/Source/CustomizableObjectEditor/Public/MuCOE/Nodes/CustomizableObjectNodeAnimationPose.h
@@ -24,7 +24,7 @@ public:
 
 	UCustomizableObjectNodeAnimationPose();
 
-	UPROPERTY(EditAnywhere, Category=CustomizableObject)
+	UPROPERTY(EditAnywhere, Category=NoCategory)
 	TObjectPtr<UPoseAsset> PoseAsset;
 
 	// UObject interface.
