[Mutable] Minor UI changes.

- Typos.
- Removed new CO dialog.
- Added UCustomizableObjectNode::IsExperimental.
- Moved experimental nodes to is own category.
- Moved Pose Mesh to Mesh category.

#rnx

diff --git a/Engine/Plugins/Mutable/Source/CustomizableObjectEditor/Private/MuCOE/CustomizableObjectEditorFunctionLibrary.cpp b/Engine/Plugins/Mutable/Source/CustomizableObjectEditor/Private/MuCOE/CustomizableObjectEditorFunctionLibrary.cpp
--- a/Engine/Plugins/Mutable/Source/CustomizableObjectEditor/Private/MuCOE/CustomizableObjectEditorFunctionLibrary.cpp
+++ b/Engine/Plugins/Mutable/Source/CustomizableObjectEditor/Private/MuCOE/CustomizableObjectEditorFunctionLibrary.cpp
@@ -91,10 +91,6 @@ UCustomizableObject* UCustomizableObjectEditorFunctionLibrary::NewCustomizableOb
 	}
 	
 	UCustomizableObjectFactory* Factory = NewObject<UCustomizableObjectFactory>();
-	Factory->CreationSettings.bEmptyObject = false;
-	Factory->CreationSettings.bIsChildObject = Parameters.ParentObject != nullptr;
-	Factory->CreationSettings.ParentObject = Parameters.ParentObject;
-	Factory->CreationSettings.GroupNodeName = Parameters.ParentGroupNode;
 
 	const FString PackagePath = FPackageName::GetLongPackagePath(PackageName);
 
diff --git a/Engine/Plugins/Mutable/Source/CustomizableObjectEditor/Private/MuCOE/CustomizableObjectFactory.cpp b/Engine/Plugins/Mutable/Source/CustomizableObjectEditor/Private/MuCOE/CustomizableObjectFactory.cpp
--- a/Engine/Plugins/Mutable/Source/CustomizableObjectEditor/Private/MuCOE/CustomizableObjectFactory.cpp
+++ b/Engine/Plugins/Mutable/Source/CustomizableObjectEditor/Private/MuCOE/CustomizableObjectFactory.cpp
@@ -5,35 +5,23 @@
 #include "MuCO/CustomizableObject.h"
 #include "MuCO/CustomizableObjectPrivate.h"
 #include "MuCOE/CustomizableObjectGraph.h"
-#include "MuCOE/GraphTraversal.h"
-#include "MuCOE/Nodes/CustomizableObjectNodeObject.h"
-#include "MuCOE/Nodes/CustomizableObjectNodeObjectGroup.h"
-#include "MuCOE/EdGraphSchema_CustomizableObject.h"
-
-#include "ContentBrowserModule.h"
-#include "DetailLayoutBuilder.h"
 #include "Editor.h"
 #include "Editor/EditorEngine.h"
 #include "Engine/SkeletalMesh.h"
-#include "IContentBrowserSingleton.h"
 #include "Input/Reply.h"
 #include "Layout/Visibility.h"
-#include "PropertyCustomizationHelpers.h"
 #include "PropertyEditorModule.h"
 #include "Styling/SlateTypes.h"
 #include "Types/SlateEnums.h"
 #include "Widgets/Input/SButton.h"
 #include "Widgets/Input/SCheckBox.h"
 #include "Widgets/Input/SEditableTextBox.h"
-#include "Widgets/Input/SSpinBox.h"
-#include "Widgets/Input/STextComboBox.h"
 #include "Widgets/Layout/SBorder.h"
 #include "Widgets/SWindow.h"
 
 
 #define LOCTEXT_NAMESPACE "CustomizableObjectFactory"
 
-static const FVector2D NewCOWindowsSize = FVector2D(300, 230);
 
 UCustomizableObjectFactory::UCustomizableObjectFactory()
 	: Super()
@@ -57,15 +45,6 @@ UClass* UCustomizableObjectFactory::ResolveSupportedClass()
 }
 
 
-bool UCustomizableObjectFactory::ConfigureProperties()
-{
-	TSharedRef<FCustomizableObjectFactoryUI> COSettingsWindow = MakeShareable(new FCustomizableObjectFactoryUI());
-	CreationSettings = COSettingsWindow->ConstructFactoryUI();
-	
-	return COSettingsWindow->CanCreateObject();
-}
-
-
 UObject* UCustomizableObjectFactory::FactoryCreateNew(UClass* Class, UObject* InParent, FName Name, EObjectFlags Flags, UObject* Context, FFeedbackContext* Warn)
 {
 	UCustomizableObject* NewObj = NewObject<UCustomizableObject>(InParent, Class, Name, Flags);
@@ -85,379 +64,9 @@ UObject* UCustomizableObjectFactory::FactoryCreateNew(UClass* Class, UObject* In
 	ObjectPrivate->GetSource() = Source;
 	
 	Source->AddEssentialGraphNodes();
-
-	UCustomizableObjectNodeObject* BaseObjectNode = nullptr;
-
-	for (const TObjectPtr<UEdGraphNode>& AuxNode : Source->Nodes)
-	{
-		UCustomizableObjectNodeObject* BaseNode = Cast<UCustomizableObjectNodeObject>(AuxNode);
-
-		if (BaseNode && BaseNode->bIsBase)
-		{
-			BaseObjectNode = BaseNode;
-			break;
-		}
-	}
-
-	if (!BaseObjectNode)
-	{
-		return NewObj;
-	}
 	
-	if (!CreationSettings.bEmptyObject)
-	{
-		if (CreationSettings.bIsChildObject && CreationSettings.ParentObject.IsValid())
-		{
-			BaseObjectNode->ParentObject = CreationSettings.ParentObject.Get();
-
-			TArray<UCustomizableObjectNodeObjectGroup*> GroupNodes;
-			CreationSettings.ParentObject->GetPrivate()->GetSource()->GetNodesOfClass<UCustomizableObjectNodeObjectGroup>(GroupNodes);
-
-			for (UCustomizableObjectNodeObjectGroup* GroupNode : GroupNodes)
-			{
-				if (GroupNode->GroupName == CreationSettings.GroupNodeName)
-				{
-					BaseObjectNode->ParentObjectGroupId = GroupNode->NodeGuid;
-
-					break;
-				}
-			}
-
-			ObjectPrivate->SetIsChildObject(true);
-		}
-	}
-
 	return NewObj;
 }
 
 
-// Factory Settings UI ----------------------------------------------------------------------------------------------------------------------------------------------------------------
-
-const FCustomizableObjectOptions FCustomizableObjectFactoryUI::ConstructFactoryUI()
-{
-	FPropertyEditorModule& PropertyEditorModule = FModuleManager::GetModuleChecked<FPropertyEditorModule>("PropertyEditor");
-	FDetailsViewArgs DetailsViewArgs;
-	DetailsViewArgs.NameAreaSettings = FDetailsViewArgs::HideNameArea;
-	
-	TSharedRef<IDetailsView> DetailsView = PropertyEditorModule.CreateDetailView(DetailsViewArgs);
-	DetailsView->SetObject(nullptr);
-	
-	// Settings window
-	COSettingsWindow = SNew(SWindow)
-	.Title(LOCTEXT("CustomizableObjectFactoryptions", "New Costumizable Object"))
-	.ClientSize(NewCOWindowsSize)
-	.SupportsMinimize(false)
-	.SupportsMaximize(false)
-	[
-		SNew(SBorder)
-		.BorderImage(FAppStyle::GetBrush("Menu.Background")) // Adds lighter background
-		.Padding(5)
-		[
-			// Is Child Object Widgets
-			SNew(SVerticalBox)
-			+ SVerticalBox::Slot()
-			.AutoHeight()
-			.HAlign(HAlign_Left)
-			.Padding(0.0f, 5.0f, 0.0f, 0.0f)
-			[
-				SNew(SHorizontalBox)
-				+ SHorizontalBox::Slot()
-				.AutoWidth()
-				.Padding(0.0f, 2.0f, 0.0f, 0.0f)
-				[
-					SNew(STextBlock)
-					.Text(LOCTEXT("IsChildObject_CheckBoxName","Is Child Object:"))
-					.Font(IDetailLayoutBuilder::GetDetailFont())
-				]
-
-				+ SHorizontalBox::Slot()
-				.AutoWidth()
-				.Padding(5.0f, 0.0f, 0.0f, 0.0f)
-				[
-					SNew(SCheckBox)
-					.IsChecked(Options.bIsChildObject ? ECheckBoxState::Checked : ECheckBoxState::Unchecked)
-					.OnCheckStateChanged(this, &FCustomizableObjectFactoryUI::OnCheckBoxChanged)
-				]
-			]
-				
-			// Parent Selector Widgets
-			+ SVerticalBox::Slot()
-			.AutoHeight()
-			.Padding(0.0f, 10.0f, 0.0f, 10.0f)
-			.HAlign(HAlign_Fill)
-			.VAlign(VAlign_Fill)
-			[
-				SNew(SBorder)
-				.Padding(5.0f, 10.0f, 0.0f, 10.0f)
-				.Visibility(this, &FCustomizableObjectFactoryUI::GetParentWidgetsVisibility)
-				.BorderImage(FAppStyle::Get().GetBrush("Brushes.Panel"))
-				.HAlign(HAlign_Fill)
-				[
-					SNew(SVerticalBox)
-					+ SVerticalBox::Slot()
-					.AutoHeight()
-					.HAlign(HAlign_Left)
-					[
-						SNew(STextBlock)
-						.Text(LOCTEXT("SelectedParent_Text", "Parent Object: "))
-						.Font(IDetailLayoutBuilder::GetDetailFont())
-					]
-
-					+ SVerticalBox::Slot()
-					.Padding(5.0f, 5.0f, 10.0f, 0.0f)
-					.AutoHeight()
-					[
-						SNew(SObjectPropertyEntryBox)
-						.AllowedClass(UCustomizableObject::StaticClass())
-						.OnObjectChanged(this, &FCustomizableObjectFactoryUI::OnPickedCustomizableObjectParent)
-						.ObjectPath(this, &FCustomizableObjectFactoryUI::GetSelectedCustomizableObjectPath)
-						.DisplayThumbnail(true)
-						.ThumbnailPool(UThumbnailManager::Get().GetSharedThumbnailPool())
-						.AllowClear(false)
-						.AllowCreate(false)
-						.DisplayBrowse(false)
-						.DisplayUseSelected(false)
-					]
-
-					+ SVerticalBox::Slot()
-					.AutoHeight()
-					.HAlign(HAlign_Left)
-					.Padding(0.0f, 15.0f, 0.0f, 0.0f)
-					[
-						SNew(STextBlock)
-						.Text(LOCTEXT("SelecteGroup_Text", "Group Node: "))
-						.Font(IDetailLayoutBuilder::GetDetailFont())
-					]
-
-					+ SVerticalBox::Slot()
-					.Padding(5.0f, 5.0f, 10.0f, 0.0f)
-					.AutoHeight()
-					[
-						SAssignNew(GroupSelector, STextComboBox)
-						.OptionsSource(&GroupOptions)
-						.IsEnabled(this, &FCustomizableObjectFactoryUI::IsNodeGroupSelectorEnabled)
-						.OnSelectionChanged(this, &FCustomizableObjectFactoryUI::OnSelectGroupComboBox)
-						.Font(IDetailLayoutBuilder::GetDetailFont())
-					]
-				]
-			]
-
-			+ SVerticalBox::Slot()
-			[
-				SNew(SSpacer)
-			]
-
-			+ SVerticalBox::Slot()
-			.HAlign(HAlign_Fill)
-			.AutoHeight()
-			[
-				SNew(SHorizontalBox)
-				+ SHorizontalBox::Slot()
-				.AutoWidth()
-				.HAlign(HAlign_Left)
-				[
-					SNew(SButton)
-						.Text(LOCTEXT("CreateEmptyCO", "Create Empty"))
-						.ToolTipText(LOCTEXT("CreateEmptyCO_Tooltip","Create a Customizable Object without a components or reference skeletal mesh."))
-						.OnClicked(this, &FCustomizableObjectFactoryUI::OnCreate, true)
-				]
-
-				+ SHorizontalBox::Slot()
-				[
-					SNew(SSpacer)
-				]
-
-				+ SHorizontalBox::Slot()
-				.AutoWidth()
-				.HAlign(HAlign_Right)
-				[
-					SNew(SHorizontalBox)
-					+ SHorizontalBox::Slot()
-					.AutoWidth()
-					[
-						SNew(SButton)
-							.Text(LOCTEXT("OK", "OK"))
-							.ToolTipText(this, &FCustomizableObjectFactoryUI::GetOKButtonTooltip)
-							.IsEnabled(this, &FCustomizableObjectFactoryUI::IsConfigurationValid)
-							.OnClicked(this, &FCustomizableObjectFactoryUI::OnCreate, false)
-					]
-					+ SHorizontalBox::Slot()
-					.AutoWidth()
-					[
-						SNew(SButton)
-							.Text(LOCTEXT("Cancel", "Cancel"))
-							.OnClicked(this, &FCustomizableObjectFactoryUI::OnCancel)
-					]
-				]
-			]
-		]
-	];
-
-	GEditor->EditorAddModalWindow(COSettingsWindow.ToSharedRef());
-	COSettingsWindow.Reset();
-
-	return Options;
-}
-
-
-void FCustomizableObjectFactoryUI::OnCheckBoxChanged(ECheckBoxState State)
-{
-	Options.bIsChildObject = State == ECheckBoxState::Checked;
-
-	// Reset settings
-	Options.ParentObject = nullptr;
-	Options.GroupNodeName.Empty();
-}
-
-
-FReply FCustomizableObjectFactoryUI::OnCreate(bool bIsEmpty)
-{
-	bCreateObject = true;
-	Options.bEmptyObject = bIsEmpty;
-
-	if (COSettingsWindow.IsValid())
-	{
-		COSettingsWindow->RequestDestroyWindow();
-	}
-
-	return FReply::Handled();
-}
-
-
-FText FCustomizableObjectFactoryUI::GetOKButtonTooltip() const
-{
-	FString Tooltip = "Create a Customizable Object with the settings selected.";
-
-	if (Options.bIsChildObject)
-	{
-		if (!IsConfigurationValid())
-		{
-			Tooltip = "Select a Parent Object and the name of a Group Object Node.";
-		}
-	}
-	else
-	{
-		if (!IsConfigurationValid())
-		{
-			Tooltip = "Set the number of components that the Customizable Object will have. Then select a Reference Skeletal Mesh and a Name for each component.";
-		}
-	}
-
-	return FText::FromString(Tooltip);
-}
-
-
-FReply FCustomizableObjectFactoryUI::OnCancel()
-{
-	bCreateObject = false;
-
-	if (COSettingsWindow.IsValid())
-	{
-		COSettingsWindow->RequestDestroyWindow();
-	}
-
-	return FReply::Handled();
-}
-
-
-bool FCustomizableObjectFactoryUI::IsConfigurationValid() const
-{
-	if (Options.bIsChildObject)
-	{
-		if (Options.ParentObject.IsValid() && !Options.GroupNodeName.IsEmpty())
-		{
-			return true;
-		}
-	}
-
-	return false;
-}
-
-
-bool FCustomizableObjectFactoryUI::CanCreateObject() const
-{
-	return bCreateObject;
-}
-
-
-void FCustomizableObjectFactoryUI::OnPickedCustomizableObjectParent(const FAssetData& SelectedAsset)
-{
-	if (UCustomizableObject* Parent = Cast<UCustomizableObject>(SelectedAsset.GetAsset()))
-	{
-		Options.ParentObject = Parent;
-
-		if (GroupSelector.IsValid())
-		{
-			GroupSelector->ClearSelection();
-			GenerateGroupOptions();
-		}
-	}
-}
-
-
-FString FCustomizableObjectFactoryUI::GetSelectedCustomizableObjectPath() const
-{
-	return Options.ParentObject.IsValid() ? Options.ParentObject->GetPathName() : FString();
-}
-
-
-EVisibility FCustomizableObjectFactoryUI::GetParentWidgetsVisibility() const
-{
-	return Options.bIsChildObject ? EVisibility::Visible : EVisibility::Collapsed;
-}
-
-
-bool FCustomizableObjectFactoryUI::IsNodeGroupSelectorEnabled() const
-{
-	return Options.ParentObject.IsValid() ? true : false;
-}
-
-
-void FCustomizableObjectFactoryUI::GenerateGroupOptions()
-{
-	GroupOptions.Empty();
-
-	if (Options.ParentObject.IsValid())
-	{
-		TArray<UCustomizableObjectNodeObjectGroup*> GroupNodes;
-		Options.ParentObject->GetPrivate()->GetSource()->GetNodesOfClass<UCustomizableObjectNodeObjectGroup>(GroupNodes);
-
-		for (UCustomizableObjectNodeObjectGroup* GroupNode : GroupNodes)
-		{
-			GroupOptions.Add(MakeShareable(new FString(GroupNode->GroupName)));
-		}
-
-		if (GroupSelector.IsValid())
-		{
-			GroupSelector->RefreshOptions();
-		}
-	}
-}
-
-
-void FCustomizableObjectFactoryUI::OnSelectGroupComboBox(TSharedPtr<FString> Selection, ESelectInfo::Type SelectInfo)
-{
-	if (Selection.IsValid())
-	{
-		Options.GroupNodeName = *Selection;
-		GroupSelector->SetSelectedItem(Selection);
-	}
-}
-
-
-EVisibility FCustomizableObjectFactoryUI::GetComponentWidgetsVisibility() const
-{
-	return !Options.bIsChildObject ? EVisibility::Visible : EVisibility::Collapsed;
-}
-
-
-void FCustomizableObjectFactoryUI::OnSelectComponentComboBox(TSharedPtr<FString> Selection, ESelectInfo::Type SelectInfo)
-{
-	if (Selection)
-	{
-		ComponentSelector->SetSelectedItem(Selection);
-	}
-}
-
-
 #undef LOCTEXT_NAMESPACE
diff --git a/Engine/Plugins/Mutable/Source/CustomizableObjectEditor/Private/MuCOE/CustomizableObjectGraph.cpp b/Engine/Plugins/Mutable/Source/CustomizableObjectEditor/Private/MuCOE/CustomizableObjectGraph.cpp
--- a/Engine/Plugins/Mutable/Source/CustomizableObjectEditor/Private/MuCOE/CustomizableObjectGraph.cpp
+++ b/Engine/Plugins/Mutable/Source/CustomizableObjectEditor/Private/MuCOE/CustomizableObjectGraph.cpp
@@ -190,6 +190,6 @@ void UCustomizableObjectGraph::AddEssentialGraphNodes()
 	if (!bGraphHasBase)
 	{
 		UCustomizableObjectNodeObject* NodeTemplate = NewObject<UCustomizableObjectNodeObject>();
-		UCustomizableObjectNodeObject* BaseObjectNode = Cast<UCustomizableObjectNodeObject>(FCustomizableObjectSchemaAction_NewNode::CreateNode(this, nullptr, FVector2D::ZeroVector, Cast<UEdGraphNode>(NodeTemplate)));
+		FCustomizableObjectSchemaAction_NewNode::CreateNode(this, nullptr, FVector2D::ZeroVector, Cast<UEdGraphNode>(NodeTemplate));
 	}
 }
diff --git a/Engine/Plugins/Mutable/Source/CustomizableObjectEditor/Private/MuCOE/EdGraphSchema_CustomizableObject.cpp b/Engine/Plugins/Mutable/Source/CustomizableObjectEditor/Private/MuCOE/EdGraphSchema_CustomizableObject.cpp
--- a/Engine/Plugins/Mutable/Source/CustomizableObjectEditor/Private/MuCOE/EdGraphSchema_CustomizableObject.cpp
+++ b/Engine/Plugins/Mutable/Source/CustomizableObjectEditor/Private/MuCOE/EdGraphSchema_CustomizableObject.cpp
@@ -403,7 +403,6 @@ void UEdGraphSchema_CustomizableObject::GetGraphContextActions(FGraphContextMenu
 			ContextMenuBuilder.CreateTemplateNode<UCustomizableObjectNodeModifierMorphMeshSection>(),
 			ContextMenuBuilder.CreateTemplateNode<UCustomizableObjectNodeModifierClipMorph>(),
 			ContextMenuBuilder.CreateTemplateNode<UCustomizableObjectNodeModifierClipWithMesh>(),
-			ContextMenuBuilder.CreateTemplateNode<UCustomizableObjectNodeModifierClipDeform>(),
 			ContextMenuBuilder.CreateTemplateNode<UCustomizableObjectNodeModifierClipWithUVMask>(),
 		};
 
@@ -415,7 +414,6 @@ void UEdGraphSchema_CustomizableObject::GetGraphContextActions(FGraphContextMenu
 		{
 			ContextMenuBuilder.CreateTemplateNode<UCustomizableObjectNodeComponentMesh>(),
 			ContextMenuBuilder.CreateTemplateNode<UCustomizableObjectNodeComponentMeshAddTo>(),
-			ContextMenuBuilder.CreateTemplateNode<UCustomizableObjectNodeComponentPassthroughMesh>(),
 			ContextMenuBuilder.CreateTemplateNode<UCustomizableObjectNodeComponentVariation>(),
 			ContextMenuBuilder.CreateTemplateNode<UCustomizableObjectNodeComponentSwitch>(),
 		};
@@ -434,21 +432,12 @@ void UEdGraphSchema_CustomizableObject::GetGraphContextActions(FGraphContextMenu
 			ContextMenuBuilder.CreateTemplateNode<UCustomizableObjectNodeMeshMorphStackDefinition>(),
 			ContextMenuBuilder.CreateTemplateNode<UCustomizableObjectNodeMeshMorphStackApplication>(),
 			// ContextMenuBuilder.CreateTemplateNode<UCustomizableObjectNodeMeshGeometryOperation>(),
-			ContextMenuBuilder.CreateTemplateNode<UCustomizableObjectNodeMeshReshape>(),
-		};
-
-		AddNewNodeCategoryActionsFiltered(MeshTemplateNodes, ContextMenuBuilder, TEXT("Mesh"), GeneralGrouping, Filter);
-	}
-
-	{
-		UCustomizableObjectNode* PoseMeshTemplateNodes[]
-		{
 			ContextMenuBuilder.CreateTemplateNode<UCustomizableObjectNodeAnimationPose>(),
 		};
 
-		AddNewNodeCategoryActionsFiltered(PoseMeshTemplateNodes, ContextMenuBuilder, TEXT("PoseMesh"), GeneralGrouping, Filter);
+		AddNewNodeCategoryActionsFiltered(MeshTemplateNodes, ContextMenuBuilder, TEXT("Mesh"), GeneralGrouping, Filter);
 	}
-
+	
 	{
 		UCustomizableObjectNode* TextureTemplateNodes[] =
 		{
@@ -541,6 +530,17 @@ void UEdGraphSchema_CustomizableObject::GetGraphContextActions(FGraphContextMenu
 
 		AddNewNodeCategoryActionsFiltered(ProjectorTemplateNodes, ContextMenuBuilder, TEXT("Projector"), GeneralGrouping, Filter);
 	}
+	
+	{
+		UCustomizableObjectNode* ExperimentalTemplateNodes[] =
+		{
+			ContextMenuBuilder.CreateTemplateNode<UCustomizableObjectNodeComponentPassthroughMesh>(),
+			ContextMenuBuilder.CreateTemplateNode<UCustomizableObjectNodeMeshReshape>(),
+			ContextMenuBuilder.CreateTemplateNode<UCustomizableObjectNodeModifierClipDeform>(),
+		};
+
+		AddNewNodeCategoryActionsFiltered(ExperimentalTemplateNodes, ContextMenuBuilder, TEXT("Experimental"), GeneralGrouping, Filter);
+	}
 
 	{
 		// External Pin Nodes
diff --git a/Engine/Plugins/Mutable/Source/CustomizableObjectEditor/Private/MuCOE/Nodes/CustomizableObjectNode.cpp b/Engine/Plugins/Mutable/Source/CustomizableObjectEditor/Private/MuCOE/Nodes/CustomizableObjectNode.cpp
--- a/Engine/Plugins/Mutable/Source/CustomizableObjectEditor/Private/MuCOE/Nodes/CustomizableObjectNode.cpp
+++ b/Engine/Plugins/Mutable/Source/CustomizableObjectEditor/Private/MuCOE/Nodes/CustomizableObjectNode.cpp
@@ -491,6 +491,12 @@ void UCustomizableObjectNode::AddPinData(const UEdGraphPin& Pin, UCustomizableOb
 }
 
 
+bool UCustomizableObjectNode::IsExperimental() const
+{
+	return false;
+}
+
+
 TArray<UEdGraphPin*> UCustomizableObjectNode::GetAllOrphanPins() const
 {
 	TArray<UEdGraphPin*> OrphanPins;
diff --git a/Engine/Plugins/Mutable/Source/CustomizableObjectEditor/Private/MuCOE/Nodes/CustomizableObjectNodeComponentPassthroughMesh.cpp b/Engine/Plugins/Mutable/Source/CustomizableObjectEditor/Private/MuCOE/Nodes/CustomizableObjectNodeComponentPassthroughMesh.cpp
--- a/Engine/Plugins/Mutable/Source/CustomizableObjectEditor/Private/MuCOE/Nodes/CustomizableObjectNodeComponentPassthroughMesh.cpp
+++ b/Engine/Plugins/Mutable/Source/CustomizableObjectEditor/Private/MuCOE/Nodes/CustomizableObjectNodeComponentPassthroughMesh.cpp
@@ -131,6 +131,12 @@ void UCustomizableObjectNodeComponentPassthroughMesh::AllocateDefaultPins(UCusto
 }
 
 
+bool UCustomizableObjectNodeComponentPassthroughMesh::IsExperimental() const
+{
+	return true;
+}
+
+
 FText UCustomizableObjectNodeComponentPassthroughMesh::GetNodeTitle(ENodeTitleType::Type TitleType) const
 {
 	if (TitleType == ENodeTitleType::ListView)
diff --git a/Engine/Plugins/Mutable/Source/CustomizableObjectEditor/Private/MuCOE/Nodes/CustomizableObjectNodeMeshClipDeform.cpp b/Engine/Plugins/Mutable/Source/CustomizableObjectEditor/Private/MuCOE/Nodes/CustomizableObjectNodeMeshClipDeform.cpp
--- a/Engine/Plugins/Mutable/Source/CustomizableObjectEditor/Private/MuCOE/Nodes/CustomizableObjectNodeMeshClipDeform.cpp
+++ b/Engine/Plugins/Mutable/Source/CustomizableObjectEditor/Private/MuCOE/Nodes/CustomizableObjectNodeMeshClipDeform.cpp
@@ -36,6 +36,12 @@ void UCustomizableObjectNodeModifierClipDeform::BackwardsCompatibleFixup(int32 C
 }
 
 
+bool UCustomizableObjectNodeModifierClipDeform::IsExperimental() const
+{
+	return true;
+}
+
+
 UEdGraphPin* UCustomizableObjectNodeModifierClipDeform::ClipShapePin() const
 {
 	return FindPin(TEXT("Clip Shape"), EGPD_Input);
diff --git a/Engine/Plugins/Mutable/Source/CustomizableObjectEditor/Private/MuCOE/Nodes/CustomizableObjectNodeMeshReshape.cpp b/Engine/Plugins/Mutable/Source/CustomizableObjectEditor/Private/MuCOE/Nodes/CustomizableObjectNodeMeshReshape.cpp
--- a/Engine/Plugins/Mutable/Source/CustomizableObjectEditor/Private/MuCOE/Nodes/CustomizableObjectNodeMeshReshape.cpp
+++ b/Engine/Plugins/Mutable/Source/CustomizableObjectEditor/Private/MuCOE/Nodes/CustomizableObjectNodeMeshReshape.cpp
@@ -60,6 +60,12 @@ FString UCustomizableObjectNodeMeshReshape::GetRefreshMessage() const
 }
 
 
+bool UCustomizableObjectNodeMeshReshape::IsExperimental() const
+{
+	return true;
+}
+
+
 void UCustomizableObjectNodeMeshReshape::Serialize(FArchive& Ar)
 {
 	Super::Serialize(Ar);
diff --git a/Engine/Plugins/Mutable/Source/CustomizableObjectEditor/Private/MuCOE/Nodes/CustomizableObjectNodePassThroughTexture.cpp b/Engine/Plugins/Mutable/Source/CustomizableObjectEditor/Private/MuCOE/Nodes/CustomizableObjectNodePassThroughTexture.cpp
--- a/Engine/Plugins/Mutable/Source/CustomizableObjectEditor/Private/MuCOE/Nodes/CustomizableObjectNodePassThroughTexture.cpp
+++ b/Engine/Plugins/Mutable/Source/CustomizableObjectEditor/Private/MuCOE/Nodes/CustomizableObjectNodePassThroughTexture.cpp
@@ -44,11 +44,11 @@ FText UCustomizableObjectNodePassThroughTexture::GetNodeTitle(ENodeTitleType::Ty
 		FFormatNamedArguments Args;
 		Args.Add(TEXT("TextureName"), FText::FromString(Texture->GetName()));
 
-		return FText::Format(LOCTEXT("PassThrough Texture_Title", "{TextureName}\nPassThrough Texture"), Args);
+		return FText::Format(LOCTEXT("Passthrough Texture_Title", "{TextureName}\nPassthrough Texture"), Args);
 	}
 	else
 	{
-		return LOCTEXT("PassThrough Texture", "PassThrough Texture");
+		return LOCTEXT("Passthrough Texture", "Passthrough Texture");
 	}
 }
 
diff --git a/Engine/Plugins/Mutable/Source/CustomizableObjectEditor/Private/MuCOE/Nodes/SCustomizableObjectNode.cpp b/Engine/Plugins/Mutable/Source/CustomizableObjectEditor/Private/MuCOE/Nodes/SCustomizableObjectNode.cpp
--- a/Engine/Plugins/Mutable/Source/CustomizableObjectEditor/Private/MuCOE/Nodes/SCustomizableObjectNode.cpp
+++ b/Engine/Plugins/Mutable/Source/CustomizableObjectEditor/Private/MuCOE/Nodes/SCustomizableObjectNode.cpp
@@ -3,11 +3,16 @@
 #include "SCustomizableObjectNode.h"
 
 #include "SCustomizableObjectNodePin.h"
+#include "MuCOE/Nodes/CustomizableObjectNode.h"
+#include "Widgets/SBoxPanel.h"
+
+#define LOCTEXT_NAMESPACE "CustomizableObjectEditor"
 
 
 void SCustomizableObjectNode::Construct(const FArguments& InArgs, UEdGraphNode* InGraphNode)
 {
     GraphNode = InGraphNode;
+	
     UpdateGraphNode();
 }
 
@@ -16,3 +21,28 @@ TSharedPtr<SGraphPin> SCustomizableObjectNode::CreatePinWidget(UEdGraphPin* Pin)
 {
 	return SNew(SCustomizableObjectNodePin, Pin);
 }
+
+
+void SCustomizableObjectNode::CreateBelowPinControls(TSharedPtr<SVerticalBox> MainBox)
+{
+	SGraphNode::CreateBelowPinControls(MainBox);
+
+	if (UCustomizableObjectNode* Node = Cast<UCustomizableObjectNode>(GraphNode))
+	{
+		if (Node->IsExperimental())
+		{
+			MainBox->AddSlot()
+			.Padding(FMargin(0.0, 2.0, 0.0, 0.0))
+			.AutoHeight()
+			[
+				SNew(SErrorText)
+				.BackgroundColor(FAppStyle::GetColor("ErrorReporting.WarningBackgroundColor"))
+				.ErrorText(LOCTEXT("Experimental", "EXPERIMENTAL"))
+			];
+		}
+	}
+}
+
+
+#undef LOCTEXT_NAMESPACE
+
diff --git a/Engine/Plugins/Mutable/Source/CustomizableObjectEditor/Private/MuCOE/Nodes/SCustomizableObjectNode.h b/Engine/Plugins/Mutable/Source/CustomizableObjectEditor/Private/MuCOE/Nodes/SCustomizableObjectNode.h
--- a/Engine/Plugins/Mutable/Source/CustomizableObjectEditor/Private/MuCOE/Nodes/SCustomizableObjectNode.h
+++ b/Engine/Plugins/Mutable/Source/CustomizableObjectEditor/Private/MuCOE/Nodes/SCustomizableObjectNode.h
@@ -19,4 +19,5 @@ public:
 
 	// SGraphNode interface
     virtual TSharedPtr<SGraphPin> CreatePinWidget(UEdGraphPin* Pin) const override;
+	virtual void CreateBelowPinControls(TSharedPtr<SVerticalBox> MainBox) override;
 };
diff --git a/Engine/Plugins/Mutable/Source/CustomizableObjectEditor/Public/MuCOE/CustomizableObjectFactory.h b/Engine/Plugins/Mutable/Source/CustomizableObjectEditor/Public/MuCOE/CustomizableObjectFactory.h
--- a/Engine/Plugins/Mutable/Source/CustomizableObjectEditor/Public/MuCOE/CustomizableObjectFactory.h
+++ b/Engine/Plugins/Mutable/Source/CustomizableObjectEditor/Public/MuCOE/CustomizableObjectFactory.h
@@ -23,23 +23,6 @@ class USkeletalMesh;
 struct EVisibility;
 
 
-struct FCustomizableObjectOptions
-{
-	// Determines if the CO asset to create has not been configurated
-	bool bEmptyObject = true;
-
-	// Child Object Settings
-	// Determines if the CO asset to create is child of another CO
-	bool bIsChildObject = false;
-
-	// Parent of the CO
-	TWeakObjectPtr<UCustomizableObject> ParentObject = nullptr;
-
-	// Name of the group node of the parent where this CO asset will be linked
-	FString GroupNodeName;
-};
-
-
 UCLASS(MinimalAPI)
 class UCustomizableObjectFactory : public UFactory
 {
@@ -49,67 +32,7 @@ class UCustomizableObjectFactory : public UFactory
 
 public:
 	// Begin UFactory Interface
-	virtual bool ConfigureProperties() override;
 	virtual bool DoesSupportClass(UClass * Class) override;
 	virtual UClass* ResolveSupportedClass() override;
 	virtual UObject* FactoryCreateNew(UClass* Class, UObject* InParent,FName Name,EObjectFlags Flags,UObject* Context,FFeedbackContext* Warn) override;
-	
-	FCustomizableObjectOptions CreationSettings;
-};
-
-
-class FCustomizableObjectFactoryUI : public TSharedFromThis<FCustomizableObjectFactoryUI>
-{
-public:
-
-	// Generates the Window widget and its content
-	const FCustomizableObjectOptions ConstructFactoryUI();
-
-	// Callback that generates a CO asset using the options set
-	FReply OnCreate(bool bIsEmpty = true);
-
-	// Creates a the tooltip of the OK button
-	FText GetOKButtonTooltip() const;
-
-	// Callback that cancels the generation of the CO asset
-	FReply OnCancel();
-
-	// Enables or disables the OK button in function of the options set.
-	bool IsConfigurationValid() const;
-
-	// Returns true if the CO asset can be created when clossing the window
-	bool CanCreateObject() const;
-
-	// Callback to set if the CO is a child object
-	void OnCheckBoxChanged(ECheckBoxState State);
-
-	// Parent Selecto Widget methods
-	EVisibility GetParentWidgetsVisibility() const;
-	FString GetSelectedCustomizableObjectPath() const;
-	void OnPickedCustomizableObjectParent(const FAssetData& SelectedAsset);
-	void OnSelectGroupComboBox(TSharedPtr<FString> Selection, ESelectInfo::Type SelectInfo);
-	void GenerateGroupOptions();
-	bool IsNodeGroupSelectorEnabled() const;
-
-	// Component Selectior Widget methods
-	EVisibility GetComponentWidgetsVisibility() const;
-	void OnSelectComponentComboBox(TSharedPtr<FString> Selection, ESelectInfo::Type SelectInfo);
-
-private:
-
-	// Pointer to the created widgets
-	TSharedPtr<SWindow> COSettingsWindow;
-	TSharedPtr<STextComboBox> GroupSelector;
-	TSharedPtr<STextComboBox> ComponentSelector;
-
-	// bool that determines if a CO asset can be created
-	bool bCreateObject = false;
-
-	// Setting options to generate the CO asset
-	FCustomizableObjectOptions Options;
-
-	// Combobox options
-	TArray<TSharedPtr<FString>> GroupOptions;
-	TArray<TSharedPtr<FString>> ComponentsOptions;
-
 };
diff --git a/Engine/Plugins/Mutable/Source/CustomizableObjectEditor/Public/MuCOE/Nodes/CustomizableObjectNode.h b/Engine/Plugins/Mutable/Source/CustomizableObjectEditor/Public/MuCOE/Nodes/CustomizableObjectNode.h
--- a/Engine/Plugins/Mutable/Source/CustomizableObjectEditor/Public/MuCOE/Nodes/CustomizableObjectNode.h
+++ b/Engine/Plugins/Mutable/Source/CustomizableObjectEditor/Public/MuCOE/Nodes/CustomizableObjectNode.h
@@ -323,6 +323,10 @@ protected:
 	/** Attach data to a given pin. */
 	void AddPinData(const UEdGraphPin& Pin, UCustomizableObjectNodePinData& PinData);
 
+public:
+	/** Return true to mark this node as experimental. */
+	virtual bool IsExperimental() const;
+	
 private:
 	/** Data attached to a given pin. Not all pins contain data.
 	 * Data must be eliminated once the pins is removed. This is done automatically on CustomRemovePin.
diff --git a/Engine/Plugins/Mutable/Source/CustomizableObjectEditor/Public/MuCOE/Nodes/CustomizableObjectNodeComponentMeshBase.h b/Engine/Plugins/Mutable/Source/CustomizableObjectEditor/Public/MuCOE/Nodes/CustomizableObjectNodeComponentMeshBase.h
--- a/Engine/Plugins/Mutable/Source/CustomizableObjectEditor/Public/MuCOE/Nodes/CustomizableObjectNodeComponentMeshBase.h
+++ b/Engine/Plugins/Mutable/Source/CustomizableObjectEditor/Public/MuCOE/Nodes/CustomizableObjectNodeComponentMeshBase.h
@@ -10,8 +10,8 @@
 UENUM()
 enum class ECustomizableObjectAutomaticLODStrategy : uint8
 {
-	// Use the same strategy than the parent object. If root, then use "Manual".
-	Inherited = 0 UMETA(DisplayName = "Inherit from parent object"),
+	// Use the same strategy than the parent component. If root, then use "Manual".
+	Inherited = 0 UMETA(DisplayName = "Inherit from parent component"),
 	// Don't try to generate LODs automatically for the child nodes. Only the ones tha explicitely define them will be used.
 	Manual = 1 UMETA(DisplayName = "Only manually created LODs"),
 	// Try to generate the same material structure than LOd 0 if the source meshes have LODs.
@@ -32,7 +32,7 @@ public:
 	UPROPERTY(EditAnywhere, Category = ComponentMesh)
 	int32 NumLODs = 1;
 	
-	UPROPERTY(EditAnywhere, Category = ComponentMesh)
+	UPROPERTY(EditAnywhere, Category = ComponentMesh, DisplayName = "Auto LOD Strategy")
 	ECustomizableObjectAutomaticLODStrategy AutoLODStrategy = ECustomizableObjectAutomaticLODStrategy::AutomaticFromMesh;
 	
 	UPROPERTY()
diff --git a/Engine/Plugins/Mutable/Source/CustomizableObjectEditor/Public/MuCOE/Nodes/CustomizableObjectNodeComponentPassthroughMesh.h b/Engine/Plugins/Mutable/Source/CustomizableObjectEditor/Public/MuCOE/Nodes/CustomizableObjectNodeComponentPassthroughMesh.h
--- a/Engine/Plugins/Mutable/Source/CustomizableObjectEditor/Public/MuCOE/Nodes/CustomizableObjectNodeComponentPassthroughMesh.h
+++ b/Engine/Plugins/Mutable/Source/CustomizableObjectEditor/Public/MuCOE/Nodes/CustomizableObjectNodeComponentPassthroughMesh.h
@@ -83,6 +83,7 @@ public:
 
 	// UCustomizableObjectNode interface
 	virtual void AllocateDefaultPins(UCustomizableObjectNodeRemapPins* RemapPins) override;
+	virtual bool IsExperimental() const override;
 
 	// Own interface
 	
diff --git a/Engine/Plugins/Mutable/Source/CustomizableObjectEditor/Public/MuCOE/Nodes/CustomizableObjectNodeMeshReshape.h b/Engine/Plugins/Mutable/Source/CustomizableObjectEditor/Public/MuCOE/Nodes/CustomizableObjectNodeMeshReshape.h
--- a/Engine/Plugins/Mutable/Source/CustomizableObjectEditor/Public/MuCOE/Nodes/CustomizableObjectNodeMeshReshape.h
+++ b/Engine/Plugins/Mutable/Source/CustomizableObjectEditor/Public/MuCOE/Nodes/CustomizableObjectNodeMeshReshape.h
@@ -55,6 +55,7 @@ public:
 	// UCustomizableObjectNode interface
 	void AllocateDefaultPins(UCustomizableObjectNodeRemapPins* RemapPins) override;
 	FString GetRefreshMessage() const override;
+	virtual bool IsExperimental() const override;
 
 	inline UEdGraphPin* BaseMeshPin() const
 	{
diff --git a/Engine/Plugins/Mutable/Source/CustomizableObjectEditor/Public/MuCOE/Nodes/CustomizableObjectNodeModifierClipDeform.h b/Engine/Plugins/Mutable/Source/CustomizableObjectEditor/Public/MuCOE/Nodes/CustomizableObjectNodeModifierClipDeform.h
--- a/Engine/Plugins/Mutable/Source/CustomizableObjectEditor/Public/MuCOE/Nodes/CustomizableObjectNodeModifierClipDeform.h
+++ b/Engine/Plugins/Mutable/Source/CustomizableObjectEditor/Public/MuCOE/Nodes/CustomizableObjectNodeModifierClipDeform.h
@@ -42,6 +42,7 @@ public:
 	// UCustomizableObjectNode interface
 	virtual void AllocateDefaultPins(UCustomizableObjectNodeRemapPins* RemapPins) override;
 	virtual void BackwardsCompatibleFixup(int32 CustomizableObjectCustomVersion) override;
+	virtual bool IsExperimental() const override;
 
 	// Own interface
 
