[mutable] Enable base modifier details for Remove Mesh node.
#rnx

diff --git a/Engine/Plugins/Mutable/Source/CustomizableObjectEditor/Private/MuCOE/Nodes/CustomizableObjectNodeModifierRemoveMeshDetails.cpp b/Engine/Plugins/Mutable/Source/CustomizableObjectEditor/Private/MuCOE/Nodes/CustomizableObjectNodeModifierRemoveMeshDetails.cpp
--- a/Engine/Plugins/Mutable/Source/CustomizableObjectEditor/Private/MuCOE/Nodes/CustomizableObjectNodeModifierRemoveMeshDetails.cpp
+++ b/Engine/Plugins/Mutable/Source/CustomizableObjectEditor/Private/MuCOE/Nodes/CustomizableObjectNodeModifierRemoveMeshDetails.cpp
@@ -25,7 +25,7 @@ TSharedRef<IDetailCustomization> FCustomizableObjectNodeModifierRemoveMeshDetail
 
 void FCustomizableObjectNodeModifierRemoveMeshDetails::CustomizeDetails(IDetailLayoutBuilder& DetailBuilder)
 {
-	FCustomizableObjectNodeDetails::CustomizeDetails(DetailBuilder);
+	FCustomizableObjectNodeModifierBaseDetails::CustomizeDetails(DetailBuilder);
 
 	// This property is not relevant for this node
 	DetailBuilder.HideProperty(GET_MEMBER_NAME_CHECKED(UCustomizableObjectNodeModifierWithMaterial, ReferenceMaterial), UCustomizableObjectNodeModifierWithMaterial::StaticClass());
diff --git a/Engine/Plugins/Mutable/Source/CustomizableObjectEditor/Private/MuCOE/Nodes/CustomizableObjectNodeModifierRemoveMeshDetails.h b/Engine/Plugins/Mutable/Source/CustomizableObjectEditor/Private/MuCOE/Nodes/CustomizableObjectNodeModifierRemoveMeshDetails.h
--- a/Engine/Plugins/Mutable/Source/CustomizableObjectEditor/Private/MuCOE/Nodes/CustomizableObjectNodeModifierRemoveMeshDetails.h
+++ b/Engine/Plugins/Mutable/Source/CustomizableObjectEditor/Private/MuCOE/Nodes/CustomizableObjectNodeModifierRemoveMeshDetails.h
@@ -2,14 +2,12 @@
 
 #pragma once
 
-#include "CustomizableObjectNodeDetails.h"
+#include "CustomizableObjectNodeModifierBaseDetails.h"
 #include "IDetailCustomization.h"
-#include "Types/SlateEnums.h"
-#include "Misc/AssertionMacros.h"
 
 class IDetailLayoutBuilder;
 
-class FCustomizableObjectNodeModifierRemoveMeshDetails : public FCustomizableObjectNodeDetails
+class FCustomizableObjectNodeModifierRemoveMeshDetails : public FCustomizableObjectNodeModifierBaseDetails
 {
 public:
 	/** Makes a new instance of this detail layout class for a specific detail view requesting it */
