[mutable] Cleanup redudant functionality from the clip-with-mesh node that was used to directly select objects to modify and generated artifical tags.
Rename MeshClipMorph node to the more explicit "Clip Mesh With Plane And Morph"
#rnx

diff --git a/Engine/Plugins/Mutable/Source/CustomizableObject/Internal/MuCO/CustomizableObjectCustomVersion.h b/Engine/Plugins/Mutable/Source/CustomizableObject/Internal/MuCO/CustomizableObjectCustomVersion.h
--- a/Engine/Plugins/Mutable/Source/CustomizableObject/Internal/MuCO/CustomizableObjectCustomVersion.h
+++ b/Engine/Plugins/Mutable/Source/CustomizableObject/Internal/MuCO/CustomizableObjectCustomVersion.h
@@ -156,6 +156,8 @@ struct CUSTOMIZABLEOBJECT_API FCustomizableObjectCustomVersion
 
 		FixMaterialPinsRename,
 
+		ModifierClipWithMeshCleanup,
+
 		// -----<new versions can be added above this line>-------------------------------------------------
 		VersionPlusOne,
 		LatestVersion = VersionPlusOne - 1
diff --git a/Engine/Plugins/Mutable/Source/CustomizableObjectEditor/Private/MuCOE/CustomizableObjectCompiler.cpp b/Engine/Plugins/Mutable/Source/CustomizableObjectEditor/Private/MuCOE/CustomizableObjectCompiler.cpp
--- a/Engine/Plugins/Mutable/Source/CustomizableObjectEditor/Private/MuCOE/CustomizableObjectCompiler.cpp
+++ b/Engine/Plugins/Mutable/Source/CustomizableObjectEditor/Private/MuCOE/CustomizableObjectCompiler.cpp
@@ -731,7 +731,6 @@ mu::NodeObjectPtr FCustomizableObjectCompiler::GenerateMutableRoot(
 	UE_LOG(LogMutable, Verbose, TEXT("PROFILE: [ %16.8f ] GenerateMutableSource end."), FPlatformTime::Seconds());
 
     ActualRoot->RealTimeMorphSelectionOverrides = GenerationContext.RealTimeMorphTargetsOverrides;
-	GenerationContext.GenerateClippingCOInternalTags();
 
 	GenerationContext.GenerateSharedSurfacesUniqueIds();
 
diff --git a/Engine/Plugins/Mutable/Source/CustomizableObjectEditor/Private/MuCOE/GenerateMutableSource/GenerateMutableSource.cpp b/Engine/Plugins/Mutable/Source/CustomizableObjectEditor/Private/MuCOE/GenerateMutableSource/GenerateMutableSource.cpp
--- a/Engine/Plugins/Mutable/Source/CustomizableObjectEditor/Private/MuCOE/GenerateMutableSource/GenerateMutableSource.cpp
+++ b/Engine/Plugins/Mutable/Source/CustomizableObjectEditor/Private/MuCOE/GenerateMutableSource/GenerateMutableSource.cpp
@@ -273,62 +273,6 @@ const FGuid FMutableGraphGenerationContext::GetNodeIdUnique(const UCustomizableO
 }
 
 
-void FMutableGraphGenerationContext::GenerateClippingCOInternalTags()
-{
-	int32 i = 0;
-	int32 j;
-
-	// The following array ArrayCachedTaggedMaterialNodes and related functionality was removed because it prevented two clip nodes from
-	// clipping the same material node
-	// In the case more than one CO clips a particular CO, the associated mu::NodeSurfaceNewPtr nodes with that particular CO only need to be tagged once
-	//TArray<FGuid> ArrayCachedTaggedMaterialNodes;
-
-	for (auto It = MapClipMeshNodeToMutableClipMeshNodeArray.CreateConstIterator(); It; ++It)
-	{
-		FString TagName = FString::Printf(TEXT("ClippingTag_%d"), i);
-		i++;
-
-		for (j = 0; j < It->Value.Num(); ++j)
-		{
-			It->Value[j]->RequiredTags.Add(TagName);
-		}
-
-		UCustomizableObjectNodeModifierClipWithMesh* CustomizableObjectNodeModifierClipWithMesh = It->Key;
-
-		// Tag also those nodes in the material side of the clipping functionality (since a CO can have several materials with the same LOD,
-		// all of them are considered, so CustomizableObjectNodeMeshClipWithMesh->ArrayMaterialNodeToClipWithID contains the Guids of all
-		// the material nodes with the same name selected for being affected by the clipping operation
-		TArray<FGuid>& ArrayMaterialNodeToClipWithID = CustomizableObjectNodeModifierClipWithMesh->ArrayMaterialNodeToClipWithID;
-
-		for (j = 0; j < ArrayMaterialNodeToClipWithID.Num(); ++j)
-		{
-			// The following lines were commented for the same reasons as ArrayCachedTaggedMaterialNodes above
-			//if (ArrayCachedTaggedMaterialNodes.Find(CustomizableObjectNodeMeshClipWithMesh->ArrayMaterialNodeToClipWithID[j]) != INDEX_NONE)
-			//{
-			//    continue;
-			//}
-
-			bool FoundEntryInMap = false;
-			for (auto It2 = MapMaterialNodeToMutableSurfaceNodeArray.CreateConstIterator(); (It2 && !FoundEntryInMap); ++It2)
-			{
-				if (It2->Key->NodeGuid == ArrayMaterialNodeToClipWithID[j])
-				{
-					FoundEntryInMap = true;
-
-					for (int32 k = 0; k < It2->Value.Num(); ++k)
-					{
-						It2->Value[k]->Tags.Add(TagName);
-					}
-				}
-			}
-
-			// The following lines were commented for the same reasons as ArrayCachedTaggedMaterialNodes above
-			//ArrayCachedTaggedMaterialNodes.Add(CustomizableObjectNodeMeshClipWithMesh->ArrayMaterialNodeToClipWithID[j]);
-		}
-	}
-}
-
-
 void FMutableGraphGenerationContext::GenerateSharedSurfacesUniqueIds()
 {
 	int32 UniqueId = 0;
diff --git a/Engine/Plugins/Mutable/Source/CustomizableObjectEditor/Private/MuCOE/GenerateMutableSource/GenerateMutableSource.h b/Engine/Plugins/Mutable/Source/CustomizableObjectEditor/Private/MuCOE/GenerateMutableSource/GenerateMutableSource.h
--- a/Engine/Plugins/Mutable/Source/CustomizableObjectEditor/Private/MuCOE/GenerateMutableSource/GenerateMutableSource.h
+++ b/Engine/Plugins/Mutable/Source/CustomizableObjectEditor/Private/MuCOE/GenerateMutableSource/GenerateMutableSource.h
@@ -738,11 +738,6 @@ struct FMutableGraphGenerationContext
 	// Check if the Id of the node Node already exists, if it's new adds it to NodeIds array, otherwise, returns new Id
 	const FGuid GetNodeIdUnique(const UCustomizableObjectNode* Node);
 
-	/** Generates new tags for the UCustomizableObjectNodeModifierClipWithMesh nodes which have assigned a CO in the
-	* UCustomizableObjectNodeModifierClipWithMesh::CustomizableObjectToClipWith field
-	* @return nothing */
-	void GenerateClippingCOInternalTags();
-
 	/** Generates shared surface IDs for all surface nodes. If one or more nodes are equal, they will use the same SharedSurfaceId */
 	void GenerateSharedSurfacesUniqueIds();
 
@@ -793,9 +788,6 @@ struct FMutableGraphGenerationContext
 	// Map with pairs (Unreal Mutable material node Guid, array with its corresponding Mutable surface node) to add tags from clipping nodes
 	TMap<class UCustomizableObjectNodeMaterial*, TArray<mu::Ptr<mu::NodeSurfaceNew>>> MapMaterialNodeToMutableSurfaceNodeArray;
 
-	// Map with pairs (Unreal Mutable clip mesh node, array with its corresponding Mutable mesh modifier nodes) to add tags
-	TMap<class UCustomizableObjectNodeModifierClipWithMesh*, TArray<mu::Ptr<mu::NodeModifierMeshClipWithMesh>>> MapClipMeshNodeToMutableClipMeshNodeArray;
-
 	// Data used for MorphTarget reconstruction.
 	TMap<uint32, FRealTimeMorphMeshData> RealTimeMorphTargetPerMeshData;
 
diff --git a/Engine/Plugins/Mutable/Source/CustomizableObjectEditor/Private/MuCOE/GenerateMutableSource/GenerateMutableSourceModifier.cpp b/Engine/Plugins/Mutable/Source/CustomizableObjectEditor/Private/MuCOE/GenerateMutableSource/GenerateMutableSourceModifier.cpp
--- a/Engine/Plugins/Mutable/Source/CustomizableObjectEditor/Private/MuCOE/GenerateMutableSource/GenerateMutableSourceModifier.cpp
+++ b/Engine/Plugins/Mutable/Source/CustomizableObjectEditor/Private/MuCOE/GenerateMutableSource/GenerateMutableSourceModifier.cpp
@@ -198,23 +198,6 @@ mu::Ptr<mu::NodeModifier> GenerateMutableSourceModifier(const UEdGraphPin * Pin,
 		ClipNode->MultipleTagsPolicy = TypedNodeClipMesh->MultipleTagPolicy;
 		ClipNode->RequiredTags = TypedNodeClipMesh->RequiredTags;
 
-		if (TypedNodeClipMesh->CustomizableObjectToClipWith != nullptr)
-		{
-			TArray<mu::Ptr<mu::NodeModifierMeshClipWithMesh>>* ArrayDataPtr = GenerationContext.MapClipMeshNodeToMutableClipMeshNodeArray.Find(Cast<UCustomizableObjectNodeModifierClipWithMesh>(Pin->GetOwningNode()));
-
-			if (ArrayDataPtr == nullptr)
-			{
-				TArray<mu::Ptr<mu::NodeModifierMeshClipWithMesh>> ArrayData;
-				ArrayData.Add(ClipNode);
-				UCustomizableObjectNodeModifierClipWithMesh* CastedNode = Cast<UCustomizableObjectNodeModifierClipWithMesh>(Pin->GetOwningNode());
-				GenerationContext.MapClipMeshNodeToMutableClipMeshNodeArray.Add(CastedNode, ArrayData);
-			}
-			else
-			{
-				ArrayDataPtr->AddUnique(ClipNode);
-			}
-		}
-
 		GenerationContext.MeshGenerationFlags.Pop();
 	}
 
diff --git a/Engine/Plugins/Mutable/Source/CustomizableObjectEditor/Private/MuCOE/Nodes/CustomizableObjectNodeMeshClipMorph.cpp b/Engine/Plugins/Mutable/Source/CustomizableObjectEditor/Private/MuCOE/Nodes/CustomizableObjectNodeMeshClipMorph.cpp
--- a/Engine/Plugins/Mutable/Source/CustomizableObjectEditor/Private/MuCOE/Nodes/CustomizableObjectNodeMeshClipMorph.cpp
+++ b/Engine/Plugins/Mutable/Source/CustomizableObjectEditor/Private/MuCOE/Nodes/CustomizableObjectNodeMeshClipMorph.cpp
@@ -143,7 +143,7 @@ void UCustomizableObjectNodeModifierClipMorph::AllocateDefaultPins(UCustomizable
 
 FText UCustomizableObjectNodeModifierClipMorph::GetNodeTitle(ENodeTitleType::Type TitleType) const
 {
-	return LOCTEXT("Clip_Morph_Mesh", "Clip Morph Mesh");
+	return LOCTEXT("ClipMeshWithPlaneAndMorph", "Clip Mesh With Plane and Morph");
 }
 
 
diff --git a/Engine/Plugins/Mutable/Source/CustomizableObjectEditor/Private/MuCOE/Nodes/CustomizableObjectNodeMeshClipWithMesh.cpp b/Engine/Plugins/Mutable/Source/CustomizableObjectEditor/Private/MuCOE/Nodes/CustomizableObjectNodeMeshClipWithMesh.cpp
--- a/Engine/Plugins/Mutable/Source/CustomizableObjectEditor/Private/MuCOE/Nodes/CustomizableObjectNodeMeshClipWithMesh.cpp
+++ b/Engine/Plugins/Mutable/Source/CustomizableObjectEditor/Private/MuCOE/Nodes/CustomizableObjectNodeMeshClipWithMesh.cpp
@@ -6,43 +6,16 @@
 #include "MuCOE/EdGraphSchema_CustomizableObject.h"
 #include "MuCOE/ICustomizableObjectEditor.h"
 #include "MuCO/CustomizableObjectCustomVersion.h"
+#include "MuCOE/GraphTraversal.h"
 
 class UCustomizableObjectNodeRemapPins;
 
 #define LOCTEXT_NAMESPACE "CustomizableObjectEditor"
 
 
-UCustomizableObjectNodeModifierClipWithMesh::UCustomizableObjectNodeModifierClipWithMesh() :
-	CustomizableObjectToClipWith(nullptr)
+UCustomizableObjectNodeModifierClipWithMesh::UCustomizableObjectNodeModifierClipWithMesh()
 {
-	// We allow this node to be shared
 	Transform = FTransform::Identity;
-
-	// When initialize we don't use materials neither tags
-	bUseMaterials = false;
-	bUseTags = false;
-}
-
-
-void UCustomizableObjectNodeModifierClipWithMesh::BeginPostDuplicate(bool bDuplicateForPIE)
-{
-	Super::BeginPostDuplicate(bDuplicateForPIE);
-
-	if (UCustomizableObjectGraph* Graph = Cast<UCustomizableObjectGraph>(GetGraph()))
-	{
-		TArray<FGuid> NotifiedNodes;
-		for (const FGuid& NodeToClipWithId : ArrayMaterialNodeToClipWithID)
-		{
-			FGuid NodeToClipAlreadyUpdated = Graph->RequestNotificationForNodeIdChange(NodeToClipWithId, NodeGuid);
-
-			if (NodeToClipWithId != NodeToClipAlreadyUpdated)
-			{
-				NotifiedNodes.Add(NodeToClipAlreadyUpdated);
-			}
-		}
-
-		ArrayMaterialNodeToClipWithID = NotifiedNodes;
-	}
 }
 
 
@@ -68,16 +41,51 @@ void UCustomizableObjectNodeModifierClipWithMesh::BackwardsCompatibleFixup(int32
 		RequiredTags = Tags_DEPRECATED;
 		Tags_DEPRECATED.Empty();
 	}
-}
-
+	
+	if (CustomizableObjectCustomVersion == FCustomizableObjectCustomVersion::ModifierClipWithMeshCleanup)
+	{
+		for ( const FGuid& MeshSectionNodeID : ArrayMaterialNodeToClipWithID_DEPRECATED )
+		{
+			// Look for the parent material and set it as the modifier reference material 
+			UCustomizableObjectNode* ParentNode = GetCustomizableObjectExternalNode<UCustomizableObjectNode>(CustomizableObjectToClipWith_DEPRECATED.Get(), MeshSectionNodeID);
 
-void UCustomizableObjectNodeModifierClipWithMesh::ReconstructNode(UCustomizableObjectNodeRemapPins* RemapPinsMode)
-{
-	Super::ReconstructNode(RemapPinsMode);
+			if (ParentNode)
+			{
+				// Add an autogenerated tag to the legacy parent, so that it can be referenced in this modifier.
+				FString NewLegacyTag = MakeNodeAutoTag(ParentNode);
+				RequiredTags.Add(NewLegacyTag);
+
+				TArray<FString>* NodeEnableTags = ParentNode->GetEnableTags();
+				if (NodeEnableTags)
+				{
+					NodeEnableTags->AddUnique(NewLegacyTag);
+				}
+				else
+				{
+					// Conversion failed?
+					ensure(false);
+					UE_LOG(LogMutable, Warning, TEXT("[%s] UCustomizableObjectNodeModifierClipWithMesh version upgrade failed."), *GetOutermost()->GetName());
+				}
+
+				// If the tag was added to another CO, keep track of the addition to be able to repeat it in case
+				// that CO is not re-saved but this one is.
+				UCustomizableObject* ThisNodeObject = GetRootObject(*this);
+				if (CustomizableObjectToClipWith_DEPRECATED != ThisNodeObject)
+				{
+					FLegacyTag LegacyTag;
+					LegacyTag.ParentObject = CustomizableObjectToClipWith_DEPRECATED;
+					LegacyTag.ParentNode = MeshSectionNodeID;
+					LegacyTag.Tag = NewLegacyTag;
+					LegacyBackportsRequiredTags.AddUnique(LegacyTag);
+				}
+			}
+			else
+			{
+				UE_LOG(LogMutable, Log, TEXT("[%s] UCustomizableObjectNodeModifierClipWithMesh has no parent. It will not be upgraded."), *GetOutermost()->GetName());
+			}
+		}
 
-	if (!CustomizableObjectToClipWith)
-	{
-		CustomizableObjectToClipWith = Cast<UCustomizableObject>(GetOutermostObject());
+		ArrayMaterialNodeToClipWithID_DEPRECATED.Empty();
 	}
 }
 
@@ -88,19 +96,6 @@ FText UCustomizableObjectNodeModifierClipWithMesh::GetNodeTitle(ENodeTitleType::
 }
 
 
-void UCustomizableObjectNodeModifierClipWithMesh::PinConnectionListChanged(UEdGraphPin* Pin)
-{
-	Super::PinConnectionListChanged(Pin);
-
-	TSharedPtr<ICustomizableObjectEditor> Editor = GetGraphEditor();
-
-	if (Editor.IsValid())
-	{
-		Editor->UpdateGraphNodeProperties();
-	}
-}
-
-
 UEdGraphPin* UCustomizableObjectNodeModifierClipWithMesh::OutputPin() const
 {
 	return FindPin(TEXT("Modifier"));
@@ -113,15 +108,9 @@ UEdGraphPin* UCustomizableObjectNodeModifierClipWithMesh::ClipMeshPin() const
 }
 
 
-void UCustomizableObjectNodeModifierClipWithMesh::UpdateReferencedNodeId(const FGuid& NewGuid)
-{
-	ArrayMaterialNodeToClipWithID.Add(NewGuid);
-}
-
-
 FText UCustomizableObjectNodeModifierClipWithMesh::GetTooltipText() const
 {
-	return LOCTEXT("Clip_Mesh_Mesh_Tooltip", "Removes the part of a material that is completely enclosed in a mesh volume.\nIt only removes the faces that fall completely inside the cutting volume, along with the vertices and edges that define only faces that are deleted.");
+	return LOCTEXT("Clip_Mesh_Mesh_Tooltip", "Removes the part of a mesh section that is completely enclosed in a mesh volume.\nIt only removes the faces that fall completely inside the cutting volume, along with the vertices and edges that define only faces that are deleted.");
 }
 
 #undef LOCTEXT_NAMESPACE
diff --git a/Engine/Plugins/Mutable/Source/CustomizableObjectEditor/Private/MuCOE/Nodes/CustomizableObjectNodeModifierClipWithMeshDetails.cpp b/Engine/Plugins/Mutable/Source/CustomizableObjectEditor/Private/MuCOE/Nodes/CustomizableObjectNodeModifierClipWithMeshDetails.cpp
--- a/Engine/Plugins/Mutable/Source/CustomizableObjectEditor/Private/MuCOE/Nodes/CustomizableObjectNodeModifierClipWithMeshDetails.cpp
+++ b/Engine/Plugins/Mutable/Source/CustomizableObjectEditor/Private/MuCOE/Nodes/CustomizableObjectNodeModifierClipWithMeshDetails.cpp
@@ -11,11 +11,7 @@
 #include "MuCOE/Nodes/CustomizableObjectNodeModifierClipWithMesh.h"
 #include "MuCOE/Nodes/CustomizableObjectNodeStaticMesh.h"
 #include "PropertyCustomizationHelpers.h"
-#include "ScopedTransaction.h"
 #include "MuCO/CustomizableObjectPrivate.h"
-#include "Widgets/Input/SButton.h"
-#include "Widgets/Input/SEditableTextBox.h"
-#include "Widgets/Input/STextComboBox.h"
 
 class IPropertyHandle;
 
@@ -42,166 +38,9 @@ void FCustomizableObjectNodeModifierClipWithMeshDetails::CustomizeDetails(IDetai
 		Node = Cast<UCustomizableObjectNodeModifierClipWithMesh>(DetailsView->GetSelectedObjects()[0].Get());
 	}
 
-	IDetailCategoryBuilder& CustomizableObjectToClipCategory = DetailBuilder.EditCategory("CustomizableObjectToClip");
+	IDetailCategoryBuilder& CustomizableObjectToClipCategory = DetailBuilder.EditCategory("ClipMesh");
 
-	if(Node != nullptr)
-	{
-		SelectedCO = Node->CustomizableObjectToClipWith;
-
-		TSharedRef<IPropertyHandle> ParentObjectProperty = DetailBuilder.GetProperty("ParentObject");
-
-		// Hidding unnecessary properties
-		DetailBuilder.HideProperty("CustomizableObjectToClipWith");
-		DetailBuilder.HideProperty("ArrayMaterialNodeToClipWithID");
-
-		if (const UEdGraphPin* ConnectedPin = FollowInputPin(*Node->ClipMeshPin());
-			ConnectedPin && !Cast<UCustomizableObjectNodeStaticMesh>(ConnectedPin->GetOwningNode()))
-		{
-			DetailBuilder.HideProperty("Transform");
-		}
-
-		ArrayMaterialNodeOption.Empty();
-		ArrayMaterialNodeName.Empty();
-
-		TArray<FGuid> ArrayMaterialNodeGuid;
-		TSharedPtr<FString> SelectedMethod;
-
-		CustomizableObjectToClipCategory.AddCustomRow(LOCTEXT("FCustomizableObjectNodeModifierClipWithMeshDetails", "Blocks"))
-		[
-			SNew(SHorizontalBox)
-			+ SHorizontalBox::Slot()
-			.Padding(FMargin(0.0f, 3.0f, 0.0f, 10.0f))
-			.FillWidth(0.15f)
-			[
-				SNew(STextBlock)
-				.Text(LOCTEXT("CustomizableObjectText", "Object:"))
-			]
-			+ SHorizontalBox::Slot()
-			.FillWidth(0.85f)
-			[
-				SNew(SObjectPropertyEntryBox)
-				.AllowedClass(UCustomizableObject::StaticClass())
-				.OnObjectChanged(this, &FCustomizableObjectNodeModifierClipWithMeshDetails::ClipperCustomizableObjectSelectionChanged)
-				.ObjectPath(SelectedCO != nullptr ? SelectedCO->GetPathName() : "")
-				.ForceVolatile(true)
-			]
-		];
-
-		SelectAMethod = MakeShareable(new FString("Select a Method"));
-
-		// Filling ComboBox options
-		ClippingMethods.Empty();
-		ClippingMethods.Add(MakeShareable(new FString("Mesh Section")));
-		ClippingMethods.Add(MakeShareable(new FString("Tags")));
-		ClippingMethods.Add(MakeShareable(new FString("Mesh Section & Tags")));
-
-		// Initializing the inital value of the ComboBox
-		SelectedMethod = SetInitialClippingMethod();
-		
-		// This flag is now up to date and can be used to hide ui.
-		if (!Node->bUseTags)
-		{
-			DetailBuilder.HideProperty("Tags");
-		}
-
-		// ComboBox to select the Clipping method (material or/and tags)
-		CustomizableObjectToClipCategory.AddCustomRow(LOCTEXT("FCustomizableObjectNodeModifierClipWithMeshDetails", "Blocks"))
-		[
-			SNew(SHorizontalBox)
-			+SHorizontalBox::Slot()
-			.AutoWidth()
-			.Padding(FMargin(0.0f, 8.0f, 0.0f, 5.0f))
-			[
-				SNew(STextBlock)
-				.Text(LOCTEXT("ClippingMethod", "Clipping Method:"))
-			]
-			
-			+ SHorizontalBox::Slot()
-			.AutoWidth()
-			.Padding(FMargin(4.0f, 5.0f, 0.0f, 5.0f))
-			[
-				SNew(STextComboBox)
-				.OptionsSource(&ClippingMethods)
-				.InitiallySelectedItem(SelectedMethod)
-				.OnSelectionChanged(this, &FCustomizableObjectNodeModifierClipWithMeshDetails::OnClippingMethodComboBoxSelectionChanged)
-			]
-		];
-
-
-		// Filling materials ComboBox with all the possible materials
-		int32 NumElementBeforeAdd;
-		int32 NumElementAfterAdd;
-		TSharedPtr<FString> ItemToSelect;
-		TArray<UCustomizableObjectNodeMaterial*> ArrayMaterialNodes;
-		if (SelectedCO != nullptr)
-		{
-			SelectedCO->GetPrivate()->GetSource()->GetNodesOfClass<UCustomizableObjectNodeMaterial>(ArrayMaterialNodes);
-		}
-
-		for (const UCustomizableObjectNodeMaterial* MaterialNode : ArrayMaterialNodes)
-		{
-			if (MaterialNode->GetMaterial())
-			{
-				NumElementBeforeAdd = ArrayMaterialNodeName.Num();
-				ArrayMaterialNodeName.AddUnique(MaterialNode->GetMaterial()->GetName());
-				NumElementAfterAdd = ArrayMaterialNodeName.Num();
-
-				if (NumElementBeforeAdd != NumElementAfterAdd)
-				{
-					ArrayMaterialNodeGuid.Add(MaterialNode->NodeGuid);
-				}
-			}
-		}
-
-		// Adding a none option
-		ArrayMaterialNodeOption.Add(MakeShareable(new FString("None")));
-		ItemToSelect = ArrayMaterialNodeOption.Last();
-
-		// Initializing the Selected material if it was already selected
-		const int32 MaxIndex = ArrayMaterialNodeName.Num();
-		for (int32 i = 0; i < MaxIndex; ++i)
-		{
-			ArrayMaterialNodeOption.Add(MakeShareable(new FString(ArrayMaterialNodeName[i])));
-
-			if (Node->ArrayMaterialNodeToClipWithID.Find(ArrayMaterialNodeGuid[i]) != INDEX_NONE)
-			{
-				ItemToSelect = ArrayMaterialNodeOption.Last();
-			}
-		}
-
-        ArrayMaterialNodeOption.Sort(CompareNames);
-
-		// Widget that shows the Tags widget or/and the materials ComboBox
-		if (Node->bUseMaterials)
-		{
-			CustomizableObjectToClipCategory.AddCustomRow(LOCTEXT("FCustomizableObjectNodeModifierClipWithMeshDetails", "Blocks"))
-			[
-				SNew(SVerticalBox)
-				+SVerticalBox::Slot()
-				.AutoHeight()
-				[
-					SNew(SHorizontalBox)
-					+ SHorizontalBox::Slot()
-					.FillWidth(0.15f)
-					.Padding(0.0f, 3.0f, 0.0f, 0.0f)
-					[
-						SNew(STextBlock)
-						.Text(LOCTEXT("MaterialsText", "Material:"))
-					]
-					+ SHorizontalBox::Slot()
-					.FillWidth(0.85f)
-					.Padding(2.0f, 0.0f, 0.0f, 0.0f)
-					[
-						SNew(STextComboBox)
-						.OptionsSource(&ArrayMaterialNodeOption)
-						.InitiallySelectedItem(ItemToSelect)
-						.OnSelectionChanged(this, &FCustomizableObjectNodeModifierClipWithMeshDetails::OnMeshClipWithMeshNodeComboBoxSelectionChanged)
-					]
-				]
-			];
-		}
-	}
-	else
+	if(!Node)
 	{
 		CustomizableObjectToClipCategory.AddCustomRow(LOCTEXT("Node", "Node"))
 			[
@@ -212,126 +51,4 @@ void FCustomizableObjectNodeModifierClipWithMeshDetails::CustomizeDetails(IDetai
 }
 
 
-void FCustomizableObjectNodeModifierClipWithMeshDetails::ClipperCustomizableObjectSelectionChanged(const FAssetData& AssetData)
-{
-	UCustomizableObject* NewClipperCO = Cast<UCustomizableObject>(AssetData.GetAsset());
-
-	if (Node != nullptr)
-	{
-		Node->CustomizableObjectToClipWith = NewClipperCO;
-		Node->ArrayMaterialNodeToClipWithID.Empty();
-	}
-
-	if (DetailBuilderPtr)
-	{
-		DetailBuilderPtr->ForceRefreshDetails();
-	}
-}
-
-
-TSharedPtr<FString> FCustomizableObjectNodeModifierClipWithMeshDetails::SetInitialClippingMethod()
-{
-	TSharedPtr<FString> InitialMethod = SelectAMethod;
-
-	if (!Node->bUseMaterials && !Node->bUseTags)
-	{
-		if (Node->ArrayMaterialNodeToClipWithID.Num() > 0)
-		{
-			InitialMethod = ClippingMethods[0];
-			Node->bUseMaterials = true;
-		}
-
-		if (Node->RequiredTags.Num() > 0)
-		{
-			InitialMethod = ClippingMethods[1];
-			Node->bUseTags = true;
-		}
-
-		if (Node->bUseMaterials && Node->bUseTags)
-		{
-			InitialMethod = ClippingMethods[2];
-		}
-	}
-	else
-	{
-		if (Node->bUseMaterials && !Node->bUseTags)
-		{
-			InitialMethod = ClippingMethods[0];
-		}
-		else if (!Node->bUseMaterials && Node->bUseTags)
-		{
-			InitialMethod = ClippingMethods[1];
-		}
-		else
-		{
-			InitialMethod = ClippingMethods[2];
-		}
-	}
-
-	return InitialMethod;
-}
-
-
-void FCustomizableObjectNodeModifierClipWithMeshDetails::OnClippingMethodComboBoxSelectionChanged(TSharedPtr<FString> Selection, ESelectInfo::Type SelectInfo)
-{
-	if (Selection.IsValid())
-	{
-		if (*Selection == "Tags")
-		{
-			Node->bUseMaterials = false;
-			Node->bUseTags = true;
-			Node->ArrayMaterialNodeToClipWithID.Empty();
-		}
-		else if(*Selection == "Material")
-		{
-			Node->bUseMaterials = true;
-			Node->bUseTags = false;
-			Node->RequiredTags.Empty();
-		}
-		else
-		{
-			Node->bUseMaterials = true;
-			Node->bUseTags = true;
-		}
-
-		DetailBuilderPtr->ForceRefreshDetails();
-	}	
-}
-
-
-void FCustomizableObjectNodeModifierClipWithMeshDetails::OnMeshClipWithMeshNodeComboBoxSelectionChanged(TSharedPtr<FString> Selection, ESelectInfo::Type SelectInfo)
-{
-	if (Selection.IsValid() && (SelectedCO != nullptr))
-	{
-		TArray<UCustomizableObjectNodeMaterial*> ArrayMaterialNode;
-		SelectedCO->GetPrivate()->GetSource()->GetNodesOfClass<UCustomizableObjectNodeMaterial>(ArrayMaterialNode);
-
-		const FScopedTransaction Transaction(LOCTEXT("ChangedClipMeshWithMeshMaterialTransaction", "Changed Clip Mesh With Mesh Material"));
-		Node->Modify();
-		Node->ArrayMaterialNodeToClipWithID.Empty();
-		for (const UCustomizableObjectNodeMaterial* MaterialNode : ArrayMaterialNode)
-		{
-			if (MaterialNode && MaterialNode->GetMaterial() && *Selection == MaterialNode->GetMaterial()->GetName())
-			{
-				Node->ArrayMaterialNodeToClipWithID.Add(MaterialNode->NodeGuid);
-			}
-		}
-	}
-}
-
-
-FReply FCustomizableObjectNodeModifierClipWithMeshDetails::OnAddTagPressed()
-{
-	if (Node)
-	{
-		Node->RequiredTags.Add("");
-		DetailBuilderPtr->ForceRefreshDetails();
-
-		return FReply::Handled();
-	}
-
-	return FReply::Unhandled();
-}
-
-
 #undef LOCTEXT_NAMESPACE
diff --git a/Engine/Plugins/Mutable/Source/CustomizableObjectEditor/Private/MuCOE/Nodes/CustomizableObjectNodeModifierClipWithMeshDetails.h b/Engine/Plugins/Mutable/Source/CustomizableObjectEditor/Private/MuCOE/Nodes/CustomizableObjectNodeModifierClipWithMeshDetails.h
--- a/Engine/Plugins/Mutable/Source/CustomizableObjectEditor/Private/MuCOE/Nodes/CustomizableObjectNodeModifierClipWithMeshDetails.h
+++ b/Engine/Plugins/Mutable/Source/CustomizableObjectEditor/Private/MuCOE/Nodes/CustomizableObjectNodeModifierClipWithMeshDetails.h
@@ -4,17 +4,9 @@
 
 #include "CustomizableObjectNodeModifierBaseDetails.h"
 #include "IDetailCustomization.h"
-#include "Widgets/SCompoundWidget.h"
 
-namespace ESelectInfo { enum Type : int; }
-namespace ETextCommit { enum Type : int; }
-
-class FText;
 class IDetailLayoutBuilder;
-class SHorizontalBox;
-class SVerticalBox;
 class UCustomizableObjectNodeModifierClipWithMesh;
-struct FAssetData;
 
 
 class FCustomizableObjectNodeModifierClipWithMeshDetails : public FCustomizableObjectNodeModifierBaseDetails
@@ -26,40 +18,12 @@ public:
 	/** ILayoutDetails interface */
 	virtual void CustomizeDetails( IDetailLayoutBuilder& DetailBuilder ) override;
 
-	/** Callback when new CO is selected */
-	void ClipperCustomizableObjectSelectionChanged(const FAssetData& AssetData);
-
 private:
-	/** Callback when the clipping method is selected */
-	void OnClippingMethodComboBoxSelectionChanged(TSharedPtr<FString> Selection, ESelectInfo::Type SelectInfo);
-
-	/** Callback when new material name is selected */
-	void OnMeshClipWithMeshNodeComboBoxSelectionChanged(TSharedPtr<FString> Selection, ESelectInfo::Type SelectInfo);
-
-	/** Initializes the text of the method combobox and the view to display */
-	TSharedPtr<FString> SetInitialClippingMethod();
 
 	/** The node for which details are being customized */
-	class UCustomizableObjectNodeModifierClipWithMesh* Node;
+	UCustomizableObjectNodeModifierClipWithMesh* Node = nullptr;
 
 	/** Pointer to the builder passed by CustomizeDetails method */
 	IDetailLayoutBuilder* DetailBuilderPtr = nullptr;
 
-	/** References to the names of all possible methods to select the material */
-	TArray<TSharedPtr<FString>> ClippingMethods;
-
-	/** References to the names of all possible selectable material nodes in SelectedCO */
-	TArray<TSharedPtr<FString>> ArrayMaterialNodeOption;
-
-	/** To cover all LODs from a single selection, only unique names of materials are considered */
-	TArray<FString> ArrayMaterialNodeName;
-
-	/** Selected CO in the details tab of the Node node, for which to select a material that will be clipped by Node node */
-	class UCustomizableObject* SelectedCO;
-
-	/** Pointer to a FString to indicate that a method must be selected */
-	TSharedPtr<FString> SelectAMethod;
-
-	/** Callback to add a tag to the arry of tags */
-	FReply OnAddTagPressed();
 };
diff --git a/Engine/Plugins/Mutable/Source/CustomizableObjectEditor/Public/MuCOE/Nodes/CustomizableObjectNodeModifierClipWithMesh.h b/Engine/Plugins/Mutable/Source/CustomizableObjectEditor/Public/MuCOE/Nodes/CustomizableObjectNodeModifierClipWithMesh.h
--- a/Engine/Plugins/Mutable/Source/CustomizableObjectEditor/Public/MuCOE/Nodes/CustomizableObjectNodeModifierClipWithMesh.h
+++ b/Engine/Plugins/Mutable/Source/CustomizableObjectEditor/Public/MuCOE/Nodes/CustomizableObjectNodeModifierClipWithMesh.h
@@ -18,33 +18,29 @@ struct FGuid;
 UCLASS()
 class CUSTOMIZABLEOBJECTEDITOR_API UCustomizableObjectNodeModifierClipWithMesh : public UCustomizableObjectNodeModifierBase
 {
-public:
 	GENERATED_BODY()
 
-	UCustomizableObjectNodeModifierClipWithMesh();
+public:
 
 	UPROPERTY()
 	TArray<FString> Tags_DEPRECATED;
 
 	//!< If assigned, then a material inside this CO will be clipped by this node.
     //!< If several materials with the same name, all are considered (to cover all LOD levels)
-    UPROPERTY(EditAnywhere, Category = CustomizableObjectToClip)
-	TObjectPtr<UCustomizableObject> CustomizableObjectToClipWith;
+    UPROPERTY()
+	TObjectPtr<UCustomizableObject> CustomizableObjectToClipWith_DEPRECATED = nullptr;
 
     //!< Array with the Guids of the nodes with the same material inside the CustomizableObjectToClipWith CO (if any is assigned)
-    UPROPERTY(EditAnywhere, Category = CustomizableObjectToClip)
-    TArray<FGuid> ArrayMaterialNodeToClipWithID;
+    UPROPERTY()
+    TArray<FGuid> ArrayMaterialNodeToClipWithID_DEPRECATED;
 
-	/** Mesh Transform*/
-	UPROPERTY(EditAnywhere, Category = CustomizableObject)
+	/** Transform to apply to the clip mesh before clipping. */
+	UPROPERTY(EditAnywhere, Category = ClipMesh)
 	FTransform Transform;
 
-	// Details view variables
-	// The clipping node uses tags
-	bool bUseTags;
+public:
 
-	// The clipping node uses a material name
-	bool bUseMaterials;
+	UCustomizableObjectNodeModifierClipWithMesh();
 
 	// EdGraphNode interface
 	virtual FText GetNodeTitle(ENodeTitleType::Type TitleType) const override;
@@ -52,10 +48,6 @@ public:
 
 	// UCustomizableObjectNode interface
 	virtual void AllocateDefaultPins(UCustomizableObjectNodeRemapPins* RemapPins) override;
-	virtual void ReconstructNode(UCustomizableObjectNodeRemapPins* RemapPinsMode) override;
-	virtual void PinConnectionListChanged(UEdGraphPin* Pin) override;
-	virtual void UpdateReferencedNodeId(const FGuid& NewGuid) override;
-	virtual void BeginPostDuplicate(bool bDuplicateForPIE) override;
 	virtual void BackwardsCompatibleFixup(int32 CustomizableObjectCustomVersion) override;
 
 	// Own interface
