[Mutable] Morph global override rework.

- changed data structure to be a able to represent with autogenerated details.
- removed custom widget in favor of the autogenerated one. Thew old one was lacking basic functionality.
- Now skeletal mesh selection is prioritzed over morph name selections.
- Now the list is not autogenerated on the first compilation.

#rnx

diff --git a/Engine/Plugins/Mutable/Source/CustomizableObject/Internal/MuCO/CustomizableObjectCustomVersion.h b/Engine/Plugins/Mutable/Source/CustomizableObject/Internal/MuCO/CustomizableObjectCustomVersion.h
--- a/Engine/Plugins/Mutable/Source/CustomizableObject/Internal/MuCO/CustomizableObjectCustomVersion.h
+++ b/Engine/Plugins/Mutable/Source/CustomizableObject/Internal/MuCO/CustomizableObjectCustomVersion.h
@@ -158,6 +158,8 @@ struct CUSTOMIZABLEOBJECT_API FCustomizableObjectCustomVersion
 
 		ModifierClipWithMeshCleanup,
 
+		RealTimeMorphTargetOverrideDataStructureRework,
+
 		// -----<new versions can be added above this line>-------------------------------------------------
 		VersionPlusOne,
 		LatestVersion = VersionPlusOne - 1
diff --git a/Engine/Plugins/Mutable/Source/CustomizableObjectEditor/Private/MuCOE/CustomizableObjectCompiler.cpp b/Engine/Plugins/Mutable/Source/CustomizableObjectEditor/Private/MuCOE/CustomizableObjectCompiler.cpp
--- a/Engine/Plugins/Mutable/Source/CustomizableObjectEditor/Private/MuCOE/CustomizableObjectCompiler.cpp
+++ b/Engine/Plugins/Mutable/Source/CustomizableObjectEditor/Private/MuCOE/CustomizableObjectCompiler.cpp
@@ -716,7 +716,6 @@ mu::NodeObjectPtr FCustomizableObjectCompiler::GenerateMutableRoot(
 	}
 	
     GenerationContext.RealTimeMorphTargetsOverrides = ActualRoot->RealTimeMorphSelectionOverrides;
-    GenerationContext.RealTimeMorphTargetsOverrides.Reset();
 
 	if (!GenerationContext.ParamNamesToSelectedOptions.IsEmpty())
 	{
@@ -730,8 +729,6 @@ mu::NodeObjectPtr FCustomizableObjectCompiler::GenerateMutableRoot(
 	mu::NodeObjectPtr MutableRoot = GenerateMutableSource(ActualRoot->OutputPin(), GenerationContext);
 	UE_LOG(LogMutable, Verbose, TEXT("PROFILE: [ %16.8f ] GenerateMutableSource end."), FPlatformTime::Seconds());
 
-    ActualRoot->RealTimeMorphSelectionOverrides = GenerationContext.RealTimeMorphTargetsOverrides;
-
 	GenerationContext.GenerateSharedSurfacesUniqueIds();
 
 	// Generate ReferenceSkeletalMeshes data;
diff --git a/Engine/Plugins/Mutable/Source/CustomizableObjectEditor/Private/MuCOE/GenerateMutableSource/GenerateMutableSourceMesh.cpp b/Engine/Plugins/Mutable/Source/CustomizableObjectEditor/Private/MuCOE/GenerateMutableSource/GenerateMutableSourceMesh.cpp
--- a/Engine/Plugins/Mutable/Source/CustomizableObjectEditor/Private/MuCOE/GenerateMutableSource/GenerateMutableSourceMesh.cpp
+++ b/Engine/Plugins/Mutable/Source/CustomizableObjectEditor/Private/MuCOE/GenerateMutableSource/GenerateMutableSourceMesh.cpp
@@ -1178,28 +1178,6 @@ mu::MeshPtr ConvertSkeletalMeshToMutable(const USkeletalMesh* InSkeletalMesh, co
 		const UCustomizableObjectNodeTable* NodeTypedTable = Cast<UCustomizableObjectNodeTable>(CurrentNode);
 		check(NodeTypedSkMesh || (NodeTypedTable));
 
-        // Add SkeletalMesh morphs to the usage override data structure.
-        // Notice this will only be populated here, when compiling.
-		TArray<FRealTimeMorphSelectionOverride>& RealTimeMorphTargetOverrides = GenerationContext.RealTimeMorphTargetsOverrides;
-        for (const UMorphTarget* MorphTarget : SkeletalMeshMorphTargets)
-		{
-            // Find if the morph target global override is already present
-            // and add it if needed.
-            FRealTimeMorphSelectionOverride* MorphFound = RealTimeMorphTargetOverrides.FindByPredicate(
-                    [MorphTarget](const FRealTimeMorphSelectionOverride& O) { return O.MorphName == MorphTarget->GetFName(); });
-
-            MorphFound = MorphFound 
-                    ? MorphFound 
-                    : &RealTimeMorphTargetOverrides.Emplace_GetRef(MorphTarget->GetFName());
-
-            const int32 AddedMeshNameIdx = MorphFound->SkeletalMeshesNames.AddUnique(InSkeletalMesh->GetFName());
-
-			if (AddedMeshNameIdx >= MorphFound->Override.Num())
-			{
-				MorphFound->Override.Add(ECustomizableObjectSelectionOverride::NoOverride);
-			}
-        }
-
         // Add SkeletalMesh node used defined realtime morph targets to a temporal array where
         // the actual to be used real-time morphs names will be placed.        
         TArray<FName> UsedMorphTargetsNames = Invoke([&]()
@@ -1241,28 +1219,24 @@ mu::MeshPtr ConvertSkeletalMeshToMutable(const USkeletalMesh* InSkeletalMesh, co
             return MorphTargetsNames;
         });
 
-        //Apply global morph targets overrides to the SkeletalMesh user defined RT morph targets. 
+		TArray<FRealTimeMorphSelectionOverride>& RealTimeMorphTargetOverrides = GenerationContext.RealTimeMorphTargetsOverrides;
         for (FRealTimeMorphSelectionOverride& MorphTargetOverride : RealTimeMorphTargetOverrides)
         {
             const ECustomizableObjectSelectionOverride OverrideValue = 
 				Invoke([&]() -> ECustomizableObjectSelectionOverride
 				{
-					const ECustomizableObjectSelectionOverride GlobalOverrideValue = MorphTargetOverride.SelectionOverride;
-					
-					if (GlobalOverrideValue != ECustomizableObjectSelectionOverride::NoOverride)
-					{
-						return GlobalOverrideValue;
-					} 
-				
-					const int32 FoundIdx = 
-							MorphTargetOverride.SkeletalMeshesNames.Find(InSkeletalMesh->GetFName());
+					const int32 FoundMeshIndex = MorphTargetOverride.SkeletalMeshes.IndexOfByPredicate(
+							[Name = InSkeletalMesh->GetFName()](const FSkeletalMeshMorphTargetOverride& Elem)
+							{
+								return Name == Elem.SkeletalMeshName;
+							});
 
-					if (FoundIdx != INDEX_NONE)
+					if (FoundMeshIndex != INDEX_NONE)
 					{
-						return MorphTargetOverride.Override[FoundIdx];
+						return MorphTargetOverride.SkeletalMeshes[FoundMeshIndex].SelectionOverride;
 					}
 
-					return ECustomizableObjectSelectionOverride::NoOverride;
+					return MorphTargetOverride.SelectionOverride;
 				});
 
             if (OverrideValue == ECustomizableObjectSelectionOverride::Enable)
@@ -1275,6 +1249,8 @@ mu::MeshPtr ConvertSkeletalMeshToMutable(const USkeletalMesh* InSkeletalMesh, co
             }
         }
 
+
+
 		for (const UMorphTarget* MorphTarget : SkeletalMeshMorphTargets)
 		{
 			if (!MorphTarget)
diff --git a/Engine/Plugins/Mutable/Source/CustomizableObjectEditor/Private/MuCOE/Nodes/CustomizableObjectNodeObject.cpp b/Engine/Plugins/Mutable/Source/CustomizableObjectEditor/Private/MuCOE/Nodes/CustomizableObjectNodeObject.cpp
--- a/Engine/Plugins/Mutable/Source/CustomizableObjectEditor/Private/MuCOE/Nodes/CustomizableObjectNodeObject.cpp
+++ b/Engine/Plugins/Mutable/Source/CustomizableObjectEditor/Private/MuCOE/Nodes/CustomizableObjectNodeObject.cpp
@@ -512,6 +512,24 @@ void UCustomizableObjectNodeObject::BackwardsCompatibleFixup(int32 CustomizableO
 			}
 		}
 	}
+
+	if (CustomizableObjectCustomVersion == FCustomizableObjectCustomVersion::RealTimeMorphTargetOverrideDataStructureRework)
+	{
+		for (FRealTimeMorphSelectionOverride& Override : RealTimeMorphSelectionOverrides)
+		{
+			for (int32 SkeletalMeshIndex = 0; SkeletalMeshIndex < Override.SkeletalMeshesNames_DEPRECATED.Num(); ++SkeletalMeshIndex)
+			{
+				// Only add SkeletalMesh override if the morph selection was NoOverride and the the skeletalmesh had override. 
+				if (Override.SelectionOverride == ECustomizableObjectSelectionOverride::NoOverride &&
+					Override.Override_DEPRECATED[SkeletalMeshIndex] != ECustomizableObjectSelectionOverride::NoOverride)
+				{
+					FSkeletalMeshMorphTargetOverride& SkeletalMeshOverride = Override.SkeletalMeshes.AddDefaulted_GetRef();
+					SkeletalMeshOverride.SkeletalMeshName = Override.SkeletalMeshesNames_DEPRECATED[SkeletalMeshIndex];
+					SkeletalMeshOverride.SelectionOverride = Override.Override_DEPRECATED[SkeletalMeshIndex];
+				}
+			}
+		}
+	}
 }
 
 
diff --git a/Engine/Plugins/Mutable/Source/CustomizableObjectEditor/Private/MuCOE/Nodes/CustomizableObjectNodeObjectDetails.cpp b/Engine/Plugins/Mutable/Source/CustomizableObjectEditor/Private/MuCOE/Nodes/CustomizableObjectNodeObjectDetails.cpp
--- a/Engine/Plugins/Mutable/Source/CustomizableObjectEditor/Private/MuCOE/Nodes/CustomizableObjectNodeObjectDetails.cpp
+++ b/Engine/Plugins/Mutable/Source/CustomizableObjectEditor/Private/MuCOE/Nodes/CustomizableObjectNodeObjectDetails.cpp
@@ -169,7 +169,6 @@ void FCustomizableObjectNodeObjectDetails::CustomizeDetails( IDetailLayoutBuilde
 
 	IDetailCategoryBuilder& StatesCategory = DetailBuilder.EditCategory("States");
 	IDetailCategoryBuilder& ExternalCategory = DetailBuilder.EditCategory("AttachedToExternalObject");
-	IDetailCategoryBuilder& RealTimeMorphTargets = DetailBuilder.EditCategory("RealTime Morph Targets");
 	
 	if (BaseObjectNode.IsValid())
 	{
@@ -248,16 +247,6 @@ void FCustomizableObjectNodeObjectDetails::CustomizeDetails( IDetailLayoutBuilde
 				
 				DetailBuilder.HideProperty(StatesProperty);
 			}
-            else
-            {
-                RealTimeMorphTargets.AddCustomRow(
-                    LOCTEXT("FCustomizableObjectNodeObjectMorphTargetsDetails", 
-                            "Realtime Morph Targets Override "))
-                [
-                    SNew(SCustomizableObjectNodeSkeletalMeshRTMorphTargetOverride)
-                    .Node(BaseObjectNode.Get())
-                ];
-            }
 		}
 		else
 		{
diff --git a/Engine/Plugins/Mutable/Source/CustomizableObjectEditor/Private/MuCOE/SCustomizableObjectNodeObjectRTMorphTargetOverride.cpp b/Engine/Plugins/Mutable/Source/CustomizableObjectEditor/Private/MuCOE/SCustomizableObjectNodeObjectRTMorphTargetOverride.cpp
--- a/Engine/Plugins/Mutable/Source/CustomizableObjectEditor/Private/MuCOE/SCustomizableObjectNodeObjectRTMorphTargetOverride.cpp
+++ b/Engine/Plugins/Mutable/Source/CustomizableObjectEditor/Private/MuCOE/SCustomizableObjectNodeObjectRTMorphTargetOverride.cpp
@@ -9,6 +9,7 @@
 #include "Widgets/Input/SSearchBox.h"
 #include "Widgets/Layout/SExpandableArea.h"
 #include "Widgets/Views/SListView.h"
+#include "DetailLayoutBuilder.h"
 
 class ITableRow;
 class STableViewBase;
@@ -66,18 +67,20 @@ TSharedRef<SWidget> SRowWidget::GenerateWidgetForColumn(const FName& ColumnName)
     {
         return
             SNew(SVerticalBox)
-            + SVerticalBox::Slot().AutoHeight().Padding(4).VAlign(VAlign_Center)
+            + SVerticalBox::Slot().AutoHeight().VAlign(VAlign_Center)
             [
                 SNew(SHorizontalBox)
-                + SHorizontalBox::Slot().Padding(4).VAlign(VAlign_Center)
+                + SHorizontalBox::Slot().Padding(4,0,0,0).VAlign(VAlign_Center)
                 [
                     SNew(STextBlock).Text(FText::FromName(Item->SkeletalMeshName))
+                        .Font(IDetailLayoutBuilder::GetDetailFont())
                 ]
-                +SHorizontalBox::Slot().Padding(4).HAlign(HAlign_Right)
+                +SHorizontalBox::Slot().Padding(4,0,0,0).HAlign(HAlign_Right)
                 [
                     SNew(SEnumComboBox, StaticEnum<ECustomizableObjectSelectionOverride>())
                     .CurrentValue(this, &ThisType::OverrideCurrentValue)
                     .OnEnumSelectionChanged(this, &ThisType::OnEnumSelectionChanged)
+                    .Font(IDetailLayoutBuilder::GetDetailFont())
                 ]
             ];
     }
@@ -102,14 +105,18 @@ ECustomizableObjectSelectionOverride* SRowWidget::GetItemEnumPtr(UCustomizableOb
         return nullptr;
     }
 
-    int32 SkeletalMeshIndex = MorphFound->SkeletalMeshesNames.Find(Item->SkeletalMeshName);
+    FSkeletalMeshMorphTargetOverride* SkeletalMeshOverride = MorphFound->SkeletalMeshes.FindByPredicate(
+            [SkeletalMeshName = Item->SkeletalMeshName](FSkeletalMeshMorphTargetOverride& O) 
+            { 
+                return O.SkeletalMeshName == SkeletalMeshName; 
+            });
     
-    if (SkeletalMeshIndex == INDEX_NONE)
+    if (!SkeletalMeshOverride)
     {
         return nullptr;
     }
 
-    return &MorphFound->Override[SkeletalMeshIndex];
+    return &SkeletalMeshOverride->SelectionOverride;
 }
 
 void SRowWidget::OnEnumSelectionChanged(int32 InEnumValue, ESelectInfo::Type)
@@ -157,8 +164,8 @@ void SSkeletalMeshListViewWidget::UpdateWidget()
     
     ChildSlot
     [
-        SNew(SBorder)
-        [
+        //SNew(SBorder)
+        //[
             SNew(SVerticalBox)
             +SVerticalBox::Slot()
             [
@@ -172,9 +179,9 @@ void SSkeletalMeshListViewWidget::UpdateWidget()
                     SNew(SHeaderRow)
                     + SHeaderRow::Column(ViewRows::NameColumn)
                     .DefaultLabel(LOCTEXT("SkeletalMeshLabel", "Skeletal Meshes"))
-                )
+                ) 
             ]
-        ]
+        //]
     ];
 
     View->RequestListRefresh();
@@ -221,15 +228,14 @@ void SSkeletalMeshListViewWidget::PopulateItems()
             FilterText.IsEmpty() || 
             MorphName.ToString().Contains(FilterText);
 
-    for (const FName& SkeletalMeshName : Found->SkeletalMeshesNames)
+    for (const FSkeletalMeshMorphTargetOverride& SkeletalMeshOverride : Found->SkeletalMeshes)
     {
         const bool bFilterPass =
-            bMorphFilterPass || 
-            SkeletalMeshName.ToString().Contains(FilterText);
+            bMorphFilterPass || SkeletalMeshOverride.SkeletalMeshName.ToString().Contains(FilterText);
 
         if (bFilterPass)
         {
-            Items.Emplace(MakeShareable(new Internal::FItemInfo(MorphName, SkeletalMeshName)));
+            Items.Emplace(MakeShareable(new Internal::FItemInfo(MorphName, SkeletalMeshOverride.SkeletalMeshName)));
         }
     } 
 }
@@ -278,30 +284,34 @@ TSharedRef<SWidget> SRowWidget::GenerateWidgetForColumn(const FName& ColumnName)
             SNew(SBorder)
             [ 
                 SNew(SVerticalBox)
-                + SVerticalBox::Slot().AutoHeight().Padding(4).VAlign(VAlign_Center)
+                + SVerticalBox::Slot().AutoHeight().VAlign(VAlign_Center)
                 [
                     SNew(SExpandableArea)
+                    .BorderImage(FAppStyle::GetBrush("DetailsView.GridLine"))
                     .InitiallyCollapsed(true)
                     .HeaderContent()
                     [
                     SNew(SHorizontalBox)
-                        + SHorizontalBox::Slot().Padding(4).VAlign(VAlign_Center)
+                    + SHorizontalBox::Slot().VAlign(VAlign_Center)
                     [
-                        SNew(STextBlock).Text(FText::FromName(Item->Name))
+                        SNew(STextBlock)
+                        .Text(FText::FromName(Item->Name))
+                        .Font(IDetailLayoutBuilder::GetDetailFont())
                     ]
-                    + SHorizontalBox::Slot().Padding(4).HAlign(HAlign_Right)
+                    + SHorizontalBox::Slot().HAlign(HAlign_Right)
                     [
                         SNew(SEnumComboBox, StaticEnum<ECustomizableObjectSelectionOverride>())
                         .CurrentValue(this, &ThisType::SelectionCurrentValue)
                         .OnEnumSelectionChanged(this, &ThisType::OnEnumSelectionChanged)
+                        .Font(IDetailLayoutBuilder::GetDetailFont())
+                    ]
                     ]
-                ]
                     .BodyContent()
-                [
+                    [
                     SNew(SSkeletalMeshListViewWidget)
                     .MorphName(Item->Name)
                     .MainWidget(MainWidget)
-                ]
+                    ]
                 ]
             ];
     }
@@ -410,9 +420,9 @@ void SCustomizableObjectNodeSkeletalMeshRTMorphTargetOverride::PopulateItems()
 
         const bool bSearchBoxContainsSkeletalMeshName = [&]()
             {
-                for (const FName& SkeletalMeshName : Override.SkeletalMeshesNames)
+                for (const FSkeletalMeshMorphTargetOverride& SkeletalMeshOverride : Override.SkeletalMeshes)
                 {
-                    if (SkeletalMeshName.ToString().Contains(SearchBoxFilter))
+                    if (SkeletalMeshOverride.SkeletalMeshName.ToString().Contains(SearchBoxFilter))
                     {
                         return true;
                     }
diff --git a/Engine/Plugins/Mutable/Source/CustomizableObjectEditor/Public/MuCOE/Nodes/CustomizableObjectNodeObject.h b/Engine/Plugins/Mutable/Source/CustomizableObjectEditor/Public/MuCOE/Nodes/CustomizableObjectNodeObject.h
--- a/Engine/Plugins/Mutable/Source/CustomizableObjectEditor/Public/MuCOE/Nodes/CustomizableObjectNodeObject.h
+++ b/Engine/Plugins/Mutable/Source/CustomizableObjectEditor/Public/MuCOE/Nodes/CustomizableObjectNodeObject.h
@@ -74,36 +74,39 @@ private:
 	bool bDontCompressRuntimeTextures_DEPRECATED = false;
 };
 
+USTRUCT()
+struct FSkeletalMeshMorphTargetOverride
+{
+	GENERATED_BODY()
+	
+	UPROPERTY(EditAnywhere, Category = RealTimeMorphTargets)
+	FName SkeletalMeshName;
+
+	UPROPERTY(EditAnywhere, Category = RealTimeMorphTargets)
+	ECustomizableObjectSelectionOverride SelectionOverride = ECustomizableObjectSelectionOverride::NoOverride;
+};
 
 USTRUCT()
 struct FRealTimeMorphSelectionOverride
 {
 	GENERATED_BODY()
 
-    FRealTimeMorphSelectionOverride()
-    {
-    }
+	UPROPERTY(EditAnywhere, Category = RealTimeMorphTargets)
+	FName MorphName;
 
-    FRealTimeMorphSelectionOverride(const FName& InMorphName)
-        : MorphName(InMorphName)
-    {
-    }
+	UPROPERTY(EditAnywhere, Category = RealTimeMorphTargets)
+	ECustomizableObjectSelectionOverride SelectionOverride = ECustomizableObjectSelectionOverride::NoOverride;
 
-    UPROPERTY()
-    FName MorphName;
-
-    UPROPERTY()
-    ECustomizableObjectSelectionOverride SelectionOverride = 
-        ECustomizableObjectSelectionOverride::NoOverride;
+	UPROPERTY(EditAnywhere, Category = RealTimeMorphTargets, meta=(TitleProperty=SkeletalMeshName))
+	TArray<FSkeletalMeshMorphTargetOverride> SkeletalMeshes;
 
     UPROPERTY()
-    TArray<FName> SkeletalMeshesNames;
+    TArray<FName> SkeletalMeshesNames_DEPRECATED;
 
     UPROPERTY()
-    TArray<ECustomizableObjectSelectionOverride> Override;
+    TArray<ECustomizableObjectSelectionOverride> Override_DEPRECATED;
 };
 
-
 UCLASS()
 class CUSTOMIZABLEOBJECTEDITOR_API UCustomizableObjectNodeObject : public UCustomizableObjectNode
 {
@@ -139,12 +142,12 @@ public:
 	UPROPERTY()
 	FGuid Identifier;
 
-    UPROPERTY()
-    TArray<FRealTimeMorphSelectionOverride> RealTimeMorphSelectionOverrides;
+	UPROPERTY(EditAnywhere, Category = RealTimeMorphTargets, meta=(TitleProperty=MorphName, EditCondition="ParentObject==nullptr"))
+	TArray<FRealTimeMorphSelectionOverride> RealTimeMorphSelectionOverrides;
 
 	UPROPERTY()
 	TArray<FComponentSettings> ComponentSettings_DEPRECATED;
-	
+		
 	// UObject interface.
 	virtual void PostEditChangeProperty(FPropertyChangedEvent& PropertyChangedEvent) override;
 
