[Mutable] Set automatic blocks strategy to Ignore for legacy nodes.

#rnx

diff --git a/Engine/Plugins/Mutable/Source/CustomizableObject/Internal/MuCO/CustomizableObjectCustomVersion.h b/Engine/Plugins/Mutable/Source/CustomizableObject/Internal/MuCO/CustomizableObjectCustomVersion.h
--- a/Engine/Plugins/Mutable/Source/CustomizableObject/Internal/MuCO/CustomizableObjectCustomVersion.h
+++ b/Engine/Plugins/Mutable/Source/CustomizableObject/Internal/MuCO/CustomizableObjectCustomVersion.h
@@ -159,6 +159,8 @@ struct CUSTOMIZABLEOBJECT_API FCustomizableObjectCustomVersion
 		ModifierClipWithMeshCleanup,
 
 		RealTimeMorphTargetOverrideDataStructureRework,
+		
+		FixAutomaticBlocksStrategyLegacyNodes,
 
 		// -----<new versions can be added above this line>-------------------------------------------------
 		VersionPlusOne,
diff --git a/Engine/Plugins/Mutable/Source/CustomizableObjectEditor/Private/MuCOE/CustomizableObjectLayout.cpp b/Engine/Plugins/Mutable/Source/CustomizableObjectEditor/Private/MuCOE/CustomizableObjectLayout.cpp
--- a/Engine/Plugins/Mutable/Source/CustomizableObjectEditor/Private/MuCOE/CustomizableObjectLayout.cpp
+++ b/Engine/Plugins/Mutable/Source/CustomizableObjectEditor/Private/MuCOE/CustomizableObjectLayout.cpp
@@ -21,6 +21,9 @@ UCustomizableObjectLayout::UCustomizableObjectLayout()
 	GridSize = FIntPoint(4, 4);
 	MaxGridSize = FIntPoint(4, 4);
 
+	FCustomizableObjectLayoutBlock Block(FIntPoint(0, 0), FIntPoint(4, 4));
+	Blocks.Add(Block);
+
 	PackingStrategy = ECustomizableObjectTextureLayoutPackingStrategy::Resizable;
 	BlockReductionMethod = ECustomizableObjectLayoutBlockReductionMethod::Halve;
 }
diff --git a/Engine/Plugins/Mutable/Source/CustomizableObjectEditor/Private/MuCOE/GenerateMutableSource/GenerateMutableSourceLayout.cpp b/Engine/Plugins/Mutable/Source/CustomizableObjectEditor/Private/MuCOE/GenerateMutableSource/GenerateMutableSourceLayout.cpp
--- a/Engine/Plugins/Mutable/Source/CustomizableObjectEditor/Private/MuCOE/GenerateMutableSource/GenerateMutableSourceLayout.cpp
+++ b/Engine/Plugins/Mutable/Source/CustomizableObjectEditor/Private/MuCOE/GenerateMutableSource/GenerateMutableSourceLayout.cpp
@@ -80,7 +80,7 @@ mu::Ptr<mu::NodeLayout> CreateMutableLayoutNode(FMutableGraphGenerationContext&
 			LayoutNode->Blocks[0].bReduceByTwo = false;
 		}
 	}
-	else if (UnrealLayout->Blocks.IsEmpty())
+	else
 	{
 		// Convert the UE mesh in the layout into a Mutable mesh.
 		mu::Ptr<mu::Mesh> MutableMesh = nullptr;
diff --git a/Engine/Plugins/Mutable/Source/CustomizableObjectEditor/Private/MuCOE/Nodes/CustomizableObjectNodeSkeletalMesh.cpp b/Engine/Plugins/Mutable/Source/CustomizableObjectEditor/Private/MuCOE/Nodes/CustomizableObjectNodeSkeletalMesh.cpp
--- a/Engine/Plugins/Mutable/Source/CustomizableObjectEditor/Private/MuCOE/Nodes/CustomizableObjectNodeSkeletalMesh.cpp
+++ b/Engine/Plugins/Mutable/Source/CustomizableObjectEditor/Private/MuCOE/Nodes/CustomizableObjectNodeSkeletalMesh.cpp
@@ -661,6 +661,23 @@ void UCustomizableObjectNodeSkeletalMesh::BackwardsCompatibleFixup(int32 Customi
 			}
 		}
 	}
+
+	if (CustomizableObjectCustomVersion == FCustomizableObjectCustomVersion::FixAutomaticBlocksStrategyLegacyNodes)
+	{
+		for (UEdGraphPin* Pin : GetAllPins())
+		{
+			if (UCustomizableObjectNodeSkeletalMeshPinDataMesh* MeshPinData = Cast<UCustomizableObjectNodeSkeletalMeshPinDataMesh>(GetPinData(*Pin)))
+			{
+				for (UCustomizableObjectLayout* Layout : MeshPinData->Layouts)
+				{
+					if (Layout)
+					{
+						Layout->AutomaticBlocksStrategy = ECustomizableObjectLayoutAutomaticBlocksStrategy::Ignore;
+					}
+				}
+			}
+		}
+	}
 }
 
 
diff --git a/Engine/Plugins/Mutable/Source/CustomizableObjectEditor/Private/MuCOE/Nodes/CustomizableObjectNodeTable.cpp b/Engine/Plugins/Mutable/Source/CustomizableObjectEditor/Private/MuCOE/Nodes/CustomizableObjectNodeTable.cpp
--- a/Engine/Plugins/Mutable/Source/CustomizableObjectEditor/Private/MuCOE/Nodes/CustomizableObjectNodeTable.cpp
+++ b/Engine/Plugins/Mutable/Source/CustomizableObjectEditor/Private/MuCOE/Nodes/CustomizableObjectNodeTable.cpp
@@ -181,6 +181,23 @@ void UCustomizableObjectNodeTable::BackwardsCompatibleFixup(int32 CustomizableOb
 			}
 		}
 	}
+
+	if (CustomizableObjectCustomVersion == FCustomizableObjectCustomVersion::FixAutomaticBlocksStrategyLegacyNodes)
+	{
+		for (UEdGraphPin* Pin : GetAllPins())
+		{
+			if (UCustomizableObjectNodeTableMeshPinData* MeshPinData = Cast<UCustomizableObjectNodeTableMeshPinData>(GetPinData(*Pin)))
+			{
+				for (UCustomizableObjectLayout* Layout : MeshPinData->Layouts)
+				{
+					if (Layout)
+					{
+						Layout->AutomaticBlocksStrategy = ECustomizableObjectLayoutAutomaticBlocksStrategy::Ignore;
+					}
+				}
+			}
+		}
+	}
 }
 
 
