[Mutable] Fix missing blocks in RemoveMeshBlocks and EditMeshSections nodes. GetLayouts from NodeLayoutBlocks for versions older than FCustomizableObjectCustomVersion::RemoveNodeLayout.

#rnx

diff --git a/Engine/Plugins/Mutable/Source/CustomizableObjectEditor/Private/MuCOE/Nodes/CustomizableObjectNodeModifierEditMeshSection.cpp b/Engine/Plugins/Mutable/Source/CustomizableObjectEditor/Private/MuCOE/Nodes/CustomizableObjectNodeModifierEditMeshSection.cpp
--- a/Engine/Plugins/Mutable/Source/CustomizableObjectEditor/Private/MuCOE/Nodes/CustomizableObjectNodeModifierEditMeshSection.cpp
+++ b/Engine/Plugins/Mutable/Source/CustomizableObjectEditor/Private/MuCOE/Nodes/CustomizableObjectNodeModifierEditMeshSection.cpp
@@ -8,6 +8,10 @@
 #include "MuCOE/EdGraphSchema_CustomizableObject.h"
 #include "MuCOE/CustomizableObjectPin.h"
 #include "MuCO/CustomizableObjectCustomVersion.h"
+#include "MuCOE/Nodes/CustomizableObjectNodeSkeletalMesh.h"
+#include "MuCOE/Nodes/CustomizableObjectNodeTable.h"
+#include "MuCOE/Nodes/CustomizableObjectNodeLayoutBlocks.h"
+#include "MuCOE/GraphTraversal.h"
 #include "MaterialTypes.h"
 
 class FCustomizableObjectNodeParentedMaterial;
@@ -90,13 +94,82 @@ void UCustomizableObjectNodeModifierEditMeshSection::BackwardsCompatibleFixup(in
 {
 	Super::BackwardsCompatibleFixup(CustomizableObjectCustomVersion);
 
+	const auto GetLayouts = [](UCustomizableObjectNodeMaterialBase& ParentMaterialNode, TArray<UCustomizableObjectLayout*>& OutLayouts)
+		{
+			UCustomizableObjectNodeMaterial* NodeMaterial = ParentMaterialNode.GetMaterialNode();
+			if (!NodeMaterial)
+			{
+				return;
+			}
+
+			UEdGraphPin* MeshPin = NodeMaterial->GetMeshPin();
+			if (!MeshPin)
+			{
+				return;
+			}
+
+			const UEdGraphPin* ConnectedPin = FollowInputPin(*MeshPin);
+			if (!ConnectedPin)
+			{
+				return;
+			}
+
+			const UEdGraphPin* SourceMeshPin = FindMeshBaseSource(*ConnectedPin, false);
+			if (!SourceMeshPin)
+			{
+				return;
+			}
+
+			if (const UCustomizableObjectNodeSkeletalMesh* MeshNode = Cast<UCustomizableObjectNodeSkeletalMesh>(SourceMeshPin->GetOwningNode()))
+			{
+				UCustomizableObjectNodeSkeletalMeshPinDataMesh* MeshPinData = Cast<UCustomizableObjectNodeSkeletalMeshPinDataMesh>(MeshNode->GetPinData(*SourceMeshPin));
+				
+				if (!MeshPinData)
+				{
+					return;
+				}
+
+				// The custom version of SkeletalMesh nodes may be up to date if they're in a different CO. 
+				OutLayouts = MeshPinData->Layouts;
+
+				if (OutLayouts.IsEmpty())
+				{
+					// Pre FCustomizableObjectCustomVersion::RemoveNodeLayout code. Get Layouts from NodeLayoutBlocks
+					TArray<UEdGraphPin*> NonOrphanPins = MeshNode->GetAllNonOrphanPins();
+					for (const UEdGraphPin* Pin : NonOrphanPins)
+					{
+						if (const UCustomizableObjectNodeSkeletalMeshPinDataLayout* PinData = Cast<UCustomizableObjectNodeSkeletalMeshPinDataLayout>(MeshNode->GetPinData(*Pin)))
+						{
+							if (PinData->GetLODIndex() == MeshPinData->GetLODIndex() &&
+								PinData->GetSectionIndex() == MeshPinData->GetSectionIndex())
+							{
+								if (const UEdGraphPin* SourceLayoutConnectedPin = FollowInputPin(*Pin))
+								{
+									const UCustomizableObjectNodeLayoutBlocks* LayoutNode = Cast<UCustomizableObjectNodeLayoutBlocks>(SourceLayoutConnectedPin->GetOwningNode());
+									if (LayoutNode && LayoutNode->Layout)
+									{
+										OutLayouts.Add(LayoutNode->Layout);
+									}
+								}
+							}
+						}
+					}
+				}
+			}
+			else if (const UCustomizableObjectNodeTable* TableNode = Cast<UCustomizableObjectNodeTable>(SourceMeshPin->GetOwningNode()))
+			{
+				OutLayouts = TableNode->GetLayouts(SourceMeshPin);
+			}
+		};
+
 	// Convert deprecated node index list to the node id list.
 	if (CustomizableObjectCustomVersion == FCustomizableObjectCustomVersion::PostLoadToCustomVersion
 		&& BlockIds_DEPRECATED.Num() < Blocks_DEPRECATED.Num())
 	{
 		UCustomizableObjectNodeMaterialBase* ParentMaterialNode = GetCustomizableObjectExternalNode<UCustomizableObjectNodeMaterialBase>(ParentMaterialObject_DEPRECATED.Get(), ParentMaterialNodeId_DEPRECATED);
 
-		TArray<UCustomizableObjectLayout*> Layouts = ParentMaterialNode->GetLayouts();
+		TArray<UCustomizableObjectLayout*> Layouts;
+		GetLayouts(*ParentMaterialNode, Layouts);
 
 		if (!Layouts.IsValidIndex(ParentLayoutIndex))
 		{
@@ -143,7 +216,8 @@ void UCustomizableObjectNodeModifierEditMeshSection::BackwardsCompatibleFixup(in
 
 		UCustomizableObjectNodeMaterialBase* ParentMaterialNode = GetCustomizableObjectExternalNode<UCustomizableObjectNodeMaterialBase>(ParentMaterialObject_DEPRECATED.Get(), ParentMaterialNodeId_DEPRECATED);
 
-		TArray<UCustomizableObjectLayout*> ParentLayouts = ParentMaterialNode->GetLayouts();
+		TArray<UCustomizableObjectLayout*> ParentLayouts;
+		GetLayouts(*ParentMaterialNode, ParentLayouts);
 
 		if (!ParentLayouts.IsValidIndex(ParentLayoutIndex))
 		{
diff --git a/Engine/Plugins/Mutable/Source/CustomizableObjectEditor/Private/MuCOE/Nodes/CustomizableObjectNodeModifierRemoveMeshBlocks.cpp b/Engine/Plugins/Mutable/Source/CustomizableObjectEditor/Private/MuCOE/Nodes/CustomizableObjectNodeModifierRemoveMeshBlocks.cpp
--- a/Engine/Plugins/Mutable/Source/CustomizableObjectEditor/Private/MuCOE/Nodes/CustomizableObjectNodeModifierRemoveMeshBlocks.cpp
+++ b/Engine/Plugins/Mutable/Source/CustomizableObjectEditor/Private/MuCOE/Nodes/CustomizableObjectNodeModifierRemoveMeshBlocks.cpp
@@ -8,6 +8,10 @@
 #include "MuCOE/EdGraphSchema_CustomizableObject.h"
 #include "MuCOE/ICustomizableObjectEditor.h"
 #include "MuCOE/Nodes/CustomizableObjectNodeMaterial.h"
+#include "MuCOE/Nodes/CustomizableObjectNodeSkeletalMesh.h"
+#include "MuCOE/Nodes/CustomizableObjectNodeTable.h"
+#include "MuCOE/Nodes/CustomizableObjectNodeLayoutBlocks.h"
+#include "MuCOE/GraphTraversal.h"
 
 class UCustomizableObjectNodeRemapPins;
 class UEdGraphPin;
@@ -85,7 +89,72 @@ void UCustomizableObjectNodeModifierRemoveMeshBlocks::BackwardsCompatibleFixup(i
 		}
 		else
 		{
-			ParentLayouts = ParentMaterialNode->GetLayouts();
+			[&]()
+				{
+					UCustomizableObjectNodeMaterial* NodeMaterial = ParentMaterialNode->GetMaterialNode();
+					if (!NodeMaterial)
+					{
+						return;
+					}
+
+					UEdGraphPin* MeshPin = NodeMaterial->GetMeshPin();
+					if (!MeshPin)
+					{
+						return;
+					}
+
+					const UEdGraphPin* ConnectedPin = FollowInputPin(*MeshPin);
+					if (!ConnectedPin)
+					{
+						return;
+					}
+
+					const UEdGraphPin* SourceMeshPin = FindMeshBaseSource(*ConnectedPin, false);
+					if (!SourceMeshPin)
+					{
+						return;
+					}
+
+					if (const UCustomizableObjectNodeSkeletalMesh* MeshNode = Cast<UCustomizableObjectNodeSkeletalMesh>(SourceMeshPin->GetOwningNode()))
+					{
+						UCustomizableObjectNodeSkeletalMeshPinDataMesh* MeshPinData = Cast<UCustomizableObjectNodeSkeletalMeshPinDataMesh>(MeshNode->GetPinData(*SourceMeshPin));
+
+						if (!MeshPinData)
+						{
+							return;
+						}
+
+						// The custom version of SkeletalMesh nodes may be up to date if it belongs to a different CO. 
+						ParentLayouts = MeshPinData->Layouts;
+
+						if (ParentLayouts.IsEmpty())
+						{
+							// Pre FCustomizableObjectCustomVersion::RemoveNodeLayout code. Get Layouts from NodeLayoutBlocks
+							for (const UEdGraphPin* Pin : MeshNode->GetAllNonOrphanPins())
+							{
+								if (const UCustomizableObjectNodeSkeletalMeshPinDataLayout* PinData = Cast<UCustomizableObjectNodeSkeletalMeshPinDataLayout>(MeshNode->GetPinData(*Pin)))
+								{
+									if (PinData->GetLODIndex() == MeshPinData->GetLODIndex() &&
+										PinData->GetSectionIndex() == MeshPinData->GetSectionIndex())
+									{
+										if (const UEdGraphPin* SourceLayoutConnectedPin = FollowInputPin(*Pin))
+										{
+											const UCustomizableObjectNodeLayoutBlocks* LayoutNode = Cast<UCustomizableObjectNodeLayoutBlocks>(SourceLayoutConnectedPin->GetOwningNode());
+											if (LayoutNode && LayoutNode->Layout)
+											{
+												ParentLayouts.Add(LayoutNode->Layout);
+											}
+										}
+									}
+								}
+							}
+						}
+					}
+					else if (const UCustomizableObjectNodeTable* TableNode = Cast<UCustomizableObjectNodeTable>(SourceMeshPin->GetOwningNode()))
+					{
+						ParentLayouts = TableNode->GetLayouts(SourceMeshPin);
+					}
+				}(); // invoke lambda GetLayouts();
 		}
 
 		if (!ParentLayouts.IsValidIndex(ParentLayoutIndex))
