[Mutable] Fix crash when using Passthrough Mesh Component nodes

#rnx

diff --git a/Engine/Plugins/Mutable/Source/CustomizableObject/Internal/MuCO/CustomizableObjectPrivate.h b/Engine/Plugins/Mutable/Source/CustomizableObject/Internal/MuCO/CustomizableObjectPrivate.h
--- a/Engine/Plugins/Mutable/Source/CustomizableObject/Internal/MuCO/CustomizableObjectPrivate.h
+++ b/Engine/Plugins/Mutable/Source/CustomizableObject/Internal/MuCO/CustomizableObjectPrivate.h
@@ -1081,10 +1081,6 @@ struct FModelResources
 	UPROPERTY()
 	TArray<FCustomizableObjectStreamedResourceData> StreamedResourceData;
 
-	/** Max number of components in the compiled Model. */
-	UPROPERTY()
-	uint8 NumComponents = 0; 
-
 	/** Max number of LODs in the compiled Model. */
 	UPROPERTY()
 	uint8 NumLODs = 0;
@@ -1471,6 +1467,8 @@ public:
 		DerivedDataCache,
 
 		ComponentsArray,
+
+		FixComponentNames,
 		
 		// -----<new versions can be added above this line>--------
 		LastCustomizableObjectVersion
diff --git a/Engine/Plugins/Mutable/Source/CustomizableObject/Private/MuCO/CustomizableObject.cpp b/Engine/Plugins/Mutable/Source/CustomizableObject/Private/MuCO/CustomizableObject.cpp
--- a/Engine/Plugins/Mutable/Source/CustomizableObject/Private/MuCO/CustomizableObject.cpp
+++ b/Engine/Plugins/Mutable/Source/CustomizableObject/Private/MuCO/CustomizableObject.cpp
@@ -738,7 +738,6 @@ void FModelResources::Serialize(FObjectAndNameAsStringProxyArchive& MemoryWriter
 	MemoryWriter << ClothingAssetsData;
 	MemoryWriter << ClothSharedConfigsData;
 
-	MemoryWriter << NumComponents;
 	MemoryWriter << NumLODs;
 	MemoryWriter << NumLODsToStream;
 	MemoryWriter << FirstLODAvailable;
@@ -890,7 +889,6 @@ bool FModelResources::Unserialize(FObjectAndNameAsStringProxyArchive& MemoryRead
 	MemoryReader << ClothingAssetsData;
 	MemoryReader << ClothSharedConfigsData;
 
-	MemoryReader << NumComponents;
 	MemoryReader << NumLODs;
 	MemoryReader << NumLODsToStream;
 	MemoryReader << FirstLODAvailable;
@@ -1405,7 +1403,7 @@ int32 UCustomizableObject::GetComponentCount() const
 {
 	if (IsCompiled())
 	{
-		return GetPrivate()->GetModelResources().NumComponents;
+		return GetPrivate()->GetModelResources().ComponentNames.Num();
 	}
 
 	return 0;
diff --git a/Engine/Plugins/Mutable/Source/CustomizableObjectEditor/Private/MuCOE/CustomizableObjectCompiler.cpp b/Engine/Plugins/Mutable/Source/CustomizableObjectEditor/Private/MuCOE/CustomizableObjectCompiler.cpp
--- a/Engine/Plugins/Mutable/Source/CustomizableObjectEditor/Private/MuCOE/CustomizableObjectCompiler.cpp
+++ b/Engine/Plugins/Mutable/Source/CustomizableObjectEditor/Private/MuCOE/CustomizableObjectCompiler.cpp
@@ -1134,12 +1134,7 @@ void FCustomizableObjectCompiler::CompileInternal(bool bAsync)
 		ModelResources.CustomizableObjectPathMap = MoveTemp(GenerationContext.CustomizableObjectPathMap);
 #endif
 
-		ModelResources.NumComponents = GenerationContext.NumComponents;
-
-		for (int32 Index = 0; Index <  GenerationContext.ComponentInfos.Num(); ++Index)
-		{
-			ModelResources.ComponentNames.Add(GenerationContext.ComponentInfos[Index].ComponentName);
-		}
+		ModelResources.ComponentNames = MoveTemp(GenerationContext.ComponentNames);
 		
 		ModelResources.NumLODs = GenerationContext.NumLODsInRoot;
 		ModelResources.NumLODsToStream = GenerationContext.bEnableLODStreaming ? GenerationContext.NumMaxLODsToStream : 0;
diff --git a/Engine/Plugins/Mutable/Source/CustomizableObjectEditor/Private/MuCOE/GenerateMutableSource/GenerateMutableSource.cpp b/Engine/Plugins/Mutable/Source/CustomizableObjectEditor/Private/MuCOE/GenerateMutableSource/GenerateMutableSource.cpp
--- a/Engine/Plugins/Mutable/Source/CustomizableObjectEditor/Private/MuCOE/GenerateMutableSource/GenerateMutableSource.cpp
+++ b/Engine/Plugins/Mutable/Source/CustomizableObjectEditor/Private/MuCOE/GenerateMutableSource/GenerateMutableSource.cpp
@@ -1244,7 +1244,7 @@ void PopulateReferenceSkeletalMeshesData(FMutableGraphGenerationContext& Generat
 	const FString PlatformName = GenerationContext.Options.TargetPlatform->IniPlatformName();
 	
 	const uint32 LODCount = GenerationContext.NumLODsInRoot;
-	const uint32 ComponentCount = GenerationContext.NumComponents;
+	const uint32 ComponentCount = GenerationContext.ComponentInfos.Num();
 
 	GenerationContext.ReferenceSkeletalMeshesData.AddDefaulted(ComponentCount);
 	for(uint32 ComponentIndex = 0; ComponentIndex < ComponentCount; ++ComponentIndex)
diff --git a/Engine/Plugins/Mutable/Source/CustomizableObjectEditor/Private/MuCOE/GenerateMutableSource/GenerateMutableSource.h b/Engine/Plugins/Mutable/Source/CustomizableObjectEditor/Private/MuCOE/GenerateMutableSource/GenerateMutableSource.h
--- a/Engine/Plugins/Mutable/Source/CustomizableObjectEditor/Private/MuCOE/GenerateMutableSource/GenerateMutableSource.h
+++ b/Engine/Plugins/Mutable/Source/CustomizableObjectEditor/Private/MuCOE/GenerateMutableSource/GenerateMutableSource.h
@@ -755,7 +755,12 @@ struct FMutableGraphGenerationContext
 
 	TMap<FName, UCustomizableObjectNodeComponentMesh*> MeshComponents;
 
+	/** Only Mesh Components (no passthrough). */
 	TArray<FMutableComponentInfo> ComponentInfos;
+
+	/** All components. Index is the ObjectComponentIndex. */
+	TArray<FName> ComponentNames;
+	
 	TArray<FMutableRefSkeletalMeshData> ReferenceSkeletalMeshesData;
 	
 	TArray<UMaterialInterface*> ReferencedMaterials;
@@ -814,7 +819,6 @@ struct FMutableGraphGenerationContext
 	mu::Ptr<mu::NodeMesh> ComponentMeshOverride;
 
 	uint8 NumLODsInRoot = 0;
-	uint8 NumComponents = 0;
 
 	uint8 FirstLODAvailable = MAX_MESH_LOD_COUNT;
 	uint8 NumMaxLODsToStream = MAX_MESH_LOD_COUNT;
diff --git a/Engine/Plugins/Mutable/Source/CustomizableObjectEditor/Private/MuCOE/GenerateMutableSource/GenerateMutableSourceComponent.cpp b/Engine/Plugins/Mutable/Source/CustomizableObjectEditor/Private/MuCOE/GenerateMutableSource/GenerateMutableSourceComponent.cpp
--- a/Engine/Plugins/Mutable/Source/CustomizableObjectEditor/Private/MuCOE/GenerateMutableSource/GenerateMutableSourceComponent.cpp
+++ b/Engine/Plugins/Mutable/Source/CustomizableObjectEditor/Private/MuCOE/GenerateMutableSource/GenerateMutableSourceComponent.cpp
@@ -277,8 +277,9 @@ mu::Ptr<mu::NodeComponent> GenerateMutableSourceComponent(const UEdGraphPin* Pin
 		}
 		
 		mu::Ptr<mu::NodeComponentNew> NodeComponentNew = new mu::NodeComponentNew();
-		NodeComponentNew->Id = GenerationContext.NumComponents++;
-		
+		NodeComponentNew->Id = GenerationContext.ComponentNames.Num();
+		GenerationContext.ComponentNames.Add(TypedComponentMesh->ComponentName);
+	
 		NodeComponentNew->SetMessageContext(Node);
 		ObjectNode->Components.Add(NodeComponentNew);
 
@@ -323,29 +324,29 @@ mu::Ptr<mu::NodeComponent> GenerateMutableSourceComponent(const UEdGraphPin* Pin
 		}
 	}
 	
-	else if (const UCustomizableObjectNodeComponentPassthroughMesh* TypeComponentPassthroughMesh = Cast<UCustomizableObjectNodeComponentPassthroughMesh>(Node))
+	else if (const UCustomizableObjectNodeComponentPassthroughMesh* TypedComponentPassthroughMesh = Cast<UCustomizableObjectNodeComponentPassthroughMesh>(Node))
 	{
-		GenerationContext.CurrentMeshComponent = TypeComponentPassthroughMesh->ComponentName;
+		GenerationContext.CurrentMeshComponent = TypedComponentPassthroughMesh->ComponentName;
 
-		if (TypeComponentPassthroughMesh->ComponentName.IsNone())
+		if (TypedComponentPassthroughMesh->ComponentName.IsNone())
 		{
 			FString Msg = FString::Printf(TEXT("Invalid Component Name."));
-			GenerationContext.Compiler->CompilerLog(FText::FromString(Msg), TypeComponentPassthroughMesh, EMessageSeverity::Warning);
+			GenerationContext.Compiler->CompilerLog(FText::FromString(Msg), TypedComponentPassthroughMesh, EMessageSeverity::Warning);
 			return nullptr;
 		}
 
-		if (!TypeComponentPassthroughMesh->Mesh.IsValid())
+		if (!TypedComponentPassthroughMesh->Mesh.IsValid())
 		{
 			FString Msg = FString::Printf(TEXT("No mesh set for component node."));
-			GenerationContext.Compiler->CompilerLog(FText::FromString(Msg), TypeComponentPassthroughMesh, EMessageSeverity::Warning);
+			GenerationContext.Compiler->CompilerLog(FText::FromString(Msg), TypedComponentPassthroughMesh, EMessageSeverity::Warning);
 			return nullptr;
 		}
 
-		USkeletalMesh* SkeletalMesh = Cast<USkeletalMesh>(TypeComponentPassthroughMesh->Mesh.TryLoad());
+		USkeletalMesh* SkeletalMesh = Cast<USkeletalMesh>(TypedComponentPassthroughMesh->Mesh.TryLoad());
 		if (!SkeletalMesh)
 		{
 			FString Msg = FString::Printf(TEXT("Only SkeletalMeshes are supported in this node, for now."));
-			GenerationContext.Compiler->CompilerLog(FText::FromString(Msg), TypeComponentPassthroughMesh, EMessageSeverity::Warning);
+			GenerationContext.Compiler->CompilerLog(FText::FromString(Msg), TypedComponentPassthroughMesh, EMessageSeverity::Warning);
 			return nullptr;
 		}
 
@@ -358,14 +359,15 @@ mu::Ptr<mu::NodeComponent> GenerateMutableSourceComponent(const UEdGraphPin* Pin
 			constexpr bool bIsReference = true;
 			TSoftClassPtr<UAnimInstance> AnimInstance;
 			mu::Ptr<mu::Mesh> MutableMesh = GenerateMutableMesh(SkeletalMesh, AnimInstance, 0, 0, 0, 0, MeshUniqueTags, 
-																GenerationContext, TypeComponentPassthroughMesh, nullptr, bIsReference);
+																GenerationContext, TypedComponentPassthroughMesh, nullptr, bIsReference);
 
 			MeshNode->SetValue(MutableMesh);
 		}
 
 		// Create the component node
 		mu::Ptr<mu::NodeComponentNew> ComponentNode = new mu::NodeComponentNew;
-		ComponentNode->Id = GenerationContext.NumComponents++;
+		ComponentNode->Id = GenerationContext.ComponentNames.Num();
+		GenerationContext.ComponentNames.Add(TypedComponentPassthroughMesh->ComponentName);
 
 		// Create a LOD for each pass-through mesh LOD.
 		const FSkeletalMeshModel* Model = SkeletalMesh->GetImportedModel();
@@ -380,7 +382,7 @@ mu::Ptr<mu::NodeComponent> GenerateMutableSourceComponent(const UEdGraphPin* Pin
 			for (int32 SectionIndex = 0; SectionIndex < SectionCount; ++SectionIndex)
 			{
 				// Is there a pin in the unreal node for this section?
-				if (UEdGraphPin* InMaterialPin = TypeComponentPassthroughMesh->GetMaterialPin(LODIndex,SectionIndex))
+				if (UEdGraphPin* InMaterialPin = TypedComponentPassthroughMesh->GetMaterialPin(LODIndex,SectionIndex))
 				{
 					if (UEdGraphPin* ConnectedMaterialPin = FollowInputPin(*InMaterialPin))
 					{
