[HairStrandsMutableExtension] Mark Groom Constant node as experimental.

#rnx

diff --git a/Engine/Plugins/Experimental/HairStrandsMutable/Source/HairStrandsMutableEditor/Private/CustomizableObjectNodeGroomConstant.cpp b/Engine/Plugins/Experimental/HairStrandsMutable/Source/HairStrandsMutableEditor/Private/CustomizableObjectNodeGroomConstant.cpp
--- a/Engine/Plugins/Experimental/HairStrandsMutable/Source/HairStrandsMutableEditor/Private/CustomizableObjectNodeGroomConstant.cpp
+++ b/Engine/Plugins/Experimental/HairStrandsMutable/Source/HairStrandsMutableEditor/Private/CustomizableObjectNodeGroomConstant.cpp
@@ -37,10 +37,17 @@ void UCustomizableObjectNodeGroomConstant::AllocateDefaultPins(UCustomizableObje
 
 bool UCustomizableObjectNodeGroomConstant::ShouldAddToContextMenu(FText& OutCategory) const
 {
-	OutCategory = UHairStrandsMutableExtension::GroomNodeCategory;
+	OutCategory = UEdGraphSchema_CustomizableObject::NC_Experimental;
 	return true;
 }
 
+
+bool UCustomizableObjectNodeGroomConstant::IsExperimental() const
+{
+	return true;
+}
+
+
 mu::Ptr<mu::NodeExtensionData> UCustomizableObjectNodeGroomConstant::GenerateMutableNode(FExtensionDataCompilerInterface& CompilerInterface) const
 {
 	check(IsInGameThread());
diff --git a/Engine/Plugins/Experimental/HairStrandsMutable/Source/HairStrandsMutableEditor/Public/CustomizableObjectNodeGroomConstant.h b/Engine/Plugins/Experimental/HairStrandsMutable/Source/HairStrandsMutableEditor/Public/CustomizableObjectNodeGroomConstant.h
--- a/Engine/Plugins/Experimental/HairStrandsMutable/Source/HairStrandsMutableEditor/Public/CustomizableObjectNodeGroomConstant.h
+++ b/Engine/Plugins/Experimental/HairStrandsMutable/Source/HairStrandsMutableEditor/Public/CustomizableObjectNodeGroomConstant.h
@@ -23,7 +23,8 @@ public:
 	/** UCustomizableObjectNode interface */
 	virtual void AllocateDefaultPins(UCustomizableObjectNodeRemapPins* RemapPins) override;
 	virtual bool ShouldAddToContextMenu(FText& OutCategory) const override;
-
+	virtual bool IsExperimental() const override;
+	
 	/** ICustomizableObjectExtensionNode interface */
 	virtual mu::Ptr<mu::NodeExtensionData> GenerateMutableNode(FExtensionDataCompilerInterface& CompilerInterface) const override;
 
diff --git a/Engine/Plugins/Mutable/Source/CustomizableObjectEditor/Private/MuCOE/EdGraphSchema_CustomizableObject.cpp b/Engine/Plugins/Mutable/Source/CustomizableObjectEditor/Private/MuCOE/EdGraphSchema_CustomizableObject.cpp
--- a/Engine/Plugins/Mutable/Source/CustomizableObjectEditor/Private/MuCOE/EdGraphSchema_CustomizableObject.cpp
+++ b/Engine/Plugins/Mutable/Source/CustomizableObjectEditor/Private/MuCOE/EdGraphSchema_CustomizableObject.cpp
@@ -245,6 +245,10 @@ const FName UEdGraphSchema_CustomizableObject::PC_Wildcard("wildcard");
 const FName UEdGraphSchema_CustomizableObject::PC_PoseAsset("poseAsset");
 
 
+// Node categories
+const FText UEdGraphSchema_CustomizableObject::NC_Experimental(FText::FromString("Experimental"));
+
+
 TSharedPtr<FCustomizableObjectSchemaAction_NewNode> UEdGraphSchema_CustomizableObject::AddNewNodeAction(FGraphActionListBuilderBase& ContextMenuBuilder, const FString& Category, const FText& MenuDesc, const FText& Tooltip, const int32 Grouping, const FString& Keywords)
 {
 	TSharedPtr<FCustomizableObjectSchemaAction_NewNode> NewActionNode = TSharedPtr<FCustomizableObjectSchemaAction_NewNode>(new FCustomizableObjectSchemaAction_NewNode(Category, MenuDesc, Tooltip, Grouping, FText::FromString(Keywords)));
@@ -539,7 +543,7 @@ void UEdGraphSchema_CustomizableObject::GetGraphContextActions(FGraphContextMenu
 			ContextMenuBuilder.CreateTemplateNode<UCustomizableObjectNodeModifierClipDeform>(),
 		};
 
-		AddNewNodeCategoryActionsFiltered(ExperimentalTemplateNodes, ContextMenuBuilder, TEXT("Experimental"), GeneralGrouping, Filter);
+		AddNewNodeCategoryActionsFiltered(ExperimentalTemplateNodes, ContextMenuBuilder, NC_Experimental.ToString(), GeneralGrouping, Filter);
 	}
 
 	{
diff --git a/Engine/Plugins/Mutable/Source/CustomizableObjectEditor/Public/MuCOE/EdGraphSchema_CustomizableObject.h b/Engine/Plugins/Mutable/Source/CustomizableObjectEditor/Public/MuCOE/EdGraphSchema_CustomizableObject.h
--- a/Engine/Plugins/Mutable/Source/CustomizableObjectEditor/Public/MuCOE/EdGraphSchema_CustomizableObject.h
+++ b/Engine/Plugins/Mutable/Source/CustomizableObjectEditor/Public/MuCOE/EdGraphSchema_CustomizableObject.h
@@ -108,6 +108,9 @@ public:
 	static const FName PC_Wildcard;
 	static const FName PC_PoseAsset;
 
+	// Node categories
+	static const FText NC_Experimental;
+
 	// EdGraphSchema interface
 	virtual void GetAssetsGraphHoverMessage(const TArray<FAssetData>& Assets, const UEdGraph* HoverGraph, FString& OutTooltipText, bool& OutOkIcon) const override;
 	virtual void GetGraphContextActions(FGraphContextMenuBuilder& ContextMenuBuilder) const override;
