[mutable] Remove unneccessary assert.
#rnx

diff --git a/Engine/Plugins/Mutable/Source/MutableTools/Private/MuT/NodeComponentEdit.cpp b/Engine/Plugins/Mutable/Source/MutableTools/Private/MuT/NodeComponentEdit.cpp
--- a/Engine/Plugins/Mutable/Source/MutableTools/Private/MuT/NodeComponentEdit.cpp
+++ b/Engine/Plugins/Mutable/Source/MutableTools/Private/MuT/NodeComponentEdit.cpp
@@ -40,7 +40,6 @@ namespace mu
 			Parent = m_pParent->GetParentComponentNew();
 		}
 
-		check(Parent);
 		return Parent;
 	}
 
