[mutable] Remove layout pins from static mesh nodes.
#rnx

diff --git a/Engine/Plugins/Mutable/Source/CustomizableObjectEditor/Private/MuCOE/Nodes/CustomizableObjectNodeStaticMesh.cpp b/Engine/Plugins/Mutable/Source/CustomizableObjectEditor/Private/MuCOE/Nodes/CustomizableObjectNodeStaticMesh.cpp
--- a/Engine/Plugins/Mutable/Source/CustomizableObjectEditor/Private/MuCOE/Nodes/CustomizableObjectNodeStaticMesh.cpp
+++ b/Engine/Plugins/Mutable/Source/CustomizableObjectEditor/Private/MuCOE/Nodes/CustomizableObjectNodeStaticMesh.cpp
@@ -148,13 +148,6 @@ void UCustomizableObjectNodeStaticMesh::AllocateDefaultPins(UCustomizableObjectN
 				Pin->bDefaultValueIsIgnored = true;
 				LODs[i].Materials[MaterialIndex].MeshPinRef.SetPin(Pin);
 
-				// Layout
-				// TODO: Multiple layout ui
-				const FString LayoutName = MaterialName + TEXT(" Layout");
-				Pin = CustomCreatePin(EGPD_Input, Schema->PC_Layout, FName(*LayoutName));
-				Pin->bDefaultValueIsIgnored = true;
-				LODs[i].Materials[MaterialIndex].LayoutPinRef.SetPin(Pin);
-
 				// Images
 				if (MaterialInterface)
 				{
@@ -203,13 +196,6 @@ FText UCustomizableObjectNodeStaticMesh::GetNodeTitle(ENodeTitleType::Type Title
 }
 
 
-FLinearColor UCustomizableObjectNodeStaticMesh::GetNodeTitleColor() const
-{
-	const UEdGraphSchema_CustomizableObject* Schema = GetDefault<UEdGraphSchema_CustomizableObject>();
-	return Schema->GetPinTypeColor(Schema->PC_Mesh);
-}
-
-
 UTexture2D* UCustomizableObjectNodeStaticMesh::FindTextureForPin(const UEdGraphPin* Pin) const
 {
 	if (!StaticMesh)
@@ -306,11 +292,6 @@ bool UCustomizableObjectNodeStaticMesh::IsPinRelevant(const UEdGraphPin* Pin) co
 {
 	const UEdGraphSchema_CustomizableObject* Schema = GetDefault<UEdGraphSchema_CustomizableObject>();
 
-	if (Pin->Direction == EEdGraphPinDirection::EGPD_Output)
-	{
-		return Pin->PinType.PinCategory == Schema->PC_Layout;
-	}
-
 	if (Pin->Direction == EEdGraphPinDirection::EGPD_Input)
 	{	
 		return Pin->PinType.PinCategory == Schema->PC_Mesh;
diff --git a/Engine/Plugins/Mutable/Source/CustomizableObjectEditor/Public/MuCOE/Nodes/CustomizableObjectNodeStaticMesh.h b/Engine/Plugins/Mutable/Source/CustomizableObjectEditor/Public/MuCOE/Nodes/CustomizableObjectNodeStaticMesh.h
--- a/Engine/Plugins/Mutable/Source/CustomizableObjectEditor/Public/MuCOE/Nodes/CustomizableObjectNodeStaticMesh.h
+++ b/Engine/Plugins/Mutable/Source/CustomizableObjectEditor/Public/MuCOE/Nodes/CustomizableObjectNodeStaticMesh.h
@@ -137,7 +137,6 @@ public:
 
 	// UEdGraphNode interface
 	virtual FText GetNodeTitle(ENodeTitleType::Type TitleType) const override;
-	virtual FLinearColor GetNodeTitleColor() const override;
 	virtual FText GetTooltipText() const override;
 
 	// UCustomizableObjectNode interface
