[mutable] Fixed Text conflict from LOCTEXT macro for namespace for CustomizableObjectNodeModifierExtendMeshSectionDetails.cpp.
#rnx

diff --git a/Engine/Plugins/Mutable/Source/CustomizableObjectEditor/Private/MuCOE/Nodes/CustomizableObjectNodeModifierExtendMeshSectionDetails.cpp b/Engine/Plugins/Mutable/Source/CustomizableObjectEditor/Private/MuCOE/Nodes/CustomizableObjectNodeModifierExtendMeshSectionDetails.cpp
--- a/Engine/Plugins/Mutable/Source/CustomizableObjectEditor/Private/MuCOE/Nodes/CustomizableObjectNodeModifierExtendMeshSectionDetails.cpp
+++ b/Engine/Plugins/Mutable/Source/CustomizableObjectEditor/Private/MuCOE/Nodes/CustomizableObjectNodeModifierExtendMeshSectionDetails.cpp
@@ -65,7 +65,7 @@ void FCustomizableObjectNodeModifierExtendMeshSectionDetails::CustomizeDetails(
 					.Padding(FMargin(0, 4.0f, 0, 4.0f))
 					[
 						SNew(STextBlock)
-							.Text(LOCTEXT("MeshSectionDetails_Tags", "Tags enabled for extended data"))
+							.Text(LOCTEXT("ExtendMeshSectionDetails_Tags", "Tags enabled for extended data"))
 							.Font(IDetailLayoutBuilder::GetDetailFont())
 					]
 			]
