[mutable] Fix the SetWorkingMemory method in the system, which wasn't exposed in client builds. Fix some unit conversion problems. Add some mutable core stats related to memory usage.
#rnx

diff --git a/Engine/Plugins/Mutable/Source/CustomizableObject/Private/MuCO/CustomizableObjectSystem.cpp b/Engine/Plugins/Mutable/Source/CustomizableObject/Private/MuCO/CustomizableObjectSystem.cpp
--- a/Engine/Plugins/Mutable/Source/CustomizableObject/Private/MuCO/CustomizableObjectSystem.cpp
+++ b/Engine/Plugins/Mutable/Source/CustomizableObject/Private/MuCO/CustomizableObjectSystem.cpp
@@ -67,16 +67,16 @@ static FAutoConsoleVariableRef CVarMutableEnabled(
 	TEXT("true/false - Disabling Mutable will turn off CO compilation, mesh generation, and texture streaming and will remove the system ticker. "),
 	FConsoleVariableDelegate::CreateStatic(&UCustomizableObjectSystemPrivate::OnMutableEnabledChanged));
 
-int32 WorkingMemory =
+int32 WorkingMemoryKB =
 #if !PLATFORM_DESKTOP
 (10 * 1024);
 #else
 (50 * 1024);
 #endif
 
-static FAutoConsoleVariableRef CVarWorkingMemory(
+static FAutoConsoleVariableRef CVarWorkingMemoryKB(
 	TEXT("mutable.WorkingMemory"),
-	WorkingMemory,
+	WorkingMemoryKB,
 	TEXT("Limit the amount of memory (in KB) to use as working memory when building characters. More memory reduces the object construction time. 0 means no restriction. Defaults: Desktop = 50,000 KB, Others = 10,000 KB"),
 	ECVF_Scalability);
 
@@ -532,7 +532,7 @@ void UCustomizableObjectSystem::InitSystem()
 	Private->CurrentMutableOperation = nullptr;
 	Private->CurrentInstanceBeingUpdated = nullptr;
 
-	Private->LastWorkingMemoryBytes = CVarWorkingMemory->GetInt() * 1024;
+	Private->LastWorkingMemoryBytes = CVarWorkingMemoryKB->GetInt() * 1024;
 	Private->LastGeneratedResourceCacheSize = CVarGeneratedResourcesCacheSize.GetValueOnGameThread();
 
 	const mu::Ptr<mu::Settings> pSettings = new mu::Settings;
@@ -1528,7 +1528,7 @@ void UCustomizableObjectSystemPrivate::UpdateMemoryLimit()
 	// This must run on game thread, and when the mutable thread is not running
 	check(IsInGameThread());
 
-	const uint64 MemoryBytes = CVarWorkingMemory->GetInt() * 1024;
+	const uint64 MemoryBytes = uint64(CVarWorkingMemoryKB->GetInt()) * 1024;
 	if (MemoryBytes != LastWorkingMemoryBytes)
 	{
 		LastWorkingMemoryBytes = MemoryBytes;
@@ -4191,21 +4191,21 @@ bool UCustomizableObjectSystem::ShouldReuseTexturesBetweenInstances()
 }
 
 
-#if WITH_EDITOR
-
-void UCustomizableObjectSystem::SetWorkingMemory(int32 Bytes)
+void UCustomizableObjectSystem::SetWorkingMemory(int32 KBytes)
 {
-	WorkingMemory = Bytes;
-	UE_LOG(LogMutable, Log, TEXT("Working Memory set to %i bytes."), Bytes);
+	CVarWorkingMemoryKB->Set( KBytes );
+	UE_LOG(LogMutable, Log, TEXT("Working Memory set to %i kilobytes."), KBytes);
 }
 
 
 int32 UCustomizableObjectSystem::GetWorkingMemory() const
 {
-	return UCustomizableObjectSystemPrivate::IsUsingBenchmarkingSettings() ? 16384 : WorkingMemory;
+	return UCustomizableObjectSystemPrivate::IsUsingBenchmarkingSettings() ? 16384 : CVarWorkingMemoryKB->GetInt();
 }
 
 
+#if WITH_EDITOR
+
 uint64 UCustomizableObjectSystem::GetMaxChunkSizeForPlatform(const ITargetPlatform* TargetPlatform)
 {
 	if (!TargetPlatform || !TargetPlatform->RequiresCookedData())
diff --git a/Engine/Plugins/Mutable/Source/CustomizableObject/Public/MuCO/CustomizableObjectSystem.h b/Engine/Plugins/Mutable/Source/CustomizableObject/Public/MuCO/CustomizableObjectSystem.h
--- a/Engine/Plugins/Mutable/Source/CustomizableObject/Public/MuCO/CustomizableObjectSystem.h
+++ b/Engine/Plugins/Mutable/Source/CustomizableObject/Public/MuCO/CustomizableObjectSystem.h
@@ -341,13 +341,15 @@ public:
 	UFUNCTION(BlueprintCallable, BlueprintPure, Category = CustomizableObjectSystem)
 	bool IsUpdating(const UCustomizableObjectInstance* Instance) const;
 
-	/** Set Mutable's working memory limit (bytes). Mutable will flush internal caches to try to keep its memory consumption below the WorkingMemory (i.e., it is not a hard limit).
+	/** Set Mutable's working memory limit (kilobytes). Mutable will flush internal caches to try to keep its memory consumption below the WorkingMemory (i.e., it is not a hard limit).
 	 * The working memory limit will especially reduce the memory required to perform Instance Updates and Texture Streaming.
  	 * Notice that Mutable does not track all its memory (e.g., UObjects memory is no tracked).
 	 * This value can also be set using "mutable.WorkingMemory" CVar. */
-	void SetWorkingMemory(int32 Bytes);
+	UFUNCTION(BlueprintCallable, Category = CustomizableObjectSystem)
+	void SetWorkingMemory(int32 KiloBytes);
 
-	/** Get Mutable's working memory limit (bytes). See SetWorkingMemory(int32). */
+	/** Get Mutable's working memory limit (kilobytes). See SetWorkingMemory(int32). */
+	UFUNCTION(BlueprintCallable, BlueprintPure, Category = CustomizableObjectSystem)
 	int32 GetWorkingMemory() const;
 
 	/**
diff --git a/Engine/Plugins/Mutable/Source/MutableRuntime/Internal/MuR/SystemPrivate.h b/Engine/Plugins/Mutable/Source/MutableRuntime/Internal/MuR/SystemPrivate.h
--- a/Engine/Plugins/Mutable/Source/MutableRuntime/Internal/MuR/SystemPrivate.h
+++ b/Engine/Plugins/Mutable/Source/MutableRuntime/Internal/MuR/SystemPrivate.h
@@ -1949,6 +1949,9 @@ namespace mu
 		//!
 		void PrepareCache(const Model*, int32 State);
 
+		//! Update some mutable core unreal stats.
+		void UpdateStats();
+
 	private:
 
 		/** Owned by this system. */
diff --git a/Engine/Plugins/Mutable/Source/MutableRuntime/Private/MuR/System.cpp b/Engine/Plugins/Mutable/Source/MutableRuntime/Private/MuR/System.cpp
--- a/Engine/Plugins/Mutable/Source/MutableRuntime/Private/MuR/System.cpp
+++ b/Engine/Plugins/Mutable/Source/MutableRuntime/Private/MuR/System.cpp
@@ -36,6 +36,12 @@
 #include "ProfilingDebugging/CountersTrace.h"
 #include "PackedNormal.h"
 
+
+DECLARE_STATS_GROUP(TEXT("MutableCore"), STATGROUP_MutableCore, STATCAT_Advanced);
+DECLARE_DWORD_ACCUMULATOR_STAT(TEXT("Working Memory"), STAT_MutableWorkingMemory, STATGROUP_MutableCore);
+DECLARE_DWORD_ACCUMULATOR_STAT(TEXT("Working Memory Excess"), STAT_MutableWorkingMemoryExcess, STATGROUP_MutableCore);
+DECLARE_DWORD_ACCUMULATOR_STAT(TEXT("Current Memory"), STAT_MutableCurrentMemory, STATGROUP_MutableCore);
+
 namespace 
 {
 
@@ -118,6 +124,8 @@ namespace mu
 		{
 			ExtensionDataStreamer = MakeShared<NullExtensionDataStreamer>();
 		}
+
+		UpdateStats();
 	}
 
 
@@ -151,6 +159,8 @@ namespace mu
 
 		m_pD->WorkingMemoryManager.BudgetBytes = InBytes;
 		m_pD->WorkingMemoryManager.EnsureBudgetBelow(0);
+
+		m_pD->UpdateStats();
 	}
 
 
@@ -195,6 +205,8 @@ namespace mu
 		m_pD->WorkingMemoryManager.CacheResources.Empty();
 		check(m_pD->WorkingMemoryManager.TempImages.IsEmpty());
 		check(m_pD->WorkingMemoryManager.TempMeshes.IsEmpty());
+
+		m_pD->UpdateStats();
 	}
 
 	
@@ -604,8 +616,8 @@ namespace mu
 
 
     //---------------------------------------------------------------------------------------------
-    void System::EndUpdate( Instance::ID instanceID )
-    {
+	void System::EndUpdate(Instance::ID instanceID)
+	{
 		LLM_SCOPE_BYNAME(TEXT("MutableRuntime"));
 		MUTABLE_CPUPROFILER_SCOPE(EndUpdate);
 
@@ -613,7 +625,7 @@ namespace mu
 		if (pLiveInstance)
 		{
 			pLiveInstance->Instance = nullptr;
-			
+
 			// Debug check to see if we managed the op-hit-counts correctly
 			pLiveInstance->Cache->CheckHitCountsCleared();
 
@@ -639,6 +651,7 @@ namespace mu
 			m_pD->WorkingMemoryManager.PooledImages.Empty();
 		}
 
+		m_pD->UpdateStats();
 	}
 
 
@@ -1103,6 +1116,21 @@ namespace mu
 	}
 
 
+	//---------------------------------------------------------------------------------------------
+	void System::Private::UpdateStats()
+	{
+		// Updater stats
+		int32 WorkingMemoryKb = WorkingMemoryManager.BudgetBytes / 1024;
+		SET_DWORD_STAT(STAT_MutableWorkingMemory, WorkingMemoryKb);
+
+		int32 WorkingMemoryExcessKb = WorkingMemoryManager.BudgetExcessBytes / 1024;
+		SET_DWORD_STAT(STAT_MutableWorkingMemoryExcess, WorkingMemoryExcessKb);
+
+		int32 CurrentMemoryKb = WorkingMemoryManager.GetCurrentMemoryBytes() / 1024;
+		SET_DWORD_STAT(STAT_MutableCurrentMemory, CurrentMemoryKb);
+	}
+
+
 	//---------------------------------------------------------------------------------------------
 	void FWorkingMemoryManager::LogWorkingMemory(const CodeRunner* CurrentRunner) const
 	{
diff --git a/Engine/Plugins/Mutable/Source/MutableValidation/Private/MuV/ValidationUtils.cpp b/Engine/Plugins/Mutable/Source/MutableValidation/Private/MuV/ValidationUtils.cpp
--- a/Engine/Plugins/Mutable/Source/MutableValidation/Private/MuV/ValidationUtils.cpp
+++ b/Engine/Plugins/Mutable/Source/MutableValidation/Private/MuV/ValidationUtils.cpp
@@ -28,9 +28,9 @@ void PrepareAssetRegistry()
 void LogGlobalSettings()
 {
 	// Mutable Settings
-	const int32 WorkingMemory = UCustomizableObjectSystem::GetInstanceChecked()->GetWorkingMemory() ;
-	UE_LOG(LogMutable,Log, TEXT("(int) working_memory_bytes : %d"), WorkingMemory * 1024)
-	UE_LOG(LogMutable, Display, TEXT("The mutable updates will use as working memory the value of %d KB"), WorkingMemory)
+	const int32 WorkingMemoryKB = UCustomizableObjectSystem::GetInstanceChecked()->GetWorkingMemory() ;
+	UE_LOG(LogMutable,Log, TEXT("(int) working_memory_bytes : %d"), WorkingMemoryKB*1024)
+	UE_LOG(LogMutable, Display, TEXT("The mutable updates will use as working memory the value of %d KB"), WorkingMemoryKB)
 	
 	// Expand this when adding new controls from the .xml file
 	
