[Mutable] Minor comoment improvements.

#rnx

diff --git a/Engine/Plugins/Mutable/Source/CustomizableObject/Public/MuCO/CustomizableObjectInstance.h b/Engine/Plugins/Mutable/Source/CustomizableObject/Public/MuCO/CustomizableObjectInstance.h
--- a/Engine/Plugins/Mutable/Source/CustomizableObject/Public/MuCO/CustomizableObjectInstance.h
+++ b/Engine/Plugins/Mutable/Source/CustomizableObject/Public/MuCO/CustomizableObjectInstance.h
@@ -402,7 +402,7 @@ public:
 	UFUNCTION(BlueprintCallable, Category = CustomizableObjectInstance)
 	void SetColorParameterSelectedOption(const FString& ColorParamName, const FLinearColor& ColorValue);
 
-	/** Sets the bool value "BoolValue" of a bool parameter with name "BoolParamName". */
+	/** Gets the bool value "BoolValue" of a bool parameter with name "BoolParamName". */
 	UFUNCTION(BlueprintCallable, Category = CustomizableObjectInstance)
 	bool GetBoolParameterSelectedOption(const FString& BoolParamName) const;
 
@@ -410,7 +410,7 @@ public:
 	UFUNCTION(BlueprintCallable, Category = CustomizableObjectInstance)
 	void SetBoolParameterSelectedOption(const FString& BoolParamName, bool BoolValue);
 
-	/** Sets the vector value "VectorValue" of a bool parameter with index "VectorParamIndex". */
+	/** Sets the vector value "VectorValue" of a vector parameter with index "VectorParamIndex". */
 	UFUNCTION(BlueprintCallable, Category = CustomizableObjectInstance)
 	void SetVectorParameterSelectedOption(const FString& VectorParamName, const FLinearColor& VectorValue);
 
diff --git a/Engine/Plugins/Mutable/Source/CustomizableObject/Public/MuCO/CustomizableObjectSystem.h b/Engine/Plugins/Mutable/Source/CustomizableObject/Public/MuCO/CustomizableObjectSystem.h
--- a/Engine/Plugins/Mutable/Source/CustomizableObject/Public/MuCO/CustomizableObjectSystem.h
+++ b/Engine/Plugins/Mutable/Source/CustomizableObject/Public/MuCO/CustomizableObjectSystem.h
@@ -115,7 +115,10 @@ struct FCustomizableObjectExternalTexture
 };
 
 
-/** Base class for Image provider. */
+/** Base class for Image provider.
+ * 
+ *  An image provider only needs to be alive when calling UpdateSkeletalMeshAsync. Once called it can be safety destructed.
+ *  Mutable will cache those images in case they are needed due to MIP updates. */
 UCLASS(abstract)
 class CUSTOMIZABLEOBJECT_API UCustomizableSystemImageProvider : public UObject
 {
