[mutable] Fixed an issue when counting the amount of triangles in each component.
- Changes in the API have made the usage of the previous method to get the skeletal mesh of the instance deprecated.

#rnx

diff --git a/Engine/Plugins/Mutable/Source/CustomizableObject/Private/MuCO/LogBenchmarkUtil.cpp b/Engine/Plugins/Mutable/Source/CustomizableObject/Private/MuCO/LogBenchmarkUtil.cpp
--- a/Engine/Plugins/Mutable/Source/CustomizableObject/Private/MuCO/LogBenchmarkUtil.cpp
+++ b/Engine/Plugins/Mutable/Source/CustomizableObject/Private/MuCO/LogBenchmarkUtil.cpp
@@ -313,11 +313,13 @@ void FLogBenchmarkUtil::FinishUpdateMesh(const TSharedRef<FUpdateContextPrivate>
 	}
 
 	// Cache the amount of triangles of this instance
-	uint32 TriangleCount = 0;
-	for (int32 ComponentIndex = 0; ComponentIndex < Instance->GetNumComponents(); ComponentIndex++)
+	uint32 TriangleCount = 0;	
+	for (int32 ComponentIndex = 0; ComponentIndex < Object->GetComponentCount(); ComponentIndex++)
 	{
+		const FName ComponentName = Object->GetComponentName(ComponentIndex);
+		
 		// Process the generated components (not null)
-		if (const USkeletalMesh* InstanceSkeletalMesh = Instance->GetSkeletalMesh(ComponentIndex))
+		if (const USkeletalMesh* InstanceSkeletalMesh = Instance->GetComponentMeshSkeletalMesh(ComponentName))
 		{
 			const FSkeletalMeshRenderData* RenderData = InstanceSkeletalMesh->GetResourceForRendering();
 			check(RenderData);
