[mutable] Add option to select strategy of face clipping to various operations.
Default value should keep the current behaviour, but the current behaviour doesn't match past versions.
#rnx

diff --git a/Engine/Plugins/Mutable/Source/CustomizableObject/Internal/MuCO/CustomizableObjectPrivate.h b/Engine/Plugins/Mutable/Source/CustomizableObject/Internal/MuCO/CustomizableObjectPrivate.h
--- a/Engine/Plugins/Mutable/Source/CustomizableObject/Internal/MuCO/CustomizableObjectPrivate.h
+++ b/Engine/Plugins/Mutable/Source/CustomizableObject/Internal/MuCO/CustomizableObjectPrivate.h
@@ -1469,7 +1469,9 @@ public:
 		ComponentsArray,
 
 		FixComponentNames,
-		
+
+		AddedFaceCullStrategyToSomeOperations,
+
 		// -----<new versions can be added above this line>--------
 		LastCustomizableObjectVersion
 	};
diff --git a/Engine/Plugins/Mutable/Source/CustomizableObjectEditor/Private/MuCOE/GenerateMutableSource/GenerateMutableSourceModifier.cpp b/Engine/Plugins/Mutable/Source/CustomizableObjectEditor/Private/MuCOE/GenerateMutableSource/GenerateMutableSourceModifier.cpp
--- a/Engine/Plugins/Mutable/Source/CustomizableObjectEditor/Private/MuCOE/GenerateMutableSource/GenerateMutableSourceModifier.cpp
+++ b/Engine/Plugins/Mutable/Source/CustomizableObjectEditor/Private/MuCOE/GenerateMutableSource/GenerateMutableSourceModifier.cpp
@@ -79,6 +79,8 @@ mu::Ptr<mu::NodeModifier> GenerateMutableSourceModifier(const UEdGraphPin * Pin,
 		ClipNode->MultipleTagsPolicy = TypedNodeClip->MultipleTagPolicy;
 		ClipNode->RequiredTags = TypedNodeClip->RequiredTags;
 
+		ClipNode->Parameters.FaceCullStrategy = TypedNodeClip->FaceCullStrategy;
+
 		GenerationContext.MeshGenerationFlags.Pop();
 	}
 
@@ -92,6 +94,8 @@ mu::Ptr<mu::NodeModifier> GenerateMutableSourceModifier(const UEdGraphPin * Pin,
 		mu::Ptr<mu::NodeModifierMeshClipDeform> ClipNode = new mu::NodeModifierMeshClipDeform();
 		Result = ClipNode;
 	
+		ClipNode->FaceCullStrategy = TypedNodeClipDeform->FaceCullStrategy;
+
 		if (const UEdGraphPin* ConnectedPin = FollowInputPin(*TypedNodeClipDeform->ClipShapePin()))
 		{
 			FMutableGraphMeshGenerationData DummyMeshData;
@@ -145,6 +149,8 @@ mu::Ptr<mu::NodeModifier> GenerateMutableSourceModifier(const UEdGraphPin * Pin,
 		mu::Ptr<mu::NodeModifierMeshClipWithMesh> ClipNode = new mu::NodeModifierMeshClipWithMesh();
 		Result = ClipNode;
 
+		ClipNode->FaceCullStrategy = TypedNodeClipMesh->FaceCullStrategy;
+
 		if (const UEdGraphPin* ConnectedPin = FollowInputPin(*TypedNodeClipMesh->ClipMeshPin()))
 		{
 			FMutableGraphMeshGenerationData DummyMeshData;
@@ -215,6 +221,8 @@ mu::Ptr<mu::NodeModifier> GenerateMutableSourceModifier(const UEdGraphPin * Pin,
 		mu::Ptr<mu::NodeModifierMeshClipWithUVMask> ClipNode = new mu::NodeModifierMeshClipWithUVMask();
 		Result = ClipNode;
 
+		ClipNode->FaceCullStrategy = TypedNodeClipUVMask->FaceCullStrategy;
+
 		if (const UEdGraphPin* ConnectedPin = FollowInputPin(*TypedNodeClipUVMask->ClipMaskPin()))
 		{
 			FMutableGraphMeshGenerationData DummyMeshData;
@@ -401,6 +409,8 @@ mu::Ptr<mu::NodeModifier> GenerateMutableSourceModifier(const UEdGraphPin * Pin,
 			int32 LODCount = GenerationContext.NumLODsInRoot;
 			SurfNode->LODs.SetNum(LODCount);
 
+			SurfNode->FaceCullStrategy = TypedNodeRem->FaceCullStrategy;
+
 			for (int32 LODIndex = 0; LODIndex < LODCount; ++LODIndex)
 			{
 				GenerationContext.FromLOD = 0;
@@ -427,6 +437,8 @@ mu::Ptr<mu::NodeModifier> GenerateMutableSourceModifier(const UEdGraphPin * Pin,
 		mu::Ptr<mu::NodeModifierMeshClipWithUVMask> ClipNode = new mu::NodeModifierMeshClipWithUVMask();
 		Result = ClipNode;
 
+		ClipNode->FaceCullStrategy = TypedNodeRemBlocks->FaceCullStrategy;
+
 		ClipNode->MultipleTagsPolicy = TypedNodeRemBlocks->MultipleTagPolicy;
 		ClipNode->RequiredTags = TypedNodeRemBlocks->RequiredTags;
 
diff --git a/Engine/Plugins/Mutable/Source/CustomizableObjectEditor/Private/MuCOE/SMutableCodeViewer.cpp b/Engine/Plugins/Mutable/Source/CustomizableObjectEditor/Private/MuCOE/SMutableCodeViewer.cpp
--- a/Engine/Plugins/Mutable/Source/CustomizableObjectEditor/Private/MuCOE/SMutableCodeViewer.cpp
+++ b/Engine/Plugins/Mutable/Source/CustomizableObjectEditor/Private/MuCOE/SMutableCodeViewer.cpp
@@ -2402,7 +2402,7 @@ bool SMutableCodeViewer::IsConstantResourceUsedByOperation(const int32 IndexOnCo
 						break;
 					}
 
-					if (Arguments.vertexSelectionType == mu::OP::MeshClipMorphPlaneArgs::VS_SHAPE)
+					if (Arguments.VertexSelectionType == EClipVertexSelectionType::Shape)
 					{
 						// Selection Shape
 						bResourceLocated = IndexOnConstantsArray == InProgram.GetOpArgs<mu::OP::MeshClipMorphPlaneArgs>(OperationAddress).vertexSelectionShapeOrBone;
diff --git a/Engine/Plugins/Mutable/Source/CustomizableObjectEditor/Public/MuCOE/Nodes/CustomizableObjectNodeModifierClipDeform.h b/Engine/Plugins/Mutable/Source/CustomizableObjectEditor/Public/MuCOE/Nodes/CustomizableObjectNodeModifierClipDeform.h
--- a/Engine/Plugins/Mutable/Source/CustomizableObjectEditor/Public/MuCOE/Nodes/CustomizableObjectNodeModifierClipDeform.h
+++ b/Engine/Plugins/Mutable/Source/CustomizableObjectEditor/Public/MuCOE/Nodes/CustomizableObjectNodeModifierClipDeform.h
@@ -32,7 +32,10 @@ public:
 
 	UPROPERTY(EditAnywhere, Category = MeshClipDeform)
 	EShapeBindingMethod BindingMethod;
-	
+
+	UPROPERTY(EditAnywhere, Category = RemoveOptions)
+	EFaceCullStrategy FaceCullStrategy = EFaceCullStrategy::AllVerticesCulled;
+
 public:
 
 	// EdGraphNode interface
diff --git a/Engine/Plugins/Mutable/Source/CustomizableObjectEditor/Public/MuCOE/Nodes/CustomizableObjectNodeModifierClipMorph.h b/Engine/Plugins/Mutable/Source/CustomizableObjectEditor/Public/MuCOE/Nodes/CustomizableObjectNodeModifierClipMorph.h
--- a/Engine/Plugins/Mutable/Source/CustomizableObjectEditor/Public/MuCOE/Nodes/CustomizableObjectNodeModifierClipMorph.h
+++ b/Engine/Plugins/Mutable/Source/CustomizableObjectEditor/Public/MuCOE/Nodes/CustomizableObjectNodeModifierClipMorph.h
@@ -4,6 +4,7 @@
 
 #include "MuCOE/Nodes/CustomizableObjectNodeModifierBase.h"
 #include "MuT/NodeModifier.h"
+#include "MuR/Types.h"
 
 #include "CustomizableObjectNodeModifierClipMorph.generated.h"
 
@@ -60,13 +61,18 @@ public:
 	UPROPERTY()
 	FVector Normal;
 
+	UPROPERTY(EditAnywhere, Category = MeshClipParameters, meta = (DisplayName = "Max Effect Radius", ToolTip = "The maximum distance from the origin of the widget where vertices will be affected. If negative, there will be no limit."))
+	float MaxEffectRadius;
+
+	UPROPERTY(EditAnywhere, Category = MeshClipParameters)
+	EFaceCullStrategy FaceCullStrategy = EFaceCullStrategy::AllVerticesCulled;
+
 private:
+
 	UPROPERTY()
 	bool bOldOffset_DEPRECATED;
 
 public:
-	UPROPERTY(EditAnywhere, Category = MeshClipParameters, meta = (DisplayName = "Max Effect Radius", ToolTip = "The maximum distance from the origin of the widget where vertices will be affected. If negative, there will be no limit."))
-	float MaxEffectRadius;
 
 	bool bUpdateViewportWidget;
 
diff --git a/Engine/Plugins/Mutable/Source/CustomizableObjectEditor/Public/MuCOE/Nodes/CustomizableObjectNodeModifierClipWithMesh.h b/Engine/Plugins/Mutable/Source/CustomizableObjectEditor/Public/MuCOE/Nodes/CustomizableObjectNodeModifierClipWithMesh.h
--- a/Engine/Plugins/Mutable/Source/CustomizableObjectEditor/Public/MuCOE/Nodes/CustomizableObjectNodeModifierClipWithMesh.h
+++ b/Engine/Plugins/Mutable/Source/CustomizableObjectEditor/Public/MuCOE/Nodes/CustomizableObjectNodeModifierClipWithMesh.h
@@ -38,6 +38,9 @@ public:
 	UPROPERTY(EditAnywhere, Category = ClipMesh)
 	FTransform Transform;
 
+	UPROPERTY(EditAnywhere, Category = ClipMesh)
+	EFaceCullStrategy FaceCullStrategy = EFaceCullStrategy::AllVerticesCulled;
+
 public:
 
 	UCustomizableObjectNodeModifierClipWithMesh();
diff --git a/Engine/Plugins/Mutable/Source/CustomizableObjectEditor/Public/MuCOE/Nodes/CustomizableObjectNodeModifierClipWithUVMask.h b/Engine/Plugins/Mutable/Source/CustomizableObjectEditor/Public/MuCOE/Nodes/CustomizableObjectNodeModifierClipWithUVMask.h
--- a/Engine/Plugins/Mutable/Source/CustomizableObjectEditor/Public/MuCOE/Nodes/CustomizableObjectNodeModifierClipWithUVMask.h
+++ b/Engine/Plugins/Mutable/Source/CustomizableObjectEditor/Public/MuCOE/Nodes/CustomizableObjectNodeModifierClipWithUVMask.h
@@ -24,11 +24,13 @@ public:
 	UPROPERTY()
 	TArray<FString> Tags_DEPRECATED;
 
-
 	/** UV channel index that will be used to get the UVs to apply the clipping mask to. */
-	UPROPERTY(EditAnywhere, BlueprintReadWrite, Category = MeshToClip)
+	UPROPERTY(EditAnywhere, BlueprintReadWrite, Category = ClipOptions)
 	int32 UVChannelForMask = 0;
 
+	UPROPERTY(EditAnywhere, Category = ClipOptions)
+	EFaceCullStrategy FaceCullStrategy = EFaceCullStrategy::AllVerticesCulled;
+
 public:
 
 	// EdGraphNode interface
diff --git a/Engine/Plugins/Mutable/Source/CustomizableObjectEditor/Public/MuCOE/Nodes/CustomizableObjectNodeModifierRemoveMesh.h b/Engine/Plugins/Mutable/Source/CustomizableObjectEditor/Public/MuCOE/Nodes/CustomizableObjectNodeModifierRemoveMesh.h
--- a/Engine/Plugins/Mutable/Source/CustomizableObjectEditor/Public/MuCOE/Nodes/CustomizableObjectNodeModifierRemoveMesh.h
+++ b/Engine/Plugins/Mutable/Source/CustomizableObjectEditor/Public/MuCOE/Nodes/CustomizableObjectNodeModifierRemoveMesh.h
@@ -15,9 +15,15 @@ class UObject;
 UCLASS()
 class CUSTOMIZABLEOBJECTEDITOR_API UCustomizableObjectNodeModifierRemoveMesh : public UCustomizableObjectNodeModifierEditMeshSectionBase
 {
-public:
 	GENERATED_BODY()
 
+public:
+
+	UPROPERTY(EditAnywhere, Category = RemoveOptions)
+	EFaceCullStrategy FaceCullStrategy = EFaceCullStrategy::AllVerticesCulled;
+
+public:
+
 	// Begin EdGraphNode interface
 	FText GetNodeTitle(ENodeTitleType::Type TitleType) const override;
 	FText GetTooltipText() const override;
diff --git a/Engine/Plugins/Mutable/Source/CustomizableObjectEditor/Public/MuCOE/Nodes/CustomizableObjectNodeModifierRemoveMeshBlocks.h b/Engine/Plugins/Mutable/Source/CustomizableObjectEditor/Public/MuCOE/Nodes/CustomizableObjectNodeModifierRemoveMeshBlocks.h
--- a/Engine/Plugins/Mutable/Source/CustomizableObjectEditor/Public/MuCOE/Nodes/CustomizableObjectNodeModifierRemoveMeshBlocks.h
+++ b/Engine/Plugins/Mutable/Source/CustomizableObjectEditor/Public/MuCOE/Nodes/CustomizableObjectNodeModifierRemoveMeshBlocks.h
@@ -17,12 +17,18 @@ struct FPropertyChangedEvent;
 UCLASS()
 class CUSTOMIZABLEOBJECTEDITOR_API UCustomizableObjectNodeModifierRemoveMeshBlocks : public UCustomizableObjectNodeModifierEditLayoutBlocks
 {
-public:
 	GENERATED_BODY()
 
+public:
+
 	UPROPERTY()
 	TArray<int32> Blocks_DEPRECATED;
 
+	UPROPERTY(EditAnywhere, Category = RemoveOptions)
+	EFaceCullStrategy FaceCullStrategy = EFaceCullStrategy::AllVerticesCulled;
+
+public:
+
 	// Begin EdGraphNode interface
 	virtual FText GetNodeTitle(ENodeTitleType::Type TitleType) const override;
 	virtual FText GetTooltipText() const override;
diff --git a/Engine/Plugins/Mutable/Source/MutableRuntime/Internal/MuR/Operations.h b/Engine/Plugins/Mutable/Source/MutableRuntime/Internal/MuR/Operations.h
--- a/Engine/Plugins/Mutable/Source/MutableRuntime/Internal/MuR/Operations.h
+++ b/Engine/Plugins/Mutable/Source/MutableRuntime/Internal/MuR/Operations.h
@@ -3,6 +3,7 @@
 #pragma once
 
 #include "MuR/Image.h"
+#include "MuR/Types.h"
 #include "MuR/SerialisationPrivate.h"
 #include "HAL/PlatformMath.h"
 #include "HAL/UnrealMemory.h"
@@ -376,9 +377,9 @@ namespace mu
 		EMeshBindColorChannelUsage A;
 	};
 
-
 	static_assert(sizeof(FMeshBindColorChannelUsages) == sizeof(uint32));
 
+
     //---------------------------------------------------------------------------------------------
     //!
     //---------------------------------------------------------------------------------------------
@@ -914,16 +915,12 @@ namespace mu
 
 			ADDRESS morphShape;
 			ADDRESS vertexSelectionShapeOrBone;
-			
-			typedef enum 
-			{
-				VS_NONE,
-				VS_SHAPE,
-				VS_BONE_HIERARCHY
-			} VERTEX_SELECTION_TYPE;
-            uint8 vertexSelectionType;
 
 			float dist, factor, maxBoneRadius;
+
+			EClipVertexSelectionType VertexSelectionType;
+
+			EFaceCullStrategy FaceCullStrategy;
 		};
 
         struct MeshClipWithMeshArgs
@@ -987,6 +984,7 @@ namespace mu
 			ADDRESS clipShape;
 
 			float clipWeightThreshold = 0.9f;
+			EFaceCullStrategy FaceCullStrategy;
 		};
 
 		struct MeshOptimizeSkinningArgs
diff --git a/Engine/Plugins/Mutable/Source/MutableRuntime/Private/MuR/CodeRunner.cpp b/Engine/Plugins/Mutable/Source/MutableRuntime/Private/MuR/CodeRunner.cpp
--- a/Engine/Plugins/Mutable/Source/MutableRuntime/Private/MuR/CodeRunner.cpp
+++ b/Engine/Plugins/Mutable/Source/MutableRuntime/Private/MuR/CodeRunner.cpp
@@ -2062,7 +2062,9 @@ namespace mu
                 const FVector3f& origin = morphShape.position;
                 const FVector3f& normal = morphShape.up;
 
-                if (args.vertexSelectionType == OP::MeshClipMorphPlaneArgs::VS_SHAPE)
+				bool bRemoveFaceIfAllVerticesCulled = args.FaceCullStrategy==EFaceCullStrategy::AllVerticesCulled;
+
+                if (args.VertexSelectionType == EClipVertexSelectionType::Shape)
                 {
                     check(args.vertexSelectionShapeOrBone < (uint32)pModel->GetPrivate()->m_program.m_constantShapes.Num());
 
@@ -2072,7 +2074,7 @@ namespace mu
 					Ptr<Mesh> Result = CreateMesh(Source ? Source->GetDataSize() : 0);
 
 					bool bOutSuccess = false;
-					MeshClipMorphPlane(Result.get(), Source.get(), origin, normal, args.dist, args.factor, morphShape.size[0], morphShape.size[1], morphShape.size[2], selectionShape, bOutSuccess, nullptr, -1);
+					MeshClipMorphPlane(Result.get(), Source.get(), origin, normal, args.dist, args.factor, morphShape.size[0], morphShape.size[1], morphShape.size[2], selectionShape, bRemoveFaceIfAllVerticesCulled, bOutSuccess, nullptr, -1);
 					
 					if (!bOutSuccess)
 					{
@@ -2086,7 +2088,7 @@ namespace mu
 					}
                 }
 
-				else if (args.vertexSelectionType == OP::MeshClipMorphPlaneArgs::VS_BONE_HIERARCHY)
+				else if (args.VertexSelectionType == EClipVertexSelectionType::BoneHierarchy)
 				{
 					FShape selectionShape;
 					selectionShape.type = (uint8)FShape::Type::None;
@@ -2097,7 +2099,7 @@ namespace mu
 					const FBoneName Bone(args.vertexSelectionShapeOrBone);
 
 					bool bOutSuccess = false;
-					MeshClipMorphPlane(Result.get(), Source.get(), origin, normal, args.dist, args.factor, morphShape.size[0], morphShape.size[1], morphShape.size[2], selectionShape, bOutSuccess, &Bone, args.maxBoneRadius);
+					MeshClipMorphPlane(Result.get(), Source.get(), origin, normal, args.dist, args.factor, morphShape.size[0], morphShape.size[1], morphShape.size[2], selectionShape, bRemoveFaceIfAllVerticesCulled, bOutSuccess, &Bone, args.maxBoneRadius);
 
 					if (!bOutSuccess)
 					{
@@ -2119,7 +2121,7 @@ namespace mu
 					Ptr<Mesh> Result = CreateMesh(Source ? Source->GetDataSize() : 0);
 
 					bool bOutSuccess = false;
-					MeshClipMorphPlane(Result.get(), Source.get(), origin, normal, args.dist, args.factor, morphShape.size[0], morphShape.size[1], morphShape.size[2], selectionShape, bOutSuccess, nullptr, -1.0f);
+					MeshClipMorphPlane(Result.get(), Source.get(), origin, normal, args.dist, args.factor, morphShape.size[0], morphShape.size[1], morphShape.size[2], selectionShape, bRemoveFaceIfAllVerticesCulled, bOutSuccess, nullptr, -1.0f);
 
 					if (!bOutSuccess)
 					{
@@ -2236,8 +2238,10 @@ namespace mu
 				{
 					Ptr<Mesh> Result = CreateMesh(BaseMesh->GetDataSize());
 
+					bool bRemoveIfAllVerticesCulled = args.FaceCullStrategy == EFaceCullStrategy::AllVerticesCulled;
+
 					bool bOutSuccess = false;
-					MeshClipDeform(Result.get(), BaseMesh.get(), ClipShape.get(), args.clipWeightThreshold, bOutSuccess);
+					MeshClipDeform(Result.get(), BaseMesh.get(), ClipShape.get(), args.clipWeightThreshold, bRemoveIfAllVerticesCulled, bOutSuccess);
 
 					Release(ClipShape);
 
@@ -2746,6 +2750,10 @@ namespace mu
             FMemory::Memcpy(&source,data,sizeof(OP::ADDRESS)); 
 			data += sizeof(OP::ADDRESS);
 
+			EFaceCullStrategy FaceCullStrategy;
+			FMemory::Memcpy(&FaceCullStrategy, data, sizeof(EFaceCullStrategy));
+			data += sizeof(EFaceCullStrategy);
+
             TArray<FScheduledOp> conditions;
 			TArray<OP::ADDRESS> masks;
 
@@ -2847,7 +2855,8 @@ namespace mu
 								Ptr<Mesh> IterResult = CreateMesh(Result->GetDataSize());
 
 								bool bOutSuccess = false;
-								MeshRemoveMask(IterResult.get(), Result.get(), Mask.get(), bOutSuccess);
+								bool bRemoveIfAllVerticesCulled = FaceCullStrategy == EFaceCullStrategy::AllVerticesCulled;
+								MeshRemoveMask(IterResult.get(), Result.get(), Mask.get(), bRemoveIfAllVerticesCulled, bOutSuccess);
 
 								Release(Mask);
 
diff --git a/Engine/Plugins/Mutable/Source/MutableRuntime/Private/MuR/OpMeshApplyLayout.h b/Engine/Plugins/Mutable/Source/MutableRuntime/Private/MuR/OpMeshApplyLayout.h
--- a/Engine/Plugins/Mutable/Source/MutableRuntime/Private/MuR/OpMeshApplyLayout.h
+++ b/Engine/Plugins/Mutable/Source/MutableRuntime/Private/MuR/OpMeshApplyLayout.h
@@ -12,18 +12,17 @@ namespace mu
 {
 
 
-	//---------------------------------------------------------------------------------------------
-    inline void MeshApplyLayout( Mesh* pApplied, const Layout* pLayout, int32 texCoordsSet )
+    inline void MeshApplyLayout( Mesh* Applied, const Layout* InLayout, int32 TexCoordsIndex )
 	{
 		MUTABLE_CPUPROFILER_SCOPE(MeshApplyLayout);
 
         int32 buffer = -1;
 		int32 channel = -1;
-		pApplied->GetVertexBuffers().FindChannel( MBS_TEXCOORDS, texCoordsSet, &buffer, &channel );
+		Applied->GetVertexBuffers().FindChannel( MBS_TEXCOORDS, TexCoordsIndex, &buffer, &channel );
 
         int32 layoutBuffer = -1;
         int32 layoutChannel = -1;
-        pApplied->GetVertexBuffers().FindChannel( MBS_LAYOUTBLOCK, texCoordsSet, &layoutBuffer, &layoutChannel );
+        Applied->GetVertexBuffers().FindChannel( MBS_LAYOUTBLOCK, TexCoordsIndex, &layoutBuffer, &layoutChannel );
 
 		if (buffer < 0 || layoutBuffer < 0)
 		{
@@ -36,12 +35,12 @@ namespace mu
 		EMeshBufferFormat format;
 		int32 components;
 		int32 offset;
-		pApplied->GetVertexBuffers().GetChannel( buffer, channel, &semantic, &semanticIndex, &format, &components, &offset );
+		Applied->GetVertexBuffers().GetChannel( buffer, channel, &semantic, &semanticIndex, &format, &components, &offset );
 		check( semantic == MBS_TEXCOORDS );
 
-        uint8* pData = pApplied->GetVertexBuffers().GetBufferData( buffer );
-		int32 elemSize = pApplied->GetVertexBuffers().GetElementSize( buffer );
-		int32 channelOffset = pApplied->GetVertexBuffers().GetChannelOffset( buffer, channel );
+        uint8* pData = Applied->GetVertexBuffers().GetBufferData( buffer );
+		int32 elemSize = Applied->GetVertexBuffers().GetElementSize( buffer );
+		int32 channelOffset = Applied->GetVertexBuffers().GetChannelOffset( buffer, channel );
 		pData += channelOffset;
 
 		struct Box
@@ -49,14 +48,14 @@ namespace mu
 			FVector2f min, size;
 		};
 		TArray< Box > transforms;
-		transforms.SetNum(pLayout->GetBlockCount());
-		for ( int b=0; b<pLayout->GetBlockCount(); ++b )
+		transforms.SetNum(InLayout->GetBlockCount());
+		for ( int b=0; b<InLayout->GetBlockCount(); ++b )
 		{
-			FIntPoint grid = pLayout->GetGridSize();
+			FIntPoint grid = InLayout->GetGridSize();
 
 			box< FIntVector2 > block;
-			block.min = pLayout->Blocks[b].Min;
-			block.size = pLayout->Blocks[b].Size;
+			block.min = InLayout->Blocks[b].Min;
+			block.size = InLayout->Blocks[b].Size;
 
 			Box rect;
 			rect.min[0] = ( (float)block.min[0] ) / (float) grid[0];
@@ -67,18 +66,18 @@ namespace mu
 		}
 
         check( layoutBuffer>=0 && layoutChannel>=0 );
-        check( pApplied->GetVertexBuffers().Buffers[layoutBuffer].Channels.Num()==1 );
-        check( pApplied->GetVertexBuffers().Buffers[layoutBuffer].Channels[layoutChannel].ComponentCount==1 );
+        check( Applied->GetVertexBuffers().Buffers[layoutBuffer].Channels.Num()==1 );
+        check( Applied->GetVertexBuffers().Buffers[layoutBuffer].Channels[layoutChannel].ComponentCount==1 );
 
-        const uint16* pLayoutData = reinterpret_cast<const uint16*>( pApplied->GetVertexBuffers().GetBufferData( layoutBuffer ) );
-		UntypedMeshBufferIterator ItLayoutData(pApplied->GetVertexBuffers(), MBS_LAYOUTBLOCK, texCoordsSet);
+        const uint16* InLayoutData = reinterpret_cast<const uint16*>( Applied->GetVertexBuffers().GetBufferData( layoutBuffer ) );
+		UntypedMeshBufferIterator ItLayoutData(Applied->GetVertexBuffers(), MBS_LAYOUTBLOCK, TexCoordsIndex);
 
 		// In some corner case involving automatic LODs and remove meshes behaving differently among them we may need to remove vertices that don't 
 		// have any block in the current layout. Track them here.
 		TArray<int32> VerticesToRemove;
 
         uint8* pVertices = pData;
-		for ( int32 v=0; v<pApplied->GetVertexBuffers().GetElementCount(); ++v )
+		for ( int32 v=0; v<Applied->GetVertexBuffers().GetElementCount(); ++v )
 		{
 			uint64 BlockId = 0;
 			if (ItLayoutData.GetFormat() == MBF_UINT16)
@@ -86,7 +85,7 @@ namespace mu
 				// Relative blocks.
 				const uint16* SourceIds = reinterpret_cast<const uint16*>(ItLayoutData.ptr());
 				BlockId = *SourceIds;
-				BlockId = BlockId | (uint64(pApplied->MeshIDPrefix) << 32);
+				BlockId = BlockId | (uint64(Applied->MeshIDPrefix) << 32);
 			}
 			else if (ItLayoutData.GetFormat() == MBF_UINT64)
 			{
@@ -102,7 +101,7 @@ namespace mu
 			++ItLayoutData;
 
 			// TODO: This could be optimised
-			int32 relBlock = pLayout->FindBlock(BlockId);
+			int32 relBlock = InLayout->FindBlock(BlockId);
 
 			// This may still happen with lower LOD and "remove meshes" in a corner case:
 			// Auto LODs with "Remove Meshes" that behave differently across LODs, and leave geometry in a block that has been removed in the higher LOD.
@@ -209,18 +208,19 @@ namespace mu
 		{
 			// Unpack vertices into a mask
 			TBitArray<> VertexMask;
-			VertexMask.SetNum(pApplied->GetVertexCount(), false);
+			VertexMask.SetNum(Applied->GetVertexCount(), false);
 			for (int32 VertexIndex : VerticesToRemove)
 			{
 				VertexMask[VertexIndex] = true;
 			}
 	
 			// Remove
-			MeshRemoveVerticesWithCullSet(pApplied, VertexMask);
+			bool bRemoveIfAllVerticesCulled = true;
+			MeshRemoveVerticesWithCullSet(Applied, VertexMask, bRemoveIfAllVerticesCulled);
 		}
 
 		//
-		pApplied->SetLayout( texCoordsSet, pLayout->Clone() );
+		Applied->SetLayout( TexCoordsIndex, InLayout->Clone() );
 
 	}
 
diff --git a/Engine/Plugins/Mutable/Source/MutableRuntime/Private/MuR/OpMeshClipDeform.h b/Engine/Plugins/Mutable/Source/MutableRuntime/Private/MuR/OpMeshClipDeform.h
--- a/Engine/Plugins/Mutable/Source/MutableRuntime/Private/MuR/OpMeshClipDeform.h
+++ b/Engine/Plugins/Mutable/Source/MutableRuntime/Private/MuR/OpMeshClipDeform.h
@@ -100,7 +100,7 @@ namespace mu
 		return true;
 	}
 
-	inline void MeshClipDeform(Mesh* Result, const Mesh* BaseMesh, const Mesh* ShapeMesh, const float ClipWeightThreshold, bool& bOutSuccess )
+	inline void MeshClipDeform(Mesh* Result, const Mesh* BaseMesh, const Mesh* ShapeMesh, const float ClipWeightThreshold, bool bRemoveIfAllVerticesCulled, bool& bOutSuccess )
 	{
 		MUTABLE_CPUPROFILER_SCOPE(ClipDeform);
 		bOutSuccess = true;
@@ -327,6 +327,6 @@ namespace mu
 			}
 		}
 
-		MeshRemoveVerticesWithCullSet(Result, VerticesToCull);
+		MeshRemoveVerticesWithCullSet(Result, VerticesToCull, bRemoveIfAllVerticesCulled);
 	}
 }
diff --git a/Engine/Plugins/Mutable/Source/MutableRuntime/Private/MuR/OpMeshClipMorphPlane.h b/Engine/Plugins/Mutable/Source/MutableRuntime/Private/MuR/OpMeshClipMorphPlane.h
--- a/Engine/Plugins/Mutable/Source/MutableRuntime/Private/MuR/OpMeshClipMorphPlane.h
+++ b/Engine/Plugins/Mutable/Source/MutableRuntime/Private/MuR/OpMeshClipMorphPlane.h
@@ -72,7 +72,8 @@ namespace mu
 	//! Reference version
 	//---------------------------------------------------------------------------------------------
 	inline void MeshClipMorphPlane(Mesh* Result, const Mesh* pBase, const FVector3f& Origin, const FVector3f& Normal, float Dist, float Factor, float Radius,
-		float Radius2, float Angle, const FShape& SelectionShape, bool& bOutSuccess, const FBoneName* BoneId = nullptr, float VertexSelectionBoneMaxRadius = -1.f)
+		float Radius2, float Angle, const FShape& SelectionShape, bool bRemoveIfAllVerticesCulled, bool& bOutSuccess, 
+		const FBoneName* BoneId = nullptr, float VertexSelectionBoneMaxRadius = -1.f)
 	{
 		bOutSuccess = true;
 		//float Radius = 8.f;
@@ -372,6 +373,6 @@ namespace mu
 			(PositionIterBegin + VertexIndex).SetFromVec3f(Position);
 		}
 
-        MeshRemoveVerticesWithCullSet(Result, VerticesToCull);
+        MeshRemoveVerticesWithCullSet(Result, VerticesToCull, bRemoveIfAllVerticesCulled);
 	}
 }
diff --git a/Engine/Plugins/Mutable/Source/MutableRuntime/Private/MuR/OpMeshRemove.cpp b/Engine/Plugins/Mutable/Source/MutableRuntime/Private/MuR/OpMeshRemove.cpp
--- a/Engine/Plugins/Mutable/Source/MutableRuntime/Private/MuR/OpMeshRemove.cpp
+++ b/Engine/Plugins/Mutable/Source/MutableRuntime/Private/MuR/OpMeshRemove.cpp
@@ -337,7 +337,7 @@ namespace mu
     //    MeshRemoveRecreateSurface(Result, UsedVertices, UsedFaces);
     //}
 
-    void MeshRemoveVerticesWithCullSet(Mesh* Result, const TBitArray<>& VerticesToCull)
+    void MeshRemoveVerticesWithCullSet(Mesh* Result, const TBitArray<>& VerticesToCull, bool bRemoveIfAllVerticesCulled)
     { 
         MUTABLE_CPUPROFILER_SCOPE(MeshRemoveVerticesWithCullSet);
 		
@@ -361,12 +361,26 @@ namespace mu
             {
                 const uint32* FaceIndicesData = reinterpret_cast<uint32*>((IndicesBegin + FaceIndex*3).ptr());
 
-                const bool bAllVertsRemoved = 
-                        VerticesToCull[FaceIndicesData[0]] & 
-                        VerticesToCull[FaceIndicesData[1]] & 
-                        VerticesToCull[FaceIndicesData[2]];
-
-                if (!bAllVertsRemoved)
+				bool bRemoved = false;
+
+				if (bRemoveIfAllVerticesCulled)
+				{
+					const bool bAllVertsRemoved =
+						VerticesToCull[FaceIndicesData[0]] &
+						VerticesToCull[FaceIndicesData[1]] &
+						VerticesToCull[FaceIndicesData[2]];
+					bRemoved = bAllVertsRemoved;
+				}
+				else
+				{
+					const bool bOneVertRemoved =
+						VerticesToCull[FaceIndicesData[0]] |
+						VerticesToCull[FaceIndicesData[1]] |
+						VerticesToCull[FaceIndicesData[2]];
+					bRemoved = bOneVertRemoved;
+				}
+
+                if (!bRemoved)
                 {
                     ++NumUsedFaces;
                     UsedFaces[FaceIndex] = true;
@@ -558,7 +572,7 @@ namespace mu
 		MeshRemoveRecreateSurface(Result, UsedVertices, UsedFaces);
     }
 
-	void MeshRemoveMask(Mesh* Result, const Mesh* Source, const Mesh* Mask, bool& bOutSuccess)
+	void MeshRemoveMask(Mesh* Result, const Mesh* Source, const Mesh* Mask, bool bRemoveIfAllVerticesCulled, bool& bOutSuccess)
     {
         MUTABLE_CPUPROFILER_SCOPE(MeshRemoveMask);
 		bOutSuccess = true;
@@ -595,6 +609,6 @@ namespace mu
             }
         }
 		
-        MeshRemoveVerticesWithCullSet(Result, RemovedVertices);
+        MeshRemoveVerticesWithCullSet(Result, RemovedVertices, bRemoveIfAllVerticesCulled);
 	}
 }
diff --git a/Engine/Plugins/Mutable/Source/MutableRuntime/Private/MuR/OpMeshRemove.h b/Engine/Plugins/Mutable/Source/MutableRuntime/Private/MuR/OpMeshRemove.h
--- a/Engine/Plugins/Mutable/Source/MutableRuntime/Private/MuR/OpMeshRemove.h
+++ b/Engine/Plugins/Mutable/Source/MutableRuntime/Private/MuR/OpMeshRemove.h
@@ -11,18 +11,19 @@ namespace mu
 {
 	class Mesh;
 
-	/** Remove a list of vertices and related faces from a mesh. The list of vertices is stored in a specially formattes Mask mesh. */
-	extern void MeshRemoveMask(Mesh* Result, const Mesh* Source, const Mesh* Mask, bool& bOutSuccess);
-
-	/** Remove a list of vertices and related faces from a mesh. The list is stored as a bool map for every vertex in the mesh. */
-	//extern void MeshRemoveVerticesWithMap(Mesh* Result, const TBitArray<>& RemovedVertices);
+	/** Remove a list of vertices and related faces from a mesh. The list of vertices is stored in a specially formatted Mask mesh. 
+	 * If bRemoveIfAllVerticesCulled is true, a face is remove if all its vertices have the bit set in VerticesToCull.
+	 * If bRemoveIfAllVerticesCulled is false, remove a face if at least one vertex is removed.
+	 */
+	extern void MeshRemoveMask(Mesh* Result, const Mesh* Source, const Mesh* Mask, bool bRemoveIfAllVerticesCulled, bool& bOutSuccess);
 
 	/** 
 	 * Remove a set of vertices and related faces from a mesh in-place. VertexToCull is a bitset where if bit i-th is set, 
-	 * the vertex i-th will be removed if all faces referencing this vertex need to be removed. A face is remove if all its 
-	 * vertices have the bit set in VerticesToCull.
+	 * the vertex i-th will be removed if all faces referencing this vertex need to be removed. 
+	 * If bRemoveIfAllVerticesCulled is true, a face is remove if all its vertices have the bit set in VerticesToCull. 
+	 * If bRemoveIfAllVerticesCulled is false, remove a face if at least one vertex is removed.
 	 */
-	extern void MeshRemoveVerticesWithCullSet(Mesh* Result, const TBitArray<>& VerticesToCull);
+	extern void MeshRemoveVerticesWithCullSet(Mesh* Result, const TBitArray<>& VerticesToCull, bool bRemoveIfAllVerticesCulled);
 
 	/**
 	 * Recreates the Surface and Surfaces Submeshes given a set of vertices and faces remaining after mesh removal.
diff --git a/Engine/Plugins/Mutable/Source/MutableRuntime/Private/MuR/Operations.cpp b/Engine/Plugins/Mutable/Source/MutableRuntime/Private/MuR/Operations.cpp
--- a/Engine/Plugins/Mutable/Source/MutableRuntime/Private/MuR/Operations.cpp
+++ b/Engine/Plugins/Mutable/Source/MutableRuntime/Private/MuR/Operations.cpp
@@ -864,11 +864,15 @@ namespace mu
 
         case OP_TYPE::ME_REMOVEMASK:
         {
-            const uint8_t* data = program.GetOpArgsPointer(at);
+            const uint8* data = program.GetOpArgsPointer(at);
             mu::OP::ADDRESS source;
             FMemory::Memcpy( &source, data, sizeof(OP::ADDRESS) ); data+=sizeof(OP::ADDRESS);
             f(source);
 
+			EFaceCullStrategy FaceCullStrategy;
+			FMemory::Memcpy(&FaceCullStrategy, data, sizeof(EFaceCullStrategy));
+			data += sizeof(EFaceCullStrategy);
+
             uint16 removes = 0;
 			FMemory::Memcpy( &removes, data, sizeof(uint16) ); data+=sizeof(uint16);
             for (uint16 r=0; r<removes; ++r)
diff --git a/Engine/Plugins/Mutable/Source/MutableRuntime/Public/MuR/Types.h b/Engine/Plugins/Mutable/Source/MutableRuntime/Public/MuR/Types.h
--- a/Engine/Plugins/Mutable/Source/MutableRuntime/Public/MuR/Types.h
+++ b/Engine/Plugins/Mutable/Source/MutableRuntime/Public/MuR/Types.h
@@ -7,6 +7,30 @@
 #include "Containers/Array.h"
 #include "Containers/UnrealString.h"
 
+#include "Types.generated.h"
+
+/** This type is used in some mesh clip operations. 
+* Warning: this type is used compiled COs. Any relevant change to the option order requires and update in the CustomizableObjectPrivate::ECustomizableObjectVersions
+*/
+UENUM()
+enum class EFaceCullStrategy : uint8
+{
+	AllVerticesCulled = 0 UMETA(DisplayName = "Remove face if all vertices removed"),
+	OneVertexCulled = 1 UMETA(DisplayName = "Remove face if one vertex removed"),
+};
+
+
+/** This type is used in some mesh clip operations.
+* Warning: this type is used compiled COs. Any relevant change to the option order requires and update in the CustomizableObjectPrivate::ECustomizableObjectVersions
+*/
+enum class EClipVertexSelectionType : uint8
+{
+	None = 0,
+	Shape = 1,
+	BoneHierarchy = 2,
+};
+
+
 namespace mu
 {
 	class OutputArchive;
diff --git a/Engine/Plugins/Mutable/Source/MutableTools/Private/MuT/ASTOpMeshClipDeform.cpp b/Engine/Plugins/Mutable/Source/MutableTools/Private/MuT/ASTOpMeshClipDeform.cpp
--- a/Engine/Plugins/Mutable/Source/MutableTools/Private/MuT/ASTOpMeshClipDeform.cpp
+++ b/Engine/Plugins/Mutable/Source/MutableTools/Private/MuT/ASTOpMeshClipDeform.cpp
@@ -31,7 +31,7 @@ namespace mu
 		if (OtherUntyped.GetOpType()==GetOpType())
 		{
 			const ASTOpMeshClipDeform* Other = static_cast<const ASTOpMeshClipDeform*>(&OtherUntyped);
-			return Mesh == Other->Mesh && ClipShape == Other->ClipShape;
+			return Mesh == Other->Mesh && ClipShape == Other->ClipShape && FaceCullStrategy == Other->FaceCullStrategy;
 		}
 
 		return false;
@@ -51,6 +51,7 @@ namespace mu
 		Ptr<ASTOpMeshClipDeform> n = new ASTOpMeshClipDeform();
 		n->Mesh = mapChild(Mesh.child());
 		n->ClipShape = mapChild(ClipShape.child());
+		n->FaceCullStrategy = FaceCullStrategy;
 		return n;
 	}
 
@@ -68,7 +69,9 @@ namespace mu
 		if (!linkedAddress)
 		{
 			OP::MeshClipDeformArgs args;
-			memset(&args, 0, sizeof(args));
+			FMemory::Memzero(&args, sizeof(args));
+
+			args.FaceCullStrategy = FaceCullStrategy;
 
 			if (Mesh)
 			{
diff --git a/Engine/Plugins/Mutable/Source/MutableTools/Private/MuT/ASTOpMeshClipDeform.h b/Engine/Plugins/Mutable/Source/MutableTools/Private/MuT/ASTOpMeshClipDeform.h
--- a/Engine/Plugins/Mutable/Source/MutableTools/Private/MuT/ASTOpMeshClipDeform.h
+++ b/Engine/Plugins/Mutable/Source/MutableTools/Private/MuT/ASTOpMeshClipDeform.h
@@ -23,6 +23,9 @@ struct FProgram;
 		ASTChild Mesh;
 		ASTChild ClipShape;
 
+		/** Strategy to decide when to cull a face. */
+		EFaceCullStrategy FaceCullStrategy = EFaceCullStrategy::AllVerticesCulled;
+
 	public:
 
 		ASTOpMeshClipDeform();
diff --git a/Engine/Plugins/Mutable/Source/MutableTools/Private/MuT/ASTOpMeshClipMorphPlane.cpp b/Engine/Plugins/Mutable/Source/MutableTools/Private/MuT/ASTOpMeshClipMorphPlane.cpp
--- a/Engine/Plugins/Mutable/Source/MutableTools/Private/MuT/ASTOpMeshClipMorphPlane.cpp
+++ b/Engine/Plugins/Mutable/Source/MutableTools/Private/MuT/ASTOpMeshClipMorphPlane.cpp
@@ -39,7 +39,8 @@ namespace mu
 				morphShape == other->morphShape &&
 				selectionShape == other->selectionShape &&
 				vertexSelectionBone == other->vertexSelectionBone &&
-				vertexSelectionType == other->vertexSelectionType &&
+				VertexSelectionType == other->VertexSelectionType &&
+				FaceCullStrategy == other->FaceCullStrategy &&
 				dist == other->dist &&
 				factor == other->factor;
 		}
@@ -63,7 +64,8 @@ namespace mu
 		n->morphShape = morphShape;
 		n->selectionShape = selectionShape;
 		n->vertexSelectionBone = vertexSelectionBone;
-		n->vertexSelectionType = vertexSelectionType;
+		n->VertexSelectionType = VertexSelectionType;
+		n->FaceCullStrategy = FaceCullStrategy;
 		n->dist = dist;
 		n->factor = factor;
 		return n;
@@ -82,20 +84,22 @@ namespace mu
 		if (!linkedAddress)
 		{
 			OP::MeshClipMorphPlaneArgs args;
-			memset(&args, 0, sizeof(args));
+			FMemory::Memzero(&args, sizeof(args));
 
 			if (source) args.source = source->linkedAddress;
 
 			args.morphShape = (OP::ADDRESS)program.m_constantShapes.Num();
 			program.m_constantShapes.Add(morphShape);
 
-			args.vertexSelectionType = (uint8_t)vertexSelectionType;
-			if (vertexSelectionType == OP::MeshClipMorphPlaneArgs::VS_BONE_HIERARCHY)
+			args.FaceCullStrategy = FaceCullStrategy;
+
+			args.VertexSelectionType = VertexSelectionType;
+			if (VertexSelectionType == EClipVertexSelectionType::BoneHierarchy)
 			{
 				check(vertexSelectionBone.Id < MAX_uint32);
 				args.vertexSelectionShapeOrBone = vertexSelectionBone.Id;
 			}
-			else if (vertexSelectionType == OP::MeshClipMorphPlaneArgs::VS_SHAPE)
+			else if (VertexSelectionType == EClipVertexSelectionType::Shape)
 			{
 				args.vertexSelectionShapeOrBone = program.AddConstant(selectionShape);
 			}
diff --git a/Engine/Plugins/Mutable/Source/MutableTools/Private/MuT/ASTOpMeshClipMorphPlane.h b/Engine/Plugins/Mutable/Source/MutableTools/Private/MuT/ASTOpMeshClipMorphPlane.h
--- a/Engine/Plugins/Mutable/Source/MutableTools/Private/MuT/ASTOpMeshClipMorphPlane.h
+++ b/Engine/Plugins/Mutable/Source/MutableTools/Private/MuT/ASTOpMeshClipMorphPlane.h
@@ -26,7 +26,8 @@ struct FProgram;
 		FShape selectionShape;
 		FBoneName vertexSelectionBone;
 
-		OP::MeshClipMorphPlaneArgs::VERTEX_SELECTION_TYPE vertexSelectionType = OP::MeshClipMorphPlaneArgs::VS_NONE;
+		EClipVertexSelectionType VertexSelectionType = EClipVertexSelectionType::None;
+		EFaceCullStrategy FaceCullStrategy = EFaceCullStrategy::AllVerticesCulled;
 
 		float dist = 0.f, factor = 0.f, vertexSelectionBoneMaxRadius = -1.f;
 
diff --git a/Engine/Plugins/Mutable/Source/MutableTools/Private/MuT/ASTOpMeshRemoveMask.cpp b/Engine/Plugins/Mutable/Source/MutableTools/Private/MuT/ASTOpMeshRemoveMask.cpp
--- a/Engine/Plugins/Mutable/Source/MutableTools/Private/MuT/ASTOpMeshRemoveMask.cpp
+++ b/Engine/Plugins/Mutable/Source/MutableTools/Private/MuT/ASTOpMeshRemoveMask.cpp
@@ -42,7 +42,7 @@ namespace mu
 		if (otherUntyped.GetOpType() == GetOpType())
 		{
 			const ASTOpMeshRemoveMask* other = static_cast<const ASTOpMeshRemoveMask*>(&otherUntyped);
-			return source == other->source && removes == other->removes;
+			return source == other->source && removes == other->removes && FaceCullStrategy==other->FaceCullStrategy;
 		}
 		return false;
 	}
@@ -53,6 +53,7 @@ namespace mu
 	{
 		Ptr<ASTOpMeshRemoveMask> n = new ASTOpMeshRemoveMask();
 		n->source = mapChild(source.child());
+		n->FaceCullStrategy = FaceCullStrategy;
 		for (const TPair<ASTChild, ASTChild>& r : removes)
 		{
 			n->removes.Add({ ASTChild(n,mapChild(r.Key.child())), ASTChild(n,mapChild(r.Value.child())) });
@@ -96,8 +97,12 @@ namespace mu
 
 			program.m_opAddress.Add((uint32_t)program.m_byteCode.Num());
 			AppendCode(program.m_byteCode, OP_TYPE::ME_REMOVEMASK);
+
 			OP::ADDRESS sourceAt = source ? source->linkedAddress : 0;
 			AppendCode(program.m_byteCode, sourceAt);
+
+			AppendCode(program.m_byteCode, FaceCullStrategy);
+
 			AppendCode(program.m_byteCode, (uint16)removes.Num());
 			for (const TPair<ASTChild, ASTChild>& b : removes)
 			{
diff --git a/Engine/Plugins/Mutable/Source/MutableTools/Private/MuT/ASTOpMeshRemoveMask.h b/Engine/Plugins/Mutable/Source/MutableTools/Private/MuT/ASTOpMeshRemoveMask.h
--- a/Engine/Plugins/Mutable/Source/MutableTools/Private/MuT/ASTOpMeshRemoveMask.h
+++ b/Engine/Plugins/Mutable/Source/MutableTools/Private/MuT/ASTOpMeshRemoveMask.h
@@ -12,9 +12,7 @@ namespace mu
 {
 	struct FProgram;
 
-	//---------------------------------------------------------------------------------------------
-	//! From a source mesh, remove a list of fragments with a condition.
-	//---------------------------------------------------------------------------------------------
+	/** Remove a list of mesh fragments with a condition from a source mesh. */
 	class ASTOpMeshRemoveMask final : public ASTOp
 	{
 	public:
@@ -25,6 +23,9 @@ namespace mu
 		//! Pairs of remove candidates: condition + mesh to remove
 		TArray< TPair<ASTChild, ASTChild> > removes;
 
+		/** Strategy to decide when to cull a face. */
+		EFaceCullStrategy FaceCullStrategy = EFaceCullStrategy::AllVerticesCulled;
+
 	public:
 
 		ASTOpMeshRemoveMask();
diff --git a/Engine/Plugins/Mutable/Source/MutableTools/Private/MuT/CodeGenerator.cpp b/Engine/Plugins/Mutable/Source/MutableTools/Private/MuT/CodeGenerator.cpp
--- a/Engine/Plugins/Mutable/Source/MutableTools/Private/MuT/CodeGenerator.cpp
+++ b/Engine/Plugins/Mutable/Source/MutableTools/Private/MuT/CodeGenerator.cpp
@@ -2067,6 +2067,7 @@ namespace mu
 					{
 						RemoveOp = new ASTOpMeshRemoveMask();
 						RemoveOp->source = LastMeshOp;
+						RemoveOp->FaceCullStrategy = Edit->FaceCullStrategy;
 						LastMeshOp = RemoveOp;
 					}
 
@@ -2197,6 +2198,7 @@ namespace mu
 				{
 					RemoveOp = new ASTOpMeshRemoveMask();
 					RemoveOp->source = LastMeshOp;
+					RemoveOp->FaceCullStrategy = TypedClipNode->FaceCullStrategy;
 					LastMeshOp = RemoveOp;
 				}
 
@@ -2273,6 +2275,7 @@ namespace mu
 					{
 						RemoveOp = new ASTOpMeshRemoveMask();
 						RemoveOp->source = LastMeshOp;
+						RemoveOp->FaceCullStrategy = TypedClipNode->FaceCullStrategy;
 						LastMeshOp = RemoveOp;
 					}
 
@@ -2292,6 +2295,7 @@ namespace mu
 				const NodeModifierMeshClipMorphPlane* TypedNode = static_cast<const NodeModifierMeshClipMorphPlane*>(m.Node);
 				Ptr<ASTOpMeshClipMorphPlane> op = new ASTOpMeshClipMorphPlane();
 				op->source = LastMeshOp;
+				op->FaceCullStrategy = TypedNode->Parameters.FaceCullStrategy;
 
 				// Morph to an ellipse
 				{
@@ -2319,25 +2323,20 @@ namespace mu
 				}
 
 				// Selection box
-				if (TypedNode->Parameters.VertexSelectionType == FClipMorphPlaneParameters::VS_SHAPE)
+				op->VertexSelectionType = TypedNode->Parameters.VertexSelectionType;
+				if (op->VertexSelectionType == EClipVertexSelectionType::Shape)
 				{
-					op->vertexSelectionType = OP::MeshClipMorphPlaneArgs::VS_SHAPE;
 					FShape selectionShape;
 					selectionShape.type = (uint8)FShape::Type::AABox;
 					selectionShape.position = TypedNode->Parameters.SelectionBoxOrigin;
 					selectionShape.size = TypedNode->Parameters.SelectionBoxRadius;
 					op->selectionShape = selectionShape;
 				}
-				else if (TypedNode->Parameters.VertexSelectionType == FClipMorphPlaneParameters::VS_BONE_HIERARCHY)
+				else if (op->VertexSelectionType == EClipVertexSelectionType::BoneHierarchy)
 				{
-					op->vertexSelectionType = OP::MeshClipMorphPlaneArgs::VS_BONE_HIERARCHY;
 					op->vertexSelectionBone = TypedNode->Parameters.VertexSelectionBone;
 					op->vertexSelectionBoneMaxRadius = TypedNode->Parameters.MaxEffectRadius;
 				}
-				else
-				{
-					op->vertexSelectionType = OP::MeshClipMorphPlaneArgs::VS_NONE;
-				}
 
 				// Parameters
 				op->dist = TypedNode->Parameters.DistanceToPlane;
@@ -2367,9 +2366,11 @@ namespace mu
 				Ptr<ASTOpMeshBindShape>  BindOp = new ASTOpMeshBindShape();
 				Ptr<ASTOpMeshClipDeform> ClipOp = new ASTOpMeshClipDeform();
 
+				ClipOp->FaceCullStrategy = TypedClipNode->FaceCullStrategy;
+
 				FMeshGenerationOptions ClipOptions;
 				ClipOptions.bLayouts = false;
-				ClipOptions.State = Options.State;
+				ClipOptions.State = Options.State;				
 
 				FMeshGenerationResult ClipShapeResult;
 				GenerateMesh(ClipOptions, ClipShapeResult, TypedClipNode->ClipMesh);
diff --git a/Engine/Plugins/Mutable/Source/MutableTools/Private/MuT/CodeGenerator_Mesh.cpp b/Engine/Plugins/Mutable/Source/MutableTools/Private/MuT/CodeGenerator_Mesh.cpp
--- a/Engine/Plugins/Mutable/Source/MutableTools/Private/MuT/CodeGenerator_Mesh.cpp
+++ b/Engine/Plugins/Mutable/Source/MutableTools/Private/MuT/CodeGenerator_Mesh.cpp
@@ -1403,6 +1403,8 @@ namespace mu
     {
         Ptr<ASTOpMeshClipMorphPlane> op = new ASTOpMeshClipMorphPlane();
 
+		op->FaceCullStrategy = ClipNode->Parameters.FaceCullStrategy;
+
         // Base
         if (ClipNode->Source)
         {
@@ -1424,8 +1426,8 @@ namespace mu
             op->morphShape.up = ClipNode->Parameters.Normal;
             op->morphShape.size = FVector3f(ClipNode->Parameters.Radius1, ClipNode->Parameters.Radius2, ClipNode->Parameters.Rotation); // TODO: Move rotation to ellipse rotation reference base instead of passing it directly
 
-                                                                                      // Generate a "side" vector.
-                                                                                      // \todo: make generic and move to the vector class
+            // Generate a "side" vector.
+            // \todo: make generic and move to the vector class
             {
                 // Generate vector perpendicular to normal for ellipse rotation reference base
 				FVector3f aux_base(0.f, 1.f, 0.f);
@@ -1440,25 +1442,20 @@ namespace mu
         }
 
         // Selection by shape
-        if (ClipNode->Parameters.VertexSelectionType== FClipMorphPlaneParameters::VS_SHAPE)
+		op->VertexSelectionType = ClipNode->Parameters.VertexSelectionType;
+        if (op->VertexSelectionType == EClipVertexSelectionType::Shape)
         {
-            op->vertexSelectionType = OP::MeshClipMorphPlaneArgs::VS_SHAPE;
             op->selectionShape.type = (uint8_t)FShape::Type::AABox;
             op->selectionShape.position = ClipNode->Parameters.SelectionBoxOrigin;
             op->selectionShape.size = ClipNode->Parameters.SelectionBoxRadius;
         }
-        else if (ClipNode->Parameters.VertexSelectionType == FClipMorphPlaneParameters::VS_BONE_HIERARCHY)
+        else if (op->VertexSelectionType == EClipVertexSelectionType::BoneHierarchy)
         {
             // Selection by bone hierarchy?
-            op->vertexSelectionType = OP::MeshClipMorphPlaneArgs::VS_BONE_HIERARCHY;
             op->vertexSelectionBone = ClipNode->Parameters.VertexSelectionBone;
 			op->vertexSelectionBoneMaxRadius = ClipNode->Parameters.MaxEffectRadius;
         }
-        else
-        {
-            op->vertexSelectionType = OP::MeshClipMorphPlaneArgs::VS_NONE;
-        }
-
+ 
         // Parameters
         op->dist = ClipNode->Parameters.DistanceToPlane;
         op->factor = ClipNode->Parameters.LinearityFactor;
diff --git a/Engine/Plugins/Mutable/Source/MutableTools/Private/MuT/CodeOptimiser.cpp b/Engine/Plugins/Mutable/Source/MutableTools/Private/MuT/CodeOptimiser.cpp
--- a/Engine/Plugins/Mutable/Source/MutableTools/Private/MuT/CodeOptimiser.cpp
+++ b/Engine/Plugins/Mutable/Source/MutableTools/Private/MuT/CodeOptimiser.cpp
@@ -1178,7 +1178,7 @@ namespace mu
 			case OP_TYPE::ME_CLIPMORPHPLANE:
 			{
 				ASTOpMeshClipMorphPlane* typedOp = static_cast<ASTOpMeshClipMorphPlane*>(node.get());
-				if (typedOp->vertexSelectionType == OP::MeshClipMorphPlaneArgs::VS_BONE_HIERARCHY)
+				if (typedOp->VertexSelectionType == EClipVertexSelectionType::BoneHierarchy)
 				{
 					// We need the skeleton for the source mesh
 					RecurseWithState( typedOp->source.child(), true );
diff --git a/Engine/Plugins/Mutable/Source/MutableTools/Private/MuT/NodeMeshClipMorphPlane.cpp b/Engine/Plugins/Mutable/Source/MutableTools/Private/MuT/NodeMeshClipMorphPlane.cpp
--- a/Engine/Plugins/Mutable/Source/MutableTools/Private/MuT/NodeMeshClipMorphPlane.cpp
+++ b/Engine/Plugins/Mutable/Source/MutableTools/Private/MuT/NodeMeshClipMorphPlane.cpp
@@ -51,7 +51,7 @@ namespace mu
 	//---------------------------------------------------------------------------------------------
 	void NodeMeshClipMorphPlane::SetVertexSelectionBox(FVector3f Center, FVector3f Radius)
 	{
-		Parameters.VertexSelectionType = FClipMorphPlaneParameters::VS_SHAPE;
+		Parameters.VertexSelectionType = EClipVertexSelectionType::Shape;
 		Parameters.SelectionBoxOrigin = Center;
 		Parameters.SelectionBoxRadius = Radius;
 	}
@@ -59,7 +59,7 @@ namespace mu
 	//---------------------------------------------------------------------------------------------
 	void NodeMeshClipMorphPlane::SetVertexSelectionBone(const FBoneName& BoneId, float maxEffectRadius)
 	{
-		Parameters.VertexSelectionType = FClipMorphPlaneParameters::VS_BONE_HIERARCHY;
+		Parameters.VertexSelectionType = EClipVertexSelectionType::BoneHierarchy;
 		Parameters.VertexSelectionBone = BoneId;
 		Parameters.MaxEffectRadius = maxEffectRadius;
 	}
diff --git a/Engine/Plugins/Mutable/Source/MutableTools/Private/MuT/NodeModifierMeshClipMorphPlane.cpp b/Engine/Plugins/Mutable/Source/MutableTools/Private/MuT/NodeModifierMeshClipMorphPlane.cpp
--- a/Engine/Plugins/Mutable/Source/MutableTools/Private/MuT/NodeModifierMeshClipMorphPlane.cpp
+++ b/Engine/Plugins/Mutable/Source/MutableTools/Private/MuT/NodeModifierMeshClipMorphPlane.cpp
@@ -39,7 +39,7 @@ namespace mu
 	//---------------------------------------------------------------------------------------------
 	void NodeModifierMeshClipMorphPlane::SetVertexSelectionBox(float centerX, float centerY, float centerZ, float radiusX, float radiusY, float radiusZ)
 	{
-		Parameters.VertexSelectionType = FClipMorphPlaneParameters::VS_SHAPE;
+		Parameters.VertexSelectionType = EClipVertexSelectionType::Shape;
 		Parameters.SelectionBoxOrigin = FVector3f(centerX, centerY, centerZ);
 		Parameters.SelectionBoxRadius = FVector3f(radiusX, radiusY, radiusZ);
 	}
@@ -47,7 +47,7 @@ namespace mu
 	//---------------------------------------------------------------------------------------------
 	void NodeModifierMeshClipMorphPlane::SetVertexSelectionBone(const FBoneName& BoneId, float maxEffectRadius)
 	{
-		Parameters.VertexSelectionType = FClipMorphPlaneParameters::VS_BONE_HIERARCHY;
+		Parameters.VertexSelectionType = EClipVertexSelectionType::BoneHierarchy;
 		Parameters.VertexSelectionBone = BoneId;
 		Parameters.MaxEffectRadius = maxEffectRadius;
 	}
diff --git a/Engine/Plugins/Mutable/Source/MutableTools/Public/MuT/NodeModifierMeshClipDeform.h b/Engine/Plugins/Mutable/Source/MutableTools/Public/MuT/NodeModifierMeshClipDeform.h
--- a/Engine/Plugins/Mutable/Source/MutableTools/Public/MuT/NodeModifierMeshClipDeform.h
+++ b/Engine/Plugins/Mutable/Source/MutableTools/Public/MuT/NodeModifierMeshClipDeform.h
@@ -22,6 +22,9 @@ namespace mu
 
 		EShapeBindingMethod BindingMethod = EShapeBindingMethod::ClipDeformClosestProject;
 
+		/** */
+		EFaceCullStrategy FaceCullStrategy = EFaceCullStrategy::AllVerticesCulled;
+
 	public:
 
 		// Node interface
diff --git a/Engine/Plugins/Mutable/Source/MutableTools/Public/MuT/NodeModifierMeshClipMorphPlane.h b/Engine/Plugins/Mutable/Source/MutableTools/Public/MuT/NodeModifierMeshClipMorphPlane.h
--- a/Engine/Plugins/Mutable/Source/MutableTools/Public/MuT/NodeModifierMeshClipMorphPlane.h
+++ b/Engine/Plugins/Mutable/Source/MutableTools/Public/MuT/NodeModifierMeshClipMorphPlane.h
@@ -4,6 +4,7 @@
 
 #include "MuR/Ptr.h"
 #include "MuR/RefCounted.h"
+#include "MuR/Types.h"
 #include "MuR/Skeleton.h"
 #include "MuT/Node.h"
 #include "MuT/NodeModifier.h"
@@ -20,10 +21,10 @@ namespace mu
 	{
 		// Morph field parameters
 
-		 //! Distance to the plane of last affected vertex
+		// Distance to the plane of last affected vertex
 		float DistanceToPlane = 0;
 
-		//! "Linearity" factor of the influence.
+		// "Linearity" factor of the influence.
 		float LinearityFactor = 0;
 
 		// Ellipse location
@@ -31,21 +32,8 @@ namespace mu
 		FVector3f Normal = FVector3f(0, 0, 0);
 		float Radius1 = 0, Radius2 = 0, Rotation = 0;
 
-		//! Typed of vertex selection
-		typedef enum
-		{
-			//! All vertices, so no extra info is needed
-			VS_ALL,
-
-			//! Select vertices inside a shape
-			VS_SHAPE,
-
-			//! Select all vertices affected by any bone in a sub hierarchy
-			VS_BONE_HIERARCHY,
-		} EVertexSelection;
-
 		// Vertex selection box
-		uint8 VertexSelectionType = VS_ALL;
+		EClipVertexSelectionType VertexSelectionType = EClipVertexSelectionType::None;
 		FVector3f SelectionBoxOrigin = FVector3f(0, 0, 0);
 		FVector3f SelectionBoxRadius = FVector3f(0, 0, 0);
 		FBoneName VertexSelectionBone;
@@ -53,6 +41,8 @@ namespace mu
 		// Max distance a vertex can have to the bone in order to be affected. A negative value
 		// means no limit.
 		float MaxEffectRadius = -1.0f;
+
+		EFaceCullStrategy FaceCullStrategy = EFaceCullStrategy::AllVerticesCulled;
 	};
 
 
diff --git a/Engine/Plugins/Mutable/Source/MutableTools/Public/MuT/NodeModifierMeshClipWithMesh.h b/Engine/Plugins/Mutable/Source/MutableTools/Public/MuT/NodeModifierMeshClipWithMesh.h
--- a/Engine/Plugins/Mutable/Source/MutableTools/Public/MuT/NodeModifierMeshClipWithMesh.h
+++ b/Engine/Plugins/Mutable/Source/MutableTools/Public/MuT/NodeModifierMeshClipWithMesh.h
@@ -4,6 +4,7 @@
 
 #include "MuR/Ptr.h"
 #include "MuR/RefCounted.h"
+#include "MuR/Types.h"
 #include "MuT/Node.h"
 #include "MuT/NodeModifier.h"
 #include "MuT/NodeMesh.h"
@@ -16,9 +17,12 @@ namespace mu
 	{
 	public:
 
-		//! 
+		/** */
 		Ptr<NodeMesh> ClipMesh;
 
+		/** */
+		EFaceCullStrategy FaceCullStrategy = EFaceCullStrategy::AllVerticesCulled;
+
 	public:
 
 		// Node interface
diff --git a/Engine/Plugins/Mutable/Source/MutableTools/Public/MuT/NodeModifierMeshClipWithUVMask.h b/Engine/Plugins/Mutable/Source/MutableTools/Public/MuT/NodeModifierMeshClipWithUVMask.h
--- a/Engine/Plugins/Mutable/Source/MutableTools/Public/MuT/NodeModifierMeshClipWithUVMask.h
+++ b/Engine/Plugins/Mutable/Source/MutableTools/Public/MuT/NodeModifierMeshClipWithUVMask.h
@@ -27,6 +27,9 @@ namespace mu
 		/** Layout index of the UVs to use inthe source mesh to ben clipped with the mask. */
 		uint8 LayoutIndex = 0;
 
+		/** */
+		EFaceCullStrategy FaceCullStrategy = EFaceCullStrategy::AllVerticesCulled;
+
 	public:
 
 		// Node interface
diff --git a/Engine/Plugins/Mutable/Source/MutableTools/Public/MuT/NodeModifierSurfaceEdit.h b/Engine/Plugins/Mutable/Source/MutableTools/Public/MuT/NodeModifierSurfaceEdit.h
--- a/Engine/Plugins/Mutable/Source/MutableTools/Public/MuT/NodeModifierSurfaceEdit.h
+++ b/Engine/Plugins/Mutable/Source/MutableTools/Public/MuT/NodeModifierSurfaceEdit.h
@@ -58,6 +58,9 @@ namespace mu
 
 		TArray<FLOD> LODs;
 
+		/** For remove operations, use this strategy to cull faces. */
+		EFaceCullStrategy FaceCullStrategy = EFaceCullStrategy::AllVerticesCulled;
+
 		/** Name of the morph to apply to the modified surface if it has it. */
 		FString MeshMorph;
 
