[Mutable] Fix GetIntParameterOptionDataTable Blueprint signature.

#rnx

diff --git a/Engine/Plugins/Mutable/Source/CustomizableObject/Private/MuCO/CustomizableObject.cpp b/Engine/Plugins/Mutable/Source/CustomizableObject/Private/MuCO/CustomizableObject.cpp
--- a/Engine/Plugins/Mutable/Source/CustomizableObject/Private/MuCO/CustomizableObject.cpp
+++ b/Engine/Plugins/Mutable/Source/CustomizableObject/Private/MuCO/CustomizableObject.cpp
@@ -1766,12 +1766,15 @@ FMutableStateUIMetadata UCustomizableObject::GetStateUIMetadata(const FString& S
 
 
 #if WITH_EDITOR
-TSet<TSoftObjectPtr<UDataTable>> UCustomizableObject::GetIntParameterOptionDataTable(FString& ParamName, const FString& OptionName)
+TArray<TSoftObjectPtr<UDataTable>> UCustomizableObject::GetIntParameterOptionDataTable(const FString& ParamName, const FString& OptionName)
 {
 	const FModelResources& ModelResources = GetPrivate()->GetModelResources();
-	const TSet<TSoftObjectPtr<UDataTable>>* Result = ModelResources.IntParameterOptionDataTable.Find(MakeTuple(ParamName, OptionName));
+	if (const TSet<TSoftObjectPtr<UDataTable>>* Result = ModelResources.IntParameterOptionDataTable.Find(MakeTuple(ParamName, OptionName)))
+	{
+		return Result->Array();
+	}
 	
-	return Result ? *Result : TSet<TSoftObjectPtr<UDataTable>>();
+	return {};
 }
 #endif
 
diff --git a/Engine/Plugins/Mutable/Source/CustomizableObject/Public/MuCO/CustomizableObject.h b/Engine/Plugins/Mutable/Source/CustomizableObject/Public/MuCO/CustomizableObject.h
--- a/Engine/Plugins/Mutable/Source/CustomizableObject/Public/MuCO/CustomizableObject.h
+++ b/Engine/Plugins/Mutable/Source/CustomizableObject/Public/MuCO/CustomizableObject.h
@@ -542,7 +542,7 @@ public:
 #if WITH_EDITOR
 	/** Return the DataTables used by the given parameter and its value (if any). */
 	UFUNCTION(BlueprintCallable, Category = CustomizableObject)
-	TSet<TSoftObjectPtr<UDataTable>> GetIntParameterOptionDataTable(FString& ParamName, const FString& OptionName);
+	TArray<TSoftObjectPtr<UDataTable>> GetIntParameterOptionDataTable(const FString& ParamName, const FString& OptionName);
 #endif
 	
 private:
