[Mutable] Moved Participating Objects to its own pass.

- Decopled GenerateMutableRoot from FMutableCompiler.
    - Decoupled ProcessChildObjectsRecursively.
    - Replaced AddCachedReferencers with GetReferencingPackages.
    - Replaced UpdateArrayGCProtect with a simple array.
    - Const UCustomizableObject.
        - AddStreamedResource and MakeStreamedExtensionData now creates containers after the GenerateMutableRoot.
    - Optional FMutableGraphGenerationContext::Compiler.
- Replaced all GenerateMutableRoot LoadSynchronous with LoadObject or LoadClass.
- Added scopes to all GenerateMutable functions.
- Added const UCustomizableObject::GetPrivate.
- GetRefSkeletalMesh will now return the compiled Reference Skeletal Mesh. If not, then will find the user specified mesh.
- Moved GetRootNode(UCustomizableObject ...) to GraphTraversal.cpp.
- Moved ParamNamesToSelectedOptions from FMutableGraphGenerationContex to FCompileOptions.
- Force Graph PostLoad before loading.
- New AddParticipatingObject SoftObjectPtr signatures.
- IsCompilationOutOfDate no longer uses a relaxed superset.
    - Removed DirtyParticipatingObjects.
- IsCompilationOutOfDate now differenciates out of date packaged, removed packages, and added packages.

#rnx

diff --git a/Engine/Plugins/Experimental/HairStrandsMutable/Source/HairStrandsMutableEditor/Private/CustomizableObjectNodeGroomConstant.cpp b/Engine/Plugins/Experimental/HairStrandsMutable/Source/HairStrandsMutableEditor/Private/CustomizableObjectNodeGroomConstant.cpp
--- a/Engine/Plugins/Experimental/HairStrandsMutable/Source/HairStrandsMutableEditor/Private/CustomizableObjectNodeGroomConstant.cpp
+++ b/Engine/Plugins/Experimental/HairStrandsMutable/Source/HairStrandsMutableEditor/Private/CustomizableObjectNodeGroomConstant.cpp
@@ -62,10 +62,12 @@ mu::Ptr<mu::NodeExtensionData> UCustomizableObjectNodeGroomConstant::GenerateMut
 	// an always-loaded constant.
 	UCustomizableObjectResourceDataContainer* Container = nullptr;
 	Result->SetValue(CompilerInterface.MakeStreamedExtensionData(Container));
-	check(Container);
 
-	// Populate instanced struct
-	Container->Data.Data.InitializeAs<FGroomPinData>(GroomData);
+	if (Container)
+	{
+		// Populate instanced struct
+		Container->Data.Data.InitializeAs<FGroomPinData>(GroomData);
+	}
 
 	return Result;
 }
diff --git a/Engine/Plugins/Mutable/Source/CustomizableObject/Internal/MuCO/CustomizableObjectCompilerTypes.h b/Engine/Plugins/Mutable/Source/CustomizableObject/Internal/MuCO/CustomizableObjectCompilerTypes.h
--- a/Engine/Plugins/Mutable/Source/CustomizableObject/Internal/MuCO/CustomizableObjectCompilerTypes.h
+++ b/Engine/Plugins/Mutable/Source/CustomizableObject/Internal/MuCO/CustomizableObjectCompilerTypes.h
@@ -126,6 +126,9 @@ struct FCompilationOptions
 	 * This limit is before any pak or filesystem compression. This limit will be broken if a single piece of data is bigger because data is not fragmented for packaging purposes.
 	 */
 	uint64 DDCBytesLimit = 64 * 1024;
+	
+	/** Stores the only option of an Int Param that should be compiled. */
+	TMap<FString, FString> ParamNamesToSelectedOptions;
 };
 
 
@@ -170,10 +173,7 @@ struct CUSTOMIZABLEOBJECT_API FCompilationRequest
 
 	TArray<FText>& GetWarnings();
 	TArray<FText>& GetErrors();
-
-	void SetParameterNamesToSelectedOptions(const TMap<FString, FString>& InParamNamesToSelectedOptions);
-	const TMap<FString, FString>& GetParameterNamesToSelectedOptions() const;
-
+	
 	bool operator==(const FCompilationRequest& Other) const;
 
 private:
@@ -189,10 +189,7 @@ private:
 
 	UE::DerivedData::ECachePolicy DDCPolicy = UE::DerivedData::ECachePolicy::None;
 	UE::DerivedData::FCacheKey DDCKey;
-
-	// Stores the only option of an Int Param that should be compiled
-	TMap<FString, FString> ParamNamesToSelectedOptions;
-
+	
 	TArray<FText> Warnings;
 	TArray<FText> Errors;
 };
diff --git a/Engine/Plugins/Mutable/Source/CustomizableObject/Internal/MuCO/CustomizableObjectPrivate.h b/Engine/Plugins/Mutable/Source/CustomizableObject/Internal/MuCO/CustomizableObjectPrivate.h
--- a/Engine/Plugins/Mutable/Source/CustomizableObject/Internal/MuCO/CustomizableObjectPrivate.h
+++ b/Engine/Plugins/Mutable/Source/CustomizableObject/Internal/MuCO/CustomizableObjectPrivate.h
@@ -1096,6 +1096,10 @@ struct FModelResources
 	/** Name of all possible components. Index is the ObjectComponentIndex. */
 	UPROPERTY()
 	TArray<FName> ComponentNames;
+
+	/** Version Bridge. */
+	UPROPERTY()
+	FString CompiledVersionBridge;
 	
 #if WITH_EDITORONLY_DATA
 	void CUSTOMIZABLEOBJECT_API Serialize(FObjectAndNameAsStringProxyArchive& Ar, bool bIsCooking);
@@ -1187,17 +1191,16 @@ class CUSTOMIZABLEOBJECT_API UCustomizableObjectPrivate : public UObject
 #endif
 
 public:
-	UCustomizableObjectPrivate();
-
 	/** Must be called after unlocking the CustomizableObject. */
 	void SetModel(const TSharedPtr<mu::Model, ESPMode::ThreadSafe>& Model, const FGuid Identifier);
 	const TSharedPtr<mu::Model, ESPMode::ThreadSafe>& GetModel();
-	TSharedPtr<const mu::Model, ESPMode::ThreadSafe> GetModel() const;
+	const TSharedPtr<const mu::Model, ESPMode::ThreadSafe> GetModel() const;
 
 	const FModelResources& GetModelResources() const;
 
 #if WITH_EDITORONLY_DATA
 	FModelResources& GetModelResources(bool bIsCooking);
+	const FModelResources& GetModelResources(bool bIsCooking) const;
 
 	void SetModelStreamableBulkData(const TSharedPtr<FModelStreamableBulkData>& StreamableData, bool bIsCooking);
 #endif
@@ -1235,7 +1238,7 @@ public:
 	FString GetCompiledDataFileName(bool bIsModel, const ITargetPlatform* InTargetPlatform = nullptr, bool bIsDiskStreamer = false);
 
 	/** DDC helpers. BuildDerivedDataKey is expensive, try to cache it as much as possible. */
-	FString BuildDerivedDataKey(FCompilationOptions Options);
+	TArray<uint8> BuildDerivedDataKey(FCompilationOptions Options);
 	UE::DerivedData::FCacheKey GetDerivedDataCacheKeyForOptions(FCompilationOptions InOptions);
 
 	/** Attempts to load the compiled data from DDC. Builds key if not supplied. */
@@ -1295,7 +1298,7 @@ public:
 	
 #if WITH_EDITOR
 	/** See ICustomizableObjectEditorModule::IsCompilationOutOfDate. */
-	bool IsCompilationOutOfDate(TArray<FName>* OutOfDatePackages = nullptr) const;
+	bool IsCompilationOutOfDate(bool bSkipIndirectReferences, TArray<FName>& OutOfDatePackages, TArray<FName>& AddedPackages, TArray<FName>& RemovedPackages, bool& bVersionDiff) const;
 
 	void OnParticipatingObjectDirty(UPackage* Package, bool);
 #endif
@@ -1311,9 +1314,11 @@ public:
 #endif
 	
 	TArray<FCustomizableObjectResourceData>& GetAlwaysLoadedExtensionData();
+	const TArray<FCustomizableObjectResourceData>& GetAlwaysLoadedExtensionData() const;
 
 	TArray<FCustomizableObjectStreamedResourceData>& GetStreamedExtensionData();
-	
+	const TArray<FCustomizableObjectStreamedResourceData>& GetStreamedExtensionData() const;
+
 	const FCustomizableObjectResourceData* LoadStreamedResource(int32 ResourceIndex);
 	void UnloadStreamedResource(int32 ResourceIndex);
 
@@ -1342,10 +1347,7 @@ public:
 
 	/** Cook requests. */
 	TArray<TSharedRef<FCompilationRequest>> CompileRequests;
-
-	/** List of Participating Objects (packages) has been marked as dirty since the last compilation. */
-	TArray<FName> DirtyParticipatingObjects;
-
+	
 	ECompilationStatePrivate CompilationState = ECompilationStatePrivate::None;
 	ECompilationResultPrivate CompilationResult = ECompilationResultPrivate::Unknown;
 	
@@ -1472,9 +1474,14 @@ public:
 
 		AddedFaceCullStrategyToSomeOperations,
 
+		DDCParticipatingObjects,
+		
 		// -----<new versions can be added above this line>--------
 		LastCustomizableObjectVersion
 	};
+	
 	static constexpr int32 CurrentSupportedVersion = ECustomizableObjectVersions::LastCustomizableObjectVersion;
+	
+	static constexpr int32 DerivedDataVersion = 0;
 };
 
diff --git a/Engine/Plugins/Mutable/Source/CustomizableObject/Internal/MuCO/ICustomizableObjectEditorModule.h b/Engine/Plugins/Mutable/Source/CustomizableObject/Internal/MuCO/ICustomizableObjectEditorModule.h
--- a/Engine/Plugins/Mutable/Source/CustomizableObject/Internal/MuCO/ICustomizableObjectEditorModule.h
+++ b/Engine/Plugins/Mutable/Source/CustomizableObject/Internal/MuCO/ICustomizableObjectEditorModule.h
@@ -18,6 +18,7 @@ class FExtensibilityManager;
 class FBakeOperationCompletedDelegate;
 struct FBakingConfiguration;
 struct FCompilationRequest;
+struct FCompilationOptions;
 
 extern const FName CustomizableObjectEditorAppIdentifier;
 extern const FName CustomizableObjectInstanceEditorAppIdentifier;
@@ -49,9 +50,13 @@ public:
 	
 	/** Return if the CO is not compiled or the ParticipatingObjects system has detected a change (participating objects dirty or re-saved since last compilation).
 	  * @param Object object to check.
+	  * @param bSkipIndirectReferences if true, do not check for added/removed indirect references.
 	  * @param OutOfDatePackages list of out of date packages.
-   	  * @return true if the compilation is out of date. */
-	virtual bool IsCompilationOutOfDate(const UCustomizableObject& Object, TArray<FName>* OutOfDatePackages = nullptr) const = 0;
+	  * @param AddedPackages list of added packages since the last compilation.
+	  * @param RemovedPackages list of removed packages since the last compilation.
+	  * @param bVersionDiff true if the Version Bridge has changed since the last compilation.
+	  * @return true if the compilation is out of date. */
+	virtual bool IsCompilationOutOfDate(const UCustomizableObject& Object, bool bSkipIndirectReferences, TArray<FName>& OutOfDatePackages, TArray<FName>& AddedPackages, TArray<FName>& RemovedPackages, bool& bVersionDiff) const = 0;
 
 	/** See GraphTraversal::IsRootObject(...) */
 	virtual bool IsRootObject(const UCustomizableObject& Object) const = 0;
@@ -60,6 +65,12 @@ public:
 	  * @return Current version as string. */
 	virtual FString GetCurrentContentVersionForObject(const UCustomizableObject& Object) const = 0;
 
+	/** See GraphTraversal::GetRootObject(...) */
+	virtual UCustomizableObject* GetRootObject(UCustomizableObject* ChildObject) const = 0;
+
+	/** See GraphTraversal::GetRootObject(...) */
+	virtual const UCustomizableObject* GetRootObject(const UCustomizableObject* ChildObject) const = 0;
+	
 	/**
 	 * Execute this method in order to bake the provided instance. It will schedule a special type of instance update before proceeding with the bake itself.
 	 * @param InTargetInstance The instance we want to bake
@@ -82,5 +93,9 @@ public:
 	virtual int32 GetNumCompileRequests() = 0;
 
 	virtual USkeletalMesh* GetReferenceSkeletalMesh(const UCustomizableObject& Object, const FName& Component) const = 0;
+	
+	/** Perform a fast compilation pass to get all participating objects.
+	 *  @param bLoadObjects Load any object. If false, no objects will load. If true, only objects strictly required to get the full list of participating objects will load. */
+	virtual TMap<FName, FGuid> GetParticipatingObjects(const UCustomizableObject* Object, bool bLoadObjects, const FCompilationOptions* Options = nullptr) const = 0;
 };
  
diff --git a/Engine/Plugins/Mutable/Source/CustomizableObject/Private/MuCO/CustomizableObject.cpp b/Engine/Plugins/Mutable/Source/CustomizableObject/Private/MuCO/CustomizableObject.cpp
--- a/Engine/Plugins/Mutable/Source/CustomizableObject/Private/MuCO/CustomizableObject.cpp
+++ b/Engine/Plugins/Mutable/Source/CustomizableObject/Private/MuCO/CustomizableObject.cpp
@@ -315,6 +315,13 @@ void UCustomizableObject::PostLoad()
 {
 	Super::PostLoad();
 
+#if	WITH_EDITORONLY_DATA
+	if (Source)
+	{
+		Source->ConditionalPostLoad();
+	}
+#endif
+	
 	const int32 CustomizableObjectCustomVersion = GetLinkerCustomVersion(FCustomizableObjectCustomVersion::GUID);
 
 #if WITH_EDITOR
@@ -743,6 +750,7 @@ void FModelResources::Serialize(FObjectAndNameAsStringProxyArchive& MemoryWriter
 	MemoryWriter << FirstLODAvailable;
 
 	MemoryWriter << ComponentNames;
+	MemoryWriter << CompiledVersionBridge;
 
 	// Editor Only data
 	if (!bIsCooking)
@@ -894,6 +902,7 @@ bool FModelResources::Unserialize(FObjectAndNameAsStringProxyArchive& MemoryRead
 	MemoryReader << FirstLODAvailable;
 
 	MemoryReader << ComponentNames;
+	MemoryReader << CompiledVersionBridge;
 
 	// Editor Only data
 	if (!bIsCooking)
@@ -967,18 +976,24 @@ void UCustomizableObjectPrivate::LoadCompiledDataFromDisk()
 				CompiledDataFileHandle->Read(CompiledDataBytes.GetData(), CompiledDataSize);
 
 				FMemoryReaderView MemoryReader(CompiledDataBytes);
-
-				DirtyParticipatingObjects.Empty();
-
-				TArray<FName> OutOfDatePackages;
-				if (LoadModelResources(MemoryReader, RunningPlatform) && !IsCompilationOutOfDate(&OutOfDatePackages))
+				
+				if (LoadModelResources(MemoryReader, RunningPlatform))
 				{
-					LoadModelStreamableBulk(MemoryReader, /* bIsCooking */false);
-					LoadModel(MemoryReader);
-				}
-				else if (!OutOfDatePackages.IsEmpty())
-				{
-					UE_LOG(LogMutable, Display, TEXT("Invalidating compiled data due to changes in %s."), *OutOfDatePackages[0].ToString());
+					TArray<FName> OutOfDatePackages;
+					TArray<FName> AddedPackages;
+					TArray<FName> RemovedPackages;
+					bool bVersionDiff;
+					const bool bOutOfDate = IsCompilationOutOfDate(false, OutOfDatePackages, AddedPackages, RemovedPackages, bVersionDiff);
+					if (!bOutOfDate)
+					{
+						LoadModelStreamableBulk(MemoryReader, /* bIsCooking */false);
+						LoadModel(MemoryReader);
+					}
+					else
+					{
+						UE_LOG(LogMutable, Display, TEXT("Invalidating compiled data due to changes in %s."), *OutOfDatePackages[0].ToString());
+						PrintParticipatingPackagesDiff(OutOfDatePackages, AddedPackages, RemovedPackages, bVersionDiff);
+					}
 				}
 			}
 		}
@@ -1193,7 +1208,14 @@ void UCustomizableObjectPrivate::LoadEmbeddedData(FArchive& Ar)
 }
 
 
-UCustomizableObjectPrivate* UCustomizableObject::GetPrivate() const
+const UCustomizableObjectPrivate* UCustomizableObject::GetPrivate() const
+{
+	check(Private);
+	return Private;
+}
+
+
+UCustomizableObjectPrivate* UCustomizableObject::GetPrivate()
 {
 	check(Private);
 	return Private;
@@ -1368,7 +1390,7 @@ UCustomizableObject* UCustomizableObjectPrivate::GetPublic() const
 }
 
 #if WITH_EDITORONLY_DATA
-FPostCompileDelegate& UCustomizableObject::GetPostCompileDelegate() const
+FPostCompileDelegate& UCustomizableObject::GetPostCompileDelegate()
 {
 	return GetPrivate()->PostCompileDelegate;
 }
@@ -1788,7 +1810,7 @@ float UCustomizableObject::GetFloatParameterDefaultValue(const FString& InParame
 		return FCustomizableObjectFloatParameterValue::DEFAULT_PARAMETER_VALUE;
 	}
 
-	const TSharedPtr<mu::Model> Model = GetPrivate()->GetModel();
+	const TSharedPtr<const mu::Model>& Model = GetPrivate()->GetModel();
 	if (!Model)
 	{
 		checkNoEntry();
@@ -1808,7 +1830,7 @@ int32 UCustomizableObject::GetIntParameterDefaultValue(const FString& InParamete
 		return FCustomizableObjectIntParameterValue::DEFAULT_PARAMETER_VALUE;
 	}
 
-	const TSharedPtr<mu::Model> Model = GetPrivate()->GetModel();
+	const TSharedPtr<const mu::Model>& Model = GetPrivate()->GetModel();
 	if (!Model)
 	{
 		checkNoEntry();
@@ -1828,7 +1850,7 @@ bool UCustomizableObject::GetBoolParameterDefaultValue(const FString& InParamete
 		return FCustomizableObjectBoolParameterValue::DEFAULT_PARAMETER_VALUE;
 	}
 
-	const TSharedPtr<mu::Model> Model = GetPrivate()->GetModel();
+	const TSharedPtr<const mu::Model>& Model = GetPrivate()->GetModel();
 	if (!Model)
 	{
 		checkNoEntry();
@@ -1848,7 +1870,7 @@ FLinearColor UCustomizableObject::GetColorParameterDefaultValue(const FString& I
 		return FCustomizableObjectVectorParameterValue::DEFAULT_PARAMETER_VALUE;
 	}
 
-	const TSharedPtr<mu::Model> Model = GetPrivate()->GetModel();
+	const TSharedPtr<const mu::Model>& Model = GetPrivate()->GetModel();
 	if (!Model)
 	{
 		checkNoEntry();
@@ -1886,7 +1908,7 @@ FCustomizableObjectProjector UCustomizableObject::GetProjectorParameterDefaultVa
 		return FCustomizableObjectProjectorParameterValue::DEFAULT_PARAMETER_VALUE;
 	}
 
-	const TSharedPtr<mu::Model> Model = GetPrivate()->GetModel();
+	const TSharedPtr<const mu::Model>& Model = GetPrivate()->GetModel();
 	if (!Model)
 	{
 		checkNoEntry();
@@ -1911,7 +1933,7 @@ FName UCustomizableObject::GetTextureParameterDefaultValue(const FString& InPara
 		return FCustomizableObjectTextureParameterValue::DEFAULT_PARAMETER_VALUE;
 	}
 
-	const TSharedPtr<mu::Model> Model = GetPrivate()->GetModel();
+	const TSharedPtr<const mu::Model>& Model = GetPrivate()->GetModel();
 	if (!Model)
 	{
 		checkNoEntry();
@@ -2185,14 +2207,6 @@ void UModelStreamableData::Serialize(FArchive& Ar)
 }
 
 
-UCustomizableObjectPrivate::UCustomizableObjectPrivate()
-{
-#if WITH_EDITOR
-	UPackage::PackageMarkedDirtyEvent.AddUObject(this, &UCustomizableObjectPrivate::OnParticipatingObjectDirty);
-#endif
-}
-
-
 void UCustomizableObjectPrivate::SetModel(const TSharedPtr<mu::Model, ESPMode::ThreadSafe>& Model, const FGuid Id)
 {
 	if (MutableModel == Model
@@ -2229,7 +2243,7 @@ const TSharedPtr<mu::Model, ESPMode::ThreadSafe>& UCustomizableObjectPrivate::Ge
 }
 
 
-TSharedPtr<const mu::Model, ESPMode::ThreadSafe> UCustomizableObjectPrivate::GetModel() const
+const TSharedPtr<const mu::Model, ESPMode::ThreadSafe> UCustomizableObjectPrivate::GetModel() const
 {
 	return MutableModel;
 }
@@ -2280,6 +2294,13 @@ const FModelResources& UCustomizableObjectPrivate::GetModelResources() const
 
 #if WITH_EDITORONLY_DATA
 FModelResources& UCustomizableObjectPrivate::GetModelResources(bool bIsCooking)
+{
+	const UCustomizableObjectPrivate* ConstThis = this;
+	return *const_cast<FModelResources*>(&ConstThis->GetModelResources(bIsCooking));
+}
+
+
+const FModelResources& UCustomizableObjectPrivate::GetModelResources(bool bIsCooking) const
 {
 	return bIsCooking ? ModelResources : ModelResourcesEditor;
 }
@@ -2287,24 +2308,15 @@ FModelResources& UCustomizableObjectPrivate::GetModelResources(bool bIsCooking)
 
 
 #if WITH_EDITOR
-bool UCustomizableObjectPrivate::IsCompilationOutOfDate(TArray<FName>* OutOfDatePackages) const
+bool UCustomizableObjectPrivate::IsCompilationOutOfDate(bool bSkipIndirectReferences, TArray<FName>& OutOfDatePackages, TArray<FName>& AddedPackages, TArray<FName>& RemovedPackages, bool& bVersionDiff) const
 {
 	if (const ICustomizableObjectEditorModule* Module = ICustomizableObjectEditorModule::Get())
 	{
-		return Module->IsCompilationOutOfDate(*GetPublic(), OutOfDatePackages);
+		return Module->IsCompilationOutOfDate(*GetPublic(), bSkipIndirectReferences, OutOfDatePackages, AddedPackages, RemovedPackages, bVersionDiff);
 	}
 
 	return false;		
 }
-
-
-void UCustomizableObjectPrivate::OnParticipatingObjectDirty(UPackage* Package, bool)
-{
-	if (GetModelResources().ParticipatingObjects.Contains(Package->GetFName()))
-	{
-		DirtyParticipatingObjects.AddUnique(Package->GetFName());
-	}
-}
 #endif
 
 
@@ -2339,6 +2351,11 @@ TArray<FCustomizableObjectResourceData>& UCustomizableObjectPrivate::GetAlwaysLo
 	return GetPublic()->AlwaysLoadedExtensionData;
 }
 
+const TArray<FCustomizableObjectResourceData>& UCustomizableObjectPrivate::GetAlwaysLoadedExtensionData() const
+{
+	return GetPublic()->AlwaysLoadedExtensionData;
+}
+
 
 TArray<FCustomizableObjectStreamedResourceData>& UCustomizableObjectPrivate::GetStreamedExtensionData()
 {
@@ -2346,6 +2363,12 @@ TArray<FCustomizableObjectStreamedResourceData>& UCustomizableObjectPrivate::Get
 }
 
 
+const TArray<FCustomizableObjectStreamedResourceData>& UCustomizableObjectPrivate::GetStreamedExtensionData() const
+{
+	return GetPublic()->StreamedExtensionData;
+}
+
+
 const FCustomizableObjectResourceData* UCustomizableObjectPrivate::LoadStreamedResource(int32 ResourceIndex)
 {
 #if WITH_EDITORONLY_DATA
@@ -2425,7 +2448,7 @@ FCompilationOptions UCustomizableObjectPrivate::GetCompileOptions() const
 	Options.bPhysicsAssetMergeEnabled = GetPublic()->bEnablePhysicsAssetMerge;
 	Options.bAnimBpPhysicsManipulationEnabled = GetPublic()->bEnableAnimBpPhysicsAssetsManipualtion;
 	Options.ImageTiling = ImageTiling;
-	
+
 	return Options;
 }
 #endif
@@ -2792,59 +2815,76 @@ void SerializeCompilationOptionsForDDC(FArchive& Ar, FCompilationOptions& Option
 	Ar << Options.bPhysicsAssetMergeEnabled;
 	Ar << Options.bAnimBpPhysicsManipulationEnabled;
 	Ar << Options.ImageTiling;
+	Ar << Options.ParamNamesToSelectedOptions;
 }
 
 
-const FString& GetCustomizableObjectDerivedDataVersion()
-{
-	static FString CachedVersionString = TEXT("53DD1D7C-C040-4A52-162B-A3F03E2E8320");
-	return CachedVersionString;
-}
 
-
-FString UCustomizableObjectPrivate::BuildDerivedDataKey(FCompilationOptions Options)
+TArray<uint8> UCustomizableObjectPrivate::BuildDerivedDataKey(FCompilationOptions Options)
 {
-	FString KeySuffix;
-
-	TArray<uint8> TempBytes;
-	TempBytes.Reserve(160);
-
-	KeySuffix += FString::Printf(TEXT("%08x-"), (uint32)CurrentSupportedVersion);
+	MUTABLE_CPUPROFILER_SCOPE(UCustomizableObjectPrivate::BuildDerivedDataKey)
 
 	check(IsInGameThread());
+	
 	UCustomizableObject& CustomizableObject = *GetPublic();
 
+	TArray<uint8> Bytes;
+	FMemoryWriter Ar(Bytes, /*bIsPersistent=*/ true);
+	
+	{
+		uint32 Version = DerivedDataVersion;
+		Ar << Version;
+	}
+	
+	{
+		int32 CurrentVersion = CurrentSupportedVersion;
+		Ar << CurrentVersion;	
+	}
+
+	// Custom Version
+	{
+		int32 CustomVersion = GetLinkerCustomVersion(FCustomizableObjectCustomVersion::GUID);
+		Ar << CustomVersion;
+	}
+	
+
 	// Customizable Object Ids
-	KeySuffix += GenerateIdentifier(CustomizableObject).ToString();
-	KeySuffix += CustomizableObject.VersionId.ToString();
+	{
+		FGuid Id = GenerateIdentifier(CustomizableObject);
+		Ar << Id;
+	}
+
+	{
+		FGuid Version = CustomizableObject.VersionId;
+		Ar << Version;
+	}
 
 	// Compile Options
 	{
-		TempBytes.Reset();
-		FMemoryWriter Ar(TempBytes, /*bIsPersistent=*/ true);
 		SerializeCompilationOptionsForDDC(Ar, Options);
-		KeySuffix += FString::Printf(TEXT("OPT_%016x"), CityHash64(reinterpret_cast<const char*>(TempBytes.GetData()), TempBytes.Num()));
 	}
 
-	// Content Version
+	// Version Bridge
 	if (const ICustomizableObjectEditorModule* Module = ICustomizableObjectEditorModule::Get())
 	{
-		KeySuffix += TEXT("VER_") + Module->GetCurrentContentVersionForObject(CustomizableObject);
+		FString Version = Module->GetCurrentContentVersionForObject(CustomizableObject);
+		Ar << Version;
 	}
 
 	// Participating objects hash
+	if (ICustomizableObjectEditorModule* Module = ICustomizableObjectEditorModule::Get())
 	{
-		TempBytes.Reset();
-		FMemoryWriter Ar(TempBytes, /*bIsPersistent=*/ true);
-		//Ar << CachedParticipatingObjectsHash; // TODO PERE: Build ParticipatingObjects Hash
-		KeySuffix += FString::FromBlob(TempBytes.GetData(), TempBytes.Num() * TempBytes.GetTypeSize());
+		TMap<FName, FGuid> ParticipatingObjects = Module->GetParticipatingObjects(GetPublic(), true, &Options);
+		for (TTuple<FName, FGuid>& Tuple : ParticipatingObjects)
+		{
+			Ar << Tuple.Get<0>();
+			Ar << Tuple.Get<1>();
+		}
 	}
 
-	return FDerivedDataCacheInterface::BuildCacheKey(
-		TEXT("CUSTOMIZABLEOBJECT"),
-		*GetCustomizableObjectDerivedDataVersion(),
-		*KeySuffix
-	);
+	// TODO List of plugins and their custom versions
+
+	return Bytes;
 }
 
 
@@ -2852,12 +2892,11 @@ UE::DerivedData::FCacheKey UCustomizableObjectPrivate::GetDerivedDataCacheKeyFor
 {
 	using namespace UE::DerivedData;
 
-	// Cache key as string
-	FString DerivedDataKey = BuildDerivedDataKey(InOptions);
+	TArray<uint8> DerivedDataKey = BuildDerivedDataKey(InOptions);
 
 	FCacheKey CacheKey;
 	CacheKey.Bucket = FCacheBucket(TEXT("CustomizableObject"));
-	CacheKey.Hash = FIoHashBuilder::HashBuffer(MakeMemoryView(FTCHARToUTF8(DerivedDataKey)));
+	CacheKey.Hash = FIoHashBuilder::HashBuffer(MakeMemoryView(DerivedDataKey));
 	return CacheKey;
 }
 
@@ -3541,18 +3580,6 @@ TArray<FText>& FCompilationRequest::GetErrors()
 }
 
 
-void FCompilationRequest::SetParameterNamesToSelectedOptions(const TMap<FString, FString>& InParamNamesToSelectedOptions)
-{
-	ParamNamesToSelectedOptions = InParamNamesToSelectedOptions;
-}
-
-
-const TMap<FString, FString>& FCompilationRequest::GetParameterNamesToSelectedOptions() const
-{
-	return ParamNamesToSelectedOptions;
-}
-
-
 bool FCompilationRequest::operator==(const FCompilationRequest& Other) const
 {
 	return CustomizableObject == Other.CustomizableObject && Options.TargetPlatform == Other.Options.TargetPlatform;
diff --git a/Engine/Plugins/Mutable/Source/CustomizableObject/Private/MuCO/CustomizableObjectModule.cpp b/Engine/Plugins/Mutable/Source/CustomizableObject/Private/MuCO/CustomizableObjectModule.cpp
--- a/Engine/Plugins/Mutable/Source/CustomizableObject/Private/MuCO/CustomizableObjectModule.cpp
+++ b/Engine/Plugins/Mutable/Source/CustomizableObject/Private/MuCO/CustomizableObjectModule.cpp
@@ -331,4 +331,46 @@ void FCustomizableObjectModule::ToggleDGGUI(const TArray<FString>& Arguments)
 	{
 		UDGGUI::OpenDGGUI(SlotID, SelectedCustomizableObjectInstanceUsage, CurrentWorld, PlayerIndex);
 	}
-}
\ No newline at end of file
+}
+
+
+void PrintParticipatingPackagesDiff(const TArray<FName>& OutOfDatePackages, const TArray<FName>& AddedPackages, const TArray<FName>& RemovedPackages, bool bVersionDiff)
+{
+	constexpr int32 MaxLogLines = 10;
+
+	if (bVersionDiff)
+	{
+		UE_LOG(LogMutable, Display, TEXT("Version Bridge changed."));
+	}
+	
+	if (OutOfDatePackages.Num())
+	{
+		UE_LOG(LogMutable, Display, TEXT("Packages out of date (%i):"), OutOfDatePackages.Num());
+	}
+				
+	for (int32 Index = 0; Index < FMath::Min(MaxLogLines, OutOfDatePackages.Num()); ++Index)
+	{
+		UE_LOG(LogMutable, Display, TEXT("%s"), *OutOfDatePackages[Index].ToString());
+	}
+
+	if (AddedPackages.Num())
+	{
+		UE_LOG(LogMutable, Display, TEXT("Added packages (%i):"), AddedPackages.Num());
+	}
+				
+	for (int32 Index = 0; Index < FMath::Min(MaxLogLines, AddedPackages.Num()); ++Index)
+	{
+		UE_LOG(LogMutable, Display, TEXT("%s"), *AddedPackages[Index].ToString());
+	}
+
+	if (RemovedPackages.Num())
+	{
+		UE_LOG(LogMutable, Display, TEXT("Removed packages (%i):"), RemovedPackages.Num());
+	}
+				
+	for (int32 Index = 0; Index < FMath::Min(MaxLogLines, RemovedPackages.Num()); ++Index)
+	{
+		UE_LOG(LogMutable, Display, TEXT("%s"), *RemovedPackages[Index].ToString());
+	}
+}
+
diff --git a/Engine/Plugins/Mutable/Source/CustomizableObject/Private/MuCO/CustomizableObjectSkeletalMesh.cpp b/Engine/Plugins/Mutable/Source/CustomizableObject/Private/MuCO/CustomizableObjectSkeletalMesh.cpp
--- a/Engine/Plugins/Mutable/Source/CustomizableObject/Private/MuCO/CustomizableObjectSkeletalMesh.cpp
+++ b/Engine/Plugins/Mutable/Source/CustomizableObject/Private/MuCO/CustomizableObjectSkeletalMesh.cpp
@@ -20,7 +20,7 @@ FName GenerateUniqueNameFromCOInstance(const UCustomizableObjectInstance& Instan
 
 
 UCustomizableObjectSkeletalMesh* UCustomizableObjectSkeletalMesh::CreateSkeletalMesh(const TSharedRef<FUpdateContextPrivate>& OperationData, 
-	const UCustomizableObjectInstance& Instance, const UCustomizableObject& CustomizableObject, const int32 InstanceComponentIndex)
+	const UCustomizableObjectInstance& Instance, UCustomizableObject& CustomizableObject, const int32 InstanceComponentIndex)
 {
 	FName SkeletalMeshName = GenerateUniqueNameFromCOInstance(Instance);
 	UCustomizableObjectSkeletalMesh* OutSkeletalMesh = NewObject<UCustomizableObjectSkeletalMesh>(GetTransientPackage(), SkeletalMeshName, RF_Transient);
diff --git a/Engine/Plugins/Mutable/Source/CustomizableObject/Private/MuCO/CustomizableObjectSystem.cpp b/Engine/Plugins/Mutable/Source/CustomizableObject/Private/MuCO/CustomizableObjectSystem.cpp
--- a/Engine/Plugins/Mutable/Source/CustomizableObject/Private/MuCO/CustomizableObjectSystem.cpp
+++ b/Engine/Plugins/Mutable/Source/CustomizableObject/Private/MuCO/CustomizableObjectSystem.cpp
@@ -1354,7 +1354,7 @@ void UCustomizableObjectSystem::ClearResourceCacheProtected()
 
 
 #if WITH_EDITOR
-bool UCustomizableObjectSystem::LockObject(const class UCustomizableObject* InObject)
+bool UCustomizableObjectSystem::LockObject(UCustomizableObject* InObject)
 {
 	check(InObject != nullptr);
 	check(InObject->GetPrivate());
@@ -1440,7 +1440,7 @@ bool UCustomizableObjectSystem::LockObject(const class UCustomizableObject* InOb
 }
 
 
-void UCustomizableObjectSystem::UnlockObject(const class UCustomizableObject* Obj)
+void UCustomizableObjectSystem::UnlockObject(UCustomizableObject* Obj)
 {
 	check(Obj != nullptr);
 	check(Obj->GetPrivate());
@@ -3211,7 +3211,7 @@ namespace impl
 		UCustomizableObjectSystem* System = UCustomizableObjectSystem::GetInstanceChecked();
 		UCustomizableObjectSystemPrivate* SystemPrivate = System->GetPrivate();
 
-		const UCustomizableObject* CustomizableObject = Operation->Instance->GetCustomizableObject();
+		UCustomizableObject* CustomizableObject = Operation->Instance->GetCustomizableObject();
 		UCustomizableObjectPrivate* CustomizableObjectPrivate = CustomizableObject->GetPrivate();
 		
 		for (const TArray<mu::FResourceID>& MeshId : Operation->MeshDescriptors)
diff --git a/Engine/Plugins/Mutable/Source/CustomizableObject/Public/MuCO/CustomizableObject.h b/Engine/Plugins/Mutable/Source/CustomizableObject/Public/MuCO/CustomizableObject.h
--- a/Engine/Plugins/Mutable/Source/CustomizableObject/Public/MuCO/CustomizableObject.h
+++ b/Engine/Plugins/Mutable/Source/CustomizableObject/Public/MuCO/CustomizableObject.h
@@ -577,14 +577,16 @@ private:
 	
 public:
 #if WITH_EDITORONLY_DATA
-	FPostCompileDelegate& GetPostCompileDelegate() const;
+	FPostCompileDelegate& GetPostCompileDelegate();
 #endif
 
 	/** Create a new instance of this object. The instance parameters will be initialized with the object default values. */
 	UFUNCTION(BlueprintCallable, Category = CustomizableObject)
 	UCustomizableObjectInstance* CreateInstance();
 
-	UCustomizableObjectPrivate* GetPrivate() const;
+	const UCustomizableObjectPrivate* GetPrivate() const;
+
+	UCustomizableObjectPrivate* GetPrivate();
 
 	/** Check if the CustomizableObject asset has been compiled. This will always be true in a packaged game, but it could be false in the editor. */
 	UFUNCTION(BlueprintCallable, Category = CustomizableObject)
diff --git a/Engine/Plugins/Mutable/Source/CustomizableObject/Public/MuCO/CustomizableObjectSkeletalMesh.h b/Engine/Plugins/Mutable/Source/CustomizableObject/Public/MuCO/CustomizableObjectSkeletalMesh.h
--- a/Engine/Plugins/Mutable/Source/CustomizableObject/Public/MuCO/CustomizableObjectSkeletalMesh.h
+++ b/Engine/Plugins/Mutable/Source/CustomizableObject/Public/MuCO/CustomizableObjectSkeletalMesh.h
@@ -38,7 +38,7 @@ class CUSTOMIZABLEOBJECT_API UCustomizableObjectSkeletalMesh : public USkeletalM
 public: // Public Methods
 
 	static UCustomizableObjectSkeletalMesh* CreateSkeletalMesh(const TSharedRef<FUpdateContextPrivate>& InOperationData, 
-		const UCustomizableObjectInstance& InInstance, const UCustomizableObject& CustomizableObject, const int32 InstanceComponentIndex);
+		const UCustomizableObjectInstance& InInstance, UCustomizableObject& CustomizableObject, const int32 InstanceComponentIndex);
 
 private: // Private Methods
 
diff --git a/Engine/Plugins/Mutable/Source/CustomizableObject/Public/MuCO/CustomizableObjectSystem.h b/Engine/Plugins/Mutable/Source/CustomizableObject/Public/MuCO/CustomizableObjectSystem.h
--- a/Engine/Plugins/Mutable/Source/CustomizableObject/Public/MuCO/CustomizableObjectSystem.h
+++ b/Engine/Plugins/Mutable/Source/CustomizableObject/Public/MuCO/CustomizableObjectSystem.h
@@ -210,8 +210,8 @@ public:
 	// Lock a CustomizableObjects, preventing the generation or update of any of its instances
 	// Will return true if successful, false if it fails to lock because an update is already underway
 	// This is usually only used in the editor
-	bool LockObject(const UCustomizableObject*);
-	void UnlockObject(const UCustomizableObject*);
+	bool LockObject(UCustomizableObject*);
+	void UnlockObject(UCustomizableObject*);
 
 	/** Checks if there are any outstanding disk or mip update operations in flight for the parameter Customizable Object that may
 	* make it unsafe to compile at the moment.
diff --git a/Engine/Plugins/Mutable/Source/CustomizableObject/Public/MuCO/ICustomizableObjectModule.h b/Engine/Plugins/Mutable/Source/CustomizableObject/Public/MuCO/ICustomizableObjectModule.h
--- a/Engine/Plugins/Mutable/Source/CustomizableObject/Public/MuCO/ICustomizableObjectModule.h
+++ b/Engine/Plugins/Mutable/Source/CustomizableObject/Public/MuCO/ICustomizableObjectModule.h
@@ -83,3 +83,6 @@ public:
 /* Utility function for commands*/
 CUSTOMIZABLEOBJECT_API
 UCustomizableObjectInstanceUsage* GetPlayerCustomizableObjectInstanceUsage(const int32 SlotID, const UWorld* CurrentWorld, const int32 PlayerIndex = 0);
+
+
+CUSTOMIZABLEOBJECT_API void PrintParticipatingPackagesDiff(const TArray<FName>& OutOfDatePackages, const TArray<FName>& AddedPackages, const TArray<FName>& RemovedPackages, bool bVersionDiff);
\ No newline at end of file
diff --git a/Engine/Plugins/Mutable/Source/CustomizableObjectEditor/Private/MuCOE/CustomizableObjectCompileRunnable.cpp b/Engine/Plugins/Mutable/Source/CustomizableObjectEditor/Private/MuCOE/CustomizableObjectCompileRunnable.cpp
--- a/Engine/Plugins/Mutable/Source/CustomizableObjectEditor/Private/MuCOE/CustomizableObjectCompileRunnable.cpp
+++ b/Engine/Plugins/Mutable/Source/CustomizableObjectEditor/Private/MuCOE/CustomizableObjectCompileRunnable.cpp
@@ -251,6 +251,8 @@ const TArray<FCustomizableObjectCompileRunnable::FError>& FCustomizableObjectCom
 
 void FCustomizableObjectCompileRunnable::Tick()
 {
+	MUTABLE_CPUPROFILER_SCOPE(FCustomizableObjectCompileRunnable::Tick);
+
 	check(IsInGameThread());
 
 	constexpr double MaxSecondsPerFrame = 0.4;
diff --git a/Engine/Plugins/Mutable/Source/CustomizableObjectEditor/Private/MuCOE/CustomizableObjectCompiler.cpp b/Engine/Plugins/Mutable/Source/CustomizableObjectEditor/Private/MuCOE/CustomizableObjectCompiler.cpp
--- a/Engine/Plugins/Mutable/Source/CustomizableObjectEditor/Private/MuCOE/CustomizableObjectCompiler.cpp
+++ b/Engine/Plugins/Mutable/Source/CustomizableObjectEditor/Private/MuCOE/CustomizableObjectCompiler.cpp
@@ -42,8 +42,6 @@ class UTexture2D;
 #define UE_MUTABLE_PRELOAD_REGION		TEXT("Mutable Preload")
 #define UE_MUTABLE_SAVEDD_REGION		TEXT("Mutable SaveDD")
 
-UCustomizableObjectNodeObject* GetRootNode(UCustomizableObject* Object, bool &bOutMultipleBaseObjectsFound);
-
 
 bool FCustomizableObjectCompiler::Tick(bool bBlocking)
 {
@@ -121,8 +119,12 @@ void FCustomizableObjectCompiler::PreloadingReferencerAssetsCallback(bool bAsync
 {
 	check(IsInGameThread());
 
-	UpdateArrayGCProtect();
-
+	check(ArrayGCProtect.IsEmpty());
+	for (FSoftObjectPath AssetToStream : ArrayAssetToStream)
+	{
+		ArrayGCProtect.Add(AssetToStream.TryLoad());
+	}
+	
 	if (AsynchronousStreamableHandlePtr)
 	{
 		AsynchronousStreamableHandlePtr.Reset();
@@ -168,7 +170,7 @@ void FCustomizableObjectCompiler::Compile(const TSharedRef<FCompilationRequest>&
 		return;
 	}
 
-	UCustomizableObject* RootObject = GetRootObject(CurrentObject);
+	UCustomizableObject* RootObject = GraphTraversal::GetRootObject(CurrentObject);
 	check(RootObject);
 
 	if (RootObject->VersionBridge && !RootObject->VersionBridge->GetClass()->ImplementsInterface(UCustomizableObjectVersionBridgeInterface::StaticClass()))
@@ -263,17 +265,15 @@ void FCustomizableObjectCompiler::Compile(const TSharedRef<FCompilationRequest>&
 	TRACE_BEGIN_REGION(UE_MUTABLE_PRELOAD_REGION);
 	UE_LOG(LogMutable, Verbose, TEXT("PROFILE: [ %16.8f ] Preload asynchronously assets start."), FPlatformTime::Seconds());
 
-	CleanCachedReferencers();
-	UpdateArrayGCProtect();
-	TArray<FName> ArrayReferenceNames;
-	AddCachedReferencers(*CurrentObject->GetOuter()->GetPathName(), ArrayReferenceNames);
+	TArray<FAssetData> ReferencingAssets;
+	GetReferencingPackages(*CurrentObject, ReferencingAssets);
 
 	ArrayAssetToStream.Empty();
-	for (FAssetData& Element : ArrayAssetData)
+	for (FAssetData& Element : ReferencingAssets)
 	{
 		ArrayAssetToStream.Add(Element.GetSoftObjectPath());
 	}
-
+	
 	bool bAssetsLoaded = true;
 
 	const bool bAsync = InCompileRequest->IsAsyncCompilation();
@@ -350,14 +350,8 @@ bool FCustomizableObjectCompiler::IsRequestQueued(const TSharedRef<FCompilationR
 
 void FCustomizableObjectCompiler::AddReferencedObjects(FReferenceCollector& Collector)
 {
-	// While compilation takes place, no COs involved can be garbage-collected
-	const int32 MaxIndex = ArrayGCProtect.Num();
-
-	for (int32 i = 0; i < MaxIndex; ++i)
-	{
-		Collector.AddReferencedObject(ArrayGCProtect[i]);
-	}
-
+	Collector.AddReferencedObjects(ArrayGCProtect);
+	
 	if (CurrentObject)
 	{
 		Collector.AddReferencedObject(CurrentObject);
@@ -365,45 +359,44 @@ void FCustomizableObjectCompiler::AddReferencedObjects(FReferenceCollector& Coll
 }
 
 
-void FCustomizableObjectCompiler::UpdateArrayGCProtect()
+void ProcessChildObjectsRecursively(const UCustomizableObject* ParentObject, FMutableGraphGenerationContext& GenerationContext)
 {
-	check(IsInGameThread());
-
-	const int32 MaxIndex = ArrayAssetData.Num();
-	ArrayGCProtect.SetNum(MaxIndex);
+	TArray<FName> ReferencedObjectNames;
+	
+	const FAssetRegistryModule& AssetRegistryModule = FModuleManager::LoadModuleChecked<FAssetRegistryModule>("AssetRegistry");
+	AssetRegistryModule.Get().GetReferencers(*ParentObject->GetOuter()->GetPathName(), ReferencedObjectNames, UE::AssetRegistry::EDependencyCategory::Package, UE::AssetRegistry::EDependencyQuery::Hard);
 
-	for (int i = 0; i < MaxIndex; ++i)
+	// Required to be deterministic.
+	ReferencedObjectNames.Sort([](const FName& A, const FName& B)
 	{
-		ArrayGCProtect[i] = Cast<UCustomizableObject>(ArrayAssetData[i].GetAsset());
-	}
-}
-
-
-void FCustomizableObjectCompiler::ProcessChildObjectsRecursively(UCustomizableObject* ParentObject, FMutableGraphGenerationContext& GenerationContext)
-{
-	TArray<FName> ArrayReferenceNames;
-	AddCachedReferencers(*ParentObject->GetOuter()->GetPathName(), ArrayReferenceNames);
-	UpdateArrayGCProtect();
-
+		return A.LexicalLess(B);
+	});
+	
 	bool bMultipleBaseObjectsFound = false;
+	
+	TArray<FAssetData> AssetDataArray;
 
-	for (const FName& ReferenceName : ArrayReferenceNames)
+	FARFilter Filter;
+	Filter.PackageNames = MoveTemp(ReferencedObjectNames);
+	Filter.ClassPaths = { UCustomizableObject::StaticClass()->GetClassPathName() };
+	AssetRegistryModule.Get().GetAssets(Filter, AssetDataArray);
+
+	for (FAssetData AssetData : AssetDataArray)
 	{
-		if (ArrayAlreadyProcessedChild.Contains(ReferenceName))
+		FSoftObjectPath SoftObjectPath = AssetData.GetSoftObjectPath();
+		
+		UCustomizableObject* ChildObject = Cast<UCustomizableObject>(SoftObjectPath.TryLoad());
+		if (!ChildObject || ChildObject->HasAnyFlags(RF_Transient))
 		{
 			continue;
 		}
 
-		const FAssetData* AssetData = GetCachedAssetData(ReferenceName.ToString());
-
-		UCustomizableObject* ChildObject = AssetData ? Cast<UCustomizableObject>(AssetData->GetAsset()) : nullptr;
-		if (!ChildObject || ChildObject->HasAnyFlags(RF_Transient))
+		UCustomizableObjectNodeObject* Root = GetRootNode(ChildObject, bMultipleBaseObjectsFound);
+		if (!Root)
 		{
-			ArrayAlreadyProcessedChild.Add(ReferenceName);
 			continue;
 		}
-
-		UCustomizableObjectNodeObject* Root = GetRootNode(ChildObject, bMultipleBaseObjectsFound);
+		
 		if (Root->ParentObject != ParentObject)
 		{
 			continue;
@@ -433,9 +426,7 @@ void FCustomizableObjectCompiler::ProcessChildObjectsRecursively(UCustomizableOb
 				ensure(false);
 			}
 		}
-
-		ArrayAlreadyProcessedChild.Add(ReferenceName);
-
+		
 		if (!bMultipleBaseObjectsFound)
 		{
 			if (const FGroupNodeIdsTempData* GroupGuid = GenerationContext.DuplicatedGroupNodeIds.FindPair(ParentObject, FGroupNodeIdsTempData(Root->ParentObjectGroupId)))
@@ -444,6 +435,7 @@ void FCustomizableObjectCompiler::ProcessChildObjectsRecursively(UCustomizableOb
 			}
 
 			GenerationContext.GroupIdToExternalNodeMap.Add(Root->ParentObjectGroupId, Root);
+			GenerationContext.AddParticipatingObject(*ParentObject);
 
 			TArray<UCustomizableObjectNodeObjectGroup*> GroupNodes;
 			ChildObject->GetPrivate()->GetSource()->GetNodesOfClass<UCustomizableObjectNodeObjectGroup>(GroupNodes);
@@ -467,73 +459,17 @@ void FCustomizableObjectCompiler::ProcessChildObjectsRecursively(UCustomizableOb
 }
 
 
-void FCustomizableObjectCompiler::DisplayParameterWarning(FMutableGraphGenerationContext& GenerationContext)
+mu::Ptr<mu::NodeObject> GenerateMutableRoot( 
+	const UCustomizableObject* Object, 
+	FMutableGraphGenerationContext& GenerationContext)
 {
-	for (const TPair<FString, TArray<const UObject*>>& It : GenerationContext.ParameterNamesMap)
-	{
-		if (It.Key == "")
-		{
-			FText MessageWarning = LOCTEXT("NodeWithNoName", ". There is at least one node with no name.");
-			CompilerLog(MessageWarning, It.Value, EMessageSeverity::Warning, true);
-		}
-		else if (It.Value.Num() > 1)
-		{
-			FText MessageWarning = FText::Format(LOCTEXT("NodeWithRepeatedName", ". Several nodes have repeated name \"{0}\""), FText::FromString(It.Key));
-			CompilerLog(MessageWarning, It.Value, EMessageSeverity::Warning, true);
-		}
-	}
-}
-
+	MUTABLE_CPUPROFILER_SCOPE(GenerateMutableRoot)
 
-void FCustomizableObjectCompiler::DisplayDuplicatedNodeIdsWarning(FMutableGraphGenerationContext & GenerationContext)
-{
-	for (const TPair<FGuid, TArray<const UObject*>>& It : GenerationContext.NodeIdsMap)
-	{
-		if (It.Value.Num() > 1)
-		{
-			FText MessageWarning = LOCTEXT("NodeWithRepeatedIds", ". Several nodes have repeated NodeIds, reconstruct the nodes.");
-			CompilerLog(MessageWarning, It.Value, EMessageSeverity::Warning, true);
-		}
-	}
-}
-
-
-void FCustomizableObjectCompiler::DisplayUnnamedNodeObjectWarning(FMutableGraphGenerationContext& GenerationContext)
-{
-	FText Message = LOCTEXT("Unnamed Node Object", "Unnamed Node Object");
-	for (const UCustomizableObjectNode* It : GenerationContext.NoNameNodeObjectArray)
-	{
-		CompilerLog(Message, It, EMessageSeverity::Warning, true);
-	}
-}
-
-
-void FCustomizableObjectCompiler::DisplayOrphanNodesWarning(FMutableGraphGenerationContext& GenerationContext)
-{
-	for (const TPair<FGeneratedKey, FGeneratedData>& It : GenerationContext.Generated)
-	{
-		if (const UCustomizableObjectNode* Node = Cast<UCustomizableObjectNode>(It.Value.Source))
-		{
-			if (Node->GetAllOrphanPins().Num() > 0)
-			{
-				CompilerLog(LOCTEXT("OrphanPinsWarningCompiler", "Node contains deprecated pins"), Node, EMessageSeverity::Warning, false);
-			}
-		}
-	}
-}
-
-
-mu::NodeObjectPtr FCustomizableObjectCompiler::GenerateMutableRoot( 
-	UCustomizableObject* Object, 
-	FMutableGraphGenerationContext& GenerationContext, 
-	FText& ErrorMsg, 
-	bool& bOutIsRootObject)
-{
 	check(Object);
 	
 	if (!Object->GetPrivate()->GetSource())
 	{
-		ErrorMsg = LOCTEXT("NoSource", "Object with no valid graph found. Object not build.");
+		GenerationContext.Log(LOCTEXT("NoSource", "Object with no valid graph found. Object not build."));
 
 		if (IsRunningCookCommandlet() || IsRunningCookOnTheFly())
 		{
@@ -548,19 +484,17 @@ mu::NodeObjectPtr FCustomizableObjectCompiler::GenerateMutableRoot(
 
 	if (bMultipleBaseObjectsFound)
 	{
-		ErrorMsg = LOCTEXT("MultipleBaseRoot","Multiple base object nodes found.");
+		GenerationContext.Log(LOCTEXT("MultipleBaseRoot","Multiple base object nodes found."));
 		return nullptr;
 	}
 
 	if (!Root)
 	{
-		ErrorMsg = LOCTEXT("NoRootBase","No base object node found. Object not built.");
+		GenerationContext.Log(LOCTEXT("NoRootBase","No base object node found. Object not built."));
 		return nullptr;
 	}
-
-	bOutIsRootObject = Root->ParentObject == nullptr;
-
-	UCustomizableObject* ActualRootObject = GetRootObject(Object);
+	
+	const UCustomizableObject* ActualRootObject = GraphTraversal::GetRootObject(Object);
 	check(ActualRootObject);
 
 	GenerationContext.RootVersionBridge = ActualRootObject->VersionBridge;
@@ -570,26 +504,24 @@ mu::NodeObjectPtr FCustomizableObjectCompiler::GenerateMutableRoot(
 	
 	if (bMultipleBaseObjectsFound)
 	{
-		ErrorMsg = LOCTEXT("MultipleBaseActualRoot", "Multiple base object nodes found.");
+		GenerationContext.Log(LOCTEXT("MultipleBaseActualRoot", "Multiple base object nodes found."));
 		return nullptr;
 	}
 
 	if (!ActualRoot)
 	{
-		ErrorMsg = LOCTEXT("NoActualRootBase", "No base object node found in root Customizable Object. Object not built.");
+		GenerationContext.Log(LOCTEXT("NoActualRootBase", "No base object node found in root Customizable Object. Object not built."));
 		return nullptr;
 	}
-
-	ArrayAlreadyProcessedChild.Empty();
-
+	
 	if (Root->ObjectName.IsEmpty())
 	{
 		GenerationContext.NoNameNodeObjectArray.AddUnique(Root);
 	}
 
-	if ((Object->MeshCompileType == EMutableCompileMeshType::Full) || CurrentOptions.bIsCooking)
+	if ((Object->MeshCompileType == EMutableCompileMeshType::Full) || GenerationContext.Options.bIsCooking)
 	{
-		if (Root->ParentObject!=nullptr && CurrentOptions.bIsCooking)
+		if (Root->ParentObject!=nullptr && GenerationContext.Options.bIsCooking)
 		{
 			// This happens while packaging.
 			return nullptr;
@@ -610,11 +542,11 @@ mu::NodeObjectPtr FCustomizableObjectCompiler::GenerateMutableRoot(
 	{
 		// Local, local with children and working set modes: add parents until whole CO graph root
 		TArray<UCustomizableObjectNodeObject*> ArrayNodeObject;
-		TArray<UCustomizableObject*> ArrayCustomizableObject;
+		TArray<const UCustomizableObject*> ArrayCustomizableObject;
 		
 		if (!GetParentsUntilRoot(Object, ArrayNodeObject, ArrayCustomizableObject))
 		{
-			CompilerLog(LOCTEXT("SkeletalMeshCycleFound", "Error! Cycle detected in the Customizable Object hierarchy."), Root);
+			GenerationContext.Log(LOCTEXT("SkeletalMeshCycleFound", "Error! Cycle detected in the Customizable Object hierarchy."), Root);
 			return nullptr;
 		}
 
@@ -624,13 +556,13 @@ mu::NodeObjectPtr FCustomizableObjectCompiler::GenerateMutableRoot(
 			const int32 MaxIndex = Object->WorkingSet.Num();
 			for (int32 i = 0; i < MaxIndex; ++i)
 			{
-				if (UCustomizableObject* WorkingSetObject = Object->WorkingSet[i].LoadSynchronous())
+				if (UCustomizableObject* WorkingSetObject = GenerationContext.LoadObject(Object->WorkingSet[i], true))
 				{
 					ArrayCustomizableObject.Reset();
 
 					if (!GetParentsUntilRoot(WorkingSetObject, ArrayNodeObject, ArrayCustomizableObject))
 					{
-						CompilerLog(LOCTEXT("NoReferenceMesh", "Error! Cycle detected in the Customizable Object hierarchy."), Root);
+						GenerationContext.Log(LOCTEXT("NoReferenceMesh", "Error! Cycle detected in the Customizable Object hierarchy."), Root);
 						return nullptr;
 					}
 				}
@@ -645,7 +577,6 @@ mu::NodeObjectPtr FCustomizableObjectCompiler::GenerateMutableRoot(
 
 			if (GroupNodes.Num() > 0) // Only graphs with group nodes should have child graphs
 			{
-				FAssetRegistryModule& AssetRegistryModule = FModuleManager::LoadModuleChecked<FAssetRegistryModule>("AssetRegistry");
 				ProcessChildObjectsRecursively(Object, GenerationContext);
 			}
 		}
@@ -660,16 +591,16 @@ mu::NodeObjectPtr FCustomizableObjectCompiler::GenerateMutableRoot(
 	}
 	
 	// Find all component nodes
-	TSet<UCustomizableObject*> ParticipatingObjects; // All COs participating in the hierarchy
+	TSet<const UCustomizableObject*> CustomizableObjects; // All COs in the hierarchy
 	{
-		auto GetParents = [&ParticipatingObjects](UCustomizableObject* Object)
+		auto GetParents = [&CustomizableObjects](const UCustomizableObject* Object)
 		{
 			TArray<UCustomizableObjectNodeObject*> ArrayNodeObject;
-			TArray<UCustomizableObject*> ArrayCustomizableObject;
+			TArray<const UCustomizableObject*> ArrayCustomizableObject;
 
 			GetParentsUntilRoot(Object, ArrayNodeObject, ArrayCustomizableObject);
 
-			ParticipatingObjects.Append(ArrayCustomizableObject);
+			CustomizableObjects.Append(ArrayCustomizableObject);
 		};
 
 		// Parents
@@ -685,13 +616,18 @@ mu::NodeObjectPtr FCustomizableObjectCompiler::GenerateMutableRoot(
 		}
 
 		// Children
-		for (FAssetData& AssetData : ArrayAssetData)
+		TArray<FAssetData> ReferencingAssets;
+		GetReferencingPackages(*Object, ReferencingAssets);
+		for (FAssetData& AssetData : ReferencingAssets)
 		{
-			ParticipatingObjects.Add(Cast<UCustomizableObject>(AssetData.GetAsset()));
+			if (UCustomizableObject* ChildObject = Cast<UCustomizableObject>(AssetData.GetAsset()))
+			{
+				CustomizableObjects.Add(ChildObject);
+			}
 		}
 	}
 	
-	for (UCustomizableObject* ParticipatingObject : ParticipatingObjects)
+	for (const UCustomizableObject* ParticipatingObject : CustomizableObjects)
 	{
 		for (UEdGraphNode* Node : ParticipatingObject->GetPrivate()->GetSource()->Nodes)
 		{
@@ -707,7 +643,7 @@ mu::NodeObjectPtr FCustomizableObjectCompiler::GenerateMutableRoot(
 				check(NodeMeshComponent);
 				
 				FText Msg = FText::Format(LOCTEXT("ComponentNodeWithSameNameExists", "Already exists a Mesh Component node with the same name in Customizable Object [{0}]"), FText::FromString(GetRootObject(*NodeMeshComponent)->GetName()));
-				CompilerLog(Msg, Node, EMessageSeverity::Error);
+				GenerationContext.Log(Msg, Node, EMessageSeverity::Error);
 				return nullptr;
 			}
 
@@ -717,12 +653,12 @@ mu::NodeObjectPtr FCustomizableObjectCompiler::GenerateMutableRoot(
 	
     GenerationContext.RealTimeMorphTargetsOverrides = ActualRoot->RealTimeMorphSelectionOverrides;
 
-	if (!GenerationContext.ParamNamesToSelectedOptions.IsEmpty())
+	if (!GenerationContext.Options.ParamNamesToSelectedOptions.IsEmpty())
 	{
 		GenerationContext.TableToParamNames = Object->GetPrivate()->GetModelResources().TableToParamNames;
 	}
 
-	GenerationContext.bPartialCompilation = !bOutIsRootObject;
+	GenerationContext.bPartialCompilation = !Root->ParentObject;
 
 	// Generate the object expression
 	UE_LOG(LogMutable, Verbose, TEXT("PROFILE: [ %16.8f ] GenerateMutableSource start."), FPlatformTime::Seconds());
@@ -739,16 +675,40 @@ mu::NodeObjectPtr FCustomizableObjectCompiler::GenerateMutableRoot(
 	//	GenerationContext.CheckPhysicsAssetInSkeletalMesh(RefSkeletalMesh);
 	//}
 
-	DisplayParameterWarning( GenerationContext );
-	DisplayUnnamedNodeObjectWarning( GenerationContext );
-	DisplayDuplicatedNodeIdsWarning( GenerationContext );
-	//DisplayDiscardedPhysicsAssetSingleWarning( GenerationContext );
-	DisplayOrphanNodesWarning( GenerationContext );
+	// Display warnings for unnamed node objects
+	FText Message = LOCTEXT("Unnamed Node Object", "Unnamed Node Object");
+	for (const UCustomizableObjectNode* It : GenerationContext.NoNameNodeObjectArray)
+	{
+		GenerationContext.Log(Message, It, EMessageSeverity::Warning, true);
+	}
+
+	// If duplicated node ids are found, usually due to duplicating CustomizableObjects Assets, a warning
+	// for the nodes with repeated ids will be generated
+	for (const TPair<FGuid, TArray<const UObject*>>& It : GenerationContext.NodeIdsMap)
+	{
+		if (It.Value.Num() > 1)
+		{
+			FText MessageWarning = LOCTEXT("NodeWithRepeatedIds", "Several nodes have repeated NodeIds, reconstruct the nodes.");
+			GenerationContext.Log(MessageWarning, It.Value, EMessageSeverity::Warning, true);
+		}
+	}
+
+	// Display a warning for each node contains an orphan pin.
+	for (const TPair<FGeneratedKey, FGeneratedData>& It : GenerationContext.Generated)
+	{
+		if (const UCustomizableObjectNode* Node = Cast<UCustomizableObjectNode>(It.Value.Source))
+		{
+			if (Node->GetAllOrphanPins().Num() > 0)
+			{
+				GenerationContext.Log(LOCTEXT("OrphanPinsWarningCompiler", "Node contains deprecated pins"), Node, EMessageSeverity::Warning, false);
+			}
+		}
+	}
 
 	if (GenerationContext.CustomizableObjectWithCycle)
 	{
-		ErrorMsg = FText::Format(LOCTEXT("CycleDetected","Cycle detected in graph of CustomizableObject {0}. Object not built."),
-			FText::FromString(GenerationContext.CustomizableObjectWithCycle->GetPathName()));
+		GenerationContext.Log(FText::Format(LOCTEXT("CycleDetected","Cycle detected in graph of CustomizableObject {0}. Object not built."),
+			FText::FromString(GenerationContext.CustomizableObjectWithCycle->GetPathName())));
 
 		return nullptr;
 	}
@@ -785,83 +745,6 @@ void FCustomizableObjectCompiler::SaveCODerivedData()
 }
 
 
-void FCustomizableObjectCompiler::AddCachedReferencers(const FName& PathName, TArray<FName>& ArrayReferenceNames)
-{
-	ArrayReferenceNames.Empty();
-	FAssetRegistryModule& AssetRegistryModule = FModuleManager::LoadModuleChecked<FAssetRegistryModule>("AssetRegistry");
-	AssetRegistryModule.Get().GetReferencers(PathName, ArrayReferenceNames, UE::AssetRegistry::EDependencyCategory::Package, UE::AssetRegistry::EDependencyQuery::Hard);
-
-	// Required to make compilations deterministic within editor runs.
-	ArrayReferenceNames.Sort([](const FName& A, const FName& B)
-		{
-			return A.LexicalLess(B);
-		});
-
-	FARFilter Filter;
-	for (const FName& ReferenceName : ArrayReferenceNames)
-	{
-		if (!IsCachedInAssetData(ReferenceName.ToString()) && !ReferenceName.ToString().StartsWith(TEXT("/TempAutosave")))
-		{
-			Filter.PackageNames.Add(ReferenceName);
-		}
-	}
-
-	Filter.bIncludeOnlyOnDiskAssets = false;
-
-	TArray<FAssetData> ArrayAssetDataTemp;
-	AssetRegistryModule.Get().GetAssets(Filter, ArrayAssetDataTemp);
-
-	// Store only those which have static class type Customizable Object, to avoid loading not needed elements
-	const int32 MaxIndex = ArrayAssetDataTemp.Num();
-	for (int32 i = 0; i < MaxIndex; ++i)
-	{
-		if (ArrayAssetDataTemp[i].GetClass() == UCustomizableObject::StaticClass())
-		{
-			ArrayAssetData.Add(ArrayAssetDataTemp[i]);
-		}
-	}
-}
-
-
-void FCustomizableObjectCompiler::CleanCachedReferencers()
-{
-	ArrayAssetData.Empty();
-}
-
-
-bool FCustomizableObjectCompiler::IsCachedInAssetData(const FString& PackageName)
-{
-	const int32 MaxIndex = ArrayAssetData.Num();
-
-	for (int32 i = 0; i < MaxIndex; ++i)
-	{
-		if (ArrayAssetData[i].PackageName.ToString() == PackageName)
-		{
-			return true;
-		}
-	}
-
-	return false;
-}
-
-
-FAssetData* FCustomizableObjectCompiler::GetCachedAssetData(const FString& PackageName)
-{
-	const int32 MaxIndex = ArrayAssetData.Num();
-
-	for (int32 i = 0; i < MaxIndex; ++i)
-	{
-		if (ArrayAssetData[i].PackageName.ToString() == PackageName)
-		{
-			return &ArrayAssetData[i];
-		}
-	}
-
-	return nullptr;
-}
-
-
-
 ECompilationResultPrivate FCustomizableObjectCompiler::GetCompilationResult() const
 {
 	if (CompilationLogsContainer.GetErrorCount())
@@ -906,20 +789,19 @@ void FCustomizableObjectCompiler::CompileInternal(bool bAsync)
 	}
 
 	FMutableGraphGenerationContext GenerationContext(CurrentObject, this, CurrentOptions);
-	GenerationContext.ParamNamesToSelectedOptions = CurrentRequest->GetParameterNamesToSelectedOptions();
+
+	// Perform a first participating objects pass
+	GenerationContext.ParticipatingObjects = ICustomizableObjectEditorModule::GetChecked().GetParticipatingObjects(CurrentObject, true, &CurrentOptions);
 
 	// Clear Messages from previous Compilations
 	CompilationLogsContainer.ClearMessageCounters();
 	CompilationLogsContainer.ClearMessagesArray();
 
 	// Generate the mutable node expression
-	FText ErrorMessage = FText(LOCTEXT("FailedToGenerateRoot", "Failed to generate the mutable node graph. Object not built."));
-	bool bIsRootObject = false;
-	mu::NodeObjectPtr MutableRoot = GenerateMutableRoot(CurrentObject, GenerationContext, ErrorMessage, bIsRootObject);
-
+	mu::NodeObjectPtr MutableRoot = GenerateMutableRoot(CurrentObject, GenerationContext);
 	if (!MutableRoot)
 	{
-		CompilerLog(ErrorMessage);
+		CompilerLog(FText(LOCTEXT("FailedToGenerateRoot", "Failed to generate the mutable node graph. Object not built.")));
 		CompleteRequest(ECompilationStatePrivate::Completed, ECompilationResultPrivate::Errors);
 	}
 	else
@@ -1123,9 +1005,23 @@ void FCustomizableObjectCompiler::CompileInternal(bool bAsync)
 		CurrentObject->GetPrivate()->GetAlwaysLoadedExtensionData() = MoveTemp(GenerationContext.AlwaysLoadedExtensionData);
 
 		CurrentObject->GetPrivate()->GetStreamedExtensionData().Empty(GenerationContext.StreamedExtensionData.Num());
-		for (UCustomizableObjectResourceDataContainer* Container : GenerationContext.StreamedExtensionData)
+		for (const TPair<FName, UCustomizableObjectResourceDataContainer*>& Pair : GenerationContext.StreamedExtensionData)
 		{
-			CurrentObject->GetPrivate()->GetStreamedExtensionData().Emplace(Container);
+			const FName& ContainerName = Pair.Key;
+			UCustomizableObjectResourceDataContainer* Container = Pair.Value;
+
+			UCustomizableObjectResourceDataContainer* NewContainer = FindObject<UCustomizableObjectResourceDataContainer>(CurrentObject, *ContainerName.ToString());
+			if (!NewContainer)
+			{
+				NewContainer = NewObject<UCustomizableObjectResourceDataContainer>(
+					CurrentObject,
+					FName(*ContainerName.ToString()),
+					RF_Public);
+			}
+			
+			NewContainer->Data = Container->Data;
+			
+			CurrentObject->GetPrivate()->GetStreamedExtensionData().Emplace(NewContainer);
 		}
 
 #if WITH_EDITORONLY_DATA
@@ -1135,11 +1031,18 @@ void FCustomizableObjectCompiler::CompileInternal(bool bAsync)
 #endif
 
 		ModelResources.ComponentNames = MoveTemp(GenerationContext.ComponentNames);
+
+		if (ICustomizableObjectVersionBridgeInterface* VersionBridge = Cast<ICustomizableObjectVersionBridgeInterface>(GraphTraversal::GetRootObject(CurrentObject)->VersionBridge))
+		{
+			ModelResources.CompiledVersionBridge = VersionBridge->GetCurrentVersionAsString();
+		}
 		
 		ModelResources.NumLODs = GenerationContext.NumLODsInRoot;
 		ModelResources.NumLODsToStream = GenerationContext.bEnableLODStreaming ? GenerationContext.NumMaxLODsToStream : 0;
 		ModelResources.FirstLODAvailable = GenerationContext.FirstLODAvailable;
 
+		ModelResources.ParticipatingObjects = MoveTemp(GenerationContext.ParticipatingObjects);
+		
 		if (CurrentOptions.bGatherReferences)
 		{
 			CurrentObject->GetPrivate()->References = ModelResources;
@@ -1147,8 +1050,26 @@ void FCustomizableObjectCompiler::CompileInternal(bool bAsync)
 			CurrentObject->Modify();
 		}
 
-		ModelResources.StreamedResourceData = MoveTemp(GenerationContext.StreamedResourceData);
+		ModelResources.StreamedResourceData.Empty(GenerationContext.StreamedResourceData.Num());
+		for (TPair<FName, UCustomizableObjectResourceDataContainer*>& Pair : GenerationContext.StreamedResourceData)
+		{
+			const FName& ContainerName = Pair.Key;
+			UCustomizableObjectResourceDataContainer* Container = Pair.Value;
 
+			UCustomizableObjectResourceDataContainer* NewContainer = FindObject<UCustomizableObjectResourceDataContainer>(CurrentObject, *ContainerName.ToString());
+			if (!NewContainer)
+			{
+				NewContainer = NewObject<UCustomizableObjectResourceDataContainer>(
+					CurrentObject,
+					FName(*ContainerName.ToString()),
+					RF_Public);
+			}
+			
+			NewContainer->Data = Container->Data;
+			
+			ModelResources.StreamedResourceData.Emplace(NewContainer);
+		}
+		
 		// Pass-through textures
 		TArray<FMutableSourceTextureData> NewCompileTimeReferencedTextures;
 		for (const TPair<TSoftObjectPtr<const UTexture>, FMutableGraphGenerationContext::FGeneratedReferencedTexture>& Pair : GenerationContext.CompileTimeTextureMap)
@@ -1159,31 +1080,6 @@ void FCustomizableObjectCompiler::CompileInternal(bool bAsync)
 			NewCompileTimeReferencedTextures.Add(Tex);
 		}
 
-		if (!CurrentRequest->GetParameterNamesToSelectedOptions().Num())
-		{
-			// Get possible objects used in the compilation that are not directly referenced.
-			// Due to this check being done also in PIE (to detect out of date compilations), it has to be performant. Therefore we are gathering a relaxed set.
-			// For example, a referencing Customizable Object may not be used if it is not assigned in any Group Node. In the relaxed set we include those regardless.
-			// Notice that, to avoid automatic compilations/warnings, the set of referencing objects set found here must coincide with the set found when loading the
-			// model (discard previous compilations) or when showing PIE warnings.
-			TArray<FName> ReferencingObjectNames;
-			GetReferencingPackages(*CurrentObject, ReferencingObjectNames);
-
-			for (const FName& ReferencingObjectName : ReferencingObjectNames)
-			{
-				const TSoftObjectPtr<UObject> SoftObjectPtr = TSoftObjectPtr<UObject>(FSoftObjectPath(ReferencingObjectName.ToString()));
-
-				if (const UObject* ReferencingObject = SoftObjectPtr.LoadSynchronous())
-				{
-					GenerationContext.AddParticipatingObject(*ReferencingObject);					
-				}
-			}
-			
-			// Copy final array of participating objects
-			ModelResources.ParticipatingObjects = MoveTemp(GenerationContext.ParticipatingObjects);
-			CurrentObject->GetPrivate()->DirtyParticipatingObjects.Empty();
-		}
-
 		CompileTask = MakeShareable(new FCustomizableObjectCompileRunnable(MutableRoot));
 		CompileTask->Options = CurrentOptions;
 		CompileTask->ReferencedTextures = NewCompileTimeReferencedTextures;
@@ -1264,8 +1160,7 @@ void FCustomizableObjectCompiler::CompleteRequest(ECompilationStatePrivate State
 	}
 
 	// Remove referenced objects
-	CleanCachedReferencers();
-	UpdateArrayGCProtect();
+	ArrayGCProtect.Empty();
 
 	// Notifications
 	RemoveCompileNotification();
@@ -1352,7 +1247,7 @@ bool FCustomizableObjectCompiler::TryLoadCompiledDataFromDDC(UCustomizableObject
 }
 
 
-mu::NodePtr FCustomizableObjectCompiler::Export(UCustomizableObject* Object, const FCompilationOptions& InCompilerOptions, 
+mu::Ptr<mu::Node> FCustomizableObjectCompiler::Export(UCustomizableObject* Object, const FCompilationOptions& InCompilerOptions, 
 	TArray<TSoftObjectPtr<const UTexture>>& OutRuntimeReferencedTextures,
 	TArray<FMutableSourceTextureData>& OutCompilerReferencedTextures )
 {
@@ -1374,26 +1269,12 @@ mu::NodePtr FCustomizableObjectCompiler::Export(UCustomizableObject* Object, con
 	CompilerOptions.bAnimBpPhysicsManipulationEnabled = Object->bEnableAnimBpPhysicsAssetsManipualtion;
 
 	FMutableGraphGenerationContext GenerationContext(Object, this, CompilerOptions);
-	if (CurrentRequest)
-	{
-		GenerationContext.ParamNamesToSelectedOptions = CurrentRequest->GetParameterNamesToSelectedOptions();
-	}
 
 	// Generate the mutable node expression
-	FText ErrorMsg;
-	bool bIsRootObject = false;
-	mu::NodeObjectPtr MutableRoot = GenerateMutableRoot(Object, GenerationContext, ErrorMsg, bIsRootObject);
-
+	mu::NodeObjectPtr MutableRoot = GenerateMutableRoot(Object, GenerationContext);
 	if (!MutableRoot)
 	{
-		if (!ErrorMsg.IsEmpty())
-		{
-			FCustomizableObjectCompiler::CompilerLog(ErrorMsg, nullptr);
-		}
-		else
-		{
-			FCustomizableObjectCompiler::CompilerLog(LOCTEXT("FailedToExport","Failed to generate the mutable node graph. Object not built."), nullptr);
-		}
+		CompilerLog(LOCTEXT("FailedToExport", "Failed to generate the mutable node graph. Object not built."), nullptr);
 		return nullptr;
 	}
 
diff --git a/Engine/Plugins/Mutable/Source/CustomizableObjectEditor/Private/MuCOE/CustomizableObjectCompiler.h b/Engine/Plugins/Mutable/Source/CustomizableObjectEditor/Private/MuCOE/CustomizableObjectCompiler.h
--- a/Engine/Plugins/Mutable/Source/CustomizableObjectEditor/Private/MuCOE/CustomizableObjectCompiler.h
+++ b/Engine/Plugins/Mutable/Source/CustomizableObjectEditor/Private/MuCOE/CustomizableObjectCompiler.h
@@ -11,6 +11,7 @@
 #include "MuCOE/CustomizableObjectEditorLogger.h"
 #include "UObject/GCObject.h"
 #include "TickableEditorObject.h"
+#include "MuR/Ptr.h"
 
 #include "Framework/Notifications/NotificationManager.h"
 
@@ -21,6 +22,17 @@ class FRunnableThread;
 class FText;
 class UCustomizableObject;
 class UCustomizableObjectNode;
+struct FMutableGraphGenerationContext;
+
+namespace mu
+{
+	class NodeObject;
+	class Node;
+}
+
+
+mu::Ptr<mu::NodeObject> GenerateMutableRoot(const UCustomizableObject* Object, FMutableGraphGenerationContext& GenerationContext);
+
 
 class FCustomizableObjectCompiler : public FTickableEditorObject, public FTickableCookObject, public FGCObject
 {
@@ -39,10 +51,9 @@ public:
 
 	// FTickableCookObject interface
 	virtual void TickCook(float DeltaTime, bool bCookCompete) override;
-
-
+	
 	/** Generate the Mutable Graph from the Unreal Graph. */
-	mu::NodePtr Export(UCustomizableObject* Object, const FCompilationOptions& Options, TArray<TSoftObjectPtr<const UTexture>>& OutRuntimeReferencedTextures, TArray<FMutableSourceTextureData>& OutCompilerReferencedTextures);
+	mu::Ptr<mu::Node> Export(UCustomizableObject* Object, const FCompilationOptions& Options, TArray<TSoftObjectPtr<const UTexture>>& OutRuntimeReferencedTextures, TArray<FMutableSourceTextureData>& OutCompilerReferencedTextures);
 
 	void CompilerLog(const FText& Message, const TArray<const UObject*>& UObject, const EMessageSeverity::Type MessageSeverity = EMessageSeverity::Warning, const bool bAddBaseObjectInfo = true, const ELoggerSpamBin SpamBin = ELoggerSpamBin::ShowAll);
 	void CompilerLog(const FText& Message, const UObject* Context = nullptr, const EMessageSeverity::Type MessageSeverity = EMessageSeverity::Warning, const bool bAddBaseObjectInfo = true, const ELoggerSpamBin SpamBin = ELoggerSpamBin::ShowAll);
@@ -72,9 +83,6 @@ public:
 		return TEXT("FCustomizableObjectCompiler");
 	}
 
-	/** Simply add CO elements from ArrayAssetData to ArrayGCProtect when they've been loaded from ArrayAssetData */
-	void UpdateArrayGCProtect();
-
 private:
 
 	// Object containing all error and warning logs raised during compilation.
@@ -91,29 +99,10 @@ private:
 
 	void PreloadingReferencerAssetsCallback(bool bAsync);
 	
-	void ProcessChildObjectsRecursively(UCustomizableObject* Object, FMutableGraphGenerationContext &GenerationContext);
-	
 	// Will output to Mutable Log the warning and error messages generated during the CO compilation
 	// and update the values of NumWarnings and NumErrors
 	void UpdateCompilerLogData();
-
-	/** If duplicated elements are found in each entry of ParameterNamesMap, a warning for
-	the parameters with repeated name will be generated */
-	void DisplayParameterWarning(FMutableGraphGenerationContext& GenerationContext);
-	
-	/** If duplicated node ids are found, usually due to duplicating CustomizableObjects Assets, a warning
-	for the nodes with repeated ids will be generated */
-	void DisplayDuplicatedNodeIdsWarning(FMutableGraphGenerationContext& GenerationContext);
-	
-	/** Display warnings for unnamed node objects */
-	void DisplayUnnamedNodeObjectWarning(FMutableGraphGenerationContext& GenerationContext);
-	
-	/** Display a warning for each node contains an orphan pin. */
-	void DisplayOrphanNodesWarning(FMutableGraphGenerationContext& GenerationContext);
 	
-	mu::NodeObjectPtr GenerateMutableRoot(UCustomizableObject* Object, FMutableGraphGenerationContext& GenerationContext, FText& ErrorMessage, bool& bOutIsRootObject);
-
-
 	/** Launches the compile task in another thread when compiling a CO in the editor
 	* @param bShowNotification [in] whether to show the compiling CO notification or not
 	* @return nothing */
@@ -124,28 +113,7 @@ private:
 	* @param bShowNotification [in] whether to show the saving DD notification or not
 	* @return nothing */
 	void SaveCODerivedData();
-
-	/** Add to ArrayAssetData the FAssetData information of all referencers of static class type UCustomizableObject::StaticClass()
-	* that reference the package given by the PathName parameter
-	* @param PathName            [in]  path to the CO to be analyzed (for instance, CO->GetOuter()->GetPathName())
-	* @param ArrayReferenceNames [out] array with the package names which are referenced by the CO with PathName given as parameter
-	* @return nothing */
-	void AddCachedReferencers(const FName& PathName, TArray<FName>& ArrayReferenceNames);
-
-	/** Just used to clean ArrayAssetData
-	* @return nothing */
-	void CleanCachedReferencers();
-
-	/** Will test if package path given by PackageName parameter is one of ArrayAssetData's elements FAssetData::PackageName value
-	* @param PackageName [in] package name to test
-	* @return true if cached, false otherwise */
-	bool IsCachedInAssetData(const FString& PackageName);
-
-	/** Find FAssetData ArrayAssetData with PackageName given by parameter
-	* @param PackageName [in] package name to find in ArrayAssetData
-	* @return pointer to element if any found, nullptr otherwise */
-	FAssetData* GetCachedAssetData(const FString& PackageName);
-
+	
 	ECompilationResultPrivate GetCompilationResult() const;
 
 	/** Pointer to the Asynchronous Preloading process call back */
@@ -159,15 +127,9 @@ private:
 	/** SaveDD task and thread. */
 	TSharedPtr<FCustomizableObjectSaveDDRunnable> SaveDDTask;
 	TSharedPtr<FRunnableThread> SaveDDThread;
-
-	/** Array where to put the names of the already processed child in ProcessChildObjectsRecursively */
-	TArray<FName> ArrayAlreadyProcessedChild;
-
-	/** Array with all the packages used to compile current Customizable Object */
-	TArray<FAssetData> ArrayAssetData;
-
+	
 	/** Array used to protect from garbage collection those COs loaded asynchronously */
-	TArray<TObjectPtr<UCustomizableObject>> ArrayGCProtect;
+	TArray<TObjectPtr<UObject>> ArrayGCProtect;
 
 	TSharedPtr<mu::Model, ESPMode::ThreadSafe> Model;
 	TSharedPtr<FModelStreamableBulkData> ModelStreamableBulkData;
diff --git a/Engine/Plugins/Mutable/Source/CustomizableObjectEditor/Private/MuCOE/CustomizableObjectCookPackageSplitter.cpp b/Engine/Plugins/Mutable/Source/CustomizableObjectEditor/Private/MuCOE/CustomizableObjectCookPackageSplitter.cpp
--- a/Engine/Plugins/Mutable/Source/CustomizableObjectEditor/Private/MuCOE/CustomizableObjectCookPackageSplitter.cpp
+++ b/Engine/Plugins/Mutable/Source/CustomizableObjectEditor/Private/MuCOE/CustomizableObjectCookPackageSplitter.cpp
@@ -151,7 +151,7 @@ void GenerateNewPackage(const FCustomizableObjectStreamedResourceData& StreamedD
 
 bool FCustomizableObjectCookPackageSplitter::ShouldSplit(UObject* SplitData)
 {
-	const UCustomizableObject* Object = CastChecked<UCustomizableObject>(SplitData);
+	UCustomizableObject* Object = CastChecked<UCustomizableObject>(SplitData);
 	const FModelResources& ModelResources = Object->GetPrivate()->GetModelResources(true);
 
 	return !Object->IsChildObject() && (ModelResources.StreamedResourceData.Num() > 0 || Object->GetPrivate()->GetStreamedExtensionData().Num() > 0);
@@ -165,7 +165,7 @@ TArray<ICookPackageSplitter::FGeneratedPackage> FCustomizableObjectCookPackageSp
 	StrongObject.Reset(OwnerObject);
 	
 	const UCustomizableObject* Object = CastChecked<UCustomizableObject>(OwnerObject);
-	FModelResources& ModelResources = Object->GetPrivate()->GetModelResources(true);
+	const FModelResources& ModelResources = Object->GetPrivate()->GetModelResources(true);
 
 	TArray<ICookPackageSplitter::FGeneratedPackage> Result;
 
diff --git a/Engine/Plugins/Mutable/Source/CustomizableObjectEditor/Private/MuCOE/CustomizableObjectDetails.cpp b/Engine/Plugins/Mutable/Source/CustomizableObjectEditor/Private/MuCOE/CustomizableObjectDetails.cpp
--- a/Engine/Plugins/Mutable/Source/CustomizableObjectEditor/Private/MuCOE/CustomizableObjectDetails.cpp
+++ b/Engine/Plugins/Mutable/Source/CustomizableObjectEditor/Private/MuCOE/CustomizableObjectDetails.cpp
@@ -164,7 +164,7 @@ void FCustomizableObjectDetails::CustomizeDetails( IDetailLayoutBuilder& DetailB
 
 void FCustomizableObjectDetails::UpdateTree()
 {
-	UCustomizableObject* RootObject = GetRootObject(CustomizableObject);
+	UCustomizableObject* RootObject = GraphTraversal::GetRootObject(CustomizableObject);
 
 	if (RootObject && RootObject->GetPrivate()->GetModel() )
 	{
diff --git a/Engine/Plugins/Mutable/Source/CustomizableObjectEditor/Private/MuCOE/CustomizableObjectEditor.cpp b/Engine/Plugins/Mutable/Source/CustomizableObjectEditor/Private/MuCOE/CustomizableObjectEditor.cpp
--- a/Engine/Plugins/Mutable/Source/CustomizableObjectEditor/Private/MuCOE/CustomizableObjectEditor.cpp
+++ b/Engine/Plugins/Mutable/Source/CustomizableObjectEditor/Private/MuCOE/CustomizableObjectEditor.cpp
@@ -1510,21 +1510,23 @@ void FCustomizableObjectEditor::CompileObject(bool bOnlySelectedParameters, bool
 	}
 
 	TSharedRef<FCompilationRequest> CompileRequest = MakeShared<FCompilationRequest>(*CustomizableObject, true);
-	CompileRequest->GetCompileOptions().bSilentCompilation = false;
-	CompileRequest->GetCompileOptions().bGatherReferences = bGatherReferences;
 
-	if (bOnlySelectedParameters && PreviewInstance)
+	FCompilationOptions& CompilationOptions = CompileRequest->GetCompileOptions();
+	CompilationOptions.bSilentCompilation = false;
+	CompilationOptions.bGatherReferences = bGatherReferences;
+
+	if (bOnlySelectedParameters)
 	{
-		const TArray<FCustomizableObjectIntParameterValue>& IntParameters = PreviewInstance->GetPrivate()->GetDescriptor().GetIntParameters();
-		TMap<FString, FString> ParamNamesToSelectedOptions;
-		ParamNamesToSelectedOptions.Reserve(IntParameters.Num());
+		UCustomizableObjectInstance* Instance = GetPreviewInstance();
+		check(Instance);
+		
+		const TArray<FCustomizableObjectIntParameterValue>& IntParameters = Instance->GetPrivate()->GetDescriptor().GetIntParameters();
+		CompilationOptions.ParamNamesToSelectedOptions.Reserve(IntParameters.Num());
 
 		for (const FCustomizableObjectIntParameterValue& IntParam : IntParameters)
 		{
-			ParamNamesToSelectedOptions.Add(IntParam.ParameterName, IntParam.ParameterValueName);
+			CompilationOptions.ParamNamesToSelectedOptions.Add(IntParam.ParameterName, IntParam.ParameterValueName);
 		}
-
-		CompileRequest->SetParameterNamesToSelectedOptions(ParamNamesToSelectedOptions);
 	}
 
 	ICustomizableObjectEditorModule::GetChecked().CompileCustomizableObject(CompileRequest);
@@ -2339,50 +2341,6 @@ UCustomizableObject* FCustomizableObjectEditor::GetAbsoluteCOParent(const UCusto
 }
 
 
-void FCustomizableObjectEditor::AddCachedReferencers(const FName& PathName, TArray<FName>& ArrayReferenceNames, TArray<FAssetData>& ArrayAssetData)
-{
-	ArrayReferenceNames.Empty();
-	FAssetRegistryModule& AssetRegistryModule = FModuleManager::LoadModuleChecked<FAssetRegistryModule>("AssetRegistry");
-	AssetRegistryModule.Get().GetReferencers(PathName, ArrayReferenceNames);
-
-	FARFilter Filter;
-	for (const FName& ReferenceName : ArrayReferenceNames)
-	{
-		bool IsCachedInAssetData = false;
-
-		const int32 MaxIndex = ArrayAssetData.Num();
-
-		for (int32 i = 0; i < MaxIndex; ++i)
-		{
-			if (ArrayAssetData[i].PackageName.ToString() == ReferenceName.ToString())
-			{
-				IsCachedInAssetData = true;
-			}
-		}
-
-		if (!IsCachedInAssetData && !ReferenceName.ToString().StartsWith(TEXT("/TempAutosave")))
-		{
-			Filter.PackageNames.Add(ReferenceName);
-		}
-	}
-
-	Filter.bIncludeOnlyOnDiskAssets = false;
-
-	TArray<FAssetData> ArrayAssetDataTemp;
-	AssetRegistryModule.Get().GetAssets(Filter, ArrayAssetDataTemp);
-
-	// Store only those which have static class type Customizable Object, to avoid loading not needed elements
-	const int32 MaxIndex = ArrayAssetDataTemp.Num();
-	for (int32 i = 0; i < MaxIndex; ++i)
-	{
-		if (ArrayAssetDataTemp[i].GetClass() == UCustomizableObject::StaticClass())
-		{
-			ArrayAssetData.Add(ArrayAssetDataTemp[i]);
-		}
-	}
-}
-
-
 void FCustomizableObjectEditor::CreatePreviewActor()
 {
 	if (Actor)
diff --git a/Engine/Plugins/Mutable/Source/CustomizableObjectEditor/Private/MuCOE/CustomizableObjectEditor.h b/Engine/Plugins/Mutable/Source/CustomizableObjectEditor/Private/MuCOE/CustomizableObjectEditor.h
--- a/Engine/Plugins/Mutable/Source/CustomizableObjectEditor/Private/MuCOE/CustomizableObjectEditor.h
+++ b/Engine/Plugins/Mutable/Source/CustomizableObjectEditor/Private/MuCOE/CustomizableObjectEditor.h
@@ -380,7 +380,6 @@ public:
 
 	// Helpers to get the absolute parent of a Customizable Object
 	static UCustomizableObject* GetAbsoluteCOParent(const UCustomizableObjectNodeObject* const Root);
-	static void AddCachedReferencers(const FName& PathName, TArray<FName>& ArrayReferenceNames, TArray<FAssetData>& ArrayAssetData);
 
 	/**	The tab ids for all the tabs used */
 	static const FName ViewportTabId;
diff --git a/Engine/Plugins/Mutable/Source/CustomizableObjectEditor/Private/MuCOE/CustomizableObjectEditorModule.cpp b/Engine/Plugins/Mutable/Source/CustomizableObjectEditor/Private/MuCOE/CustomizableObjectEditorModule.cpp
--- a/Engine/Plugins/Mutable/Source/CustomizableObjectEditor/Private/MuCOE/CustomizableObjectEditorModule.cpp
+++ b/Engine/Plugins/Mutable/Source/CustomizableObjectEditor/Private/MuCOE/CustomizableObjectEditorModule.cpp
@@ -77,6 +77,7 @@
 #include "MuCOE/GraphTraversal.h"
 #include "MuCOE/CustomizableObjectInstanceBaker.h"
 #include "Editor.h"
+#include "MuCO/CustomizableObjectSystemPrivate.h"
 #include "Nodes/CustomizableObjectNodeComponentMeshDetails.h"
 
 class AActor;
@@ -164,18 +165,19 @@ void ShowOnScreenCompileWarnings()
 		
 		const uint64 KeyCompiledOutOfDate = reinterpret_cast<uint64>(Object) + KEY_OFFSET_COMPILATION_OUT_OF_DATE; // Offset added to avoid collision with bIsCompiledWithOptimization warning
 		TArray<FName> OutOfDatePackages;
-		if (Object->GetPrivate()->IsCompilationOutOfDate(&OutOfDatePackages))
+		TArray<FName> AddedPackages;
+		TArray<FName> RemovedPackages;
+		bool bVersionDiff;
+		if (Object->GetPrivate()->IsCompilationOutOfDate(true, OutOfDatePackages, AddedPackages, RemovedPackages, bVersionDiff))
 		{
 			FString Msg = FString::Printf(TEXT("Customizable Object [%s] compilation out of date. See the Output Log for more information."), *Object->GetName());
 			GEngine->AddOnScreenDebugMessage(KeyCompiledOutOfDate, ShowOnScreenCompileWarningsTickerTime * 2.0f, FColor::Yellow, Msg);
 			
 			if (!GEngine->OnScreenDebugMessageExists(KeyCompiledOutOfDate))
 			{
-				UE_LOG(LogMutable, Verbose, TEXT("Warning: Customizable Object [%s] compilation out of date. Modified packages since last compilation:"), *Object->GetName());
-				for (const FName& OutOfDatePackage : OutOfDatePackages)
-				{
-					UE_LOG(LogMutable, Verbose, TEXT("%s"), *OutOfDatePackage.ToString());
-				}
+				UE_LOG(LogMutable, Display, TEXT("Customizable Object [%s] compilation out of date. Changes since last compilation:"), *Object->GetName());
+
+				PrintParticipatingPackagesDiff(OutOfDatePackages, AddedPackages, RemovedPackages, bVersionDiff);
 			}
 		}
 		else
@@ -474,7 +476,7 @@ void FCustomizableObjectEditorModule::RegisterFactory()
 
 
 /** Recursively get all Customizable Objects that reference the given Customizable Object. */
-void GetReferencingCustomizableObjects(FName CustomizableObjectName, TArray<FName>& VisitedObjectNames, TArray<FName>& ObjectNames)
+void GetReferencingCustomizableObjects(FName CustomizableObjectName, TArray<FName>& VisitedObjectNames, TArray<FAssetData>& ReferencingAssets)
 {
 	if (VisitedObjectNames.Contains(CustomizableObjectName))
 	{
@@ -488,11 +490,17 @@ void GetReferencingCustomizableObjects(FName CustomizableObjectName, TArray<FNam
 	const FAssetRegistryModule& AssetRegistryModule = FModuleManager::LoadModuleChecked<FAssetRegistryModule>("AssetRegistry");
 	AssetRegistryModule.Get().GetReferencers(CustomizableObjectName, ReferencedObjectNames, UE::AssetRegistry::EDependencyCategory::Package, UE::AssetRegistry::EDependencyQuery::Hard);
 
+	// Required to be deterministic.
+	ReferencedObjectNames.Sort([](const FName& A, const FName& B)
+	{
+		return A.LexicalLess(B);
+	});
+	
 	TArray<FAssetData> AssetDataArray;
 
 	FARFilter Filter;
 	Filter.PackageNames = MoveTemp(ReferencedObjectNames);
-	
+
 	AssetRegistryModule.Get().GetAssets(Filter, AssetDataArray);
 
 	for (FAssetData AssetData : AssetDataArray)
@@ -501,25 +509,25 @@ void GetReferencingCustomizableObjects(FName CustomizableObjectName, TArray<FNam
 		{
 			FName ReferencedObjectName = AssetData.GetPackage()->GetFName();
 	
-			ObjectNames.Add(ReferencedObjectName);
+			ReferencingAssets.Add(AssetData);
 
-			GetReferencingCustomizableObjects(ReferencedObjectName, VisitedObjectNames, ObjectNames);
-		}			
+			GetReferencingCustomizableObjects(ReferencedObjectName, VisitedObjectNames, ReferencingAssets);
+		}
 	}
 }
 
 
-void GetReferencingPackages(const UCustomizableObject& Object, TArray<FName>& ObjectNames)
+void GetReferencingPackages(const UCustomizableObject& Object, TArray<FAssetData>& ReferencingAssets)
 {
 	// Gather all child CustomizableObjects
 	TArray<FName> VisitedObjectNames;
-	GetReferencingCustomizableObjects(Object.GetPackage()->GetFName(), VisitedObjectNames, ObjectNames);
+	GetReferencingCustomizableObjects(Object.GetPackage()->GetFName(), VisitedObjectNames, ReferencingAssets);
 
 	// Gather all tables which will composite the final tables
-	TArray<FName> CustomizableObjectNames = ObjectNames;
-	for (const FName& CustomizableObjectName : CustomizableObjectNames)
+	TArray<FAssetData> ReferencingCustomizableObjects = ReferencingAssets;
+	for (const FAssetData& ReferencingCustomizableObject : ReferencingCustomizableObjects)
 	{
-		const TSoftObjectPtr<UObject> SoftObjectPtr = TSoftObjectPtr<UObject>(FSoftObjectPath(CustomizableObjectName.ToString()));
+		const TSoftObjectPtr SoftObjectPtr(ReferencingCustomizableObject.ToSoftObjectPath());
 
 		const UCustomizableObject* ChildCustomizableObject = Cast<UCustomizableObject>(SoftObjectPtr.LoadSynchronous());
 		if (!ChildCustomizableObject)
@@ -541,7 +549,7 @@ void GetReferencingPackages(const UCustomizableObject& Object, TArray<FName>& Ob
 			{
 				if (DataTableAsset.IsValid())
 				{
-					ObjectNames.AddUnique(DataTableAsset.PackageName);
+					ReferencingAssets.AddUnique(DataTableAsset);
 				}
 			}
 		}		
@@ -549,26 +557,21 @@ void GetReferencingPackages(const UCustomizableObject& Object, TArray<FName>& Ob
 }
 
 
-bool FCustomizableObjectEditorModule::IsCompilationOutOfDate(const UCustomizableObject& Object, TArray<FName>* OutOfDatePackages) const
+bool FCustomizableObjectEditorModule::IsCompilationOutOfDate(const UCustomizableObject& Object, bool bSkipIndirectReferences, TArray<FName>& OutOfDatePackages, TArray<FName>& AddedPackages, TArray<FName>& RemovedPackages, bool& bVersionDiff) const
 {
-	if (!Object.GetPrivate()->DirtyParticipatingObjects.IsEmpty())
-	{
-		if (OutOfDatePackages)
-		{
-			OutOfDatePackages->Append(Object.GetPrivate()->DirtyParticipatingObjects);
-		}
-		else
-		{
-			return true;
-		}
-	}
+	MUTABLE_CPUPROFILER_SCOPE(FCustomizableObjectEditorModule::IsCompilationOutOfDate)
+	
+	// TODO CO Custom version
+	// TODO List of plugins and their custom versions
+	// Maybe use BuildDerivedDataKey? BuildDerivedDataKey should also consider bSkipIndirectReferences
 	
 	const FAssetRegistryModule& AssetRegistryModule = FModuleManager::LoadModuleChecked<FAssetRegistryModule>("AssetRegistry");
 
-	const FModelResources& ModelResources = Object.GetPrivate()->GetModelResources();
-	for (const TTuple<FName, FGuid>& ParticipatingObject : ModelResources.ParticipatingObjects)
+	const TMap<FName, FGuid>& OldParticipatingObjects = Object.GetPrivate()->GetModelResources().ParticipatingObjects;
+	
+	for (const TTuple<FName, FGuid>& ParticipatingObject : OldParticipatingObjects)
 	{
-		TSoftObjectPtr<UObject> SoftObjectPtr = TSoftObjectPtr<UObject>(FSoftObjectPath(ParticipatingObject.Key.ToString()));
+		TSoftObjectPtr SoftObjectPtr(FSoftObjectPath(ParticipatingObject.Key.ToString()));
 		if (SoftObjectPtr) // If loaded
 		{
 			PRAGMA_DISABLE_DEPRECATION_WARNINGS
@@ -577,14 +580,7 @@ bool FCustomizableObjectEditorModule::IsCompilationOutOfDate(const UCustomizable
 			
 			if (PackageGuid != ParticipatingObject.Value)
 			{
-				if (OutOfDatePackages)
-				{
-					OutOfDatePackages->AddUnique(ParticipatingObject.Key);
-				}
-				else
-				{
-					return true;
-				}
+				OutOfDatePackages.AddUnique(ParticipatingObject.Key);
 			}
 		}
 		else // Not loaded
@@ -594,14 +590,7 @@ bool FCustomizableObjectEditorModule::IsCompilationOutOfDate(const UCustomizable
 				
 			if (Result != UE::AssetRegistry::EExists::Exists)
 			{
-				if (OutOfDatePackages)
-				{
-					OutOfDatePackages->AddUnique(ParticipatingObject.Key);
-				}
-				else
-				{
-					return true;
-				}
+				OutOfDatePackages.AddUnique(ParticipatingObject.Key);
 			}
 
 			PRAGMA_DISABLE_DEPRECATION_WARNINGS
@@ -610,47 +599,49 @@ bool FCustomizableObjectEditorModule::IsCompilationOutOfDate(const UCustomizable
 			
 			if (PackageGuid != ParticipatingObject.Value)
 			{
-				if (OutOfDatePackages)
-				{
-					OutOfDatePackages->AddUnique(ParticipatingObject.Key);
-				}
-				else
-				{
-					return true;
-				}
+				OutOfDatePackages.AddUnique(ParticipatingObject.Key);
 			}
 		}
 	}
-
-	TArray<FName> ReferencingObjectNames;
-	GetReferencingPackages(Object, ReferencingObjectNames);
 	
-	for (const FName& ObjectName : ReferencingObjectNames)
+	// Check that we have the exact same set of participating object as before. This can change due to indirect references and versioning.
+	if (!bSkipIndirectReferences)
 	{
-		TSoftObjectPtr<UObject> ReferencingObject = TSoftObjectPtr<UObject>(FSoftObjectPath(ObjectName.ToString())); 
-
-		if ((ReferencingObject && ReferencingObject->GetPackage()->IsDirty()) ||
-			!ModelResources.ParticipatingObjects.Contains(ObjectName)) // Must be in the participating objects, if not it means it did not exist when compiling the object.
+		// Due to performance issues, we will skip loading all objects. We can do that since loading/not loading objects do not affect the number of indirect objects discovered
+		// (e.g., we will traverse the same number of COs/Tables regardless if we do not load meshes/textures...). 
+		TMap<FName, FGuid> ParticipatingObjects = GetParticipatingObjects(&Object, false);
+		
+		for (const TTuple<FName, FGuid>& ParticipatingObject : ParticipatingObjects)
 		{
-			if (OutOfDatePackages)
+			// Since here we are if the smaller set (objects found now without loading all objects) is contained in the larger set (objects found in the compilation pass),
+			// there is no need to check if the asset is a indirect reference (CO or Table).
+			if (!OldParticipatingObjects.Contains(ParticipatingObject.Key))
 			{
-				OutOfDatePackages->AddUnique(ObjectName);
+				AddedPackages.AddUnique(ParticipatingObject.Key);
 			}
-			else
+		}
+
+		for (const TTuple<FName, FGuid>& OldParticipatingObject : OldParticipatingObjects)
+		{
+			FAssetData AssetData = AssetRegistryModule.Get().GetAssetByObjectPath(FSoftObjectPath(OldParticipatingObject.Key.ToString()));
+			if (AssetData.AssetClassPath == UCustomizableObject::StaticClass()->GetClassPathName() ||
+				AssetData.AssetClassPath == UDataTable::StaticClass()->GetClassPathName())
 			{
-				return true;
+				if (!ParticipatingObjects.Contains(OldParticipatingObject.Key))
+				{
+					RemovedPackages.AddUnique(OldParticipatingObject.Key);
+				}
 			}
 		}
 	}
 
-	if (OutOfDatePackages)
-	{
-		return !OutOfDatePackages->IsEmpty();	
-	}
-	else
+	bVersionDiff = false;
+	if (ICustomizableObjectVersionBridgeInterface* VersionBridge = Cast<ICustomizableObjectVersionBridgeInterface>(Object.VersionBridge))
 	{
-		return false;
+		bVersionDiff = Object.GetPrivate()->GetModelResources().CompiledVersionBridge != VersionBridge->GetCurrentVersionAsString();
 	}
+
+	return bVersionDiff || !OutOfDatePackages.IsEmpty() || !AddedPackages.IsEmpty() || !RemovedPackages.IsEmpty();
 }
 
 
@@ -675,6 +666,18 @@ FString FCustomizableObjectEditorModule::GetCurrentContentVersionForObject(const
 }
 
 
+UCustomizableObject* FCustomizableObjectEditorModule::GetRootObject(UCustomizableObject* ChildObject) const
+{
+	return GraphTraversal::GetRootObject(ChildObject);
+}
+
+
+const UCustomizableObject* FCustomizableObjectEditorModule::GetRootObject(const UCustomizableObject* ChildObject) const
+{
+	return GraphTraversal::GetRootObject(ChildObject);
+}
+
+
 void FCustomizableObjectEditorModule::BakeCustomizableObjectInstance(UCustomizableObjectInstance* InTargetInstance, const FBakingConfiguration& InBakingConfig)
 {
 	UCustomizableObjectInstanceBaker* InstanceBaker = NewObject<UCustomizableObjectInstanceBaker>();
@@ -719,6 +722,22 @@ USkeletalMesh* FCustomizableObjectEditorModule::GetReferenceSkeletalMesh(const U
 }
 
 
+TMap<FName, FGuid> FCustomizableObjectEditorModule::GetParticipatingObjects(const UCustomizableObject* Object, bool bLoadObjects, const FCompilationOptions* InOptions) const
+{
+	MUTABLE_CPUPROFILER_SCOPE(FCustomizableObjectEditorModule::GetParticipatingObjects)
+
+	FCompilationOptions Options = InOptions ? *InOptions : Object->GetPrivate()->GetCompileOptions();
+		
+	FMutableGraphGenerationContext Context(Object, nullptr, Options);
+	Context.bParticipatingObjectsPass = true;
+	Context.bLoadObjects = bLoadObjects;
+
+	GenerateMutableRoot(Object, Context);
+
+	return Context.ParticipatingObjects;
+}
+
+
 void FCustomizableObjectEditorModule::CompileCustomizableObject(const TSharedRef<FCompilationRequest>& InCompilationRequest, bool bForceRequest)
 {
 	if (IsRunningGame())
diff --git a/Engine/Plugins/Mutable/Source/CustomizableObjectEditor/Private/MuCOE/CustomizableObjectEditorModule.h b/Engine/Plugins/Mutable/Source/CustomizableObjectEditor/Private/MuCOE/CustomizableObjectEditorModule.h
--- a/Engine/Plugins/Mutable/Source/CustomizableObjectEditor/Private/MuCOE/CustomizableObjectEditorModule.h
+++ b/Engine/Plugins/Mutable/Source/CustomizableObjectEditor/Private/MuCOE/CustomizableObjectEditorModule.h
@@ -23,7 +23,7 @@ class FBakeOperationCompletedDelegate;
   * List includes:
   * - Child UCustomizableObjects: Have inverted references.
   * - UDataTable: Data Tables used by Composite Data Tables are indirectly referenced by the UStruct and filtered by path. */
-void GetReferencingPackages(const UCustomizableObject& Object, TArray<FName>& ObjectNames);
+void GetReferencingPackages(const UCustomizableObject& Object, TArray<FAssetData>& ObjectNames);
 
 
 /**
@@ -38,12 +38,15 @@ public:
 
 	// ICustomizableObjectEditorModule interface
 	virtual FCustomizableObjectEditorLogger& GetLogger() override;
-	virtual bool IsCompilationOutOfDate(const UCustomizableObject& Object, TArray<FName>* OutOfDatePackages) const override;
+	virtual bool IsCompilationOutOfDate(const UCustomizableObject& Object, bool bSkipIndirectReferences, TArray<FName>& OutOfDatePackages, TArray<FName>& AddedPackages, TArray<FName>& RemovedPackages, bool& bVersionDiff) const override;
 	virtual bool IsRootObject(const UCustomizableObject& Object) const override;
 	virtual FString GetCurrentContentVersionForObject(const UCustomizableObject& Object) const override;
+	virtual UCustomizableObject* GetRootObject(UCustomizableObject* ChildObject) const override;
+	virtual const UCustomizableObject* GetRootObject(const UCustomizableObject* ChildObject) const override;
 	virtual void BakeCustomizableObjectInstance(UCustomizableObjectInstance* InTargetInstance, const FBakingConfiguration& InBakingConfig) override;
 	virtual USkeletalMesh* GetReferenceSkeletalMesh(const UCustomizableObject& Object, const FName& ComponentName) const override;
-	
+	virtual TMap<FName, FGuid> GetParticipatingObjects(const UCustomizableObject* Object, bool bLoadObjects, const FCompilationOptions* Options = nullptr) const override;
+
 	/** Request for a given customizable object to be compiled. Async compile requests will be queued and processed sequentially. 
 	 * @param InCompilationRequest - Request to compile an object. 
 	 * @param bForceRequest - Queue request even if already in the pending list. */
diff --git a/Engine/Plugins/Mutable/Source/CustomizableObjectEditor/Private/MuCOE/CustomizableObjectInstanceEditor.cpp b/Engine/Plugins/Mutable/Source/CustomizableObjectEditor/Private/MuCOE/CustomizableObjectInstanceEditor.cpp
--- a/Engine/Plugins/Mutable/Source/CustomizableObjectEditor/Private/MuCOE/CustomizableObjectInstanceEditor.cpp
+++ b/Engine/Plugins/Mutable/Source/CustomizableObjectEditor/Private/MuCOE/CustomizableObjectInstanceEditor.cpp
@@ -329,7 +329,7 @@ void FCustomizableObjectInstanceEditor::InitCustomizableObjectInstanceEditor( co
 	// Clears selection highlight.
 	OnInstancePropertySelectionChanged(NULL);
 	
-	if (const UCustomizableObject* CustomizableObject = CustomizableObjectInstance->GetCustomizableObject())
+	if (UCustomizableObject* CustomizableObject = CustomizableObjectInstance->GetCustomizableObject())
 	{
 		UCustomizableObjectPrivate* CustomizableObjectPrivate = CustomizableObject->GetPrivate();
 
diff --git a/Engine/Plugins/Mutable/Source/CustomizableObjectEditor/Private/MuCOE/ExtensionDataCompilerInterface.cpp b/Engine/Plugins/Mutable/Source/CustomizableObjectEditor/Private/MuCOE/ExtensionDataCompilerInterface.cpp
--- a/Engine/Plugins/Mutable/Source/CustomizableObjectEditor/Private/MuCOE/ExtensionDataCompilerInterface.cpp
+++ b/Engine/Plugins/Mutable/Source/CustomizableObjectEditor/Private/MuCOE/ExtensionDataCompilerInterface.cpp
@@ -18,30 +18,20 @@ mu::ExtensionDataPtrConst FExtensionDataCompilerInterface::MakeStreamedExtension
 	mu::ExtensionDataPtr Result = new mu::ExtensionData;
 	Result->Origin = mu::ExtensionData::EOrigin::ConstantStreamed;
 	Result->Index = GenerationContext.StreamedExtensionData.Num();
-
-	// Generate a deterministic name to help with deterministic cooking
-	const FString ContainerName = FString::Printf(TEXT("Streamed_%d"), Result->Index);
-
-	UObject* ExistingObject = FindObject<UObject>(GenerationContext.Object, *ContainerName);
-	if (ExistingObject)
+	
+	if (!GenerationContext.bParticipatingObjectsPass)
 	{
-		// This must have been left behind from a previous compilation and hasn't been deleted by 
-		// GC yet.
-		//
-		// Move it into the transient package to get it out of the way.
-		ExistingObject->Rename(nullptr, GetTransientPackage(), REN_DontCreateRedirectors);
-
-		check(!FindObject<UObject>(GenerationContext.Object, *ContainerName));
-	}
-
-	check(GenerationContext.Object);
-	OutContainer = NewObject<UCustomizableObjectResourceDataContainer>(
-		GenerationContext.Object,
-		FName(*ContainerName),
-		RF_Public);
+		// Generate a deterministic name to help with deterministic cooking
+		const FString ContainerName = FString::Printf(TEXT("Streamed_%d"), Result->Index);
 
-	GenerationContext.StreamedExtensionData.Add(OutContainer);
+		OutContainer = NewObject<UCustomizableObjectResourceDataContainer>(
+			GetTransientPackage(),
+			FName(*ContainerName),
+			RF_Public);
 
+		GenerationContext.StreamedExtensionData.Emplace(ContainerName, OutContainer);
+	}
+	
 	return Result;
 }
 
@@ -57,7 +47,7 @@ mu::ExtensionDataPtrConst FExtensionDataCompilerInterface::MakeAlwaysLoadedExten
 	return Result;
 }
 
-UObject* FExtensionDataCompilerInterface::GetOuterForAlwaysLoadedObjects()
+const UObject* FExtensionDataCompilerInterface::GetOuterForAlwaysLoadedObjects()
 {
 	check(GenerationContext.Object);
 	return GenerationContext.Object;
@@ -76,5 +66,11 @@ void FExtensionDataCompilerInterface::AddGeneratedNode(const UCustomizableObject
 
 void FExtensionDataCompilerInterface::CompilerLog(const FText& InLogText, const UCustomizableObjectNode* InNode)
 {
-	GenerationContext.Compiler->CompilerLog(InLogText, InNode);
+	GenerationContext.Log(InLogText, InNode);
 }
+
+void FExtensionDataCompilerInterface::AddParticipatingObject(const UObject& Object)
+{
+	GenerationContext.AddParticipatingObject(Object);
+}
+
diff --git a/Engine/Plugins/Mutable/Source/CustomizableObjectEditor/Private/MuCOE/GenerateMutableSource/GenerateMutableSource.cpp b/Engine/Plugins/Mutable/Source/CustomizableObjectEditor/Private/MuCOE/GenerateMutableSource/GenerateMutableSource.cpp
--- a/Engine/Plugins/Mutable/Source/CustomizableObjectEditor/Private/MuCOE/GenerateMutableSource/GenerateMutableSource.cpp
+++ b/Engine/Plugins/Mutable/Source/CustomizableObjectEditor/Private/MuCOE/GenerateMutableSource/GenerateMutableSource.cpp
@@ -52,6 +52,7 @@
 #include "MuR/Mesh.h"
 #include "PhysicsEngine/PhysicsAsset.h"
 #include "PlatformInfo.h"
+#include "AssetRegistry/AssetRegistryModule.h"
 #include "Math/NumericLimits.h"
 #include "Hash/CityHash.h"
 #include "MuCOE/Nodes/CustomizableObjectNodeCopyMaterial.h"
@@ -84,7 +85,7 @@ bool FGraphCycle::FoundCycle() const
 
 	if (const UCustomizableObject** Result = Context.VisitedPins.Find(Key))
 	{
-		Context.Compiler->CompilerLog(LOCTEXT("CycleFoundNode", "Cycle detected."), &Node, EMessageSeverity::Error, true);
+		Context.Log(LOCTEXT("CycleFoundNode", "Cycle detected."), &Node, EMessageSeverity::Error, true);
 		Context.CustomizableObjectWithCycle = *Result;
 		return true;
 	}
@@ -114,14 +115,14 @@ void CheckNumOutputs(const UEdGraphPin& Pin, const FMutableGraphGenerationContex
 
 			if (numOutLinks > 1)
 			{
-				GenerationContext.Compiler->CompilerLog(LOCTEXT("MultipleOutgoing", "The node has several outgoing connections, but it should be limited to 1."), CastChecked<UCustomizableObjectNode>(Pin.GetOwningNode()));
+				GenerationContext.Log(LOCTEXT("MultipleOutgoing", "The node has several outgoing connections, but it should be limited to 1."), CastChecked<UCustomizableObjectNode>(Pin.GetOwningNode()));
 			}
 		}
 	}
 }
 
 
-FMutableGraphGenerationContext::FMutableGraphGenerationContext(UCustomizableObject* InObject, FCustomizableObjectCompiler* InCompiler, const FCompilationOptions& InOptions)
+FMutableGraphGenerationContext::FMutableGraphGenerationContext(const UCustomizableObject* InObject, FCustomizableObjectCompiler* InCompiler, const FCompilationOptions& InOptions)
 	: Object(InObject), Compiler(InCompiler), Options(InOptions)
 	, ExtensionDataCompilerInterface(*this)
 {
@@ -135,13 +136,56 @@ FMutableGraphGenerationContext::FMutableGraphGenerationContext(UCustomizableObje
 FMutableGraphGenerationContext::~FMutableGraphGenerationContext() = default;
 
 
+void FMutableGraphGenerationContext::AddParticipatingObject(const FSoftObjectPath& SoftPath)
+{
+	const FAssetRegistryModule& AssetRegistryModule = FModuleManager::LoadModuleChecked<FAssetRegistryModule>("AssetRegistry");
+
+	FAssetPackageData AssetPackageData;
+	const UE::AssetRegistry::EExists Result = AssetRegistryModule.Get().TryGetAssetPackageData(SoftPath.GetLongPackageFName(), AssetPackageData);
+	if (Result != UE::AssetRegistry::EExists::Exists)
+	{
+		return;
+	}
+
+	PRAGMA_DISABLE_DEPRECATION_WARNINGS
+	const FGuid PackageGuid = AssetPackageData.PackageGuid;
+	PRAGMA_ENABLE_DEPRECATION_WARNINGS
+		
+	AddParticipatingObjectChecked(SoftPath.GetLongPackageFName(), PackageGuid);
+}
+
+
+void FMutableGraphGenerationContext::AddParticipatingObject(const FSoftObjectPtr& SoftObject)
+{
+	AddParticipatingObject(SoftObject.ToSoftObjectPath());
+}
+
+
 void FMutableGraphGenerationContext::AddParticipatingObject(const UObject& InObject)
 {
 	PRAGMA_DISABLE_DEPRECATION_WARNINGS
 	const FGuid PackageGuid = InObject.GetPackage()->GetGuid();
 	PRAGMA_ENABLE_DEPRECATION_WARNINGS
 
-	ParticipatingObjects.Add(InObject.GetPackage()->GetFName(), PackageGuid);
+	AddParticipatingObjectChecked(InObject.GetPackage()->GetFName(), PackageGuid);
+}
+
+
+void FMutableGraphGenerationContext::Log(const FText& Message, const TArray<const UObject*>& Context, const EMessageSeverity::Type MessageSeverity, const bool bAddBaseObjectInfo, const ELoggerSpamBin SpamBin) const
+{
+	if (Compiler)
+	{
+		Compiler->CompilerLog(Message, Context, MessageSeverity, bAddBaseObjectInfo, SpamBin);
+	}
+}
+
+
+void FMutableGraphGenerationContext::Log(const FText& Message, const UObject* Context, const EMessageSeverity::Type MessageSeverity, const bool bAddBaseObjectInfo, const ELoggerSpamBin SpamBin) const
+{
+	if (!bParticipatingObjectsPass)
+	{
+		Compiler->CompilerLog(Message, Context, MessageSeverity, bAddBaseObjectInfo, SpamBin);
+	}
 }
 
 
@@ -161,6 +205,11 @@ mu::Ptr<mu::Mesh> FMutableGraphGenerationContext::FindGeneratedMesh( const FGene
 
 int32 FMutableGraphGenerationContext::AddStreamedResource(uint32 InResourceHash, UCustomizableObjectResourceDataContainer*& OutNewResource)
 {
+	if (bParticipatingObjectsPass)
+	{
+		return INDEX_NONE;	
+	}
+	
 	OutNewResource = nullptr;
 
 	// Return resource index if found.
@@ -168,17 +217,16 @@ int32 FMutableGraphGenerationContext::AddStreamedResource(uint32 InResourceHash,
 	{
 		return *ResourceIndex;
 	}
-
+	
 	int32 NewResourceIndex = StreamedResourceData.Num();
 	const FString ContainerName = GetNameSafe(Object) + FString::Printf(TEXT("_SR_%d"), NewResourceIndex);
 
-	UCustomizableObjectResourceDataContainer* ExistingContainer = FindObject<UCustomizableObjectResourceDataContainer>(Object, *ContainerName);
-	OutNewResource = ExistingContainer ? ExistingContainer : NewObject<UCustomizableObjectResourceDataContainer>(
-		Object,
+	OutNewResource = NewObject<UCustomizableObjectResourceDataContainer>(
+		GetTransientPackage(),
 		FName(*ContainerName),
 		RF_Public);
 
-	StreamedResourceData.Add(OutNewResource);
+	StreamedResourceData.Emplace(ContainerName, OutNewResource);
 	StreamedResourceIndices.Add({ InResourceHash, NewResourceIndex });
 
 	return NewResourceIndex;
@@ -202,7 +250,6 @@ int32 FMutableGraphGenerationContext::AddAssetUserDataToStreamedResources(UAsset
 		NewResource->Data.Data = FInstancedStruct::Make(ResourceData);
 	}
 
-	check(StreamedResourceData[ResourceIndex].GetLoadedData().Type == ECOResourceDataType::AssetUserData);
 	return ResourceIndex;
 }
 
@@ -406,6 +453,31 @@ FMutableComponentInfo* FMutableGraphGenerationContext::GetCurrentComponentInfo()
 }
 
 
+UObject* FMutableGraphGenerationContext::LoadObject(const FSoftObjectPtr& SoftObject, bool bParticipatingObjectsPassLoad)
+{
+	return bLoadObjects && (!bParticipatingObjectsPass || bParticipatingObjectsPassLoad) ?
+		SoftObject.LoadSynchronous() :
+		nullptr;
+}
+
+
+void FMutableGraphGenerationContext::AddParticipatingObjectChecked(const FName& PackageName, const FGuid& PackageGuid)
+{
+	if (bParticipatingObjectsPass)
+	{
+		ParticipatingObjects.Emplace(PackageName, PackageGuid);
+	}
+	else
+	{
+		checkCode
+		(
+			FGuid* Result = ParticipatingObjects.Find(PackageName);
+			check(Result && *Result == PackageGuid); // If this check is hit it means that this Participating Object is not being discovered in the Participating Objects pass.
+		)
+	}
+}
+
+
 FGeneratedKey::FGeneratedKey(void* InFunctionAddress, const UEdGraphPin& InPin, const UCustomizableObjectNode& Node, FMutableGraphGenerationContext& GenerationContext, const bool UseMesh, const bool InbOnlyConnectedLOD)
 {
 	FunctionAddress = InFunctionAddress;
@@ -798,7 +870,7 @@ mu::Ptr<mu::NodeObject> GenerateMutableSource(const UEdGraphPin * Pin, FMutableG
 		FGuid FinalGuid = GenerationContext.GetNodeIdUnique(TypedNodeObj);
 		if (FinalGuid != TypedNodeObj->NodeGuid)
 		{
-			GenerationContext.Compiler->CompilerLog(FText::FromString(TEXT("Warning: Node has a duplicated GUID. A new ID has been generated, but cooked data will not be deterministic.")), Node, EMessageSeverity::Warning);
+			GenerationContext.Log(FText::FromString(TEXT("Warning: Node has a duplicated GUID. A new ID has been generated, but cooked data will not be deterministic.")), Node, EMessageSeverity::Warning);
 		}
 		ObjectNode->SetUid(FinalGuid.ToString());
 		
@@ -906,7 +978,7 @@ mu::Ptr<mu::NodeObject> GenerateMutableSource(const UEdGraphPin * Pin, FMutableG
 				FString Msg = FString::Printf(TEXT("Extension input %s has multiple incoming connections but is only expecting one connection."),
 					*ExtensionInputPin.InputPin.DisplayName.ToString());
 
-				GenerationContext.Compiler->CompilerLog(FText::FromString(Msg), Node, EMessageSeverity::Warning);
+				GenerationContext.Log(FText::FromString(Msg), Node, EMessageSeverity::Warning);
 			}
 
 			for (const UEdGraphPin* ConnectedPin : ConnectedPins)
@@ -966,7 +1038,7 @@ mu::Ptr<mu::NodeObject> GenerateMutableSource(const UEdGraphPin * Pin, FMutableG
 		case ECustomizableObjectGroupType::COGT_ONE: Type = mu::NodeObjectGroup::CS_ALWAYS_ONE; break;
 		case ECustomizableObjectGroupType::COGT_ONE_OR_NONE: Type = mu::NodeObjectGroup::CS_ONE_OR_NONE; break;
 		default:
-			GenerationContext.Compiler->CompilerLog(LOCTEXT("UnsupportedGroupType", "Object Group Type not supported. Setting to 'ALL'."), Node);
+			GenerationContext.Log(LOCTEXT("UnsupportedGroupType", "Object Group Type not supported. Setting to 'ALL'."), Node);
 			break;
 		}
 		GroupNode->SetSelectionType(Type);
@@ -1002,7 +1074,7 @@ mu::Ptr<mu::NodeObject> GenerateMutableSource(const UEdGraphPin * Pin, FMutableG
 
 			UCustomizableObjectNodeObject* CustomizableObjectNodeObject = Cast<UCustomizableObjectNodeObject>(ConnectedChildrenPins[ChildIndex]->GetOwningNode());
 
-			FString* SelectedOptionName = GenerationContext.ParamNamesToSelectedOptions.Find(TypedNodeGroup->GroupName); // If the param is in the map restrict to only the selected option
+			const FString* SelectedOptionName = GenerationContext.Options.ParamNamesToSelectedOptions.Find(TypedNodeGroup->GroupName); // If the param is in the map restrict to only the selected option
 			mu::NodeObjectPtr ChildNode = nullptr;
 
 			if (bConnectAtLeastTheLastChild || !SelectedOptionName || (CustomizableObjectNodeObject && *SelectedOptionName == CustomizableObjectNodeObject->ObjectName) )
@@ -1110,7 +1182,7 @@ mu::Ptr<mu::NodeObject> GenerateMutableSource(const UEdGraphPin * Pin, FMutableG
 
 			UCustomizableObjectNodeObject* CustomizableObjectNodeObject = Cast<UCustomizableObjectNodeObject>(ExternalChildNode->OutputPin()->GetOwningNode());
 
-			FString* SelectedOptionName = GenerationContext.ParamNamesToSelectedOptions.Find(TypedNodeGroup->GroupName); // If the param is in the map restrict to only the selected option
+			const FString* SelectedOptionName = GenerationContext.Options.ParamNamesToSelectedOptions.Find(TypedNodeGroup->GroupName); // If the param is in the map restrict to only the selected option
 			mu::NodeObjectPtr ChildNode = nullptr;
 
 			if (bConnectAtLeastTheLastChild || !SelectedOptionName || (CustomizableObjectNodeObject  && *SelectedOptionName == CustomizableObjectNodeObject->ObjectName) )
@@ -1199,7 +1271,7 @@ mu::Ptr<mu::NodeObject> GenerateMutableSource(const UEdGraphPin * Pin, FMutableG
 	}
 	else
 	{
-		GenerationContext.Compiler->CompilerLog(LOCTEXT("UnimplementedNode", "Node type not implemented yet."), Node);
+		GenerationContext.Log(LOCTEXT("UnimplementedNode", "Node type not implemented yet."), Node);
 	}
 
 	GenerationContext.Generated.Add(Key, FGeneratedData(Node, Result));
@@ -1364,9 +1436,11 @@ void PopulateReferenceSkeletalMeshesData(FMutableGraphGenerationContext& Generat
 				{
 					FMutableRefAssetUserData& MutAssetUserData = Data.AssetUserData.AddDefaulted_GetRef();
 					MutAssetUserData.AssetUserDataIndex = GenerationContext.AddAssetUserDataToStreamedResources(AssetUserData);
-					MutAssetUserData.AssetUserData = GenerationContext.StreamedResourceData[MutAssetUserData.AssetUserDataIndex].GetPath().Get();
-					check(MutAssetUserData.AssetUserData);
-					check(MutAssetUserData.AssetUserData->Data.Type == ECOResourceDataType::AssetUserData);
+					if (MutAssetUserData.AssetUserData != INDEX_NONE)
+					{
+						MutAssetUserData.AssetUserData = GenerationContext.StreamedResourceData[MutAssetUserData.AssetUserDataIndex].Value;
+						check(MutAssetUserData.AssetUserData->Data.Type == ECOResourceDataType::AssetUserData);
+					}
 				}
 			}
 		}
@@ -1586,6 +1660,8 @@ mu::Ptr<mu::Image> GenerateImageConstant(UTexture* Texture, FMutableGraphGenerat
 
 mu::Ptr<mu::Mesh> GenerateMeshConstant(USkeletalMesh* Mesh, FMutableGraphGenerationContext& GenerationContext, bool bIsReference)
 {
+	MUTABLE_CPUPROFILER_SCOPE(GenerateMeshConstant)
+
 	if (!Mesh)
 	{
 		return nullptr;
diff --git a/Engine/Plugins/Mutable/Source/CustomizableObjectEditor/Private/MuCOE/GenerateMutableSource/GenerateMutableSource.h b/Engine/Plugins/Mutable/Source/CustomizableObjectEditor/Private/MuCOE/GenerateMutableSource/GenerateMutableSource.h
--- a/Engine/Plugins/Mutable/Source/CustomizableObjectEditor/Private/MuCOE/GenerateMutableSource/GenerateMutableSource.h
+++ b/Engine/Plugins/Mutable/Source/CustomizableObjectEditor/Private/MuCOE/GenerateMutableSource/GenerateMutableSource.h
@@ -9,6 +9,7 @@
 #include "MuCO/CustomizableObject.h"
 #include "MuCO/CustomizableObjectIdentifier.h"
 #include "MuCO/CustomizableObjectCompilerTypes.h"
+#include "MuCOE/CustomizableObjectEditorLogger.h"
 #include "MuCOE/ExtensionDataCompilerInterface.h"
 #include "MuCOE/Nodes/CustomizableObjectNodeMaterialBase.h"
 #include "MuCOE/Nodes/CustomizableObjectNodeObject.h"
@@ -542,36 +543,42 @@ private:
 
 struct FMutableGraphGenerationContext
 {
-	FMutableGraphGenerationContext(UCustomizableObject* CustomizableObject, class FCustomizableObjectCompiler* InCompiler, const FCompilationOptions& InOptions);
+	FMutableGraphGenerationContext(const UCustomizableObject* CustomizableObject, class FCustomizableObjectCompiler* InCompiler, const FCompilationOptions& InOptions);
 	~FMutableGraphGenerationContext();
 
 	/** See FCustomizableObjectPrivateData::ParticipatingObjects. */
-	void AddParticipatingObject(const UObject& Object);
+	void AddParticipatingObject(const FSoftObjectPath& SoftPath);
+	
+	/** See FCustomizableObjectPrivateData::ParticipatingObjects. */
+	void AddParticipatingObject(const FSoftObjectPtr& SoftObject);
 
 	/** See FCustomizableObjectPrivateData::ParticipatingObjects. */
-	template<typename Type>
-	void AddParticipatingObject(const TArray<Type>& InObjects)
+	template<typename T>
+	void AddParticipatingObject(const TSoftClassPtr<T>& SoftClass)
 	{
-		for (Type InObject : InObjects)
-		{
-			if (InObject)
-			{	
-				AddParticipatingObject(*InObject);
-			}
-		}
+		AddParticipatingObject(SoftClass.ToSoftObjectPath());
 	}
 	
-	UCustomizableObject* Object = nullptr;
+	/** See FCustomizableObjectPrivateData::ParticipatingObjects. */
+	void AddParticipatingObject(const UObject& Object);
+
+	void Log(const FText& Message, const TArray<const UObject*>& UObject, const EMessageSeverity::Type MessageSeverity = EMessageSeverity::Warning, const bool bAddBaseObjectInfo = true, const ELoggerSpamBin SpamBin = ELoggerSpamBin::ShowAll) const;
+
+	void Log(const FText& Message, const UObject* Context = nullptr, const EMessageSeverity::Type MessageSeverity = EMessageSeverity::Warning, const bool bAddBaseObjectInfo = true, const ELoggerSpamBin SpamBin = ELoggerSpamBin::ShowAll) const;
+	
+	const UCustomizableObject* Object = nullptr;
 
 	/** Full hierarchy root. */
 	UCustomizableObjectNodeObject* Root = nullptr;
-	
+
+private:
 	// Non-owned reference to the compiler object
 	FCustomizableObjectCompiler* Compiler = nullptr;
 
+public:
 	// Compilation options, including target platform
-	FCompilationOptions Options;
-
+	const FCompilationOptions& Options;
+	
 	// Cache of generated pins per LOD
 	TMap<FGeneratedKey, FGeneratedData> Generated;
 
@@ -753,6 +760,31 @@ struct FMutableGraphGenerationContext
 	/** Get the reference skeletal mesh associated to the current mesh component being generated */
 	FMutableComponentInfo* GetCurrentComponentInfo();
 
+	UObject* LoadObject(const FSoftObjectPtr& SoftObject, bool bParticipatingObjectsPassLoad = false);
+
+	template<typename T>
+	T* LoadObject(const TSoftObjectPtr<T>& SoftObject, bool bParticipatingObjectsPassLoad = false)
+	{
+		return bLoadObjects && (!bParticipatingObjectsPass || bParticipatingObjectsPassLoad) ?
+			SoftObject.LoadSynchronous() :
+			nullptr;
+	}
+
+	template<typename T>
+	UClass* LoadClass(const TSoftClassPtr<T>& SoftClass, bool bParticipatingObjectsPassLoad = false)
+	{
+		return bLoadObjects && (!bParticipatingObjectsPass || bParticipatingObjectsPassLoad) ?
+			SoftClass.LoadSynchronous() :
+			nullptr;
+	}
+
+private:
+	/** Two operation modes:
+	 * 1. bParticipatingObjectsPass = true, add the participating object.
+	 * 2. bParticipatingObjectsPass = false, check that the participating object was discovered in the participating object pass. */
+	void AddParticipatingObjectChecked(const FName& PackageName, const FGuid& PackageGuid);
+
+public:
 	TMap<FName, UCustomizableObjectNodeComponentMesh*> MeshComponents;
 
 	/** Only Mesh Components (no passthrough). */
@@ -827,6 +859,11 @@ struct FMutableGraphGenerationContext
 
 	bool bPartialCompilation = false;
 
+	bool bParticipatingObjectsPass = false;
+
+	/** Load any Soft Object/Class Pointers. */
+	bool bLoadObjects = true;
+	
 	// Based on the last object visited.
 	ECustomizableObjectAutomaticLODStrategy CurrentAutoLODStrategy = ECustomizableObjectAutomaticLODStrategy::Manual;
 
@@ -855,10 +892,7 @@ struct FMutableGraphGenerationContext
 	* It's a stack because group nodes are recursive
 	*/
 	TArray<int32> SocketPriorityStack;
-
-	// Stores the only option of an Int Param that should be compiled in a partial compilation
-	TMap<FString, FString> ParamNamesToSelectedOptions;
-
+	
 	// Stores what param names use a certain table as a table can be used from multiple table nodes, useful for partial compilations to restrict params
 	TMap<FString, FMutableParamNameSet> TableToParamNames;
 
@@ -898,12 +932,12 @@ struct FMutableGraphGenerationContext
 
 	/** Resource Data constants */
 	TMap<uint32, int32> StreamedResourceIndices;
-	TArray<FCustomizableObjectStreamedResourceData> StreamedResourceData;
+	TArray<TPair<FName, UCustomizableObjectResourceDataContainer*>> StreamedResourceData;
 
 	/** Extension Data constants are collected here */
 	FExtensionDataCompilerInterface ExtensionDataCompilerInterface;
 	TArray<FCustomizableObjectResourceData> AlwaysLoadedExtensionData;
-	TArray<UCustomizableObjectResourceDataContainer*> StreamedExtensionData;
+	TArray<TPair<FName, UCustomizableObjectResourceDataContainer*>> StreamedExtensionData;
 
 	/** See FCustomizableObjectPrivateData::ParticipatingObjects. */
 	TMap<FName, FGuid> ParticipatingObjects;
diff --git a/Engine/Plugins/Mutable/Source/CustomizableObjectEditor/Private/MuCOE/GenerateMutableSource/GenerateMutableSourceColor.cpp b/Engine/Plugins/Mutable/Source/CustomizableObjectEditor/Private/MuCOE/GenerateMutableSource/GenerateMutableSourceColor.cpp
--- a/Engine/Plugins/Mutable/Source/CustomizableObjectEditor/Private/MuCOE/GenerateMutableSource/GenerateMutableSourceColor.cpp
+++ b/Engine/Plugins/Mutable/Source/CustomizableObjectEditor/Private/MuCOE/GenerateMutableSource/GenerateMutableSourceColor.cpp
@@ -89,7 +89,7 @@ mu::Ptr<mu::NodeColour> GenerateMutableSourceColor(const UEdGraphPin* Pin, FMuta
 					? LOCTEXT("NoEnumParamInSwitch", "Switch nodes must have an enum switch parameter. Please connect an enum and refesh the switch node.")
 					: LOCTEXT("InvalidEnumInSwitch", "Switch nodes must have a single enum with all the options inside. Please remove all the enums but one and refresh the switch node.");
 
-				GenerationContext.Compiler->CompilerLog(Message, Node);
+				GenerationContext.Log(Message, Node);
 				return nullptr;
 			}
 
@@ -103,7 +103,7 @@ mu::Ptr<mu::NodeColour> GenerateMutableSourceColor(const UEdGraphPin* Pin, FMuta
 				if (EnumPin)
 				{
 					const FText Message = LOCTEXT("FailedToGenerateSwitchParam", "Could not generate switch enum parameter. Please refesh the switch node and connect an enum.");
-					GenerationContext.Compiler->CompilerLog(Message, Node);
+					GenerationContext.Log(Message, Node);
 				}
 
 				return nullptr;
@@ -112,7 +112,7 @@ mu::Ptr<mu::NodeColour> GenerateMutableSourceColor(const UEdGraphPin* Pin, FMuta
 			if (SwitchParam->GetType() != mu::NodeScalarEnumParameter::GetStaticType())
 			{
 				const FText Message = LOCTEXT("WrongSwitchParamType", "Switch parameter of incorrect type.");
-				GenerationContext.Compiler->CompilerLog(Message, Node);
+				GenerationContext.Log(Message, Node);
 
 				return nullptr;
 			}
@@ -123,7 +123,7 @@ mu::Ptr<mu::NodeColour> GenerateMutableSourceColor(const UEdGraphPin* Pin, FMuta
 			if (NumSwitchOptions != EnumParameter->GetValueCount())
 			{
 				const FText Message = LOCTEXT("MismatchedSwitch", "Switch enum and switch node have different number of options. Please refresh the switch node to make sure the outcomes are labeled properly.");
-				GenerationContext.Compiler->CompilerLog(Message, Node);
+				GenerationContext.Log(Message, Node);
 			}
 
 			mu::Ptr<mu::NodeColourSwitch> SwitchNode = new mu::NodeColourSwitch;
@@ -246,7 +246,7 @@ mu::Ptr<mu::NodeColour> GenerateMutableSourceColor(const UEdGraphPin* Pin, FMuta
 			}
 			else
 			{
-				GenerationContext.Compiler->CompilerLog(LOCTEXT("ColorFailed", "Color generation failed."), Node);
+				GenerationContext.Log(LOCTEXT("ColorFailed", "Color generation failed."), Node);
 			}
 		}
 
@@ -285,7 +285,7 @@ mu::Ptr<mu::NodeColour> GenerateMutableSourceColor(const UEdGraphPin* Pin, FMuta
 			if (!Property)
 			{
 				FString Msg = FString::Printf(TEXT("Couldn't find the column [%s] in the data table's struct."), *ColumnName);
-				GenerationContext.Compiler->CompilerLog(FText::FromString(Msg), Node);
+				GenerationContext.Log(FText::FromString(Msg), Node);
 
 				bSuccess = false;
 			}
@@ -310,7 +310,7 @@ mu::Ptr<mu::NodeColour> GenerateMutableSourceColor(const UEdGraphPin* Pin, FMuta
 						if (!bSuccess)
 						{
 							FString Msg = FString::Printf(TEXT("Failed to generate the mutable table column [%s]"), *ColumnName);
-							GenerationContext.Compiler->CompilerLog(FText::FromString(Msg), Node);
+							GenerationContext.Log(FText::FromString(Msg), Node);
 						}
 					}
 
@@ -328,19 +328,19 @@ mu::Ptr<mu::NodeColour> GenerateMutableSourceColor(const UEdGraphPin* Pin, FMuta
 				else
 				{
 					FString Msg = FString::Printf(TEXT("Couldn't generate a mutable table."), *ColumnName);
-					GenerationContext.Compiler->CompilerLog(FText::FromString(Msg), Node);
+					GenerationContext.Log(FText::FromString(Msg), Node);
 				}
 			}
 		}
 		else
 		{
-			GenerationContext.Compiler->CompilerLog(LOCTEXT("ColorTableError", "Couldn't find the data table of the node."), Node);
+			GenerationContext.Log(LOCTEXT("ColorTableError", "Couldn't find the data table of the node."), Node);
 		}
 	}
 
 	else
 	{
-		GenerationContext.Compiler->CompilerLog(LOCTEXT("UnimplementedNode", "Node type not implemented yet."), Node);
+		GenerationContext.Log(LOCTEXT("UnimplementedNode", "Node type not implemented yet."), Node);
 	}
 
 	FGeneratedData CacheData = FGeneratedData(Node, Result);
diff --git a/Engine/Plugins/Mutable/Source/CustomizableObjectEditor/Private/MuCOE/GenerateMutableSource/GenerateMutableSourceComponent.cpp b/Engine/Plugins/Mutable/Source/CustomizableObjectEditor/Private/MuCOE/GenerateMutableSource/GenerateMutableSourceComponent.cpp
--- a/Engine/Plugins/Mutable/Source/CustomizableObjectEditor/Private/MuCOE/GenerateMutableSource/GenerateMutableSourceComponent.cpp
+++ b/Engine/Plugins/Mutable/Source/CustomizableObjectEditor/Private/MuCOE/GenerateMutableSource/GenerateMutableSourceComponent.cpp
@@ -95,7 +95,7 @@ void GenerateMutableSourceComponentMesh(FMutableGraphGenerationContext& Generati
 				if (Cast<UCustomizableObjectNodeModifierBase>(ChildNodePin->GetOwningNode()))
 				{
 					FString Msg = FString::Printf(TEXT("The object has legacy modifier connections that cannot be generated. Their connections should be updated."));
-					GenerationContext.Compiler->CompilerLog(FText::FromString(Msg), &TypedComponentMesh, EMessageSeverity::Warning);
+					GenerationContext.Log(FText::FromString(Msg), &TypedComponentMesh, EMessageSeverity::Warning);
 					continue;
 				}
 				
@@ -136,7 +136,7 @@ mu::Ptr<mu::NodeComponent> GenerateMutableSourceComponent(const UEdGraphPin* Pin
 		
 		if (TypedComponentMesh->ComponentName.IsNone())
 		{
-			GenerationContext.Compiler->CompilerLog(LOCTEXT("EmptyComponentNameError", "Error! Missing name in a component of the Customizable Object."), ActualRoot, EMessageSeverity::Error);
+			GenerationContext.Log(LOCTEXT("EmptyComponentNameError", "Error! Missing name in a component of the Customizable Object."), ActualRoot, EMessageSeverity::Error);
 			return nullptr;
 		}
 
@@ -145,7 +145,7 @@ mu::Ptr<mu::NodeComponent> GenerateMutableSourceComponent(const UEdGraphPin* Pin
 			return ComponentInfo.ComponentName == ComponentName;
 		}))
 		{
-			GenerationContext.Compiler->CompilerLog(FText::Format(LOCTEXT("RepeatedComponentName", "Error! Repeated name [{0}] used in more than one Component"),
+			GenerationContext.Log(FText::Format(LOCTEXT("RepeatedComponentName", "Error! Repeated name [{0}] used in more than one Component"),
 				FText::FromName(ComponentName)), ActualRoot, EMessageSeverity::Error);
 			return nullptr;
 		}
@@ -153,7 +153,7 @@ mu::Ptr<mu::NodeComponent> GenerateMutableSourceComponent(const UEdGraphPin* Pin
 		USkeletalMesh* RefSkeletalMesh = TypedComponentMesh->ReferenceSkeletalMesh;
 		if (!RefSkeletalMesh)
 		{
-			GenerationContext.Compiler->CompilerLog(LOCTEXT("NoReferenceMeshObjectTab", "Error! Missing reference Skeletal Mesh"), ActualRoot, EMessageSeverity::Error);
+			GenerationContext.Log(LOCTEXT("NoReferenceMeshObjectTab", "Error! Missing reference Skeletal Mesh"), ActualRoot, EMessageSeverity::Error);
 			return nullptr;
 		}
 
@@ -162,7 +162,7 @@ mu::Ptr<mu::NodeComponent> GenerateMutableSourceComponent(const UEdGraphPin* Pin
 		{
 			FText Msg = FText::Format(LOCTEXT("NoReferenceSkeleton", "Error! Missing skeleton in the reference mesh [{0}]"), FText::FromString(GenerationContext.CustomizableObjectWithCycle->GetPathName()));
 
-			GenerationContext.Compiler->CompilerLog(Msg, ActualRoot, EMessageSeverity::Error);
+			GenerationContext.Log(Msg, ActualRoot, EMessageSeverity::Error);
 			return nullptr;
 		}
 
@@ -180,7 +180,7 @@ mu::Ptr<mu::NodeComponent> GenerateMutableSourceComponent(const UEdGraphPin* Pin
 		// Ensure that the CO has a valid AutoLODStrategy on the ActualRoot.
 		if (TypedComponentMesh->AutoLODStrategy == ECustomizableObjectAutomaticLODStrategy::Inherited)
 		{
-			GenerationContext.Compiler->CompilerLog(LOCTEXT("RootInheritsFromParent", "Error! Component LOD Strategy can't be set to 'Inherit from parent object'"), TypedComponentMesh, EMessageSeverity::Error);
+			GenerationContext.Log(LOCTEXT("RootInheritsFromParent", "Error! Component LOD Strategy can't be set to 'Inherit from parent object'"), TypedComponentMesh, EMessageSeverity::Error);
 			return nullptr;
 		}
 		GenerationContext.CurrentAutoLODStrategy = TypedComponentMesh->AutoLODStrategy;
@@ -191,7 +191,7 @@ mu::Ptr<mu::NodeComponent> GenerateMutableSourceComponent(const UEdGraphPin* Pin
 		FGuid FinalGuid = GenerationContext.GetNodeIdUnique(TypedComponentMesh);
 		if (FinalGuid != TypedComponentMesh->NodeGuid)
 		{
-			GenerationContext.Compiler->CompilerLog(FText::FromString(TEXT("Warning: Node has a duplicated GUID. A new ID has been generated, but cooked data will not be deterministic.")), Node, EMessageSeverity::Warning);
+			GenerationContext.Log(FText::FromString(TEXT("Warning: Node has a duplicated GUID. A new ID has been generated, but cooked data will not be deterministic.")), Node, EMessageSeverity::Warning);
 		}
 		ObjectNode->SetUid(FinalGuid.ToString());
 
@@ -213,7 +213,7 @@ mu::Ptr<mu::NodeComponent> GenerateMutableSourceComponent(const UEdGraphPin* Pin
 			{
 				FString Msg = FString::Printf(TEXT("The object has %d LODs but the reference mesh only %d. Resulting objects will have %d LODs."),
 					NumLODs, MaxRefMeshLODs, MaxRefMeshLODs);
-				GenerationContext.Compiler->CompilerLog(FText::FromString(Msg), Node, EMessageSeverity::Warning);
+				GenerationContext.Log(FText::FromString(Msg), Node, EMessageSeverity::Warning);
 				GenerationContext.NumLODsInRoot = MaxRefMeshLODs;
 			}
 			else
@@ -298,7 +298,7 @@ mu::Ptr<mu::NodeComponent> GenerateMutableSourceComponent(const UEdGraphPin* Pin
 			if (TypedComponentMeshExtend->NumLODs > TypedParentComponentMesh->NumLODs)
 			{
 				FText Msg = FText::Format(LOCTEXT("ExtendMeshComponentLODs", "Add To Mesh Component can not have more LODs than its parent Mesh Component [{0}]."), FText::FromName(TypedComponentMeshExtend->ParentComponentName));
-				GenerationContext.Compiler->CompilerLog(Msg, TypedComponentMeshExtend, EMessageSeverity::Warning);
+				GenerationContext.Log(Msg, TypedComponentMeshExtend, EMessageSeverity::Warning);
 			}
 
 			mu::Ptr<mu::NodeComponent> ParentNodeComponent = GenerateMutableSourceComponent(TypedParentComponentMesh->OutputPin.Get(), GenerationContext);
@@ -320,7 +320,7 @@ mu::Ptr<mu::NodeComponent> GenerateMutableSourceComponent(const UEdGraphPin* Pin
 		else
 		{
 			FText Msg = FText::Format(LOCTEXT("ExtendMeshComponent", "Can not find parent Mesh Component {0}."), FText::FromName(TypedComponentMeshExtend->ParentComponentName));
-			GenerationContext.Compiler->CompilerLog(Msg, TypedComponentMeshExtend, EMessageSeverity::Error);
+			GenerationContext.Log(Msg, TypedComponentMeshExtend, EMessageSeverity::Error);
 		}
 	}
 	
@@ -331,14 +331,14 @@ mu::Ptr<mu::NodeComponent> GenerateMutableSourceComponent(const UEdGraphPin* Pin
 		if (TypedComponentPassthroughMesh->ComponentName.IsNone())
 		{
 			FString Msg = FString::Printf(TEXT("Invalid Component Name."));
-			GenerationContext.Compiler->CompilerLog(FText::FromString(Msg), TypedComponentPassthroughMesh, EMessageSeverity::Warning);
+			GenerationContext.Log(FText::FromString(Msg), TypedComponentPassthroughMesh, EMessageSeverity::Warning);
 			return nullptr;
 		}
 
 		if (!TypedComponentPassthroughMesh->Mesh.IsValid())
 		{
 			FString Msg = FString::Printf(TEXT("No mesh set for component node."));
-			GenerationContext.Compiler->CompilerLog(FText::FromString(Msg), TypedComponentPassthroughMesh, EMessageSeverity::Warning);
+			GenerationContext.Log(FText::FromString(Msg), TypedComponentPassthroughMesh, EMessageSeverity::Warning);
 			return nullptr;
 		}
 
@@ -346,7 +346,7 @@ mu::Ptr<mu::NodeComponent> GenerateMutableSourceComponent(const UEdGraphPin* Pin
 		if (!SkeletalMesh)
 		{
 			FString Msg = FString::Printf(TEXT("Only SkeletalMeshes are supported in this node, for now."));
-			GenerationContext.Compiler->CompilerLog(FText::FromString(Msg), TypedComponentPassthroughMesh, EMessageSeverity::Warning);
+			GenerationContext.Log(FText::FromString(Msg), TypedComponentPassthroughMesh, EMessageSeverity::Warning);
 			return nullptr;
 		}
 
@@ -427,7 +427,7 @@ mu::Ptr<mu::NodeComponent> GenerateMutableSourceComponent(const UEdGraphPin* Pin
 						if (EnumPin)
 						{
 							const FText Message = LOCTEXT("FailedToGenerateSwitchParam", "Could not generate switch enum parameter. Please refesh the switch node and connect an enum.");
-							GenerationContext.Compiler->CompilerLog(Message, Node);
+							GenerationContext.Log(Message, Node);
 						}
 
 						return Result;
@@ -436,7 +436,7 @@ mu::Ptr<mu::NodeComponent> GenerateMutableSourceComponent(const UEdGraphPin* Pin
 					if (SwitchParam->GetType() != mu::NodeScalarEnumParameter::GetStaticType())
 					{
 						const FText Message = LOCTEXT("WrongSwitchParamType", "Switch parameter of incorrect type.");
-						GenerationContext.Compiler->CompilerLog(Message, Node);
+						GenerationContext.Log(Message, Node);
 
 						return Result;
 					}
@@ -447,7 +447,7 @@ mu::Ptr<mu::NodeComponent> GenerateMutableSourceComponent(const UEdGraphPin* Pin
 					if (NumSwitchOptions != EnumParameter->GetValueCount())
 					{
 						const FText Message = LOCTEXT("MismatchedSwitch", "Switch enum and switch node have different number of options. Please refresh the switch node to make sure the outcomes are labeled properly.");
-						GenerationContext.Compiler->CompilerLog(Message, Node);
+						GenerationContext.Log(Message, Node);
 					}
 
 					mu::Ptr<mu::NodeComponentSwitch> SwitchNode = new mu::NodeComponentSwitch;
@@ -475,7 +475,7 @@ mu::Ptr<mu::NodeComponent> GenerateMutableSourceComponent(const UEdGraphPin* Pin
 				}
 				else
 				{
-					GenerationContext.Compiler->CompilerLog(LOCTEXT("NoEnumParamInSwitch", "Switch nodes must have an enum switch parameter. Please connect an enum and refesh the switch node."), Node);
+					GenerationContext.Log(LOCTEXT("NoEnumParamInSwitch", "Switch nodes must have an enum switch parameter. Please connect an enum and refesh the switch node."), Node);
 					return Result;
 				}
 			}(); // invoke lambda.
@@ -495,7 +495,7 @@ mu::Ptr<mu::NodeComponent> GenerateMutableSourceComponent(const UEdGraphPin* Pin
 			}
 			else
 			{
-				GenerationContext.Compiler->CompilerLog(LOCTEXT("ComponentFailed", "Component generation failed."), Node);
+				GenerationContext.Log(LOCTEXT("ComponentFailed", "Component generation failed."), Node);
 			}
 		}
 
@@ -518,7 +518,7 @@ mu::Ptr<mu::NodeComponent> GenerateMutableSourceComponent(const UEdGraphPin* Pin
 					}
 					else
 					{
-						GenerationContext.Compiler->CompilerLog(LOCTEXT("ComponentFailed", "Component generation failed."), Node);
+						GenerationContext.Log(LOCTEXT("ComponentFailed", "Component generation failed."), Node);
 					}
 				}
 			}
@@ -527,7 +527,7 @@ mu::Ptr<mu::NodeComponent> GenerateMutableSourceComponent(const UEdGraphPin* Pin
 
 	else
 	{
-		GenerationContext.Compiler->CompilerLog(LOCTEXT("UnimplementedNode", "Node type not implemented yet."), Node);
+		GenerationContext.Log(LOCTEXT("UnimplementedNode", "Node type not implemented yet."), Node);
 		ensure(false);
 	}
 
diff --git a/Engine/Plugins/Mutable/Source/CustomizableObjectEditor/Private/MuCOE/GenerateMutableSource/GenerateMutableSourceFloat.cpp b/Engine/Plugins/Mutable/Source/CustomizableObjectEditor/Private/MuCOE/GenerateMutableSource/GenerateMutableSourceFloat.cpp
--- a/Engine/Plugins/Mutable/Source/CustomizableObjectEditor/Private/MuCOE/GenerateMutableSource/GenerateMutableSourceFloat.cpp
+++ b/Engine/Plugins/Mutable/Source/CustomizableObjectEditor/Private/MuCOE/GenerateMutableSource/GenerateMutableSourceFloat.cpp
@@ -126,7 +126,7 @@ mu::NodeScalarPtr GenerateMutableSourceFloat(const UEdGraphPin* Pin, FMutableGra
 					? LOCTEXT("NoEnumParamInSwitch", "Switch nodes must have an enum switch parameter. Please connect an enum and refesh the switch node.")
 					: LOCTEXT("InvalidEnumInSwitch", "Switch nodes must have a single enum with all the options inside. Please remove all the enums but one and refresh the switch node.");
 
-				GenerationContext.Compiler->CompilerLog(Message, Node);
+				GenerationContext.Log(Message, Node);
 				return nullptr;
 			}
 
@@ -140,7 +140,7 @@ mu::NodeScalarPtr GenerateMutableSourceFloat(const UEdGraphPin* Pin, FMutableGra
 				if (EnumPin)
 				{
 					const FText Message = LOCTEXT("FailedToGenerateSwitchParam", "Could not generate switch enum parameter. Please refesh the switch node and connect an enum.");
-					GenerationContext.Compiler->CompilerLog(Message, Node);
+					GenerationContext.Log(Message, Node);
 				}
 
 				return nullptr;
@@ -149,7 +149,7 @@ mu::NodeScalarPtr GenerateMutableSourceFloat(const UEdGraphPin* Pin, FMutableGra
 			if (SwitchParam->GetType() != mu::NodeScalarEnumParameter::GetStaticType())
 			{
 				const FText Message = LOCTEXT("WrongSwitchParamType", "Switch parameter of incorrect type.");
-				GenerationContext.Compiler->CompilerLog(Message, Node);
+				GenerationContext.Log(Message, Node);
 
 				return nullptr;
 			}
@@ -160,7 +160,7 @@ mu::NodeScalarPtr GenerateMutableSourceFloat(const UEdGraphPin* Pin, FMutableGra
 			if (NumSwitchOptions != EnumParameter->GetValueCount())
 			{
 				const FText Message = LOCTEXT("MismatchedSwitch", "Switch enum and switch node have different number of options. Please refresh the switch node to make sure the outcomes are labeled properly.");
-				GenerationContext.Compiler->CompilerLog(Message, Node);
+				GenerationContext.Log(Message, Node);
 			}
 
 			mu::Ptr<mu::NodeScalarSwitch> SwitchNode = new mu::NodeScalarSwitch;
@@ -243,7 +243,7 @@ mu::NodeScalarPtr GenerateMutableSourceFloat(const UEdGraphPin* Pin, FMutableGra
 			}
 			else
 			{
-				GenerationContext.Compiler->CompilerLog(LOCTEXT("FloatFailed", "Float generation failed."), Node);
+				GenerationContext.Log(LOCTEXT("FloatFailed", "Float generation failed."), Node);
 			}
 		}
 
@@ -328,7 +328,7 @@ mu::NodeScalarPtr GenerateMutableSourceFloat(const UEdGraphPin* Pin, FMutableGra
 			if (!Property)
 			{
 				FString Msg = FString::Printf(TEXT("Couldn't find the column [%s] in the data table's struct."), *ColumnName);
-				GenerationContext.Compiler->CompilerLog(FText::FromString(Msg), Node);
+				GenerationContext.Log(FText::FromString(Msg), Node);
 
 				bSuccess = false;
 			}
@@ -359,7 +359,7 @@ mu::NodeScalarPtr GenerateMutableSourceFloat(const UEdGraphPin* Pin, FMutableGra
 						if (!bSuccess)
 						{
 							FString Msg = FString::Printf(TEXT("Failed to generate the mutable table column [%s]"), *ColumnName);
-							GenerationContext.Compiler->CompilerLog(FText::FromString(Msg), Node);
+							GenerationContext.Log(FText::FromString(Msg), Node);
 						}
 					}
 
@@ -377,19 +377,19 @@ mu::NodeScalarPtr GenerateMutableSourceFloat(const UEdGraphPin* Pin, FMutableGra
 				else
 				{
 					FString Msg = FString::Printf(TEXT("Couldn't generate a mutable table."), *ColumnName);
-					GenerationContext.Compiler->CompilerLog(FText::FromString(Msg), Node);
+					GenerationContext.Log(FText::FromString(Msg), Node);
 				}
 			}
 		}
 		else
 		{
-			GenerationContext.Compiler->CompilerLog(LOCTEXT("ScalarTableError", "Couldn't find the data table of the node."), Node);
+			GenerationContext.Log(LOCTEXT("ScalarTableError", "Couldn't find the data table of the node."), Node);
 		}
 	}
 
 	else
 	{
-		GenerationContext.Compiler->CompilerLog(LOCTEXT("UnimplementedNode", "Node type not implemented yet."), Node);
+		GenerationContext.Log(LOCTEXT("UnimplementedNode", "Node type not implemented yet."), Node);
 	}
 
 	if (!bDoNotAddToGeneratedCache)
diff --git a/Engine/Plugins/Mutable/Source/CustomizableObjectEditor/Private/MuCOE/GenerateMutableSource/GenerateMutableSourceGroupProjector.cpp b/Engine/Plugins/Mutable/Source/CustomizableObjectEditor/Private/MuCOE/GenerateMutableSource/GenerateMutableSourceGroupProjector.cpp
--- a/Engine/Plugins/Mutable/Source/CustomizableObjectEditor/Private/MuCOE/GenerateMutableSource/GenerateMutableSourceGroupProjector.cpp
+++ b/Engine/Plugins/Mutable/Source/CustomizableObjectEditor/Private/MuCOE/GenerateMutableSource/GenerateMutableSourceGroupProjector.cpp
@@ -116,7 +116,7 @@ mu::NodeImagePtr GenerateMutableGroupProjection(const int32 NodeLOD, const int32
 				}();
 				
 				FString msg = FString::Printf(TEXT("Material image [%s] is connected to an image but will be replaced by a Group Projector."), *ImageName);
-				GenerationContext.Compiler->CompilerLog(FText::FromString(msg), TypedNodeMat);
+				GenerationContext.Log(FText::FromString(msg), TypedNodeMat);
 				continue;
 			}
 			
@@ -138,7 +138,7 @@ mu::NodeImagePtr GenerateMutableGroupProjection(const int32 NodeLOD, const int32
 						!ProjectorTempData.bAlternateResStateNameWarningDisplayed)
 					{
 						FString msg = FString::Printf(TEXT("All 'Alternate Projection Resolution State Name' properties in Group Projector Parameter nodes connected to same Group node must have the same value or be blank. Only the value of the last connected node will be used."));
-						GenerationContext.Compiler->CompilerLog(FText::FromString(msg), ProjectorTempData.CustomizableObjectNodeGroupProjectorParameter);
+						GenerationContext.Log(FText::FromString(msg), ProjectorTempData.CustomizableObjectNodeGroupProjectorParameter);
 						ProjectorTempData.bAlternateResStateNameWarningDisplayed = true;
 					}
 
@@ -173,7 +173,7 @@ mu::NodeImagePtr GenerateMutableGroupProjection(const int32 NodeLOD, const int32
 					if (!NodeMeshApplyPose)
 					{
 						FString msg = FString::Printf(TEXT("Couldn't get bone transform information from a Pose Asset."));
-						GenerationContext.Compiler->CompilerLog(FText::FromString(msg), TypedNodeMat);
+						GenerationContext.Log(FText::FromString(msg), TypedNodeMat);
 					}
 
 					MeshSwitchNode->SetOption(SelectorIndex + 1, NodeMeshApplyPose);
@@ -353,7 +353,7 @@ bool GenerateMutableSourceGroupProjector(const UEdGraphPin* Pin, FMutableGraphGe
 				(ProjParamNode->DataTableTextureColumnName.ToString().IsEmpty() || (ProjParamNode->DataTableTextureColumnName.ToString() == "None")))
 			{
 				FString msg = FString::Printf(TEXT("The group projection node has a table assigned to the Option Images Data Table property, but no column to read textures is specified at the Data Table Texture Column Name property."));
-				GenerationContext.Compiler->CompilerLog(FText::FromString(msg), ProjParamNode, EMessageSeverity::Error, true);
+				GenerationContext.Log(FText::FromString(msg), ProjParamNode, EMessageSeverity::Error, true);
 			}
 
 			GroupProjectorTempData.CustomizableObjectNodeGroupProjectorParameter = ProjParamNode;
@@ -411,7 +411,7 @@ bool GenerateMutableSourceGroupProjector(const UEdGraphPin* Pin, FMutableGraphGe
 			if (ArrayOptionTexture.Num() == 0)
 			{
 				FString msg = FString::Printf(TEXT("The group projection node must have at least one option image connected to a texture or at least one valid element in Option Images Data Table."));
-				GenerationContext.Compiler->CompilerLog(FText::FromString(msg), ProjParamNode, EMessageSeverity::Error, true);
+				GenerationContext.Log(FText::FromString(msg), ProjParamNode, EMessageSeverity::Error, true);
 				return false;
 			}
 			
@@ -435,7 +435,7 @@ bool GenerateMutableSourceGroupProjector(const UEdGraphPin* Pin, FMutableGraphGe
 					if (PoseAsset == nullptr) // Check if the slot has a selected pose. Could be left empty by the user
 					{
 						FString msg = FString::Printf(TEXT("The group projection node must have a pose assigned on each Option Poses element."));
-						GenerationContext.Compiler->CompilerLog(FText::FromString(msg), ProjParamNode, EMessageSeverity::Error, true);
+						GenerationContext.Log(FText::FromString(msg), ProjParamNode, EMessageSeverity::Error, true);
 						return false;
 					}
 
@@ -508,7 +508,7 @@ bool GenerateMutableSourceGroupProjector(const UEdGraphPin* Pin, FMutableGraphGe
 				else
 				{
 					FString msg = FString::Printf(TEXT("The group projection node must have a texture for all the options. Please set a texture for all the options."));
-					GenerationContext.Compiler->CompilerLog(FText::FromString(msg), ProjParamNode);
+					GenerationContext.Log(FText::FromString(msg), ProjParamNode);
 				}
 			}
 
@@ -529,7 +529,7 @@ bool GenerateMutableSourceGroupProjector(const UEdGraphPin* Pin, FMutableGraphGe
 		}
 		else
 		{
-			GenerationContext.Compiler->CompilerLog(FText::FromString(TEXT("Error getting group projection properties.")), ProjParamNode, EMessageSeverity::Error, true);
+			GenerationContext.Log(FText::FromString(TEXT("Error getting group projection properties.")), ProjParamNode, EMessageSeverity::Error, true);
 			return false;
 		}
 	}
diff --git a/Engine/Plugins/Mutable/Source/CustomizableObjectEditor/Private/MuCOE/GenerateMutableSource/GenerateMutableSourceImage.cpp b/Engine/Plugins/Mutable/Source/CustomizableObjectEditor/Private/MuCOE/GenerateMutableSource/GenerateMutableSourceImage.cpp
--- a/Engine/Plugins/Mutable/Source/CustomizableObjectEditor/Private/MuCOE/GenerateMutableSource/GenerateMutableSourceImage.cpp
+++ b/Engine/Plugins/Mutable/Source/CustomizableObjectEditor/Private/MuCOE/GenerateMutableSource/GenerateMutableSourceImage.cpp
@@ -136,7 +136,7 @@ mu::NodeImagePtr GenerateMutableSourceImage(const UEdGraphPin* Pin, FMutableGrap
 		}
 		else
 		{
-			GenerationContext.Compiler->CompilerLog(LOCTEXT("MissingImage", "Missing image in texture node."), Node, EMessageSeverity::Warning);
+			GenerationContext.Log(LOCTEXT("MissingImage", "Missing image in texture node."), Node, EMessageSeverity::Warning);
 		}
 	}
 
@@ -193,7 +193,7 @@ mu::NodeImagePtr GenerateMutableSourceImage(const UEdGraphPin* Pin, FMutableGrap
 			if (MaxNodeTextureSize <= 0)
 			{
 				TextureSize.X = TextureSize.Y = 1;
-				GenerationContext.Compiler->CompilerLog(LOCTEXT("TextureParameterSize0", "Texture size not specified. Add a reference texture or set a valid value to the Texture Size variables."), Node);
+				GenerationContext.Log(LOCTEXT("TextureParameterSize0", "Texture size not specified. Add a reference texture or set a valid value to the Texture Size variables."), Node);
 			}
 			else if (ReferenceTextureSize > 0 && ReferenceTextureSize < MaxNodeTextureSize)
 			{
@@ -293,13 +293,13 @@ mu::NodeImagePtr GenerateMutableSourceImage(const UEdGraphPin* Pin, FMutableGrap
 				case COTLE_ALPHA_OVERLAY: Type = mu::EBlendType::BT_LIGHTEN; break;
 				case COTLE_NORMAL_COMBINE: Type = mu::EBlendType::BT_NORMAL_COMBINE; break;
 				default:
-					GenerationContext.Compiler->CompilerLog(LOCTEXT("UnsupportedImageEffect", "Texture layer effect not supported. Setting to 'Blend'."), Node);
+					GenerationContext.Log(LOCTEXT("UnsupportedImageEffect", "Texture layer effect not supported. Setting to 'Blend'."), Node);
 					break;
 				}
 
 				if (Type == mu::EBlendType::BT_BLEND && !MaskNode)
 				{
-					GenerationContext.Compiler->CompilerLog(LOCTEXT("ModulateWithoutMask", "Texture layer effect uses Modulate without a mask. It will replace everything below it!"), Node);
+					GenerationContext.Log(LOCTEXT("ModulateWithoutMask", "Texture layer effect uses Modulate without a mask. It will replace everything below it!"), Node);
 				}
 				
 				if (OtherPin->PinType.PinCategory == Schema->PC_Image)
@@ -346,7 +346,7 @@ mu::NodeImagePtr GenerateMutableSourceImage(const UEdGraphPin* Pin, FMutableGrap
 				if (!SwitchParam)
 				{
 					const FText Message = LOCTEXT("FailedToGenerateSwitchParam", "Could not generate switch enum parameter. Please refesh the switch node and connect an enum.");
-					GenerationContext.Compiler->CompilerLog(Message, Node);
+					GenerationContext.Log(Message, Node);
 					
 					return Result;
 				}
@@ -354,7 +354,7 @@ mu::NodeImagePtr GenerateMutableSourceImage(const UEdGraphPin* Pin, FMutableGrap
 				if (SwitchParam->GetType() != mu::NodeScalarEnumParameter::GetStaticType())
 				{
 					const FText Message = LOCTEXT("WrongSwitchParamType", "Switch parameter of incorrect type.");
-					GenerationContext.Compiler->CompilerLog(Message, Node);
+					GenerationContext.Log(Message, Node);
 
 					return Result;
 				}
@@ -365,7 +365,7 @@ mu::NodeImagePtr GenerateMutableSourceImage(const UEdGraphPin* Pin, FMutableGrap
 				if (NumSwitchOptions != EnumParameter->GetValueCount())
 				{
 					const FText Message = LOCTEXT("MismatchedSwitch", "Switch enum and switch node have different number of options. Please refresh the switch node to make sure the outcomes are labeled properly.");
-					GenerationContext.Compiler->CompilerLog(Message, Node);
+					GenerationContext.Log(Message, Node);
 				}
 
 				mu::NodeImageSwitchPtr SwitchNode = new mu::NodeImageSwitch;
@@ -381,7 +381,7 @@ mu::NodeImagePtr GenerateMutableSourceImage(const UEdGraphPin* Pin, FMutableGrap
 					else
 					{
 						const FText Message = LOCTEXT("MissingTexture", "Unable to generate texture switch node. Required connection not found.");
-						GenerationContext.Compiler->CompilerLog(Message, Node);
+						GenerationContext.Log(Message, Node);
 						return Result;
 					}
 				}
@@ -391,7 +391,7 @@ mu::NodeImagePtr GenerateMutableSourceImage(const UEdGraphPin* Pin, FMutableGrap
 			}
 			else
 			{
-				GenerationContext.Compiler->CompilerLog(LOCTEXT("NoEnumParamInSwitch", "Switch nodes must have an enum switch parameter. Please connect an enum and refesh the switch node."), Node);
+				GenerationContext.Log(LOCTEXT("NoEnumParamInSwitch", "Switch nodes must have an enum switch parameter. Please connect an enum and refesh the switch node."), Node);
 				return Result;
 			}
 		}(); // invoke lambda.
@@ -412,7 +412,7 @@ mu::NodeImagePtr GenerateMutableSourceImage(const UEdGraphPin* Pin, FMutableGrap
 			}
 			else
 			{
-				GenerationContext.Compiler->CompilerLog(LOCTEXT("TextureFailed", "Texture generation failed."), Node);
+				GenerationContext.Log(LOCTEXT("TextureFailed", "Texture generation failed."), Node);
 			}
 		}
 
@@ -601,7 +601,7 @@ mu::NodeImagePtr GenerateMutableSourceImage(const UEdGraphPin* Pin, FMutableGrap
 
 		if (!FollowInputPin(*TypedNodeProject->MeshPin()))
 		{
-			GenerationContext.Compiler->CompilerLog(LOCTEXT("MissingMeshInProjector", "Texture projector does not have a Mesh. It will be ignored. "), Node, EMessageSeverity::Warning);
+			GenerationContext.Log(LOCTEXT("MissingMeshInProjector", "Texture projector does not have a Mesh. It will be ignored. "), Node, EMessageSeverity::Warning);
 			Result = nullptr;
 		}
 
@@ -895,7 +895,7 @@ mu::NodeImagePtr GenerateMutableSourceImage(const UEdGraphPin* Pin, FMutableGrap
 		}
 		else
 		{
-			GenerationContext.Compiler->CompilerLog(LOCTEXT("MissingImagePassThrough", "Missing image in pass-through texture node."), Node);
+			GenerationContext.Log(LOCTEXT("MissingImagePassThrough", "Missing image in pass-through texture node."), Node);
 		}
 	}
 
@@ -913,7 +913,7 @@ mu::NodeImagePtr GenerateMutableSourceImage(const UEdGraphPin* Pin, FMutableGrap
 				if (!SwitchParam)
 				{
 					const FText Message = LOCTEXT("FailedToGenerateSwitchParam", "Could not generate switch enum parameter. Please refesh the switch node and connect an enum.");
-					GenerationContext.Compiler->CompilerLog(Message, Node);
+					GenerationContext.Log(Message, Node);
 
 					return Result;
 				}
@@ -921,7 +921,7 @@ mu::NodeImagePtr GenerateMutableSourceImage(const UEdGraphPin* Pin, FMutableGrap
 				if (SwitchParam->GetType() != mu::NodeScalarEnumParameter::GetStaticType())
 				{
 					const FText Message = LOCTEXT("WrongSwitchParamType", "Switch parameter of incorrect type.");
-					GenerationContext.Compiler->CompilerLog(Message, Node);
+					GenerationContext.Log(Message, Node);
 
 					return Result;
 				}
@@ -932,7 +932,7 @@ mu::NodeImagePtr GenerateMutableSourceImage(const UEdGraphPin* Pin, FMutableGrap
 				if (NumSwitchOptions != EnumParameter->GetValueCount())
 				{
 					const FText Message = LOCTEXT("MismatchedSwitch", "Switch enum and switch node have different number of options. Please refresh the switch node to make sure the outcomes are labeled properly.");
-					GenerationContext.Compiler->CompilerLog(Message, Node);
+					GenerationContext.Log(Message, Node);
 				}
 
 				// TODO: Implement Mutable core pass-through switch nodes
@@ -950,7 +950,7 @@ mu::NodeImagePtr GenerateMutableSourceImage(const UEdGraphPin* Pin, FMutableGrap
 					else
 					{
 						const FText Message = LOCTEXT("MissingPassThroughTexture", "Unable to generate pass-through texture switch node. Required connection not found.");
-						GenerationContext.Compiler->CompilerLog(Message, Node);
+						GenerationContext.Log(Message, Node);
 						return Result;
 					}
 				}
@@ -960,7 +960,7 @@ mu::NodeImagePtr GenerateMutableSourceImage(const UEdGraphPin* Pin, FMutableGrap
 			}
 			else
 			{
-				GenerationContext.Compiler->CompilerLog(LOCTEXT("NoEnumParamInSwitch", "Switch nodes must have an enum switch parameter. Please connect an enum and refesh the switch node."), Node);
+				GenerationContext.Log(LOCTEXT("NoEnumParamInSwitch", "Switch nodes must have an enum switch parameter. Please connect an enum and refesh the switch node."), Node);
 				return Result;
 			}
 		}(); // invoke lambda.
@@ -1013,7 +1013,7 @@ mu::NodeImagePtr GenerateMutableSourceImage(const UEdGraphPin* Pin, FMutableGrap
 				if (!Property)
 				{
 					FString Msg = FString::Printf(TEXT("Couldn't find the column [%s] in the data table's struct."), *ColumnName);
-					GenerationContext.Compiler->CompilerLog(FText::FromString(Msg), Node);
+					GenerationContext.Log(FText::FromString(Msg), Node);
 
 					bSuccess = false;
 				}
@@ -1023,7 +1023,7 @@ mu::NodeImagePtr GenerateMutableSourceImage(const UEdGraphPin* Pin, FMutableGrap
 				if (bSuccess && Pin->PinType.PinCategory != Schema->PC_MaterialAsset && !DefaultTexture)
 				{
 					FString Msg = FString::Printf(TEXT("Couldn't find a default value in the data table's struct for the column [%s]. The default value is null or not a supported Texture"), *ColumnName);
-					GenerationContext.Compiler->CompilerLog(FText::FromString(Msg), Node);
+					GenerationContext.Log(FText::FromString(Msg), Node);
 					
 					bSuccess = false;
 				}
@@ -1057,7 +1057,7 @@ mu::NodeImagePtr GenerateMutableSourceImage(const UEdGraphPin* Pin, FMutableGrap
 									Nodes.Add(GenerationContext.GeneratedTableImages[ImageDataIndex].TableNode);
 
 									FString Msg = FString::Printf(TEXT("Texture pin [%s] with different texture modes found in more than one table node. This will add multiple times the texture reseource in the final cook."), *ColumnName);
-									GenerationContext.Compiler->CompilerLog(FText::FromString(Msg), Nodes);
+									GenerationContext.Log(FText::FromString(Msg), Nodes);
 								}
 							}
 							else
@@ -1082,7 +1082,7 @@ mu::NodeImagePtr GenerateMutableSourceImage(const UEdGraphPin* Pin, FMutableGrap
 							if (!bSuccess)
 							{
 								FString Msg = FString::Printf(TEXT("Failed to generate the mutable table column [%s]"), *ColumnName);
-								GenerationContext.Compiler->CompilerLog(FText::FromString(Msg), Node);
+								GenerationContext.Log(FText::FromString(Msg), Node);
 							}
 						}
 
@@ -1119,20 +1119,20 @@ mu::NodeImagePtr GenerateMutableSourceImage(const UEdGraphPin* Pin, FMutableGrap
 					else
 					{
 						FString Msg = FString::Printf(TEXT("Couldn't generate a mutable table."));
-						GenerationContext.Compiler->CompilerLog(FText::FromString(Msg), Node);
+						GenerationContext.Log(FText::FromString(Msg), Node);
 					}
 				}
 			}
 			else
 			{
-				GenerationContext.Compiler->CompilerLog(LOCTEXT("ImageTableError", "Couldn't find the data table of the node."), Node);
+				GenerationContext.Log(LOCTEXT("ImageTableError", "Couldn't find the data table of the node."), Node);
 			}
 		}
 	}
 
 	else
 	{
-		GenerationContext.Compiler->CompilerLog(LOCTEXT("UnimplementedNode", "Node type not implemented yet."), Node);
+		GenerationContext.Log(LOCTEXT("UnimplementedNode", "Node type not implemented yet."), Node);
 	}
 
 	if (!bDoNotAddToGeneratedCache)
diff --git a/Engine/Plugins/Mutable/Source/CustomizableObjectEditor/Private/MuCOE/GenerateMutableSource/GenerateMutableSourceMesh.cpp b/Engine/Plugins/Mutable/Source/CustomizableObjectEditor/Private/MuCOE/GenerateMutableSource/GenerateMutableSourceMesh.cpp
--- a/Engine/Plugins/Mutable/Source/CustomizableObjectEditor/Private/MuCOE/GenerateMutableSource/GenerateMutableSourceMesh.cpp
+++ b/Engine/Plugins/Mutable/Source/CustomizableObjectEditor/Private/MuCOE/GenerateMutableSource/GenerateMutableSourceMesh.cpp
@@ -122,7 +122,7 @@ void GetLODAndSectionForAutomaticLODs(const FMutableGraphGenerationContext& Cont
 			}
 			else
 			{
-				Context.Compiler->CompilerLog(FText::Format(LOCTEXT("MeshMultipleMaterialIndex", "Mesh {0} contains multiple sections with the same Material Index"), FText::FromString(SkeletalMesh.GetName())), &Node);
+				Context.Log(FText::Format(LOCTEXT("MeshMultipleMaterialIndex", "Mesh {0} contains multiple sections with the same Material Index"), FText::FromString(SkeletalMesh.GetName())), &Node);
 			}
 		}
 	}
@@ -425,7 +425,7 @@ void SetAndPropagatePoseBoneUsage(mu::Mesh& MutableMesh, int32 PoseIndex, mu::EB
 
 }
 
-TArray<TTuple<UPhysicsAsset*, int32>> GetPhysicsAssetsFromAnimInstance(const TSoftClassPtr<UAnimInstance>& AnimInstance)
+TArray<TTuple<UPhysicsAsset*, int32>> GetPhysicsAssetsFromAnimInstance(FMutableGraphGenerationContext& GenerationContext, const TSoftClassPtr<UAnimInstance>& AnimInstance)
 {
 	// TODO: Consider caching the result in the GenerationContext.
 	TArray<TTuple<UPhysicsAsset*, int32>> Result;
@@ -435,7 +435,7 @@ TArray<TTuple<UPhysicsAsset*, int32>> GetPhysicsAssetsFromAnimInstance(const TSo
 		return Result;
 	}
 
-	UClass* AnimInstanceClass = AnimInstance.LoadSynchronous();
+	UClass* AnimInstanceClass = GenerationContext.LoadClass(AnimInstance);
 	UAnimBlueprintGeneratedClass* AnimClass = Cast<UAnimBlueprintGeneratedClass>(AnimInstanceClass);
 
 	if (AnimClass)
@@ -700,7 +700,7 @@ mu::MeshPtr ConvertSkeletalMeshToMutable(const USkeletalMesh* InSkeletalMesh, co
 	if (!ImportedModel)
 	{
 		const FString Msg = FString::Printf(TEXT("The SkeletalMesh [%s] doesn't have an imported resource."), *InSkeletalMesh->GetName());
-		GenerationContext.Compiler->CompilerLog(FText::FromString(Msg), CurrentNode);
+		GenerationContext.Log(FText::FromString(Msg), CurrentNode);
 
 		return nullptr;
 	}
@@ -719,7 +719,7 @@ mu::MeshPtr ConvertSkeletalMeshToMutable(const USkeletalMesh* InSkeletalMesh, co
 				InSkeletalMesh ? *InSkeletalMesh->GetName() : TEXT("none"),
 				LODIndex + 1,
 				ImportedModel->LODModels.Num());
-			GenerationContext.Compiler->CompilerLog(FText::FromString(Msg), CurrentNode);
+			GenerationContext.Log(FText::FromString(Msg), CurrentNode);
 
 			return nullptr;
 		}
@@ -742,7 +742,7 @@ mu::MeshPtr ConvertSkeletalMeshToMutable(const USkeletalMesh* InSkeletalMesh, co
 				ImportedModel->LODModels.Num(),
 				SectionIndex + 1,
 				LODModel.Sections.Num());
-			GenerationContext.Compiler->CompilerLog(FText::FromString(Msg), CurrentNode);
+			GenerationContext.Log(FText::FromString(Msg), CurrentNode);
 
 			return nullptr;
 		}
@@ -769,7 +769,7 @@ mu::MeshPtr ConvertSkeletalMeshToMutable(const USkeletalMesh* InSkeletalMesh, co
 		if (!InSkeleton)
 		{
 			FString Msg = FString::Printf(TEXT("No skeleton provided when converting SkeletalMesh [%s]."), *InSkeletalMesh->GetName());
-			GenerationContext.Compiler->CompilerLog(FText::FromString(Msg), CurrentNode);
+			GenerationContext.Log(FText::FromString(Msg), CurrentNode);
 			return nullptr;
 		}
 
@@ -796,7 +796,7 @@ mu::MeshPtr ConvertSkeletalMeshToMutable(const USkeletalMesh* InSkeletalMesh, co
 			{
 				if (!ErrorMessage.IsEmpty())
 				{
-					GenerationContext.Compiler->CompilerLog(FText::FromString(ErrorMessage), CurrentNode, EMessageSeverity::Warning);
+					GenerationContext.Log(FText::FromString(ErrorMessage), CurrentNode, EMessageSeverity::Warning);
 				}
 				return nullptr;
 			}
@@ -838,7 +838,7 @@ mu::MeshPtr ConvertSkeletalMeshToMutable(const USkeletalMesh* InSkeletalMesh, co
 					TEXT("The Skeleton [%s] is missing bones that SkeletalMesh [%s] needs. The mesh will be discarded! Information about missing bones can be found in the Output Log."),
 					*InSkeleton->GetName(), *InSkeletalMesh->GetName());
 				
-				GenerationContext.Compiler->CompilerLog(FText::FromString(Msg), CurrentNode, EMessageSeverity::Warning);
+				GenerationContext.Log(FText::FromString(Msg), CurrentNode, EMessageSeverity::Warning);
 
 				return nullptr;
 			}
@@ -2015,7 +2015,7 @@ mu::MeshPtr ConvertSkeletalMeshToMutable(const USkeletalMesh* InSkeletalMesh, co
 			PhysicsSetupsRemovedMsg += FString::Printf(TEXT("have been discarded because they are not present in the SkeletalMesh [%s] Skeleton."),
 				*InSkeletalMesh->GetName());
 					
-			GenerationContext.Compiler->CompilerLog(FText::FromString(PhysicsSetupsRemovedMsg), CurrentNode, EMessageSeverity::Warning);
+			GenerationContext.Log(FText::FromString(PhysicsSetupsRemovedMsg), CurrentNode, EMessageSeverity::Warning);
 		}
 
 		mu::Ptr<mu::PhysicsBody> PhysicsBody = new mu::PhysicsBody;
@@ -2170,7 +2170,7 @@ mu::MeshPtr ConvertSkeletalMeshToMutable(const USkeletalMesh* InSkeletalMesh, co
 	if (!bIgnorePhysics && !AnimBp.IsNull() && MutableMesh->GetSkeleton() && bAnimPhysicsManipulationEnabled)
 	{
 		using AnimPhysicsInfoType = TTuple<UPhysicsAsset*, int32>;
-		const TArray<AnimPhysicsInfoType> AnimPhysicsInfo = GetPhysicsAssetsFromAnimInstance(AnimBp);
+		const TArray<AnimPhysicsInfoType> AnimPhysicsInfo = GetPhysicsAssetsFromAnimInstance(GenerationContext, AnimBp);
 
 		for (const AnimPhysicsInfoType PropertyInfo : AnimPhysicsInfo)
 		{
@@ -2251,7 +2251,7 @@ mu::MeshPtr ConvertStaticMeshToMutable(const UStaticMesh* StaticMesh, int32 LODI
 		!StaticMesh->GetRenderData()->LODResources[LODIndex].Sections.IsValidIndex(SectionIndex))
 	{
 		FString Msg = FString::Printf(TEXT("Degenerated static mesh found for LOD %d Material %d. It will be ignored. "), LODIndex, SectionIndex);
-		GenerationContext.Compiler->CompilerLog(FText::FromString(Msg), CurrentNode, EMessageSeverity::Warning);
+		GenerationContext.Log(FText::FromString(Msg), CurrentNode, EMessageSeverity::Warning);
 		return nullptr;
 	}
 
@@ -2494,7 +2494,7 @@ mu::Ptr<mu::Mesh> GenerateMutableMesh(UObject * Mesh, const TSoftClassPtr<UAnimI
 	}
 	else
 	{
-		GenerationContext.Compiler->CompilerLog(LOCTEXT("UnimplementedMesh", "Mesh type not implemented yet."), CurrentNode);
+		GenerationContext.Log(LOCTEXT("UnimplementedMesh", "Mesh type not implemented yet."), CurrentNode);
 	}
 
 	if (MutableMesh)
@@ -2576,7 +2576,7 @@ mu::Ptr<mu::Mesh> BuildMorphedMutableMesh(const UEdGraphPin* BaseSourcePin, cons
 
 	if (!BaseSourcePin)
 	{
-		GenerationContext.Compiler->CompilerLog(LOCTEXT("NULLBaseSourcePin", "Morph base not set."), nullptr);
+		GenerationContext.Log(LOCTEXT("NULLBaseSourcePin", "Morph base not set."), nullptr);
 		return nullptr;
 	}
 
@@ -2603,7 +2603,7 @@ mu::Ptr<mu::Mesh> BuildMorphedMutableMesh(const UEdGraphPin* BaseSourcePin, cons
 		if (DataTable)
 		{
 			TypedNodeTable->GetPinLODAndSection(BaseSourcePin, LODIndexConnected, SectionIndexConnected);
-			SkeletalMesh = TypedNodeTable->GetSkeletalMeshAt(BaseSourcePin, DataTable, RowName);
+			SkeletalMesh = Cast<USkeletalMesh>(GenerationContext.LoadObject(TypedNodeTable->GetSkeletalMeshAt(BaseSourcePin, DataTable, RowName)));
 		}
 	}
 
@@ -2639,19 +2639,19 @@ void GenerateMorphFactor(const UCustomizableObjectNode* Node, const UEdGraphPin&
 			{
 				validStaticFactor = false;
 				FString msg = FString::Printf(TEXT("Mesh morph nodes only accept factors between -1.0 and 1.0 inclusive but the default value of the float parameter node is (%f). Factor will be ignored."), floatParameterNode->DefaultValue);
-				GenerationContext.Compiler->CompilerLog(FText::FromString(msg), Node);
+				GenerationContext.Log(FText::FromString(msg), Node);
 			}
 			if (floatParameterNode->ParamUIMetadata.MinimumValue < -1.0f)
 			{
 				validStaticFactor = false;
 				FString msg = FString::Printf(TEXT("Mesh morph nodes only accept factors between -1.0 and 1.0 inclusive but the minimum UI value for the input float parameter node is (%f). Factor will be ignored."), floatParameterNode->ParamUIMetadata.MinimumValue);
-				GenerationContext.Compiler->CompilerLog(FText::FromString(msg), Node);
+				GenerationContext.Log(FText::FromString(msg), Node);
 			}
 			if (floatParameterNode->ParamUIMetadata.MaximumValue > 1.0f)
 			{
 				validStaticFactor = false;
 				FString msg = FString::Printf(TEXT("Mesh morph nodes only accept factors between -1.0 and 1.0 inclusive but the maximum UI value for the input float parameter node is (%f). Factor will be ignored."), floatParameterNode->ParamUIMetadata.MaximumValue);
-				GenerationContext.Compiler->CompilerLog(FText::FromString(msg), Node);
+				GenerationContext.Log(FText::FromString(msg), Node);
 			}
 		}
 		
@@ -2661,7 +2661,7 @@ void GenerateMorphFactor(const UCustomizableObjectNode* Node, const UEdGraphPin&
 			{
 				validStaticFactor = false;
 				FString msg = FString::Printf(TEXT("Mesh morph nodes only accept factors between -1.0 and 1.0 inclusive but the value of the float constant node is (%f). Factor will be ignored."), floatConstantNode->Value);
-				GenerationContext.Compiler->CompilerLog(FText::FromString(msg), Node);
+				GenerationContext.Log(FText::FromString(msg), Node);
 			}
 		}
 
@@ -2698,7 +2698,7 @@ TArray<TTuple<USkeletalMesh*, TSoftClassPtr<UAnimInstance>>> GetSkeletalMeshesIn
 		{
 			for (const FName& RowName : GetRowsToCompile(*DataTable, *TableNode, GenerationContext))
 			{
-				USkeletalMesh* SkeletalMesh = TableNode->GetSkeletalMeshAt(SourceMeshPin, DataTable, RowName);
+				USkeletalMesh* SkeletalMesh = Cast<USkeletalMesh>(GenerationContext.LoadObject(TableNode->GetSkeletalMeshAt(SourceMeshPin, DataTable, RowName)));
 				TSoftClassPtr<UAnimInstance> MeshAnimInstance = TableNode->GetAnimInstanceAt(SourceMeshPin, DataTable, RowName);
 
 				if (SkeletalMesh)
@@ -2792,7 +2792,7 @@ bool GetAndValidateReshapeBonesToDeform(
 
 			if (Algo::AnyOf(MissingBones, [](const uint8& B) { return B; }))
 			{
-				GenerationContext.Compiler->CompilerLog(
+				GenerationContext.Log(
 					FText::FromString(
 						"Could not find the selected bones to deform " +
 						MakeCompactMissingBoneListMessage() +
@@ -2922,7 +2922,7 @@ bool GetAndValidateReshapePhysicsToDeform(
 			{
 				TSoftClassPtr<UAnimInstance> AnimInstance = Mesh.Get<TSoftClassPtr<UAnimInstance>>();
 
-				TArray<TTuple<UPhysicsAsset*, int32>> AnimInstanceOverridePhysicsAssets = GetPhysicsAssetsFromAnimInstance(AnimInstance);
+				TArray<TTuple<UPhysicsAsset*, int32>> AnimInstanceOverridePhysicsAssets = GetPhysicsAssetsFromAnimInstance(GenerationContext, AnimInstance);
 
 				for (const TTuple<UPhysicsAsset*, int32>& AnimPhysicsAssetInfo : AnimInstanceOverridePhysicsAssets)
 				{
@@ -3085,7 +3085,7 @@ bool GetAndValidateReshapePhysicsToDeform(
 
 		if (Algo::AnyOf(MissingBones, IsMissingBone))
 		{
-			GenerationContext.Compiler->CompilerLog(
+			GenerationContext.Log(
 				FText::FromString(
 					"Could not find the selected physics bodies bones to deform " +
 					MakeCompactMissingBoneListMessage(IsMissingBone) +
@@ -3097,7 +3097,7 @@ bool GetAndValidateReshapePhysicsToDeform(
 
 		if (Algo::AnyOf(MissingBones, IsMissingBody))
 		{
-			GenerationContext.Compiler->CompilerLog(
+			GenerationContext.Log(
 				FText::FromString(
 					"Selected Bones to deform " +
 					MakeCompactMissingBoneListMessage(IsMissingBody) +
@@ -3277,7 +3277,7 @@ mu::NodeMeshPtr GenerateMorphMesh(const UEdGraphPin* Pin,
 
 	if(!bSuccess)
 	{
-		GenerationContext.Compiler->CompilerLog(LOCTEXT("MorphGenerationFailed", "Failed to generate morph target."), MorphNode);
+		GenerationContext.Log(LOCTEXT("MorphGenerationFailed", "Failed to generate morph target."), MorphNode);
 	}
 
 	return Result;
@@ -3344,10 +3344,7 @@ mu::NodeMeshPtr GenerateMutableSourceMesh(const UEdGraphPin* Pin,
 
 			if (!TypedNodeSkel->AnimInstance.IsNull())
 			{
-				if (UClass* AnimInstance = TypedNodeSkel->AnimInstance.LoadSynchronous())
-				{
-					GenerationContext.AddParticipatingObject(*AnimInstance);					
-				}
+				GenerationContext.AddParticipatingObject(TypedNodeSkel->AnimInstance);
 
 				FName SlotIndex = TypedNodeSkel->AnimBlueprintSlotName;
 				const int32 AnimInstanceIndex = GenerationContext.AnimBPAssets.AddUnique(TypedNodeSkel->AnimInstance);
@@ -3552,7 +3549,7 @@ mu::NodeMeshPtr GenerateMutableSourceMesh(const UEdGraphPin* Pin,
 		}
 		else
 		{
-			GenerationContext.Compiler->CompilerLog(LOCTEXT("MissingskeletlMesh", "No Skeletal Mesh set in the SkeletalMesh node."), Node);
+			GenerationContext.Log(LOCTEXT("MissingskeletlMesh", "No Skeletal Mesh set in the SkeletalMesh node."), Node);
 		}
 	}
 
@@ -3561,14 +3558,14 @@ mu::NodeMeshPtr GenerateMutableSourceMesh(const UEdGraphPin* Pin,
 		if (TypedNodeStatic->StaticMesh == nullptr)
 		{
 			FString Msg = FString::Printf(TEXT("The UCustomizableObjectNodeStaticMesh node %s has no static mesh assigned"), *Node->GetName());
-			GenerationContext.Compiler->CompilerLog(FText::FromString(Msg), Node, EMessageSeverity::Warning);
+			GenerationContext.Log(FText::FromString(Msg), Node, EMessageSeverity::Warning);
 			return {};
 		}
 
 		if (TypedNodeStatic->StaticMesh->GetNumLODs() == 0)
 		{
 			FString Msg = FString::Printf(TEXT("The UCustomizableObjectNodeStaticMesh node %s has a static mesh assigned with no RenderData"), *Node->GetName());
-			GenerationContext.Compiler->CompilerLog(FText::FromString(Msg), Node, EMessageSeverity::Warning);
+			GenerationContext.Log(FText::FromString(Msg), Node, EMessageSeverity::Warning);
 			return {};
 		}
 
@@ -3622,7 +3619,7 @@ mu::NodeMeshPtr GenerateMutableSourceMesh(const UEdGraphPin* Pin,
 					/*if (bWasEmpty)
 					{
 						FString msg = "Layout without any block found. A grid sized block will be used instead.";
-						GenerationContext.Compiler->CompilerLog(FText::FromString(msg), Node, EMessageSeverity::Warning);
+						GenerationContext.Log(FText::FromString(msg), Node, EMessageSeverity::Warning);
 					}*/
 				}
 
@@ -3720,13 +3717,13 @@ mu::NodeMeshPtr GenerateMutableSourceMesh(const UEdGraphPin* Pin,
 			}
 			else
 			{
-				GenerationContext.Compiler->CompilerLog(LOCTEXT("MorphStackGenerationFailed", "Stack definition Generation failed."), Node);
+				GenerationContext.Log(LOCTEXT("MorphStackGenerationFailed", "Stack definition Generation failed."), Node);
 				Result = nullptr;
 			}
 		}
 		else
 		{
-			GenerationContext.Compiler->CompilerLog(LOCTEXT("MorphStackConnectionFailed", "Stack definition connection not found."), Node);
+			GenerationContext.Log(LOCTEXT("MorphStackConnectionFailed", "Stack definition connection not found."), Node);
 			Result = nullptr;
 		}
 	}
@@ -3750,7 +3747,7 @@ mu::NodeMeshPtr GenerateMutableSourceMesh(const UEdGraphPin* Pin,
 					if (EnumPin)
 					{
 						const FText Message = LOCTEXT("FailedToGenerateSwitchParam", "Could not generate switch enum parameter. Please refesh the switch node and connect an enum.");
-						GenerationContext.Compiler->CompilerLog(Message, Node);
+						GenerationContext.Log(Message, Node);
 					}
 
 					return Result;
@@ -3759,7 +3756,7 @@ mu::NodeMeshPtr GenerateMutableSourceMesh(const UEdGraphPin* Pin,
 				if (SwitchParam->GetType() != mu::NodeScalarEnumParameter::GetStaticType())
 				{
 					const FText Message = LOCTEXT("WrongSwitchParamType", "Switch parameter of incorrect type.");
-					GenerationContext.Compiler->CompilerLog(Message, Node);
+					GenerationContext.Log(Message, Node);
 
 					return Result;
 				}
@@ -3770,7 +3767,7 @@ mu::NodeMeshPtr GenerateMutableSourceMesh(const UEdGraphPin* Pin,
 				if (NumSwitchOptions != EnumParameter->GetValueCount())
 				{
 					const FText Message = LOCTEXT("MismatchedSwitch", "Switch enum and switch node have different number of options. Please refresh the switch node to make sure the outcomes are labeled properly.");
-					GenerationContext.Compiler->CompilerLog(Message, Node);
+					GenerationContext.Log(Message, Node);
 				}
 
 				mu::NodeMeshSwitchPtr SwitchNode = new mu::NodeMeshSwitch;
@@ -3793,7 +3790,7 @@ mu::NodeMeshPtr GenerateMutableSourceMesh(const UEdGraphPin* Pin,
 			}
 			else
 			{
-				GenerationContext.Compiler->CompilerLog(LOCTEXT("NoEnumParamInSwitch", "Switch nodes must have an enum switch parameter. Please connect an enum and refesh the switch node."), Node);
+				GenerationContext.Log(LOCTEXT("NoEnumParamInSwitch", "Switch nodes must have an enum switch parameter. Please connect an enum and refesh the switch node."), Node);
 				return Result;
 			}
 		}(); // invoke lambda.
@@ -3815,7 +3812,7 @@ mu::NodeMeshPtr GenerateMutableSourceMesh(const UEdGraphPin* Pin,
 			}
 			else
 			{
-				GenerationContext.Compiler->CompilerLog(LOCTEXT("MeshFailed", "Mesh generation failed."), Node);
+				GenerationContext.Log(LOCTEXT("MeshFailed", "Mesh generation failed."), Node);
 			}
 		}
 
@@ -3853,13 +3850,13 @@ mu::NodeMeshPtr GenerateMutableSourceMesh(const UEdGraphPin* Pin,
 			}
 			else
 			{
-				GenerationContext.Compiler->CompilerLog(LOCTEXT("MeshGenerationFailed", "Mesh generation failed."), Node);
+				GenerationContext.Log(LOCTEXT("MeshGenerationFailed", "Mesh generation failed."), Node);
 			}
 		}
 		else
 		{
 			FText Text = FText::Format(LOCTEXT("MeshGeometryMissingDef", "Geometry Operation node requires the {0} value."), TypedNodeGeometry->MeshAPin()->PinFriendlyName);
-			GenerationContext.Compiler->CompilerLog(Text, Node);
+			GenerationContext.Log(Text, Node);
 		}
 
 		if (const UEdGraphPin* ConnectedPin = FollowInputPin(*TypedNodeGeometry->MeshBPin()))
@@ -3873,7 +3870,7 @@ mu::NodeMeshPtr GenerateMutableSourceMesh(const UEdGraphPin* Pin,
 			}
 			else
 			{
-				GenerationContext.Compiler->CompilerLog(LOCTEXT("MeshGenerationFailed", "Mesh generation failed."), Node);
+				GenerationContext.Log(LOCTEXT("MeshGenerationFailed", "Mesh generation failed."), Node);
 			}
 		}
 
@@ -3886,7 +3883,7 @@ mu::NodeMeshPtr GenerateMutableSourceMesh(const UEdGraphPin* Pin,
 			}
 			else
 			{
-				GenerationContext.Compiler->CompilerLog(LOCTEXT("ScalarGenerationFailed", "Scalar generation failed."), Node);
+				GenerationContext.Log(LOCTEXT("ScalarGenerationFailed", "Scalar generation failed."), Node);
 			}
 		}
 
@@ -3899,7 +3896,7 @@ mu::NodeMeshPtr GenerateMutableSourceMesh(const UEdGraphPin* Pin,
 			}
 			else
 			{
-				GenerationContext.Compiler->CompilerLog(LOCTEXT("ScalarGenerationFailed", "Scalar generation failed."), Node);
+				GenerationContext.Log(LOCTEXT("ScalarGenerationFailed", "Scalar generation failed."), Node);
 			}
 		}
 	}
@@ -3920,12 +3917,12 @@ mu::NodeMeshPtr GenerateMutableSourceMesh(const UEdGraphPin* Pin,
 			}
 			else
 			{
-				GenerationContext.Compiler->CompilerLog(LOCTEXT("MeshFailed", "Mesh generation failed."), Node);
+				GenerationContext.Log(LOCTEXT("MeshFailed", "Mesh generation failed."), Node);
 			}
 		}
 		else
 		{
-			GenerationContext.Compiler->CompilerLog(LOCTEXT("MeshReshapeMissingDef", "Mesh reshape node requires a default value."), Node);
+			GenerationContext.Log(LOCTEXT("MeshReshapeMissingDef", "Mesh reshape node requires a default value."), Node);
 		}
 	
 		{
@@ -3964,7 +3961,7 @@ mu::NodeMeshPtr GenerateMutableSourceMesh(const UEdGraphPin* Pin,
 						}
 					}
 
-					GenerationContext.Compiler->CompilerLog(
+					GenerationContext.Log(
 						LOCTEXT("MeshReshapeColorUsageMask", 
 								"Only one color channel with mask weight usage is allowed, multiple found. Reshape masking disabled."),
 						Node);
@@ -4049,7 +4046,7 @@ mu::NodeMeshPtr GenerateMutableSourceMesh(const UEdGraphPin* Pin,
 			}
 			else
 			{
-				GenerationContext.Compiler->CompilerLog(LOCTEXT("MeshFailed", "Mesh generation failed."), Node);
+				GenerationContext.Log(LOCTEXT("MeshFailed", "Mesh generation failed."), Node);
 			}
 		}
 
@@ -4066,7 +4063,7 @@ mu::NodeMeshPtr GenerateMutableSourceMesh(const UEdGraphPin* Pin,
 			}
 			else
 			{
-				GenerationContext.Compiler->CompilerLog(LOCTEXT("MeshFailed", "Mesh generation failed."), Node);
+				GenerationContext.Log(LOCTEXT("MeshFailed", "Mesh generation failed."), Node);
 			}
 		}
 
@@ -4084,7 +4081,7 @@ mu::NodeMeshPtr GenerateMutableSourceMesh(const UEdGraphPin* Pin,
 			{
 				if (BaseShapeTriangleCount != TargetShapeTriangleCount || BaseShapeTriangleCount == -1 || TargetShapeTriangleCount == -1)
 				{
-					GenerationContext.Compiler->CompilerLog(LOCTEXT("ReshapeMeshShapeIncompatible",
+					GenerationContext.Log(LOCTEXT("ReshapeMeshShapeIncompatible",
 						"Base and Target Shapes might not be compatible. Don't have the same number of triangles."), Node, EMessageSeverity::Warning);
 				}
 			}
@@ -4115,7 +4112,7 @@ mu::NodeMeshPtr GenerateMutableSourceMesh(const UEdGraphPin* Pin,
 					else
 					{
 						FString msg = FString::Printf(TEXT("Couldn't get bone transform information from a Pose Asset."));
-						GenerationContext.Compiler->CompilerLog(FText::FromString(msg), Node);
+						GenerationContext.Log(FText::FromString(msg), Node);
 
 						Result = nullptr;
 					}
@@ -4138,7 +4135,7 @@ mu::NodeMeshPtr GenerateMutableSourceMesh(const UEdGraphPin* Pin,
 						if (!TypedNode->PoseAsset) // Check if the slot has a selected pose. Could be left empty by the user
 						{
 							FString msg = FString::Printf(TEXT("Found pose mesh node without a pose asset assigned."));
-							GenerationContext.Compiler->CompilerLog(FText::FromString(msg), TypedNode);
+							GenerationContext.Log(FText::FromString(msg), TypedNode);
 						}
 
 					Result = InputMeshNode;
@@ -4167,7 +4164,7 @@ mu::NodeMeshPtr GenerateMutableSourceMesh(const UEdGraphPin* Pin,
 			if (!Property)
 			{
 				FString Msg = FString::Printf(TEXT("Couldn't find the column [%s] in the data table's struct."), *DataTableColumnName);
-				GenerationContext.Compiler->CompilerLog(FText::FromString(Msg), Node);
+				GenerationContext.Log(FText::FromString(Msg), Node);
 
 				bSuccess = false;
 			}
@@ -4179,7 +4176,7 @@ mu::NodeMeshPtr GenerateMutableSourceMesh(const UEdGraphPin* Pin,
 			if (bSuccess && !DefaultSkeletalMesh && !DefaultStaticMesh && !DefaultPoseAsset)
 			{
 				FString Msg = FString::Printf(TEXT("Couldn't find a default value in the data table's struct for the column [%s]."), *DataTableColumnName);
-				GenerationContext.Compiler->CompilerLog(FText::FromString(Msg), Node);
+				GenerationContext.Log(FText::FromString(Msg), Node);
 
 				bSuccess = false;
 			}
@@ -4226,7 +4223,7 @@ mu::NodeMeshPtr GenerateMutableSourceMesh(const UEdGraphPin* Pin,
 						if (!bSuccess)
 						{
 							FString Msg = FString::Printf(TEXT("Failed to generate the mutable table column [%s]"), *MutableColumnName);
-							GenerationContext.Compiler->CompilerLog(FText::FromString(Msg), Node);
+							GenerationContext.Log(FText::FromString(Msg), Node);
 						}
 					}
 
@@ -4292,7 +4289,7 @@ mu::NodeMeshPtr GenerateMutableSourceMesh(const UEdGraphPin* Pin,
 									if (bWasEmpty)
 									{
 										FString msg = "Mesh Column [" + MutableColumnName + "] Layout doesn't has any block. A grid sized block will be used instead.";
-										GenerationContext.Compiler->CompilerLog(FText::FromString(msg), Node, EMessageSeverity::Warning);
+										GenerationContext.Log(FText::FromString(msg), Node, EMessageSeverity::Warning);
 									}
 
 									MeshTableNode->SetLayout(i, LayoutNode);
@@ -4310,19 +4307,19 @@ mu::NodeMeshPtr GenerateMutableSourceMesh(const UEdGraphPin* Pin,
 				else
 				{
 					FString Msg = FString::Printf(TEXT("Couldn't generate a mutable table."));
-					GenerationContext.Compiler->CompilerLog(FText::FromString(Msg), Node);
+					GenerationContext.Log(FText::FromString(Msg), Node);
 				}
 			}
 		}
 		else
 		{
-			GenerationContext.Compiler->CompilerLog(LOCTEXT("ImageTableError", "Couldn't find the data table of the node."), Node);
+			GenerationContext.Log(LOCTEXT("ImageTableError", "Couldn't find the data table of the node."), Node);
 		}
 	}
 	
 	else
 	{
-		GenerationContext.Compiler->CompilerLog(LOCTEXT("UnimplementedMeshNode", "Mesh node type not implemented yet."), Node);
+		GenerationContext.Log(LOCTEXT("UnimplementedMeshNode", "Mesh node type not implemented yet."), Node);
 	}
 	
 	GenerationContext.Generated.Add(Key, FGeneratedData(Node, Result, &MeshData));
diff --git a/Engine/Plugins/Mutable/Source/CustomizableObjectEditor/Private/MuCOE/GenerateMutableSource/GenerateMutableSourceModifier.cpp b/Engine/Plugins/Mutable/Source/CustomizableObjectEditor/Private/MuCOE/GenerateMutableSource/GenerateMutableSourceModifier.cpp
--- a/Engine/Plugins/Mutable/Source/CustomizableObjectEditor/Private/MuCOE/GenerateMutableSource/GenerateMutableSourceModifier.cpp
+++ b/Engine/Plugins/Mutable/Source/CustomizableObjectEditor/Private/MuCOE/GenerateMutableSource/GenerateMutableSourceModifier.cpp
@@ -125,7 +125,7 @@ mu::Ptr<mu::NodeModifier> GenerateMutableSourceModifier(const UEdGraphPin * Pin,
 		else
 		{
 			FText ErrorMsg = LOCTEXT("ClipDeform mesh", "The clip deform node requires an input clip shape.");
-			GenerationContext.Compiler->CompilerLog(ErrorMsg, TypedNodeClipDeform, EMessageSeverity::Error);
+			GenerationContext.Log(ErrorMsg, TypedNodeClipDeform, EMessageSeverity::Error);
 			Result = nullptr;
 		}
 	
@@ -173,13 +173,13 @@ mu::Ptr<mu::NodeModifier> GenerateMutableSourceModifier(const UEdGraphPin * Pin,
 				{
 					// TODO: We support the clip mesh not being constant. This message is not precise enough. It should say that it hasn't been 
 					// possible to check if the mesh is closed or not.
-					GenerationContext.Compiler->CompilerLog(LOCTEXT("UnimplementedNode", "Node type not implemented yet."), MeshData.Node);
+					GenerationContext.Log(LOCTEXT("UnimplementedNode", "Node type not implemented yet."), MeshData.Node);
 				}
 
 				if (!bClosed)
 				{
 					FText ErrorMsg = FText::Format(LOCTEXT("Clipping mesh", "Clipping mesh [{0}] not closed (i.e., it does not enclose a volume)."), FText::FromName(MeshData.Mesh->GetFName()));
-					GenerationContext.Compiler->CompilerLog(ErrorMsg, MeshData.Node, EMessageSeverity::Warning);
+					GenerationContext.Log(ErrorMsg, MeshData.Node, EMessageSeverity::Warning);
 				}
 			}
 
@@ -197,7 +197,7 @@ mu::Ptr<mu::NodeModifier> GenerateMutableSourceModifier(const UEdGraphPin * Pin,
 		else
 		{
 			FText ErrorMsg = LOCTEXT("Clipping mesh missing", "The clip mesh with mesh node requires an input clip mesh.");
-			GenerationContext.Compiler->CompilerLog(ErrorMsg, TypedNodeClipMesh, EMessageSeverity::Error);
+			GenerationContext.Log(ErrorMsg, TypedNodeClipMesh, EMessageSeverity::Error);
 			Result = nullptr;
 		}
 
@@ -234,7 +234,7 @@ mu::Ptr<mu::NodeModifier> GenerateMutableSourceModifier(const UEdGraphPin * Pin,
 		else
 		{
 			FText ErrorMsg = LOCTEXT("ClipUVMask mesh", "The clip mesh with UV Mask node requires an input texture mask.");
-			GenerationContext.Compiler->CompilerLog(ErrorMsg, TypedNodeClipUVMask, EMessageSeverity::Error);
+			GenerationContext.Log(ErrorMsg, TypedNodeClipUVMask, EMessageSeverity::Error);
 			Result = nullptr;
 		}
 
@@ -307,7 +307,7 @@ mu::Ptr<mu::NodeModifier> GenerateMutableSourceModifier(const UEdGraphPin * Pin,
 				}
 				else
 				{
-					GenerationContext.Compiler->CompilerLog(LOCTEXT("ExtendMaterialLayoutMissing", "Skeletal Mesh without Layout Node linked to an Extend Material. A 4x4 layout will be added as default layout."), Node);
+					GenerationContext.Log(LOCTEXT("ExtendMaterialLayoutMissing", "Skeletal Mesh without Layout Node linked to an Extend Material. A 4x4 layout will be added as default layout."), Node);
 				}
 
 				mu::Ptr<mu::NodeMeshFormat> MeshFormat = new mu::NodeMeshFormat();
@@ -553,19 +553,19 @@ mu::Ptr<mu::NodeModifier> GenerateMutableSourceModifier(const UEdGraphPin * Pin,
 				{
 					validStaticFactor = false;
 					FString msg = FString::Printf(TEXT("Mesh morph nodes only accept factors between -1.0 and 1.0 inclusive but the default value of the float parameter node is (%f). Factor will be ignored."), floatParameterNode->DefaultValue);
-					GenerationContext.Compiler->CompilerLog(FText::FromString(msg), Node);
+					GenerationContext.Log(FText::FromString(msg), Node);
 				}
 				if (floatParameterNode->ParamUIMetadata.MinimumValue < -1.0f)
 				{
 					validStaticFactor = false;
 					FString msg = FString::Printf(TEXT("Mesh morph nodes only accept factors between -1.0 and 1.0 inclusive but the minimum UI value for the input float parameter node is (%f). Factor will be ignored."), floatParameterNode->ParamUIMetadata.MinimumValue);
-					GenerationContext.Compiler->CompilerLog(FText::FromString(msg), Node);
+					GenerationContext.Log(FText::FromString(msg), Node);
 				}
 				if (floatParameterNode->ParamUIMetadata.MaximumValue > 1.0f)
 				{
 					validStaticFactor = false;
 					FString msg = FString::Printf(TEXT("Mesh morph nodes only accept factors between -1.0 and 1.0 inclusive but the maximum UI value for the input float parameter node is (%f). Factor will be ignored."), floatParameterNode->ParamUIMetadata.MaximumValue);
-					GenerationContext.Compiler->CompilerLog(FText::FromString(msg), Node);
+					GenerationContext.Log(FText::FromString(msg), Node);
 				}
 			}
 			else if (const UCustomizableObjectNodeFloatConstant* floatConstantNode = Cast<UCustomizableObjectNodeFloatConstant>(floatNode))
@@ -574,7 +574,7 @@ mu::Ptr<mu::NodeModifier> GenerateMutableSourceModifier(const UEdGraphPin * Pin,
 				{
 					validStaticFactor = false;
 					FString msg = FString::Printf(TEXT("Mesh morph nodes only accept factors between -1.0 and 1.0 inclusive but the value of the float constant node is (%f). Factor will be ignored."), floatConstantNode->Value);
-					GenerationContext.Compiler->CompilerLog(FText::FromString(msg), Node);
+					GenerationContext.Log(FText::FromString(msg), Node);
 				}
 			}
 
@@ -590,7 +590,7 @@ mu::Ptr<mu::NodeModifier> GenerateMutableSourceModifier(const UEdGraphPin * Pin,
 
 	else
 	{
-		GenerationContext.Compiler->CompilerLog(LOCTEXT("UnimplementedNode", "Node type not implemented yet."), Node);
+		GenerationContext.Log(LOCTEXT("UnimplementedNode", "Node type not implemented yet."), Node);
 	}
 
 	if (!bDoNotAddToGeneratedCache)
diff --git a/Engine/Plugins/Mutable/Source/CustomizableObjectEditor/Private/MuCOE/GenerateMutableSource/GenerateMutableSourceProjector.cpp b/Engine/Plugins/Mutable/Source/CustomizableObjectEditor/Private/MuCOE/GenerateMutableSource/GenerateMutableSourceProjector.cpp
--- a/Engine/Plugins/Mutable/Source/CustomizableObjectEditor/Private/MuCOE/GenerateMutableSource/GenerateMutableSourceProjector.cpp
+++ b/Engine/Plugins/Mutable/Source/CustomizableObjectEditor/Private/MuCOE/GenerateMutableSource/GenerateMutableSourceProjector.cpp
@@ -137,7 +137,7 @@ mu::NodeProjectorPtr GenerateMutableSourceProjector(const UEdGraphPin* Pin, FMut
 
 	else
 	{
-		GenerationContext.Compiler->CompilerLog(LOCTEXT("UnimplementedNode", "Node type not implemented yet."), Node);
+		GenerationContext.Log(LOCTEXT("UnimplementedNode", "Node type not implemented yet."), Node);
 	}
 
 	GenerationContext.Generated.Add(Key, FGeneratedData(Node, Result));
diff --git a/Engine/Plugins/Mutable/Source/CustomizableObjectEditor/Private/MuCOE/GenerateMutableSource/GenerateMutableSourceSurface.cpp b/Engine/Plugins/Mutable/Source/CustomizableObjectEditor/Private/MuCOE/GenerateMutableSource/GenerateMutableSourceSurface.cpp
--- a/Engine/Plugins/Mutable/Source/CustomizableObjectEditor/Private/MuCOE/GenerateMutableSource/GenerateMutableSourceSurface.cpp
+++ b/Engine/Plugins/Mutable/Source/CustomizableObjectEditor/Private/MuCOE/GenerateMutableSource/GenerateMutableSourceSurface.cpp
@@ -244,7 +244,7 @@ mu::Ptr<mu::NodeSurface> GenerateMutableSourceSurface(const UEdGraphPin * Pin, F
 			if (const UEdGraphPin* StaticMeshPin = FindMeshBaseSource(*ConnectedMaterialPin, true))
 			{
 				const UCustomizableObjectNode* StaticMeshNode = CastChecked<UCustomizableObjectNode>(StaticMeshPin->GetOwningNode());
-				GenerationContext.Compiler->CompilerLog(LOCTEXT("UnsupportedStaticMeshes", "Static meshes are currently not supported as material meshes"), StaticMeshNode);
+				GenerationContext.Log(LOCTEXT("UnsupportedStaticMeshes", "Static meshes are currently not supported as material meshes"), StaticMeshNode);
 			}
 		}
 
@@ -333,7 +333,7 @@ mu::Ptr<mu::NodeSurface> GenerateMutableSourceSurface(const UEdGraphPin * Pin, F
 
 			if (const UEdGraphPin* ConnectedPin = FollowInputPin(*TypedNodeMat->GetMeshPin()))
 			{
-				GenerationContext.Compiler->CompilerLog(LOCTEXT("MeshIgnored", "The mesh nodes connected to a material node will be ignored because it is part of an explicit mesh component."), Node);
+				GenerationContext.Log(LOCTEXT("MeshIgnored", "The mesh nodes connected to a material node will be ignored because it is part of an explicit mesh component."), Node);
 			}
 		}
 		else
@@ -385,7 +385,7 @@ mu::Ptr<mu::NodeSurface> GenerateMutableSourceSurface(const UEdGraphPin * Pin, F
 				}
 				else
 				{
-					GenerationContext.Compiler->CompilerLog(LOCTEXT("MeshFailed", "Mesh generation failed."), Node);
+					GenerationContext.Log(LOCTEXT("MeshFailed", "Mesh generation failed."), Node);
 				}
 			}
 		}
@@ -416,13 +416,13 @@ mu::Ptr<mu::NodeSurface> GenerateMutableSourceSurface(const UEdGraphPin * Pin, F
 					{
 						bTableMaterialPinLinked = false;
 
-						GenerationContext.Compiler->CompilerLog(LOCTEXT("DifferentParentMaterial", "The Deafult Material Instance of the Data Table must have the same Parent Material."), TypedNodeMat->GetMaterialNode());
+						GenerationContext.Log(LOCTEXT("DifferentParentMaterial", "The Deafult Material Instance of the Data Table must have the same Parent Material."), TypedNodeMat->GetMaterialNode());
 					}
 				}
 				else
 				{
 					FText Msg = FText::Format(LOCTEXT("DefaultValueNotFound", "Couldn't find a default value in the data table's struct for the column {0}. The default value is null or not a Material Instance."), FText::FromString(TableColumnName));
-					GenerationContext.Compiler->CompilerLog(Msg, Node);
+					GenerationContext.Log(Msg, Node);
 
 					bTableMaterialPinLinked = false;
 				}
@@ -575,14 +575,14 @@ mu::Ptr<mu::NodeSurface> GenerateMutableSourceSurface(const UEdGraphPin * Pin, F
 
 									FString msg = FString::Printf(TEXT("The Reference texture [%s] is using TEXTUREGROUP_UI which does not support streaming. Please set a different TEXTURE group."),
 										*ReferenceTexture->GetName(), *ImageName);
-									GenerationContext.Compiler->CompilerLog(FText::FromString(msg), Node, EMessageSeverity::Info);
+									GenerationContext.Log(FText::FromString(msg), Node, EMessageSeverity::Info);
 								}
 							}
 							else
 							{
 								// warning!
 								FString msg = FString::Printf(TEXT("The Reference texture for material image [%s] is not set and it couldn't be found automatically."), *ImageName);
-								GenerationContext.Compiler->CompilerLog(FText::FromString(msg), Node);
+								GenerationContext.Log(FText::FromString(msg), Node);
 							}
 						}
 
@@ -745,7 +745,7 @@ mu::Ptr<mu::NodeSurface> GenerateMutableSourceSurface(const UEdGraphPin * Pin, F
 										*GenerationContext.Options.TargetPlatform->PlatformName(),
 										*ReferenceTexture->GetName());
 									const FText ReplacedImageFormatText = FText::FromString(ReplacedImageFormatMsg);
-									GenerationContext.Compiler->CompilerLog(ReplacedImageFormatText, Node, EMessageSeverity::Info);
+									GenerationContext.Log(ReplacedImageFormatText, Node, EMessageSeverity::Info);
 									UE_LOG(LogMutable, Log, TEXT("%s"), *ReplacedImageFormatMsg);
 								}
 								else if (BuildSettingsPerFormatPerLayer.Num() > 1)
@@ -755,7 +755,7 @@ mu::Ptr<mu::NodeSurface> GenerateMutableSourceSurface(const UEdGraphPin * Pin, F
 										*GenerationContext.Options.TargetPlatform->PlatformName(),
 										*ReferenceTexture->GetName());
 									const FText ReplacedImageFormatText = FText::FromString(ReplacedImageFormatMsg);
-									GenerationContext.Compiler->CompilerLog(ReplacedImageFormatText, Node, EMessageSeverity::Info);
+									GenerationContext.Log(ReplacedImageFormatText, Node, EMessageSeverity::Info);
 									UE_LOG(LogMutable, Log, TEXT("%s"), *ReplacedImageFormatMsg);
 								}
 							}
@@ -815,7 +815,7 @@ mu::Ptr<mu::NodeSurface> GenerateMutableSourceSurface(const UEdGraphPin * Pin, F
 									{
 										const FString ReplacedImageFormatMsg = FString::Printf(TEXT("In object [%s] the unsupported image format %d is used, IF_RGBA_UBYTE will be used instead."), *GenerationContext.Object->GetName(), UnrealTargetPlatformFormat );
 										const FText ReplacedImageFormatText = FText::FromString(ReplacedImageFormatMsg);
-										GenerationContext.Compiler->CompilerLog(ReplacedImageFormatText, Node, EMessageSeverity::Info);
+										GenerationContext.Log(ReplacedImageFormatText, Node, EMessageSeverity::Info);
 										UE_LOG(LogMutable, Log, TEXT("%s"), *ReplacedImageFormatMsg);
 										MutableFormat = mu::EImageFormat::IF_RGBA_UBYTE;
 									}
@@ -823,7 +823,7 @@ mu::Ptr<mu::NodeSurface> GenerateMutableSourceSurface(const UEdGraphPin * Pin, F
 									{
 										const FString ReplacedImageFormatMsg = FString::Printf(TEXT("In object [%s] the unsupported image format %d is used, IF_RGBA_UBYTE will be used instead."), *GenerationContext.Object->GetName(), UnrealTargetPlatformFormatAlpha);
 										const FText ReplacedImageFormatText = FText::FromString(ReplacedImageFormatMsg);
-										GenerationContext.Compiler->CompilerLog(ReplacedImageFormatText, Node, EMessageSeverity::Info);
+										GenerationContext.Log(ReplacedImageFormatText, Node, EMessageSeverity::Info);
 										UE_LOG(LogMutable, Log, TEXT("%s"), *ReplacedImageFormatMsg);
 										MutableFormatIfAlpha = mu::EImageFormat::IF_RGBA_UBYTE;
 									}
@@ -1076,7 +1076,7 @@ mu::Ptr<mu::NodeSurface> GenerateMutableSourceSurface(const UEdGraphPin * Pin, F
 			}
 			else
 			{
-				GenerationContext.Compiler->CompilerLog(LOCTEXT("SurfaceFailed", "Surface generation failed."), Node);
+				GenerationContext.Log(LOCTEXT("SurfaceFailed", "Surface generation failed."), Node);
 			}
 		}
 
@@ -1099,7 +1099,7 @@ mu::Ptr<mu::NodeSurface> GenerateMutableSourceSurface(const UEdGraphPin * Pin, F
 					}
 					else
 					{
-						GenerationContext.Compiler->CompilerLog(LOCTEXT("SurfaceModifierFailed", "Surface generation failed."), Node);
+						GenerationContext.Log(LOCTEXT("SurfaceModifierFailed", "Surface generation failed."), Node);
 					}
 				}
 			}
@@ -1125,7 +1125,7 @@ mu::Ptr<mu::NodeSurface> GenerateMutableSourceSurface(const UEdGraphPin * Pin, F
 					if (EnumPin)
 					{
 						const FText Message = LOCTEXT("FailedToGenerateSwitchParam", "Could not generate switch enum parameter. Please refesh the switch node and connect an enum.");
-						GenerationContext.Compiler->CompilerLog(Message, Node);
+						GenerationContext.Log(Message, Node);
 					}
 
 					return Result;
@@ -1134,7 +1134,7 @@ mu::Ptr<mu::NodeSurface> GenerateMutableSourceSurface(const UEdGraphPin * Pin, F
 				if (SwitchParam->GetType() != mu::NodeScalarEnumParameter::GetStaticType())
 				{
 					const FText Message = LOCTEXT("WrongSwitchParamType", "Switch parameter of incorrect type.");
-					GenerationContext.Compiler->CompilerLog(Message, Node);
+					GenerationContext.Log(Message, Node);
 
 					return Result;
 				}
@@ -1145,7 +1145,7 @@ mu::Ptr<mu::NodeSurface> GenerateMutableSourceSurface(const UEdGraphPin * Pin, F
 				if (NumSwitchOptions != EnumParameter->GetValueCount())
 				{
 					const FText Message = LOCTEXT("MismatchedSwitch", "Switch enum and switch node have different number of options. Please refresh the switch node to make sure the outcomes are labeled properly.");
-					GenerationContext.Compiler->CompilerLog(Message, Node);
+					GenerationContext.Log(Message, Node);
 				}
 
 				mu::Ptr<mu::NodeSurfaceSwitch> SwitchNode = new mu::NodeSurfaceSwitch;
@@ -1164,7 +1164,7 @@ mu::Ptr<mu::NodeSurface> GenerateMutableSourceSurface(const UEdGraphPin * Pin, F
 						else
 						{
 							// Probably ok
-							//GenerationContext.Compiler->CompilerLog(LOCTEXT("SurfaceModifierFailed", "Surface generation failed."), Node);
+							//GenerationContext.Log(LOCTEXT("SurfaceModifierFailed", "Surface generation failed."), Node);
 						}
 					}
 				}
@@ -1174,7 +1174,7 @@ mu::Ptr<mu::NodeSurface> GenerateMutableSourceSurface(const UEdGraphPin * Pin, F
 			}
 			else
 			{
-				GenerationContext.Compiler->CompilerLog(LOCTEXT("NoEnumParamInSwitch", "Switch nodes must have an enum switch parameter. Please connect an enum and refesh the switch node."), Node);
+				GenerationContext.Log(LOCTEXT("NoEnumParamInSwitch", "Switch nodes must have an enum switch parameter. Please connect an enum and refesh the switch node."), Node);
 				return Result;
 			}
 		}(); // invoke lambda.
@@ -1182,7 +1182,7 @@ mu::Ptr<mu::NodeSurface> GenerateMutableSourceSurface(const UEdGraphPin * Pin, F
 
 	else
 	{
-		GenerationContext.Compiler->CompilerLog(LOCTEXT("UnimplementedNode", "Node type not implemented yet."), Node);
+		GenerationContext.Log(LOCTEXT("UnimplementedNode", "Node type not implemented yet."), Node);
 	}
 
 
diff --git a/Engine/Plugins/Mutable/Source/CustomizableObjectEditor/Private/MuCOE/GenerateMutableSource/GenerateMutableSourceTable.cpp b/Engine/Plugins/Mutable/Source/CustomizableObjectEditor/Private/MuCOE/GenerateMutableSource/GenerateMutableSourceTable.cpp
--- a/Engine/Plugins/Mutable/Source/CustomizableObjectEditor/Private/MuCOE/GenerateMutableSource/GenerateMutableSourceTable.cpp
+++ b/Engine/Plugins/Mutable/Source/CustomizableObjectEditor/Private/MuCOE/GenerateMutableSource/GenerateMutableSourceTable.cpp
@@ -40,12 +40,15 @@ bool FillTableColumn(const UCustomizableObjectNodeTable* TableNode, mu::TablePtr
 	// Getting property type
 	if (const FSoftObjectProperty* SoftObjectProperty = CastField<FSoftObjectProperty>(ColumnProperty))
 	{
-		UObject* Object = SoftObjectProperty->GetPropertyValue(CellData).LoadSynchronous();
-
+		FSoftObjectPtr SoftObject = SoftObjectProperty->GetPropertyValue(CellData);
+		
 		if (SoftObjectProperty->PropertyClass->IsChildOf(USkeletalMesh::StaticClass()))
 		{
-			USkeletalMesh* SkeletalMesh = Cast<USkeletalMesh>(Object);
+			GenerationContext.AddParticipatingObject(SoftObject);
 
+			UObject* Object = GenerationContext.LoadObject(SoftObject, true);
+
+			USkeletalMesh* SkeletalMesh = Cast<USkeletalMesh>(Object);
 			if(!SkeletalMesh)
 			{
 				// Generating an Empty cell
@@ -95,7 +98,7 @@ bool FillTableColumn(const UCustomizableObjectNodeTable* TableNode, mu::TablePtr
 										FString Message = FString::Printf(
 											TEXT("The column with name [%s] for the Anim Slot property should be an FName instead of an Integer, it will be internally converted to FName but should probaly be converted in the table itself."), 
 											*AnimBP);
-										GenerationContext.Compiler->CompilerLog(FText::FromString(Message), TableNode, EMessageSeverity::Info);
+										GenerationContext.Log(FText::FromString(Message), TableNode, EMessageSeverity::Info);
 
 										SlotIndex = FName(FString::FromInt(IntProperty->GetPropertyValue(AnimSlotData)));
 									}
@@ -121,11 +124,7 @@ bool FillTableColumn(const UCustomizableObjectNodeTable* TableNode, mu::TablePtr
 
 											if (!AnimInstance.IsNull())
 											{
-												// TODO: We only need it to add it to participating objects. We should probably skip fully loading it.
-												if (UClass* Anim = AnimInstance.LoadSynchronous())
-												{
-													GenerationContext.AddParticipatingObject(*Anim);													
-												}
+												GenerationContext.AddParticipatingObject(AnimInstance);
 
 												const int32 AnimInstanceIndex = GenerationContext.AnimBPAssets.AddUnique(AnimInstance);
 
@@ -146,7 +145,7 @@ bool FillTableColumn(const UCustomizableObjectNodeTable* TableNode, mu::TablePtr
 				else
 				{
 					FString msg = FString::Printf(TEXT("Could not found the Slot column of the animation blueprint column [%s] for the mesh column [%s]."), *AnimBP, *ColumnName);
-					GenerationContext.Compiler->CompilerLog(FText::FromString(msg), TableNode);
+					GenerationContext.Log(FText::FromString(msg), TableNode);
 				}
 			}
 
@@ -186,7 +185,7 @@ bool FillTableColumn(const UCustomizableObjectNodeTable* TableNode, mu::TablePtr
 			if (!ReferenceSkeletalMesh)
 			{
 				FString msg = FString::Printf(TEXT("Reference Skeletal Mesh not found for column [%s]."), *ColumnName);
-				GenerationContext.Compiler->CompilerLog(FText::FromString(msg), TableNode);
+				GenerationContext.Log(FText::FromString(msg), TableNode);
 
 				return false;
 			}
@@ -334,8 +333,11 @@ bool FillTableColumn(const UCustomizableObjectNodeTable* TableNode, mu::TablePtr
 
 		else if (SoftObjectProperty->PropertyClass->IsChildOf(UStaticMesh::StaticClass()))
 		{
-			UStaticMesh* StaticMesh = Cast<UStaticMesh>(Object);
+			GenerationContext.AddParticipatingObject(SoftObject);
+			
+			UObject* Object = GenerationContext.LoadObject(SoftObject);
 
+			UStaticMesh* StaticMesh = Cast<UStaticMesh>(Object);
 			if (!StaticMesh)
 			{
 				return false;
@@ -347,7 +349,7 @@ bool FillTableColumn(const UCustomizableObjectNodeTable* TableNode, mu::TablePtr
 			if (!ReferenceStaticMesh)
 			{
 				FString msg = FString::Printf(TEXT("Reference Static Mesh not found for column [%s]."), *ColumnName);
-				GenerationContext.Compiler->CompilerLog(FText::FromString(msg), TableNode);
+				GenerationContext.Log(FText::FromString(msg), TableNode);
 
 				return false;
 			}
@@ -405,11 +407,14 @@ bool FillTableColumn(const UCustomizableObjectNodeTable* TableNode, mu::TablePtr
 
 		else if (SoftObjectProperty->PropertyClass->IsChildOf(UTexture::StaticClass()))
 		{
-			UTexture* Texture = Cast<UTexture>(Object);
+			GenerationContext.AddParticipatingObject(SoftObject);
 
 			// Removing encoding part
 			const FString PinName = ColumnName.Replace(TEXT("--PassThrough"), TEXT(""), ESearchCase::CaseSensitive);
+			
+			UObject* Object = GenerationContext.LoadObject(SoftObject);
 
+			UTexture* Texture = Cast<UTexture>(Object);
 			if (!Texture)
 			{
 				Texture = TableNode->GetColumnDefaultAssetByType<UTexture>(PinName);
@@ -430,21 +435,17 @@ bool FillTableColumn(const UCustomizableObjectNodeTable* TableNode, mu::TablePtr
 				CurrentColumn = MutableTable->AddColumn(ColumnName, mu::ETableColumnType::Image);
 			}
 
-			bool bIsPassthroughTexture = TableNode->GetColumnImageMode(PinName) == ETableTextureType::PASSTHROUGH_TEXTURE;
-			if (!bIsPassthroughTexture)
-			{
-				if (Texture)
-				{
-					GenerationContext.AddParticipatingObject(*Texture);
-				}
-			}
-
+			const bool bIsPassthroughTexture = TableNode->GetColumnImageMode(PinName) == ETableTextureType::PASSTHROUGH_TEXTURE;
 			mu::Ptr<mu::ResourceProxyMemory<mu::Image>> Proxy = new mu::ResourceProxyMemory<mu::Image>(GenerateImageConstant(Texture, GenerationContext, bIsPassthroughTexture));
 			MutableTable->SetCell(CurrentColumn, RowIdx, Proxy.get());
 		}
 
 		else if (SoftObjectProperty->PropertyClass->IsChildOf(UMaterialInterface::StaticClass()))
 		{
+			GenerationContext.AddParticipatingObject(SoftObject);
+			
+			UObject* Object = GenerationContext.LoadObject(SoftObject);
+
 			// Get display name of the column of the data table (name showed in the table and struct editors)
 			// Will be used in the warnings to help to identify a column with errors.
 			FString MaterialColumnDisplayName = ColumnProperty->GetDisplayNameText().ToString();
@@ -458,7 +459,7 @@ bool FillTableColumn(const UCustomizableObjectNodeTable* TableNode, mu::TablePtr
 			if (!ReferenceMaterial)
 			{
 				FString msg = FString::Printf(TEXT("Default Material Instance not found for column [%s]."), *MaterialColumnDisplayName);
-				GenerationContext.Compiler->CompilerLog(FText::FromString(msg), TableNode);
+				GenerationContext.Log(FText::FromString(msg), TableNode);
 
 				return false;
 			}
@@ -578,6 +579,10 @@ bool FillTableColumn(const UCustomizableObjectNodeTable* TableNode, mu::TablePtr
 
 		else if (SoftObjectProperty->PropertyClass->IsChildOf(UPoseAsset::StaticClass()))
 		{
+			GenerationContext.AddParticipatingObject(SoftObject);
+			
+			UObject* Object = GenerationContext.LoadObject(SoftObject);
+			
 			if (UPoseAsset* PoseAsset = Cast<UPoseAsset>(Object))
 			{
 				CurrentColumn = MutableTable.get()->FindColumn(ColumnName);
@@ -616,7 +621,7 @@ bool FillTableColumn(const UCustomizableObjectNodeTable* TableNode, mu::TablePtr
 		{
 			// Unsuported Variable Type
 			FString msg = FString::Printf(TEXT("[%s] is not a supported class for mutable nodes."), *SoftObjectProperty->PropertyClass.GetName());
-			GenerationContext.Compiler->CompilerLog(FText::FromString(msg), TableNode);
+			GenerationContext.Log(FText::FromString(msg), TableNode);
 
 			return false;
 		}
@@ -762,7 +767,7 @@ TArray<FName> GetEnabledRows(const UDataTable& DataTable, const UCustomizableObj
 
 void RestrictRowNamesToSelectedOption(TArray<FName>& InOutRowNames, const UCustomizableObjectNodeTable& TableNode, FMutableGraphGenerationContext& GenerationContext)
 {
-	if (!GenerationContext.ParamNamesToSelectedOptions.IsEmpty())
+	if (!GenerationContext.Options.ParamNamesToSelectedOptions.IsEmpty())
 	{
 		FMutableParamNameSet* ParamNameSet = GenerationContext.TableToParamNames.Find(TableNode.Table->GetPathName());
 
@@ -773,7 +778,7 @@ void RestrictRowNamesToSelectedOption(TArray<FName>& InOutRowNames, const UCusto
 			for (const FString& ParamName : ParamNameSet->ParamNames)
 			{
 				// If the param is in the map restrict to only the selected option
-				FString* SelectedOptionString = GenerationContext.ParamNamesToSelectedOptions.Find(ParamName);
+				const FString* SelectedOptionString = GenerationContext.Options.ParamNamesToSelectedOptions.Find(ParamName);
 
 				if (SelectedOptionString)
 				{
@@ -829,7 +834,7 @@ void RestrictRowContentByVersion( TArray<FName>& InOutRowNames, const UDataTable
 	if (!CustomizableObjectVersionBridgeInterface)
 	{
 		const FString Message = "Found a data table with at least a row with a Custom Version asset but the Root Object does not have a Version Bridge asset assigned.";
-		GenerationContext.Compiler->CompilerLog(FText::FromString(Message), &TableNode, EMessageSeverity::Error);
+		GenerationContext.Log(FText::FromString(Message), &TableNode, EMessageSeverity::Error);
 		return;
 	}
 
@@ -899,6 +904,7 @@ bool GenerateTableColumn(const UCustomizableObjectNodeTable* TableNode, const UE
 	TArray<FName> RowNames = GetRowsToCompile(*DataTable, *TableNode, GenerationContext);
 
 	// Pre-pass to request async loading of all data. This seems to be slightly faster because it avoids the sync after every separate load.
+	if (!GenerationContext.bParticipatingObjectsPass)
 	{
 		TArray<int32> LoadRequests;
 		LoadRequests.Reserve(RowNames.Num());
@@ -977,7 +983,7 @@ void GenerateTableParameterUIData(const UDataTable* DataTable, const UCustomizab
 		{
 			FText LogMessage = FText::Format(LOCTEXT("InvalidParamUIMetadataColumn_Warning",
 				"UI Metadata Column [{0}] is not a valid type or does not exist in the Structure of the Node."), FText::FromName(TableNode->ParamUIMetadataColumn));
-			GenerationContext.Compiler->CompilerLog(LogMessage, TableNode);
+			GenerationContext.Log(LogMessage, TableNode);
 		}
 
 		FProperty* ThumbnailColumnProperty = DataTable->FindTableProperty(TableNode->ThumbnailColumn);
@@ -988,7 +994,7 @@ void GenerateTableParameterUIData(const UDataTable* DataTable, const UCustomizab
 		{
 			FText LogMessage = FText::Format(LOCTEXT("InvalidThumbnailColumn_Warning",
 				"Thumbnail Column [{0}] is not an objet type or does not exist in the Structure of the Node."), FText::FromName(TableNode->ThumbnailColumn));
-			GenerationContext.Compiler->CompilerLog(LogMessage, TableNode);
+			GenerationContext.Log(LogMessage, TableNode);
 		}
 
 		if (!bIsValidMetadataColumn)
@@ -1059,7 +1065,7 @@ mu::TablePtr GenerateMutableSourceTable(const UDataTable* DataTable, const UCust
 {
 	check(DataTable && TableNode);
 
-	if (GenerationContext.ParamNamesToSelectedOptions.IsEmpty())
+	if (GenerationContext.Options.ParamNamesToSelectedOptions.IsEmpty())
 	{
 		FMutableParamNameSet* ParamNameSet = GenerationContext.TableToParamNames.Find(DataTable->GetPathName());
 
@@ -1085,7 +1091,7 @@ mu::TablePtr GenerateMutableSourceTable(const UDataTable* DataTable, const UCust
 			Nodes.Add(TableNode);
 			Nodes.Add(CachedTable->ReferenceNode);
 
-			GenerationContext.Compiler->CompilerLog(LOCTEXT("TableNodesCompilationRestrictionError",
+			GenerationContext.Log(LOCTEXT("TableNodesCompilationRestrictionError",
 				"Found one or more Table Nodes with the same data table but different Compilation Restrictions."), Nodes);
 		}
 
@@ -1124,7 +1130,7 @@ mu::TablePtr GenerateMutableSourceTable(const UDataTable* DataTable, const UCust
 	else
 	{
 		FString msg = "Couldn't find the Data Table's Struct asset in the Node.";
-		GenerationContext.Compiler->CompilerLog(FText::FromString(msg), DataTable);
+		GenerationContext.Log(FText::FromString(msg), DataTable);
 		
 		return nullptr;
 	}
@@ -1170,7 +1176,7 @@ UDataTable* GenerateDataTableFromStruct(const UCustomizableObjectNodeTable* Tabl
 {
 	if (!TableNode->Structure)
 	{
-		GenerationContext.Compiler->CompilerLog(LOCTEXT("EmptyStructureError", "Empty structure asset."), TableNode);
+		GenerationContext.Log(LOCTEXT("EmptyStructureError", "Empty structure asset."), TableNode);
 		return nullptr;
 	}
 
@@ -1193,7 +1199,7 @@ UDataTable* GenerateDataTableFromStruct(const UCustomizableObjectNodeTable* Tabl
 	if (TableNode->FilterPaths.IsEmpty())
 	{
 		// Preventing load all data tables of the project
-		GenerationContext.Compiler->CompilerLog(LOCTEXT("NoFilePathsError", "There are no filter paths selected. This is an error to prevent loading all data table of the project."), TableNode);
+		GenerationContext.Log(LOCTEXT("NoFilePathsError", "There are no filter paths selected. This is an error to prevent loading all data table of the project."), TableNode);
 
 		return nullptr;
 	}
@@ -1218,7 +1224,7 @@ UDataTable* GenerateDataTableFromStruct(const UCustomizableObjectNodeTable* Tabl
 
 	if (ParentTables.IsEmpty())
 	{
-		GenerationContext.Compiler->CompilerLog(LOCTEXT("NoDataTablesFoundWarning", "Could not find a data table with the specified struct in the selected paths."), TableNode);
+		GenerationContext.Log(LOCTEXT("NoDataTablesFoundWarning", "Could not find a data table with the specified struct in the selected paths."), TableNode);
 
 		return nullptr;
 	}
@@ -1270,7 +1276,7 @@ UDataTable* GenerateDataTableFromStruct(const UCustomizableObjectNodeTable* Tabl
 
 		FString Message = FString::Printf(TEXT("Row with name [%s] repeated in the following Data Tables: [%s]. The last row processed will be used [%s]."),
 			*RowName.ToString(), *TableNames, *DataTablesNames.Last()->GetName());
-		GenerationContext.Compiler->CompilerLog(FText::FromString(Message), TableNode);
+		GenerationContext.Log(FText::FromString(Message), TableNode);
 	}
 
 	CompositeDataTable->AppendParentTables(ParentTables);
@@ -1315,7 +1321,7 @@ void LogRowGenerationMessage(const UCustomizableObjectNodeTable* TableNode, cons
 		}
 	}
 
-	GenerationContext.Compiler->CompilerLog(FText::FromString(FinalMessage), TableNode);
+	GenerationContext.Log(FText::FromString(FinalMessage), TableNode);
 }
 
 
diff --git a/Engine/Plugins/Mutable/Source/CustomizableObjectEditor/Private/MuCOE/GraphTraversal.cpp b/Engine/Plugins/Mutable/Source/CustomizableObjectEditor/Private/MuCOE/GraphTraversal.cpp
--- a/Engine/Plugins/Mutable/Source/CustomizableObjectEditor/Private/MuCOE/GraphTraversal.cpp
+++ b/Engine/Plugins/Mutable/Source/CustomizableObjectEditor/Private/MuCOE/GraphTraversal.cpp
@@ -244,7 +244,7 @@ TArray<UEdGraphPin*> ReverseFollowPinArray(const UEdGraphPin& Pin, bool bIgnoreO
 }
 
 
-UCustomizableObjectNodeObject* GetRootNode(UCustomizableObject* Object, bool& bOutMultipleBaseObjectsFound)
+UCustomizableObjectNodeObject* GetRootNode(const UCustomizableObject* Object, bool& bOutMultipleBaseObjectsFound)
 {
 	// Look for the base object node
 	UCustomizableObjectNodeObject* Root = nullptr;
@@ -273,7 +273,7 @@ UCustomizableObjectNodeObject* GetRootNode(UCustomizableObject* Object, bool& bO
 }
 
 
-bool GetParentsUntilRoot(UCustomizableObject* Object, TArray<UCustomizableObjectNodeObject*>& ArrayNodeObject, TArray<UCustomizableObject*>& ArrayCustomizableObject)
+bool GetParentsUntilRoot(const UCustomizableObject* Object, TArray<UCustomizableObjectNodeObject*>& ArrayNodeObject, TArray<const  UCustomizableObject*>& ArrayCustomizableObject)
 {
 	bool MultipleBaseObjectsFound;
 	UCustomizableObjectNodeObject* Root = GetRootNode(Object, MultipleBaseObjectsFound);
@@ -329,7 +329,7 @@ bool HasCandidateAsParent(UCustomizableObjectNodeObject* Node, UCustomizableObje
 }
 
 
-UCustomizableObject* GetFullGraphRootObject(UCustomizableObjectNodeObject* Node, TArray<UCustomizableObject*>& VisitedObjects)
+UCustomizableObject* GetFullGraphRootObject(const UCustomizableObjectNodeObject* Node, TArray<UCustomizableObject*>& VisitedObjects)
 {
 	if (Node->ParentObject != nullptr)
 	{
@@ -373,7 +373,14 @@ UCustomizableObject* GetRootObject(const UCustomizableObjectNode& Node)
 }
 
 
-UCustomizableObject* GetRootObject(UCustomizableObject* ChildObject)
+UCustomizableObject* GraphTraversal::GetRootObject(UCustomizableObject* ChildObject)
+{
+	const UCustomizableObject* ConstChildObject = ChildObject;
+	return const_cast<UCustomizableObject*>(GetRootObject(ConstChildObject));
+}
+
+
+const UCustomizableObject* GraphTraversal::GetRootObject(const UCustomizableObject* ChildObject)
 {
 	// Grab a node to start the search -> Get the root since it should be always present
 	bool bMultipleBaseObjectsFound = false;
@@ -391,7 +398,7 @@ UCustomizableObject* GetRootObject(UCustomizableObject* ChildObject)
 }
 
 
-UCustomizableObjectNodeObject* GetFullGraphRootNodeObject(UCustomizableObjectNodeObject* Node, TArray<UCustomizableObject*>& VisitedObjects)
+UCustomizableObjectNodeObject* GetFullGraphRootNodeObject(const UCustomizableObjectNodeObject* Node, TArray<const UCustomizableObject*>& VisitedObjects)
 {
 	if (Node->ParentObject != nullptr)
 	{
@@ -659,7 +666,7 @@ void GetAllObjectsInGraph(UCustomizableObject* Object, TSet<UCustomizableObject*
 	}
 
 	// Search the root of the CO's graph
-	UCustomizableObject* RootObject = GetRootObject(Object);
+	UCustomizableObject* RootObject = GraphTraversal::GetRootObject(Object);
 	TMultiMap<FGuid, UCustomizableObjectNodeObject*> DummyMap;
 
 	FAssetRegistryModule& AssetRegistryModule = FModuleManager::LoadModuleChecked<FAssetRegistryModule>("AssetRegistry");
diff --git a/Engine/Plugins/Mutable/Source/CustomizableObjectEditor/Private/MuCOE/Nodes/CustomizableObjectNodeModifierBase.cpp b/Engine/Plugins/Mutable/Source/CustomizableObjectEditor/Private/MuCOE/Nodes/CustomizableObjectNodeModifierBase.cpp
--- a/Engine/Plugins/Mutable/Source/CustomizableObjectEditor/Private/MuCOE/Nodes/CustomizableObjectNodeModifierBase.cpp
+++ b/Engine/Plugins/Mutable/Source/CustomizableObjectEditor/Private/MuCOE/Nodes/CustomizableObjectNodeModifierBase.cpp
@@ -74,7 +74,7 @@ void UCustomizableObjectNodeModifierBase::GetPossiblyModifiedNodes(TArray<UCusto
 {
 	// Scan all potential receivers
 	UCustomizableObject* ThisNodeObject = GetRootObject(*this);
-	UCustomizableObject* RootObject = GetRootObject(ThisNodeObject);
+	UCustomizableObject* RootObject = GraphTraversal::GetRootObject(ThisNodeObject);
 
 	TSet<UCustomizableObject*> AllCustomizableObject;
 	GetAllObjectsInGraph(RootObject, AllCustomizableObject);
diff --git a/Engine/Plugins/Mutable/Source/CustomizableObjectEditor/Private/MuCOE/Nodes/CustomizableObjectNodeObjectDetails.cpp b/Engine/Plugins/Mutable/Source/CustomizableObjectEditor/Private/MuCOE/Nodes/CustomizableObjectNodeObjectDetails.cpp
--- a/Engine/Plugins/Mutable/Source/CustomizableObjectEditor/Private/MuCOE/Nodes/CustomizableObjectNodeObjectDetails.cpp
+++ b/Engine/Plugins/Mutable/Source/CustomizableObjectEditor/Private/MuCOE/Nodes/CustomizableObjectNodeObjectDetails.cpp
@@ -366,7 +366,7 @@ void FCustomizableObjectNodeObjectDetails::FillParameterNamesArray()
 	}
 
 	// Get full graph root customizable object
-	UCustomizableObject* RootObjet = GetRootObject(CustomizableObject);
+	UCustomizableObject* RootObjet = GraphTraversal::GetRootObject(CustomizableObject);
 
 	// Full tree graph of customizable objects
 	TSet<UCustomizableObject*> CustomObjectTree;
diff --git a/Engine/Plugins/Mutable/Source/CustomizableObjectEditor/Private/MuCOE/Nodes/CustomizableObjectNodeObjectGroup.cpp b/Engine/Plugins/Mutable/Source/CustomizableObjectEditor/Private/MuCOE/Nodes/CustomizableObjectNodeObjectGroup.cpp
--- a/Engine/Plugins/Mutable/Source/CustomizableObjectEditor/Private/MuCOE/Nodes/CustomizableObjectNodeObjectGroup.cpp
+++ b/Engine/Plugins/Mutable/Source/CustomizableObjectEditor/Private/MuCOE/Nodes/CustomizableObjectNodeObjectGroup.cpp
@@ -52,7 +52,7 @@ void UCustomizableObjectNodeObjectGroup::PostEditChangeProperty(FPropertyChanged
 			{
 				if (CurrentRootNode->ParentObject)
 				{
-					TArray<UCustomizableObject*> VisitedObjects;
+					TArray<const UCustomizableObject*> VisitedObjects;
 					CurrentRootNode = GetFullGraphRootNodeObject(CurrentRootNode, VisitedObjects);
 				}
 				if (!CurrentRootNode->ParentObject)
diff --git a/Engine/Plugins/Mutable/Source/CustomizableObjectEditor/Private/MuCOE/Nodes/CustomizableObjectNodeTable.cpp b/Engine/Plugins/Mutable/Source/CustomizableObjectEditor/Private/MuCOE/Nodes/CustomizableObjectNodeTable.cpp
--- a/Engine/Plugins/Mutable/Source/CustomizableObjectEditor/Private/MuCOE/Nodes/CustomizableObjectNodeTable.cpp
+++ b/Engine/Plugins/Mutable/Source/CustomizableObjectEditor/Private/MuCOE/Nodes/CustomizableObjectNodeTable.cpp
@@ -1030,11 +1030,11 @@ bool UCustomizableObjectNodeTable::CheckPinUpdated(const FString& PinName, const
 }
 
 
-USkeletalMesh* UCustomizableObjectNodeTable::GetSkeletalMeshAt(const UEdGraphPin* Pin, const UDataTable* DataTable, const FName& RowName) const
+FSoftObjectPtr UCustomizableObjectNodeTable::GetSkeletalMeshAt(const UEdGraphPin* Pin, const UDataTable* DataTable, const FName& RowName) const
 {
 	if (!DataTable || !DataTable->GetRowStruct() || !Pin || !DataTable->GetRowNames().Contains(RowName))
 	{
-		return nullptr;
+		return {};
 	}
 	
 	const UScriptStruct* TableStruct = DataTable->GetRowStruct();
@@ -1047,7 +1047,7 @@ USkeletalMesh* UCustomizableObjectNodeTable::GetSkeletalMeshAt(const UEdGraphPin
 
 	if (!ColumnProperty)
 	{
-		return nullptr;
+		return {};
 	}
 
 	if (const FSoftObjectProperty* SoftObjectProperty = CastField<FSoftObjectProperty>(ColumnProperty))
@@ -1056,18 +1056,12 @@ USkeletalMesh* UCustomizableObjectNodeTable::GetSkeletalMeshAt(const UEdGraphPin
 		{
 			if (uint8* CellData = ColumnProperty->ContainerPtrToValuePtr<uint8>(RowData, 0))
 			{
-				if (UObject* Object = SoftObjectProperty->GetPropertyValue(CellData).LoadSynchronous())
-				{
-					if (USkeletalMesh* SkeletalMesh = Cast<USkeletalMesh>(Object))
-					{
-						return SkeletalMesh;
-					}
-				}
+				return SoftObjectProperty->GetPropertyValue(CellData);
 			}
 		}
 	}
 
-	return nullptr;
+	return {};
 }
 
 
diff --git a/Engine/Plugins/Mutable/Source/CustomizableObjectEditor/Private/MuCOE/SMutableTagListWidget.cpp b/Engine/Plugins/Mutable/Source/CustomizableObjectEditor/Private/MuCOE/SMutableTagListWidget.cpp
--- a/Engine/Plugins/Mutable/Source/CustomizableObjectEditor/Private/MuCOE/SMutableTagListWidget.cpp
+++ b/Engine/Plugins/Mutable/Source/CustomizableObjectEditor/Private/MuCOE/SMutableTagListWidget.cpp
@@ -94,7 +94,7 @@ void SMutableTagListWidget::RefreshOptions()
 
 		// Scan all potential receivers
 		UCustomizableObject* ThisNodeObject = GetRootObject(*Node);
-		UCustomizableObject* RootObject = GetRootObject(ThisNodeObject);
+		UCustomizableObject* RootObject = GraphTraversal::GetRootObject(ThisNodeObject);
 
 		TSet<UCustomizableObject*> AllCustomizableObject;
 		GetAllObjectsInGraph(RootObject, AllCustomizableObject);
diff --git a/Engine/Plugins/Mutable/Source/CustomizableObjectEditor/Public/MuCOE/ExtensionDataCompilerInterface.h b/Engine/Plugins/Mutable/Source/CustomizableObjectEditor/Public/MuCOE/ExtensionDataCompilerInterface.h
--- a/Engine/Plugins/Mutable/Source/CustomizableObjectEditor/Public/MuCOE/ExtensionDataCompilerInterface.h
+++ b/Engine/Plugins/Mutable/Source/CustomizableObjectEditor/Public/MuCOE/ExtensionDataCompilerInterface.h
@@ -59,7 +59,7 @@ public:
 	mu::ExtensionDataPtrConst MakeAlwaysLoadedExtensionData(FInstancedStruct&& Data);
 
 	/** The Outer to use for objects owned by always-loaded Extension Data constants. */
-	UObject* GetOuterForAlwaysLoadedObjects();
+	const UObject* GetOuterForAlwaysLoadedObjects();
 
 	/**
 	* Adds a node to the Generation Context list of generated nodes. This function is meant to be called
@@ -71,5 +71,7 @@ public:
 	/** Adds a compiler log message to be displayed at the end of the compilation process */
 	void CompilerLog(const FText& InLogText, const class UCustomizableObjectNode* InNode);
 
+	void AddParticipatingObject(const UObject& Object);
+
 	FMutableGraphGenerationContext& GenerationContext;
 };
\ No newline at end of file
diff --git a/Engine/Plugins/Mutable/Source/CustomizableObjectEditor/Public/MuCOE/GraphTraversal.h b/Engine/Plugins/Mutable/Source/CustomizableObjectEditor/Public/MuCOE/GraphTraversal.h
--- a/Engine/Plugins/Mutable/Source/CustomizableObjectEditor/Public/MuCOE/GraphTraversal.h
+++ b/Engine/Plugins/Mutable/Source/CustomizableObjectEditor/Public/MuCOE/GraphTraversal.h
@@ -53,31 +53,34 @@ UEdGraphPin* FollowOutputPin(const UEdGraphPin& Pin, bool* CycleDetected = nullp
 TArray<UEdGraphPin*> ReverseFollowPinArray(const UEdGraphPin& Pin, bool bIgnoreOrphan = false, bool* bOutCycleDetected = nullptr);
 
 /** Returns the root Object Node of the Customizable Object's graph */
-UCustomizableObjectNodeObject* GetRootNode(UCustomizableObject* Object, bool& bOutMultipleBaseObjectsFound);
+UCustomizableObjectNodeObject* GetRootNode(const UCustomizableObject* Object, bool& bOutMultipleBaseObjectsFound);
 
 /** Return in ArrayNodeObject the roots nodes in each Customizable Object graph until the whole root node is found (i.e. the one with parent = nullptr)
  * return false if a cycle is found between Customizable Objects */
-bool GetParentsUntilRoot(UCustomizableObject* Object, TArray<UCustomizableObjectNodeObject*>& ArrayNodeObject, TArray<UCustomizableObject*>& ArrayCustomizableObject);
+bool GetParentsUntilRoot(const UCustomizableObject* Object, TArray<UCustomizableObjectNodeObject*>& ArrayNodeObject, TArray<const UCustomizableObject*>& ArrayCustomizableObject);
 
 /** Returns true if the Candidate is parent of the current Customizable Object */
 bool HasCandidateAsParent(UCustomizableObjectNodeObject* Node, UCustomizableObject* ParentCandidate);
 
 /** Return the full graph Customizable Object root of the given node. */
-UCustomizableObject* GetFullGraphRootObject(UCustomizableObjectNodeObject* Node, TArray<UCustomizableObject*>& VisitedObjects);
+UCustomizableObject* GetFullGraphRootObject(const UCustomizableObjectNodeObject* Node, TArray<UCustomizableObject*>& VisitedObjects);
 
 /** Return Customizable Object of the given node. */
 UCustomizableObject* GetRootObject(const UCustomizableObjectNode& Node);
 
-/** Provided a CO object it provides the root CO it is connected. In other words : it returns the root of the entire
-	 * mutable graph.
-	 * @param InObject Customizable object whose root CO we are asking for.
-	 * @return The CO that is the root of the provided Customizable Object. It can be equal to InObject if the provided
-	 * object does not have any parent.
-*/
-CUSTOMIZABLEOBJECTEDITOR_API UCustomizableObject* GetRootObject(UCustomizableObject* ChildObject);
+namespace GraphTraversal
+{
+	/** Provided a CO object it provides the root CO it is connected. In other words : it returns the root of the entire
+		 * mutable graph.
+		 * @param InObject Customizable object whose root CO we are asking for.
+		 * @return The CO that is the root of the provided Customizable Object. It can be equal to InObject if the provided
+		 * object does not have any parent.	*/
+	CUSTOMIZABLEOBJECTEDITOR_API UCustomizableObject* GetRootObject(UCustomizableObject* ChildObject);
+	CUSTOMIZABLEOBJECTEDITOR_API const UCustomizableObject* GetRootObject(const UCustomizableObject* ChildObject);
+}
 
 /** Return the full graph Customizable Object Node root of the node given as parameter */
-UCustomizableObjectNodeObject* GetFullGraphRootNodeObject(UCustomizableObjectNodeObject* Node, TArray<UCustomizableObject*>& VisitedObjects);
+UCustomizableObjectNodeObject* GetFullGraphRootNodeObject(const UCustomizableObjectNodeObject* Node, TArray<const UCustomizableObject*>& VisitedObjects);
 
 /** Given an output pin, return the output pin where the mesh is located. */
 const UEdGraphPin* FindMeshBaseSource(const UEdGraphPin& Pin, const bool bOnlyLookForStaticMesh);
diff --git a/Engine/Plugins/Mutable/Source/CustomizableObjectEditor/Public/MuCOE/Nodes/CustomizableObjectNodeTable.h b/Engine/Plugins/Mutable/Source/CustomizableObjectEditor/Public/MuCOE/Nodes/CustomizableObjectNodeTable.h
--- a/Engine/Plugins/Mutable/Source/CustomizableObjectEditor/Public/MuCOE/Nodes/CustomizableObjectNodeTable.h
+++ b/Engine/Plugins/Mutable/Source/CustomizableObjectEditor/Public/MuCOE/Nodes/CustomizableObjectNodeTable.h
@@ -357,7 +357,7 @@ public:
 
 	// Generation Mutable Source Methods
 	// We should do this in a template!
-	USkeletalMesh* GetSkeletalMeshAt(const UEdGraphPin* Pin, const UDataTable* DataTable, const FName& RowName) const;
+	FSoftObjectPtr GetSkeletalMeshAt(const UEdGraphPin* Pin, const UDataTable* DataTable, const FName& RowName) const;
 	TSoftClassPtr<UAnimInstance> GetAnimInstanceAt(const UEdGraphPin* Pin, const UDataTable* DataTable, const FName& RowName) const;
 	
 	// Returns the image mode of the column
diff --git a/Engine/Plugins/Mutable/Source/MutableValidation/Private/MuV/AssetValidator_CustomizableObjects.cpp b/Engine/Plugins/Mutable/Source/MutableValidation/Private/MuV/AssetValidator_CustomizableObjects.cpp
--- a/Engine/Plugins/Mutable/Source/MutableValidation/Private/MuV/AssetValidator_CustomizableObjects.cpp
+++ b/Engine/Plugins/Mutable/Source/MutableValidation/Private/MuV/AssetValidator_CustomizableObjects.cpp
@@ -100,7 +100,7 @@ EDataValidationResult UAssetValidator_CustomizableObjects::IsCustomizableObjectV
 	}
 	
 	// Find out which is the root for this CO (it may be itself but that is OK)
-	UCustomizableObject* RootObject = GetRootObject(InCustomizableObject);
+	UCustomizableObject* RootObject = GraphTraversal::GetRootObject(InCustomizableObject);
 	check (RootObject);
 	
 	// Check that the object to be compiled has not already been compiled
