[Mutable] Fix incorrect AssetUserDataIndex check.

#rnx

diff --git a/Engine/Plugins/Mutable/Source/CustomizableObjectEditor/Private/MuCOE/GenerateMutableSource/GenerateMutableSource.cpp b/Engine/Plugins/Mutable/Source/CustomizableObjectEditor/Private/MuCOE/GenerateMutableSource/GenerateMutableSource.cpp
--- a/Engine/Plugins/Mutable/Source/CustomizableObjectEditor/Private/MuCOE/GenerateMutableSource/GenerateMutableSource.cpp
+++ b/Engine/Plugins/Mutable/Source/CustomizableObjectEditor/Private/MuCOE/GenerateMutableSource/GenerateMutableSource.cpp
@@ -1436,7 +1436,7 @@ void PopulateReferenceSkeletalMeshesData(FMutableGraphGenerationContext& Generat
 				{
 					FMutableRefAssetUserData& MutAssetUserData = Data.AssetUserData.AddDefaulted_GetRef();
 					MutAssetUserData.AssetUserDataIndex = GenerationContext.AddAssetUserDataToStreamedResources(AssetUserData);
-					if (MutAssetUserData.AssetUserData != INDEX_NONE)
+					if (MutAssetUserData.AssetUserDataIndex != INDEX_NONE)
 					{
 						MutAssetUserData.AssetUserData = GenerationContext.StreamedResourceData[MutAssetUserData.AssetUserDataIndex].Value;
 						check(MutAssetUserData.AssetUserData->Data.Type == ECOResourceDataType::AssetUserData);
