[Mutable] Fixed dependency between Object and Graph PostLoads.

#rnx

diff --git a/Engine/Plugins/Mutable/Source/CustomizableObject/Private/MuCO/CustomizableObject.cpp b/Engine/Plugins/Mutable/Source/CustomizableObject/Private/MuCO/CustomizableObject.cpp
--- a/Engine/Plugins/Mutable/Source/CustomizableObject/Private/MuCO/CustomizableObject.cpp
+++ b/Engine/Plugins/Mutable/Source/CustomizableObject/Private/MuCO/CustomizableObject.cpp
@@ -314,13 +314,6 @@ bool UCustomizableObjectPrivate::TryLoadCompiledCookDataForPlatform(const ITarge
 void UCustomizableObject::PostLoad()
 {
 	Super::PostLoad();
-
-#if	WITH_EDITORONLY_DATA
-	if (Source)
-	{
-		Source->ConditionalPostLoad();
-	}
-#endif
 	
 	const int32 CustomizableObjectCustomVersion = GetLinkerCustomVersion(FCustomizableObjectCustomVersion::GUID);
 
@@ -388,7 +381,6 @@ void UCustomizableObject::PostLoad()
 			GetPrivate()->LoadCompiledDataFromDisk();
 		}
 	}
-
 #endif
 }
 
diff --git a/Engine/Plugins/Mutable/Source/CustomizableObjectEditor/Private/MuCOE/CustomizableObjectGraph.cpp b/Engine/Plugins/Mutable/Source/CustomizableObjectEditor/Private/MuCOE/CustomizableObjectGraph.cpp
--- a/Engine/Plugins/Mutable/Source/CustomizableObjectEditor/Private/MuCOE/CustomizableObjectGraph.cpp
+++ b/Engine/Plugins/Mutable/Source/CustomizableObjectEditor/Private/MuCOE/CustomizableObjectGraph.cpp
@@ -22,6 +22,9 @@ void UCustomizableObjectGraph::PostLoad()
 
 	// TODO UE-222779 If compatibility code has to be executed, loaded the full hierarchy and in sync the BackwardsCompatibleFixup 
 	
+	// Make sure the Object has executed its compatibility code.
+	GetOuter()->ConditionalPostLoad();
+
 	// Make sure all nodes have finished loading.
 	for (UEdGraphNode* Node : Nodes)
 	{
