[Mutable] Renamed ContentVersion/VersionBrige to ReleaseVersion.

#rnx

diff --git a/Engine/Plugins/Mutable/Source/CustomizableObject/Internal/MuCO/CustomizableObjectPrivate.h b/Engine/Plugins/Mutable/Source/CustomizableObject/Internal/MuCO/CustomizableObjectPrivate.h
--- a/Engine/Plugins/Mutable/Source/CustomizableObject/Internal/MuCO/CustomizableObjectPrivate.h
+++ b/Engine/Plugins/Mutable/Source/CustomizableObject/Internal/MuCO/CustomizableObjectPrivate.h
@@ -1097,9 +1097,8 @@ struct FModelResources
 	UPROPERTY()
 	TArray<FName> ComponentNames;
 
-	/** Version Bridge. */
 	UPROPERTY()
-	FString CompiledVersionBridge;
+	FString ReleaseVersion;
 	
 #if WITH_EDITORONLY_DATA
 	void CUSTOMIZABLEOBJECT_API Serialize(FObjectAndNameAsStringProxyArchive& Ar, bool bIsCooking);
@@ -1298,9 +1297,7 @@ public:
 	
 #if WITH_EDITOR
 	/** See ICustomizableObjectEditorModule::IsCompilationOutOfDate. */
-	bool IsCompilationOutOfDate(bool bSkipIndirectReferences, TArray<FName>& OutOfDatePackages, TArray<FName>& AddedPackages, TArray<FName>& RemovedPackages, bool& bVersionDiff) const;
-
-	void OnParticipatingObjectDirty(UPackage* Package, bool);
+	bool IsCompilationOutOfDate(bool bSkipIndirectReferences, TArray<FName>& OutOfDatePackages, TArray<FName>& AddedPackages, TArray<FName>& RemovedPackages, bool& bReleaseVersionDiff) const;
 #endif
 
 	TArray<FString>& GetCustomizableObjectClassTags();
diff --git a/Engine/Plugins/Mutable/Source/CustomizableObject/Internal/MuCO/ICustomizableObjectEditorModule.h b/Engine/Plugins/Mutable/Source/CustomizableObject/Internal/MuCO/ICustomizableObjectEditorModule.h
--- a/Engine/Plugins/Mutable/Source/CustomizableObject/Internal/MuCO/ICustomizableObjectEditorModule.h
+++ b/Engine/Plugins/Mutable/Source/CustomizableObject/Internal/MuCO/ICustomizableObjectEditorModule.h
@@ -54,16 +54,16 @@ public:
 	  * @param OutOfDatePackages list of out of date packages.
 	  * @param AddedPackages list of added packages since the last compilation.
 	  * @param RemovedPackages list of removed packages since the last compilation.
-	  * @param bVersionDiff true if the Version Bridge has changed since the last compilation.
+	  * @param bReleaseVersionDiff true if the Release Version has changed since the last compilation.
 	  * @return true if the compilation is out of date. */
-	virtual bool IsCompilationOutOfDate(const UCustomizableObject& Object, bool bSkipIndirectReferences, TArray<FName>& OutOfDatePackages, TArray<FName>& AddedPackages, TArray<FName>& RemovedPackages, bool& bVersionDiff) const = 0;
+	virtual bool IsCompilationOutOfDate(const UCustomizableObject& Object, bool bSkipIndirectReferences, TArray<FName>& OutOfDatePackages, TArray<FName>& AddedPackages, TArray<FName>& RemovedPackages, bool& bReleaseVersionDiff) const = 0;
 
 	/** See GraphTraversal::IsRootObject(...) */
 	virtual bool IsRootObject(const UCustomizableObject& Object) const = 0;
 
-	/** Get the current VersionBridge's version for Object. 
+	/** Get the current Release Version for the given Object. 
 	  * @return Current version as string. */
-	virtual FString GetCurrentContentVersionForObject(const UCustomizableObject& Object) const = 0;
+	virtual FString GetCurrentReleaseVersionForObject(const UCustomizableObject& Object) const = 0;
 
 	/** See GraphTraversal::GetRootObject(...) */
 	virtual UCustomizableObject* GetRootObject(UCustomizableObject* ChildObject) const = 0;
diff --git a/Engine/Plugins/Mutable/Source/CustomizableObject/Private/MuCO/CustomizableObject.cpp b/Engine/Plugins/Mutable/Source/CustomizableObject/Private/MuCO/CustomizableObject.cpp
--- a/Engine/Plugins/Mutable/Source/CustomizableObject/Private/MuCO/CustomizableObject.cpp
+++ b/Engine/Plugins/Mutable/Source/CustomizableObject/Private/MuCO/CustomizableObject.cpp
@@ -742,7 +742,7 @@ void FModelResources::Serialize(FObjectAndNameAsStringProxyArchive& MemoryWriter
 	MemoryWriter << FirstLODAvailable;
 
 	MemoryWriter << ComponentNames;
-	MemoryWriter << CompiledVersionBridge;
+	MemoryWriter << ReleaseVersion;
 
 	// Editor Only data
 	if (!bIsCooking)
@@ -894,7 +894,7 @@ bool FModelResources::Unserialize(FObjectAndNameAsStringProxyArchive& MemoryRead
 	MemoryReader << FirstLODAvailable;
 
 	MemoryReader << ComponentNames;
-	MemoryReader << CompiledVersionBridge;
+	MemoryReader << ReleaseVersion;
 
 	// Editor Only data
 	if (!bIsCooking)
@@ -974,8 +974,8 @@ void UCustomizableObjectPrivate::LoadCompiledDataFromDisk()
 					TArray<FName> OutOfDatePackages;
 					TArray<FName> AddedPackages;
 					TArray<FName> RemovedPackages;
-					bool bVersionDiff;
-					const bool bOutOfDate = IsCompilationOutOfDate(false, OutOfDatePackages, AddedPackages, RemovedPackages, bVersionDiff);
+					bool bReleaseVersion;
+					const bool bOutOfDate = IsCompilationOutOfDate(false, OutOfDatePackages, AddedPackages, RemovedPackages, bReleaseVersion);
 					if (!bOutOfDate)
 					{
 						LoadModelStreamableBulk(MemoryReader, /* bIsCooking */false);
@@ -984,7 +984,7 @@ void UCustomizableObjectPrivate::LoadCompiledDataFromDisk()
 					else
 					{
 						UE_LOG(LogMutable, Display, TEXT("Invalidating compiled data due to changes in %s."), *OutOfDatePackages[0].ToString());
-						PrintParticipatingPackagesDiff(OutOfDatePackages, AddedPackages, RemovedPackages, bVersionDiff);
+						PrintParticipatingPackagesDiff(OutOfDatePackages, AddedPackages, RemovedPackages, bReleaseVersion);
 					}
 				}
 			}
@@ -2300,11 +2300,11 @@ const FModelResources& UCustomizableObjectPrivate::GetModelResources(bool bIsCoo
 
 
 #if WITH_EDITOR
-bool UCustomizableObjectPrivate::IsCompilationOutOfDate(bool bSkipIndirectReferences, TArray<FName>& OutOfDatePackages, TArray<FName>& AddedPackages, TArray<FName>& RemovedPackages, bool& bVersionDiff) const
+bool UCustomizableObjectPrivate::IsCompilationOutOfDate(bool bSkipIndirectReferences, TArray<FName>& OutOfDatePackages, TArray<FName>& AddedPackages, TArray<FName>& RemovedPackages, bool& bReleaseVersionDiff) const
 {
 	if (const ICustomizableObjectEditorModule* Module = ICustomizableObjectEditorModule::Get())
 	{
-		return Module->IsCompilationOutOfDate(*GetPublic(), bSkipIndirectReferences, OutOfDatePackages, AddedPackages, RemovedPackages, bVersionDiff);
+		return Module->IsCompilationOutOfDate(*GetPublic(), bSkipIndirectReferences, OutOfDatePackages, AddedPackages, RemovedPackages, bReleaseVersionDiff);
 	}
 
 	return false;		
@@ -2856,10 +2856,10 @@ TArray<uint8> UCustomizableObjectPrivate::BuildDerivedDataKey(FCompilationOption
 		SerializeCompilationOptionsForDDC(Ar, Options);
 	}
 
-	// Version Bridge
+	// Release Version
 	if (const ICustomizableObjectEditorModule* Module = ICustomizableObjectEditorModule::Get())
 	{
-		FString Version = Module->GetCurrentContentVersionForObject(CustomizableObject);
+		FString Version = Module->GetCurrentReleaseVersionForObject(CustomizableObject);
 		Ar << Version;
 	}
 
diff --git a/Engine/Plugins/Mutable/Source/CustomizableObject/Private/MuCO/CustomizableObjectModule.cpp b/Engine/Plugins/Mutable/Source/CustomizableObject/Private/MuCO/CustomizableObjectModule.cpp
--- a/Engine/Plugins/Mutable/Source/CustomizableObject/Private/MuCO/CustomizableObjectModule.cpp
+++ b/Engine/Plugins/Mutable/Source/CustomizableObject/Private/MuCO/CustomizableObjectModule.cpp
@@ -334,13 +334,13 @@ void FCustomizableObjectModule::ToggleDGGUI(const TArray<FString>& Arguments)
 }
 
 
-void PrintParticipatingPackagesDiff(const TArray<FName>& OutOfDatePackages, const TArray<FName>& AddedPackages, const TArray<FName>& RemovedPackages, bool bVersionDiff)
+void PrintParticipatingPackagesDiff(const TArray<FName>& OutOfDatePackages, const TArray<FName>& AddedPackages, const TArray<FName>& RemovedPackages, bool bReleaseVersion)
 {
 	constexpr int32 MaxLogLines = 10;
 
-	if (bVersionDiff)
+	if (bReleaseVersion)
 	{
-		UE_LOG(LogMutable, Display, TEXT("Version Bridge changed."));
+		UE_LOG(LogMutable, Display, TEXT("Release Version changed."));
 	}
 	
 	if (OutOfDatePackages.Num())
diff --git a/Engine/Plugins/Mutable/Source/CustomizableObject/Public/MuCO/ICustomizableObjectModule.h b/Engine/Plugins/Mutable/Source/CustomizableObject/Public/MuCO/ICustomizableObjectModule.h
--- a/Engine/Plugins/Mutable/Source/CustomizableObject/Public/MuCO/ICustomizableObjectModule.h
+++ b/Engine/Plugins/Mutable/Source/CustomizableObject/Public/MuCO/ICustomizableObjectModule.h
@@ -85,4 +85,4 @@ CUSTOMIZABLEOBJECT_API
 UCustomizableObjectInstanceUsage* GetPlayerCustomizableObjectInstanceUsage(const int32 SlotID, const UWorld* CurrentWorld, const int32 PlayerIndex = 0);
 
 
-CUSTOMIZABLEOBJECT_API void PrintParticipatingPackagesDiff(const TArray<FName>& OutOfDatePackages, const TArray<FName>& AddedPackages, const TArray<FName>& RemovedPackages, bool bVersionDiff);
\ No newline at end of file
+CUSTOMIZABLEOBJECT_API void PrintParticipatingPackagesDiff(const TArray<FName>& OutOfDatePackages, const TArray<FName>& AddedPackages, const TArray<FName>& RemovedPackages, bool bReleaseVersion);
\ No newline at end of file
diff --git a/Engine/Plugins/Mutable/Source/CustomizableObjectEditor/Private/MuCOE/CustomizableObjectCompiler.cpp b/Engine/Plugins/Mutable/Source/CustomizableObjectEditor/Private/MuCOE/CustomizableObjectCompiler.cpp
--- a/Engine/Plugins/Mutable/Source/CustomizableObjectEditor/Private/MuCOE/CustomizableObjectCompiler.cpp
+++ b/Engine/Plugins/Mutable/Source/CustomizableObjectEditor/Private/MuCOE/CustomizableObjectCompiler.cpp
@@ -1034,7 +1034,7 @@ void FCustomizableObjectCompiler::CompileInternal(bool bAsync)
 
 		if (ICustomizableObjectVersionBridgeInterface* VersionBridge = Cast<ICustomizableObjectVersionBridgeInterface>(GraphTraversal::GetRootObject(CurrentObject)->VersionBridge))
 		{
-			ModelResources.CompiledVersionBridge = VersionBridge->GetCurrentVersionAsString();
+			ModelResources.ReleaseVersion = VersionBridge->GetCurrentVersionAsString();
 		}
 		
 		ModelResources.NumLODs = GenerationContext.NumLODsInRoot;
diff --git a/Engine/Plugins/Mutable/Source/CustomizableObjectEditor/Private/MuCOE/CustomizableObjectEditorModule.cpp b/Engine/Plugins/Mutable/Source/CustomizableObjectEditor/Private/MuCOE/CustomizableObjectEditorModule.cpp
--- a/Engine/Plugins/Mutable/Source/CustomizableObjectEditor/Private/MuCOE/CustomizableObjectEditorModule.cpp
+++ b/Engine/Plugins/Mutable/Source/CustomizableObjectEditor/Private/MuCOE/CustomizableObjectEditorModule.cpp
@@ -167,8 +167,8 @@ void ShowOnScreenCompileWarnings()
 		TArray<FName> OutOfDatePackages;
 		TArray<FName> AddedPackages;
 		TArray<FName> RemovedPackages;
-		bool bVersionDiff;
-		if (Object->GetPrivate()->IsCompilationOutOfDate(true, OutOfDatePackages, AddedPackages, RemovedPackages, bVersionDiff))
+		bool bReleaseVersion;
+		if (Object->GetPrivate()->IsCompilationOutOfDate(true, OutOfDatePackages, AddedPackages, RemovedPackages, bReleaseVersion))
 		{
 			FString Msg = FString::Printf(TEXT("Customizable Object [%s] compilation out of date. See the Output Log for more information."), *Object->GetName());
 			GEngine->AddOnScreenDebugMessage(KeyCompiledOutOfDate, ShowOnScreenCompileWarningsTickerTime * 2.0f, FColor::Yellow, Msg);
@@ -177,7 +177,7 @@ void ShowOnScreenCompileWarnings()
 			{
 				UE_LOG(LogMutable, Display, TEXT("Customizable Object [%s] compilation out of date. Changes since last compilation:"), *Object->GetName());
 
-				PrintParticipatingPackagesDiff(OutOfDatePackages, AddedPackages, RemovedPackages, bVersionDiff);
+				PrintParticipatingPackagesDiff(OutOfDatePackages, AddedPackages, RemovedPackages, bReleaseVersion);
 			}
 		}
 		else
@@ -557,7 +557,7 @@ void GetReferencingPackages(const UCustomizableObject& Object, TArray<FAssetData
 }
 
 
-bool FCustomizableObjectEditorModule::IsCompilationOutOfDate(const UCustomizableObject& Object, bool bSkipIndirectReferences, TArray<FName>& OutOfDatePackages, TArray<FName>& AddedPackages, TArray<FName>& RemovedPackages, bool& bVersionDiff) const
+bool FCustomizableObjectEditorModule::IsCompilationOutOfDate(const UCustomizableObject& Object, bool bSkipIndirectReferences, TArray<FName>& OutOfDatePackages, TArray<FName>& AddedPackages, TArray<FName>& RemovedPackages, bool& bReleaseVersion) const
 {
 	MUTABLE_CPUPROFILER_SCOPE(FCustomizableObjectEditorModule::IsCompilationOutOfDate)
 	
@@ -635,13 +635,13 @@ bool FCustomizableObjectEditorModule::IsCompilationOutOfDate(const UCustomizable
 		}
 	}
 
-	bVersionDiff = false;
+	bReleaseVersion = false;
 	if (ICustomizableObjectVersionBridgeInterface* VersionBridge = Cast<ICustomizableObjectVersionBridgeInterface>(Object.VersionBridge))
 	{
-		bVersionDiff = Object.GetPrivate()->GetModelResources().CompiledVersionBridge != VersionBridge->GetCurrentVersionAsString();
+		bReleaseVersion = Object.GetPrivate()->GetModelResources().ReleaseVersion != VersionBridge->GetCurrentVersionAsString();
 	}
 
-	return bVersionDiff || !OutOfDatePackages.IsEmpty() || !AddedPackages.IsEmpty() || !RemovedPackages.IsEmpty();
+	return bReleaseVersion || !OutOfDatePackages.IsEmpty() || !AddedPackages.IsEmpty() || !RemovedPackages.IsEmpty();
 }
 
 
@@ -650,7 +650,7 @@ bool FCustomizableObjectEditorModule::IsRootObject(const UCustomizableObject& Ob
 	return GraphTraversal::IsRootObject(Object);
 }
 
-FString FCustomizableObjectEditorModule::GetCurrentContentVersionForObject(const UCustomizableObject& Object) const
+FString FCustomizableObjectEditorModule::GetCurrentReleaseVersionForObject(const UCustomizableObject& Object) const
 {
 	if (Object.VersionBridge && Object.VersionBridge->GetClass()->ImplementsInterface(UCustomizableObjectVersionBridgeInterface::StaticClass()))
 	{
diff --git a/Engine/Plugins/Mutable/Source/CustomizableObjectEditor/Private/MuCOE/CustomizableObjectEditorModule.h b/Engine/Plugins/Mutable/Source/CustomizableObjectEditor/Private/MuCOE/CustomizableObjectEditorModule.h
--- a/Engine/Plugins/Mutable/Source/CustomizableObjectEditor/Private/MuCOE/CustomizableObjectEditorModule.h
+++ b/Engine/Plugins/Mutable/Source/CustomizableObjectEditor/Private/MuCOE/CustomizableObjectEditorModule.h
@@ -40,7 +40,7 @@ public:
 	virtual FCustomizableObjectEditorLogger& GetLogger() override;
 	virtual bool IsCompilationOutOfDate(const UCustomizableObject& Object, bool bSkipIndirectReferences, TArray<FName>& OutOfDatePackages, TArray<FName>& AddedPackages, TArray<FName>& RemovedPackages, bool& bVersionDiff) const override;
 	virtual bool IsRootObject(const UCustomizableObject& Object) const override;
-	virtual FString GetCurrentContentVersionForObject(const UCustomizableObject& Object) const override;
+	virtual FString GetCurrentReleaseVersionForObject(const UCustomizableObject& Object) const override;
 	virtual UCustomizableObject* GetRootObject(UCustomizableObject* ChildObject) const override;
 	virtual const UCustomizableObject* GetRootObject(const UCustomizableObject* ChildObject) const override;
 	virtual void BakeCustomizableObjectInstance(UCustomizableObjectInstance* InTargetInstance, const FBakingConfiguration& InBakingConfig) override;
