[Mutable] Fix ParticipatingObjects check when exporting a model.

#rnx

diff --git a/Engine/Plugins/Mutable/Source/CustomizableObjectEditor/Private/MuCOE/CustomizableObjectCompiler.cpp b/Engine/Plugins/Mutable/Source/CustomizableObjectEditor/Private/MuCOE/CustomizableObjectCompiler.cpp
--- a/Engine/Plugins/Mutable/Source/CustomizableObjectEditor/Private/MuCOE/CustomizableObjectCompiler.cpp
+++ b/Engine/Plugins/Mutable/Source/CustomizableObjectEditor/Private/MuCOE/CustomizableObjectCompiler.cpp
@@ -1270,6 +1270,8 @@ mu::Ptr<mu::Node> FCustomizableObjectCompiler::Export(UCustomizableObject* Objec
 
 	FMutableGraphGenerationContext GenerationContext(Object, this, CompilerOptions);
 
+	GenerationContext.bSkipParticipatingObjectsPass = true;
+	
 	// Generate the mutable node expression
 	mu::NodeObjectPtr MutableRoot = GenerateMutableRoot(Object, GenerationContext);
 	if (!MutableRoot)
diff --git a/Engine/Plugins/Mutable/Source/CustomizableObjectEditor/Private/MuCOE/GenerateMutableSource/GenerateMutableSource.cpp b/Engine/Plugins/Mutable/Source/CustomizableObjectEditor/Private/MuCOE/GenerateMutableSource/GenerateMutableSource.cpp
--- a/Engine/Plugins/Mutable/Source/CustomizableObjectEditor/Private/MuCOE/GenerateMutableSource/GenerateMutableSource.cpp
+++ b/Engine/Plugins/Mutable/Source/CustomizableObjectEditor/Private/MuCOE/GenerateMutableSource/GenerateMutableSource.cpp
@@ -471,8 +471,11 @@ void FMutableGraphGenerationContext::AddParticipatingObjectChecked(const FName&
 	{
 		checkCode
 		(
-			FGuid* Result = ParticipatingObjects.Find(PackageName);
-			check(Result && *Result == PackageGuid); // If this check is hit it means that this Participating Object is not being discovered in the Participating Objects pass.
+			if (!bSkipParticipatingObjectsPass)
+			{
+				FGuid* Result = ParticipatingObjects.Find(PackageName);
+				check(Result && *Result == PackageGuid); // If this check is hit it means that this Participating Object is not being discovered in the Participating Objects pass.
+			}
 		)
 	}
 }
diff --git a/Engine/Plugins/Mutable/Source/CustomizableObjectEditor/Private/MuCOE/GenerateMutableSource/GenerateMutableSource.h b/Engine/Plugins/Mutable/Source/CustomizableObjectEditor/Private/MuCOE/GenerateMutableSource/GenerateMutableSource.h
--- a/Engine/Plugins/Mutable/Source/CustomizableObjectEditor/Private/MuCOE/GenerateMutableSource/GenerateMutableSource.h
+++ b/Engine/Plugins/Mutable/Source/CustomizableObjectEditor/Private/MuCOE/GenerateMutableSource/GenerateMutableSource.h
@@ -859,8 +859,12 @@ public:
 
 	bool bPartialCompilation = false;
 
+	/** true if performing the Participating Objects pass. */
 	bool bParticipatingObjectsPass = false;
 
+	/** true if the Participating Objects have been skipped. */
+	bool bSkipParticipatingObjectsPass = false;
+
 	/** Load any Soft Object/Class Pointers. */
 	bool bLoadObjects = true;
 	
