[mutable] Fix crashes in code generator that assumed image swizzle always had 4 channels.
#rnx

diff --git a/Engine/Plugins/Mutable/Source/MutableTools/Private/MuT/CodeGenerator.cpp b/Engine/Plugins/Mutable/Source/MutableTools/Private/MuT/CodeGenerator.cpp
--- a/Engine/Plugins/Mutable/Source/MutableTools/Private/MuT/CodeGenerator.cpp
+++ b/Engine/Plugins/Mutable/Source/MutableTools/Private/MuT/CodeGenerator.cpp
@@ -1092,15 +1092,26 @@ namespace mu
 					else if (pImageNode->GetType() == NodeImageSwizzle::GetStaticType())
 					{
 						NodeImageSwizzle* ts = static_cast<NodeImageSwizzle*>(pImageNode.get());
-						NodeImage* Source = ts->GetSource(0).get();
-						if (!swizzleNode && Source ==ts->GetSource(1) && Source==ts->GetSource(2) && Source==ts->GetSource(3))
-						{
-							swizzleNode = ts;
-							pImageNode = Source;
-						}
-						else
+
+						if (!ts->GetPrivate()->m_sources.IsEmpty())
 						{
-							bFound = true;
+							NodeImage* Source = ts->GetSource(0).get();
+							
+							bool bAllSourcesAreTheSame = true;
+							for (int32 SourceIndex=1; SourceIndex<ts->GetPrivate()->m_sources.Num(); ++SourceIndex)
+							{
+								bAllSourcesAreTheSame = bAllSourcesAreTheSame && (Source == ts->GetSource(SourceIndex));
+							}
+
+							if (!swizzleNode && bAllSourcesAreTheSame)
+							{
+								swizzleNode = ts;
+								pImageNode = Source;
+							}
+							else
+							{
+								bFound = true;
+							}
 						}
 					}
 					else
@@ -1323,14 +1334,12 @@ namespace mu
 							{
 								Ptr<ASTOpImageSwizzle> fop = new ASTOpImageSwizzle();
 								fop->Format = swizzleNode->GetPrivate()->m_format;
-								fop->Sources[0] = imageAd;
-								fop->Sources[1] = imageAd;
-								fop->Sources[2] = imageAd;
-								fop->Sources[3] = imageAd;
-								fop->SourceChannels[0] = swizzleNode->GetPrivate()->m_sourceChannels[0];
-								fop->SourceChannels[1] = swizzleNode->GetPrivate()->m_sourceChannels[1];
-								fop->SourceChannels[2] = swizzleNode->GetPrivate()->m_sourceChannels[2];
-								fop->SourceChannels[3] = swizzleNode->GetPrivate()->m_sourceChannels[3];
+
+								for (int32 ChannelIndex = 0; ChannelIndex < swizzleNode->GetPrivate()->m_sourceChannels.Num(); ++ChannelIndex)
+								{
+									fop->Sources[ChannelIndex] = imageAd;
+									fop->SourceChannels[ChannelIndex] = swizzleNode->GetPrivate()->m_sourceChannels[ChannelIndex];
+								}
 								check(fop->Format != EImageFormat::IF_NONE);
 								imageAd = fop;
 							}
diff --git a/Engine/Plugins/Mutable/Source/MutableTools/Private/MuT/NodeImageSwizzle.cpp b/Engine/Plugins/Mutable/Source/MutableTools/Private/MuT/NodeImageSwizzle.cpp
--- a/Engine/Plugins/Mutable/Source/MutableTools/Private/MuT/NodeImageSwizzle.cpp
+++ b/Engine/Plugins/Mutable/Source/MutableTools/Private/MuT/NodeImageSwizzle.cpp
@@ -39,41 +39,64 @@ namespace mu
 	{
 		m_pD->m_format = format;
 
-		int channelCount = GetImageFormatData( format ).Channels;
+		int32 channelCount = GetImageFormatData( format ).Channels;
 		m_pD->m_sources.SetNum( channelCount );
 		m_pD->m_sourceChannels.SetNum( channelCount );
 	}
 
 
 	//---------------------------------------------------------------------------------------------
-	NodeImagePtr NodeImageSwizzle::GetSource( int t ) const
+	NodeImagePtr NodeImageSwizzle::GetSource( int32 t ) const
 	{
-		check( t>=0 && t<(int)m_pD->m_sources.Num() );
-		return m_pD->m_sources[t].get();
+		if (m_pD->m_sources.IsValidIndex(t))
+		{
+			return m_pD->m_sources[t].get();
+		}
+		ensure(false);
+		return nullptr;
 	}
 
 
 	//---------------------------------------------------------------------------------------------
-	void NodeImageSwizzle::SetSource( int t, NodeImagePtr pNode )
+	void NodeImageSwizzle::SetSource( int32 t, NodeImagePtr pNode )
 	{
-		check( t>=0 && t<(int)m_pD->m_sources.Num() );
-		m_pD->m_sources[t] = pNode;
+		if (m_pD->m_sources.IsValidIndex(t))
+		{
+			m_pD->m_sources[t] = pNode;
+		}
+		else
+		{
+			ensure(false);
+		}
 	}
 
 
 	//---------------------------------------------------------------------------------------------
-	int NodeImageSwizzle::GetSourceChannel( int t ) const
+	int32 NodeImageSwizzle::GetSourceChannel( int32 t ) const
 	{
-		check( t>=0 && t<(int)m_pD->m_sourceChannels.Num() );
-		return m_pD->m_sourceChannels[t];
+		if (m_pD->m_sourceChannels.IsValidIndex(t))
+		{
+			return m_pD->m_sourceChannels[t];
+		}
+		else
+		{
+			ensure(false);
+			return 0;
+		}
 	}
 
 
 	//---------------------------------------------------------------------------------------------
-	void NodeImageSwizzle::SetSourceChannel( int t, int channel )
+	void NodeImageSwizzle::SetSourceChannel( int32 OutputChannel, int32 SourceChannel )
 	{
-		check( t>=0 && t<(int)m_pD->m_sourceChannels.Num() );
-		m_pD->m_sourceChannels[t] = channel;
+		if (m_pD->m_sourceChannels.IsValidIndex(OutputChannel))
+		{
+			m_pD->m_sourceChannels[OutputChannel] = SourceChannel;
+		}
+		else
+		{
+			ensure(false);
+		}
 	}
 
 
diff --git a/Engine/Plugins/Mutable/Source/MutableTools/Public/MuT/NodeImageSwizzle.h b/Engine/Plugins/Mutable/Source/MutableTools/Public/MuT/NodeImageSwizzle.h
--- a/Engine/Plugins/Mutable/Source/MutableTools/Public/MuT/NodeImageSwizzle.h
+++ b/Engine/Plugins/Mutable/Source/MutableTools/Public/MuT/NodeImageSwizzle.h
@@ -42,11 +42,11 @@ namespace mu
 		void SetFormat(EImageFormat);
 
 		//!
-		NodeImagePtr GetSource( int ) const;
-		int GetSourceChannel( int ) const;
+		NodeImagePtr GetSource( int32 ) const;
+		int GetSourceChannel( int32 ) const;
 
-		void SetSource( int, NodeImagePtr );
-		void SetSourceChannel( int, int );
+		void SetSource( int32, NodeImagePtr );
+		void SetSourceChannel( int32 OutputChannel, int32 SourceChannel );
 
 		//-----------------------------------------------------------------------------------------
 		// Interface pattern
