[Mutable] Fix crash when logging CO out of date data invalidation.

#rnx

diff --git a/Engine/Plugins/Mutable/Source/CustomizableObject/Private/MuCO/CustomizableObject.cpp b/Engine/Plugins/Mutable/Source/CustomizableObject/Private/MuCO/CustomizableObject.cpp
--- a/Engine/Plugins/Mutable/Source/CustomizableObject/Private/MuCO/CustomizableObject.cpp
+++ b/Engine/Plugins/Mutable/Source/CustomizableObject/Private/MuCO/CustomizableObject.cpp
@@ -983,7 +983,11 @@ void UCustomizableObjectPrivate::LoadCompiledDataFromDisk()
 					}
 					else
 					{
-						UE_LOG(LogMutable, Display, TEXT("Invalidating compiled data due to changes in %s."), *OutOfDatePackages[0].ToString());
+						if (OutOfDatePackages.Num())
+						{
+							UE_LOG(LogMutable, Display, TEXT("Invalidating compiled data due to changes in %s."), *OutOfDatePackages[0].ToString());
+						}
+						
 						PrintParticipatingPackagesDiff(OutOfDatePackages, AddedPackages, RemovedPackages, bReleaseVersion);
 					}
 				}
