#Mutable Do not include all data tables when asking for the parent tables of a node that is not set to use the asset registry, or if it is set to use the asset registry without a structure or paths set

diff --git a/Engine/Plugins/Mutable/Source/CustomizableObjectEditor/Private/MuCOE/Nodes/CustomizableObjectNodeTable.cpp b/Engine/Plugins/Mutable/Source/CustomizableObjectEditor/Private/MuCOE/Nodes/CustomizableObjectNodeTable.cpp
--- a/Engine/Plugins/Mutable/Source/CustomizableObjectEditor/Private/MuCOE/Nodes/CustomizableObjectNodeTable.cpp
+++ b/Engine/Plugins/Mutable/Source/CustomizableObjectEditor/Private/MuCOE/Nodes/CustomizableObjectNodeTable.cpp
@@ -1348,29 +1348,37 @@ FGuid UCustomizableObjectNodeTable::GetColumnIdByName(const FName& ColumnName) c
 
 TArray<FAssetData> UCustomizableObjectNodeTable::GetParentTables() const
 {
-	FAssetRegistryModule& AssetRegistryModule = FModuleManager::LoadModuleChecked<FAssetRegistryModule>("AssetRegistry");
-	IAssetRegistry& AssetRegistry = AssetRegistryModule.GetRegistry();
-	
-	TArray<FName> ReferencedTables;
-	AssetRegistry.Get()->GetReferencers(Structure.GetPackage().GetFName(), ReferencedTables, UE::AssetRegistry::EDependencyCategory::Package, UE::AssetRegistry::EDependencyQuery::NoRequirements);
-	
-	FARFilter Filter;
-	Filter.ClassPaths.Add(FTopLevelAssetPath(UDataTable::StaticClass()));
-
-	for (const FName& ReferencedTable : ReferencedTables)
+	TArray<FAssetData> DataTableAssets;
+	if (TableDataGatheringMode == ETableDataGatheringSource::ETDGM_AssetRegistry)
 	{
-		Filter.PackageNames.Add(ReferencedTable);
-	}
+		FAssetRegistryModule& AssetRegistryModule = FModuleManager::LoadModuleChecked<FAssetRegistryModule>("AssetRegistry");
+		IAssetRegistry& AssetRegistry = AssetRegistryModule.GetRegistry();
 
-	for (int32 PathIndex = 0; PathIndex < FilterPaths.Num(); ++PathIndex)
-	{
-		Filter.PackagePaths.Add(FilterPaths[PathIndex]);
-	}
+		FARFilter Filter;
+		Filter.ClassPaths.Add(FTopLevelAssetPath(UDataTable::StaticClass()));
+		Filter.bRecursivePaths = true;
 
-	Filter.bRecursivePaths = true;
+		if (Structure)
+		{
+			TArray<FName> ReferencedTables;
+			AssetRegistry.Get()->GetReferencers(Structure.GetPackage().GetFName(), ReferencedTables, UE::AssetRegistry::EDependencyCategory::Package, UE::AssetRegistry::EDependencyQuery::NoRequirements);
 
-	TArray<FAssetData> DataTableAssets;
-	AssetRegistry.Get()->GetAssets(Filter, DataTableAssets);
+			for (const FName& ReferencedTable : ReferencedTables)
+			{
+				Filter.PackageNames.Add(ReferencedTable);
+			}
+		}
+
+		for (int32 PathIndex = 0; PathIndex < FilterPaths.Num(); ++PathIndex)
+		{
+			Filter.PackagePaths.Add(FilterPaths[PathIndex]);
+		}
+
+		if (!Filter.IsEmpty())
+		{
+			AssetRegistry.Get()->GetAssets(Filter, DataTableAssets);
+		}
+	}
 
 	return DataTableAssets;
 }
