[Mutable] Removed support for Virtual Multilayers.

#rnx

diff --git a/Engine/Plugins/Mutable/Source/CustomizableObject/Internal/MuCO/CustomizableObjectInstanceDescriptor.h b/Engine/Plugins/Mutable/Source/CustomizableObject/Internal/MuCO/CustomizableObjectInstanceDescriptor.h
--- a/Engine/Plugins/Mutable/Source/CustomizableObject/Internal/MuCO/CustomizableObjectInstanceDescriptor.h
+++ b/Engine/Plugins/Mutable/Source/CustomizableObject/Internal/MuCO/CustomizableObjectInstanceDescriptor.h
@@ -291,36 +291,6 @@ struct CUSTOMIZABLEOBJECT_API FCustomizableObjectInstanceDescriptor
 	/** Updates the parameters of the layer at Index from the multilayer projector with name ParamName. */
 	void UpdateLayer(const FName& ParamName, int32 Index, const FMultilayerProjectorLayer& Layer);
 
-
-	// Virtual layers
-
-	/** Given Multilayer Projector name, create a new Multilayer Projector Helper (if non-existent) for virtual layers. See FMultilayerProjector.
-	 *
-	 * @return ture if successfully created (or was already created).
-	 */
-	bool CreateMultiLayerProjector(const FName& ProjectorParamName);
-
-	/** Given Multilayer Projector name, remove a Multilayer Projector Helper. See FMultilayerProjector. */
-	void RemoveMultilayerProjector(const FName& ProjectorParamName);
-
-	/** See FMultilayerProjector::GetVirtualLayers. */
-	TArray<FName> MultilayerProjectorGetVirtualLayers(const FName& ProjectorParamName) const;
-	
-	/** See FMultilayerProjector::VirtualLayer. */
-	void MultilayerProjectorCreateVirtualLayer(const FName& ProjectorParamName, const FName& Id);
-
-	/** See FMultilayerProjector::FindOrCreateVirtualLayer. */
-	FMultilayerProjectorVirtualLayer MultilayerProjectorFindOrCreateVirtualLayer(const FName& ProjectorParamName, const FName& Id);
-
-	/** See FMultilayerProjector::RemoveVirtualLayer. */
-	void MultilayerProjectorRemoveVirtualLayer(const FName& ProjectorParamName, const FName& Id);
-
-	/** See FMultilayerProjector::GetVirtualLayer. */
-	FMultilayerProjectorVirtualLayer MultilayerProjectorGetVirtualLayer(const FName& ProjectorParamName, const FName& Id) const;
-
-	/** See FMultilayerProjector::UpdateVirtualLayer. */
-	void MultilayerProjectorUpdateVirtualLayer(const FName& ProjectorParamName, const FName& Id, const FMultilayerProjectorVirtualLayer& Layer);
-
 	/** Return a Mutable Core object containing all parameters. */
 	mu::Ptr<mu::Parameters> GetParameters() const;
 
@@ -362,10 +332,6 @@ private:
 	/** Array of RequestedLODs per component to generate, they MUST NOT be used in an update (Mutable thread). */
 	TArray<uint16> RequestedLODLevels;
 	
-	/** Multilayer Projector helpers. See FMultilayerProjector.*/
-	UPROPERTY()
-	TMap<FName, FMultilayerProjector> MultilayerProjectors;
-	
 	// Friends
 	friend FDescriptorHash;
 	friend UCustomizableObjectInstance;
diff --git a/Engine/Plugins/Mutable/Source/CustomizableObject/Internal/MuCO/CustomizableObjectInstancePrivate.h b/Engine/Plugins/Mutable/Source/CustomizableObject/Internal/MuCO/CustomizableObjectInstancePrivate.h
--- a/Engine/Plugins/Mutable/Source/CustomizableObject/Internal/MuCO/CustomizableObjectInstancePrivate.h
+++ b/Engine/Plugins/Mutable/Source/CustomizableObject/Internal/MuCO/CustomizableObjectInstancePrivate.h
@@ -36,6 +36,17 @@ class USkeleton;
 class UCustomizableObjectExtension;
 
 
+// Log texts
+extern const FString MULTILAYER_PROJECTOR_PARAMETERS_INVALID;
+
+	
+// Parameters encoding
+CUSTOMIZABLEOBJECT_API extern const FString NUM_LAYERS_PARAMETER_POSTFIX;
+CUSTOMIZABLEOBJECT_API extern const FString OPACITY_PARAMETER_POSTFIX;
+CUSTOMIZABLEOBJECT_API extern const FString IMAGE_PARAMETER_POSTFIX;
+CUSTOMIZABLEOBJECT_API extern const FString POSE_PARAMETER_POSTFIX;
+
+
 /** \param OnlyLOD: If not 0, extract and convert only one single LOD from the source image.
   * \param ExtractChannel: If different than -1, extract a single-channel image with the specified source channel data. */
 CUSTOMIZABLEOBJECT_API void ConvertImage(UTexture2D* Texture, mu::Ptr<const mu::Image> MutableImage, const FMutableModelImageProperties& Props, int32 OnlyLOD = -1, int32 ExtractChannel = -1);
diff --git a/Engine/Plugins/Mutable/Source/CustomizableObject/Private/MuCO/CustomizableObjectInstance.cpp b/Engine/Plugins/Mutable/Source/CustomizableObject/Private/MuCO/CustomizableObjectInstance.cpp
--- a/Engine/Plugins/Mutable/Source/CustomizableObject/Private/MuCO/CustomizableObjectInstance.cpp
+++ b/Engine/Plugins/Mutable/Source/CustomizableObject/Private/MuCO/CustomizableObjectInstance.cpp
@@ -80,6 +80,13 @@ static FAutoConsoleVariableRef CVarDisableClothingPhysicsEditsPropagation(
 }
 
 
+const FString MULTILAYER_PROJECTOR_PARAMETERS_INVALID = TEXT("Invalid Multilayer Projector Parameters.");
+
+const FString NUM_LAYERS_PARAMETER_POSTFIX = FString("_NumLayers");
+const FString OPACITY_PARAMETER_POSTFIX = FString("_Opacity");
+const FString IMAGE_PARAMETER_POSTFIX = FString("_SelectedImages");
+const FString POSE_PARAMETER_POSTFIX = FString("_SelectedPoses");
+
 
 // Struct used by BuildMaterials() to identify common materials between LODs
 struct FMutableMaterialPlaceholder
@@ -531,10 +538,10 @@ void UCustomizableObjectInstance::Serialize(FArchive& Ar)
 		// Find the num layer parameters that were int enums
 		for (int32 i = 0; i < IntParameters_DEPRECATED.Num(); ++i)
 		{
-			if (IntParameters_DEPRECATED[i].ParameterName.EndsWith(FMultilayerProjector::NUM_LAYERS_PARAMETER_POSTFIX, ESearchCase::CaseSensitive))
+			if (IntParameters_DEPRECATED[i].ParameterName.EndsWith(NUM_LAYERS_PARAMETER_POSTFIX, ESearchCase::CaseSensitive))
 			{
 				FString ParameterNamePrefix, Aux;
-				const bool bSplit = IntParameters_DEPRECATED[i].ParameterName.Split(FMultilayerProjector::NUM_LAYERS_PARAMETER_POSTFIX, &ParameterNamePrefix, &Aux);
+				const bool bSplit = IntParameters_DEPRECATED[i].ParameterName.Split(NUM_LAYERS_PARAMETER_POSTFIX, &ParameterNamePrefix, &Aux);
 				check(bSplit);
 
 				// Confirm this is actually a multilayer param by finding the corresponding pose param
@@ -543,7 +550,7 @@ void UCustomizableObjectInstance::Serialize(FArchive& Ar)
 					if (i != j)
 					{
 						if (IntParameters_DEPRECATED[j].ParameterName.StartsWith(ParameterNamePrefix, ESearchCase::CaseSensitive) &&
-							IntParameters_DEPRECATED[j].ParameterName.EndsWith(FMultilayerProjector::POSE_PARAMETER_POSTFIX, ESearchCase::CaseSensitive))
+							IntParameters_DEPRECATED[j].ParameterName.EndsWith(POSE_PARAMETER_POSTFIX, ESearchCase::CaseSensitive))
 						{
 							IntParametersToMove.Add(i);
 							break;
@@ -580,11 +587,6 @@ void UCustomizableObjectInstance::Serialize(FArchive& Ar)
 		Descriptor.VectorParameters = VectorParameters_DEPRECATED;
 		Descriptor.ProjectorParameters = ProjectorParameters_DEPRECATED;
 	}
-
-	if (CustomizableObjectCustomVersion < FCustomizableObjectCustomVersion::DescriptorMultilayerProjectors)
-	{
-		Descriptor.MultilayerProjectors = MultilayerProjectors_DEPRECATED;
-	}
 }
 
 
@@ -2449,18 +2451,6 @@ int32 UCustomizableObjectInstance::RemoveValueFromProjectorRange(const FString&
 }
 
 
-bool UCustomizableObjectInstance::CreateMultiLayerProjector(const FName& ProjectorParamName)
-{
-	return Descriptor.CreateMultiLayerProjector(ProjectorParamName);
-}
-
-
-void UCustomizableObjectInstance::RemoveMultilayerProjector(const FName& ProjectorParamName)
-{
-	Descriptor.RemoveMultilayerProjector(ProjectorParamName);
-}
-
-
 int32 UCustomizableObjectInstance::MultilayerProjectorNumLayers(const FName& ProjectorParamName) const
 {
 	return Descriptor.NumProjectorLayers(ProjectorParamName);
@@ -2491,42 +2481,6 @@ void UCustomizableObjectInstance::MultilayerProjectorUpdateLayer(const FName& Pr
 }
 
 
-TArray<FName> UCustomizableObjectInstance::MultilayerProjectorGetVirtualLayers(const FName& ProjectorParamName) const
-{
-	return Descriptor.MultilayerProjectorGetVirtualLayers(ProjectorParamName);
-}
-
-
-void UCustomizableObjectInstance::MultilayerProjectorCreateVirtualLayer(const FName& ProjectorParamName, const FName& Id)
-{
-	Descriptor.MultilayerProjectorCreateVirtualLayer(ProjectorParamName, Id);
-}
-
-
-FMultilayerProjectorVirtualLayer UCustomizableObjectInstance::MultilayerProjectorFindOrCreateVirtualLayer(const FName& ProjectorParamName, const FName& Id)
-{
-	return Descriptor.MultilayerProjectorFindOrCreateVirtualLayer(ProjectorParamName, Id);
-}
-
-
-void UCustomizableObjectInstance::MultilayerProjectorRemoveVirtualLayer(const FName& ProjectorParamName, const FName& Id)
-{
-	Descriptor.MultilayerProjectorRemoveVirtualLayer(ProjectorParamName, Id);
-}
-
-
-FMultilayerProjectorVirtualLayer UCustomizableObjectInstance::MultilayerProjectorGetVirtualLayer(const FName& ProjectorParamName, const FName& Id) const
-{
-	return Descriptor.MultilayerProjectorGetVirtualLayer(ProjectorParamName, Id);
-}
-
-
-void UCustomizableObjectInstance::MultilayerProjectorUpdateVirtualLayer(const FName& ProjectorParamName, const FName& Id, const FMultilayerProjectorVirtualLayer& Layer)
-{
-	Descriptor.MultilayerProjectorUpdateVirtualLayer(ProjectorParamName, Id, Layer);
-}
-
-
 void UCustomizableObjectInstance::SaveDescriptor(FArchive &Ar, bool bUseCompactDescriptor)
 {
 	Descriptor.SaveDescriptor(Ar, bUseCompactDescriptor);
diff --git a/Engine/Plugins/Mutable/Source/CustomizableObject/Private/MuCO/CustomizableObjectInstanceDescriptor.cpp b/Engine/Plugins/Mutable/Source/CustomizableObject/Private/MuCO/CustomizableObjectInstanceDescriptor.cpp
--- a/Engine/Plugins/Mutable/Source/CustomizableObject/Private/MuCO/CustomizableObjectInstanceDescriptor.cpp
+++ b/Engine/Plugins/Mutable/Source/CustomizableObject/Private/MuCO/CustomizableObjectInstanceDescriptor.cpp
@@ -2458,31 +2458,31 @@ bool FCustomizableObjectInstanceDescriptor::IsMultilayerProjector(const FString&
 	if (const int32 ProjectorParameterIndex = FindTypedParameterIndex(ParamName, EMutableParameterType::Projector);
 		ProjectorParameterIndex == INDEX_NONE)
 	{
-		ensureAlwaysMsgf(false, TEXT("%s"), *FMultilayerProjector::MULTILAYER_PROJECTOR_PARAMETERS_INVALID);
+		ensureAlwaysMsgf(false, TEXT("%s"), *MULTILAYER_PROJECTOR_PARAMETERS_INVALID);
 		return false;
 	}
 
 	// Num layers.
-	if (const int32 FloatParameterIndex = FindTypedParameterIndex(ParamName + FMultilayerProjector::NUM_LAYERS_PARAMETER_POSTFIX, EMutableParameterType::Float);
+	if (const int32 FloatParameterIndex = FindTypedParameterIndex(ParamName + NUM_LAYERS_PARAMETER_POSTFIX, EMutableParameterType::Float);
 		FloatParameterIndex == INDEX_NONE)
 	{
-		ensureAlwaysMsgf(false, TEXT("%s"), *FMultilayerProjector::MULTILAYER_PROJECTOR_PARAMETERS_INVALID);
+		ensureAlwaysMsgf(false, TEXT("%s"), *MULTILAYER_PROJECTOR_PARAMETERS_INVALID);
 		return false;
 	}
 
 	// Selected Image.
-	if (const int32 IntParameterIndex = FindTypedParameterIndex(ParamName + FMultilayerProjector::IMAGE_PARAMETER_POSTFIX, EMutableParameterType::Int);
+	if (const int32 IntParameterIndex = FindTypedParameterIndex(ParamName + IMAGE_PARAMETER_POSTFIX, EMutableParameterType::Int);
 		IntParameterIndex == INDEX_NONE)
 	{
-		ensureAlwaysMsgf(false, TEXT("%s"), *FMultilayerProjector::MULTILAYER_PROJECTOR_PARAMETERS_INVALID);
+		ensureAlwaysMsgf(false, TEXT("%s"), *MULTILAYER_PROJECTOR_PARAMETERS_INVALID);
 		return false;
 	}
 
 	// Opacity.	
-	if (const int32 FloatParameterIndex = FindTypedParameterIndex(ParamName + FMultilayerProjector::OPACITY_PARAMETER_POSTFIX, EMutableParameterType::Float);
+	if (const int32 FloatParameterIndex = FindTypedParameterIndex(ParamName + OPACITY_PARAMETER_POSTFIX, EMutableParameterType::Float);
 		FloatParameterIndex == INDEX_NONE)
 	{
-		ensureAlwaysMsgf(false, TEXT("%s"), *FMultilayerProjector::MULTILAYER_PROJECTOR_PARAMETERS_INVALID);
+		ensureAlwaysMsgf(false, TEXT("%s"), *MULTILAYER_PROJECTOR_PARAMETERS_INVALID);
 		return false;
 	}
 
@@ -2496,7 +2496,7 @@ int32 FCustomizableObjectInstanceDescriptor::NumProjectorLayers(const FName& Pro
 
 	if (IsMultilayerProjector(ParamName))
 	{
-		return GetFloatParameterSelectedOption(ParamName + FMultilayerProjector::NUM_LAYERS_PARAMETER_POSTFIX);
+		return GetFloatParameterSelectedOption(ParamName + NUM_LAYERS_PARAMETER_POSTFIX);
 	}
 
 	return INDEX_NONE;
@@ -2513,14 +2513,14 @@ void FCustomizableObjectInstanceDescriptor::CreateLayer(const FName& ProjectorPa
 	}
 
 	// Num Layers.
-	SetFloatParameterSelectedOption(ParamName + FMultilayerProjector::NUM_LAYERS_PARAMETER_POSTFIX, NumProjectorLayers(ProjectorParamName) + 1);
+	SetFloatParameterSelectedOption(ParamName + NUM_LAYERS_PARAMETER_POSTFIX, NumProjectorLayers(ProjectorParamName) + 1);
 
 	// Projector Range. New value is defaulted.
 	AddValueToProjectorRange(ParamName);
 
 	// Selected Image Range.
 	{
-		const FString ImageParamName = ParamName + FMultilayerProjector::IMAGE_PARAMETER_POSTFIX;
+		const FString ImageParamName = ParamName + IMAGE_PARAMETER_POSTFIX;
 		AddValueToTextureRange(ImageParamName);
 
 		const int32 ParamIndex = CustomizableObject->FindParameter(ImageParamName);
@@ -2531,7 +2531,7 @@ void FCustomizableObjectInstanceDescriptor::CreateLayer(const FName& ProjectorPa
 
 	// Opacity Range.
 	{
-		const FString OpacityParamName = ParamName + FMultilayerProjector::OPACITY_PARAMETER_POSTFIX;
+		const FString OpacityParamName = ParamName + OPACITY_PARAMETER_POSTFIX;
 		AddValueToFloatRange(OpacityParamName);
 
 		const int32 ParamIndex = CustomizableObject->FindParameter(OpacityParamName);
@@ -2553,16 +2553,16 @@ void FCustomizableObjectInstanceDescriptor::RemoveLayerAt(const FName& Projector
 	}
 
 	// Num Layers.
-	SetFloatParameterSelectedOption(ParamName + FMultilayerProjector::NUM_LAYERS_PARAMETER_POSTFIX, NumLayers - 1);
+	SetFloatParameterSelectedOption(ParamName + NUM_LAYERS_PARAMETER_POSTFIX, NumLayers - 1);
 
 	// Projector Range.
 	RemoveValueFromProjectorRange(ParamName, RangeIndex);
 
 	// Selected Image Range.
-	RemoveValueFromIntRange(ParamName + FMultilayerProjector::IMAGE_PARAMETER_POSTFIX, RangeIndex);
+	RemoveValueFromIntRange(ParamName + IMAGE_PARAMETER_POSTFIX, RangeIndex);
 
 	// Opacity Range.
-	RemoveValueFromFloatRange(ParamName + FMultilayerProjector::OPACITY_PARAMETER_POSTFIX);
+	RemoveValueFromFloatRange(ParamName + OPACITY_PARAMETER_POSTFIX);
 }
 
 
@@ -2594,79 +2594,4 @@ void FCustomizableObjectInstanceDescriptor::UpdateLayer(const FName& ProjectorPa
 }
 
 
-bool FCustomizableObjectInstanceDescriptor::CreateMultiLayerProjector(const FName& ProjectorParamName)
-{
-	if (!IsMultilayerProjector(ProjectorParamName.ToString()))
-	{
-#if WITH_EDITOR
-		UE_LOG(LogMutable, Error, TEXT("%s"), *FMultilayerProjector::MULTILAYER_PROJECTOR_PARAMETERS_INVALID);
-#else
-		ensureAlwaysMsgf(false, TEXT("%s"), *FMultilayerProjector::MULTILAYER_PROJECTOR_PARAMETERS_INVALID);
-#endif
-		return false;
-	}
-
-	if (!MultilayerProjectors.Contains(ProjectorParamName))
-	{
-		const FMultilayerProjector MultilayerProjector(ProjectorParamName);
-		MultilayerProjectors.Add(ProjectorParamName, MultilayerProjector);
-	}
-
-	return true;
-}
-
-
-void FCustomizableObjectInstanceDescriptor::RemoveMultilayerProjector(const FName& ProjectorParamName)
-{
-	MultilayerProjectors.Remove(ProjectorParamName);
-}
-
-
-TArray<FName> FCustomizableObjectInstanceDescriptor::MultilayerProjectorGetVirtualLayers(const FName& ProjectorParamName) const
-{
-	check(MultilayerProjectors.Contains(ProjectorParamName)); // Multilayer Projector not created.
-	
-	return MultilayerProjectors[ProjectorParamName].GetVirtualLayers();
-}
-
-
-void FCustomizableObjectInstanceDescriptor::MultilayerProjectorCreateVirtualLayer(const FName& ProjectorParamName, const FName& Id)
-{
-	check(MultilayerProjectors.Contains(ProjectorParamName)); // Multilayer Projector not created.
-	
-	MultilayerProjectors[ProjectorParamName].CreateVirtualLayer(*this, Id);
-}
-
-
-FMultilayerProjectorVirtualLayer FCustomizableObjectInstanceDescriptor::MultilayerProjectorFindOrCreateVirtualLayer(const FName& ProjectorParamName, const FName& Id)
-{
-	check(MultilayerProjectors.Contains(ProjectorParamName)); // Multilayer Projector not created.
-	
-	return MultilayerProjectors[ProjectorParamName].FindOrCreateVirtualLayer(*this, Id);
-}
-
-
-void FCustomizableObjectInstanceDescriptor::MultilayerProjectorRemoveVirtualLayer(const FName& ProjectorParamName, const FName& Id)
-{
-	check(MultilayerProjectors.Contains(ProjectorParamName)); // Multilayer Projector not created.
-	
-	MultilayerProjectors[ProjectorParamName].RemoveVirtualLayer(*this, Id);
-}
-
-
-FMultilayerProjectorVirtualLayer FCustomizableObjectInstanceDescriptor::MultilayerProjectorGetVirtualLayer(const FName& ProjectorParamName, const FName& Id) const
-{
-	check(MultilayerProjectors.Contains(ProjectorParamName)); // Multilayer Projector not created.
-	
-	return MultilayerProjectors[ProjectorParamName].GetVirtualLayer(*this, Id);
-}
-
-
-void FCustomizableObjectInstanceDescriptor::MultilayerProjectorUpdateVirtualLayer(const FName& ProjectorParamName, const FName& Id, const FMultilayerProjectorVirtualLayer& Layer)
-{
-	check(MultilayerProjectors.Contains(ProjectorParamName)); // Multilayer Projector not created.
-	
-	MultilayerProjectors[ProjectorParamName].UpdateVirtualLayer(*this, Id, Layer);
-}
-
 #undef RETURN_ON_CO_UNCOMPILED
diff --git a/Engine/Plugins/Mutable/Source/CustomizableObject/Private/MuCO/DescriptorHash.cpp b/Engine/Plugins/Mutable/Source/CustomizableObject/Private/MuCO/DescriptorHash.cpp
--- a/Engine/Plugins/Mutable/Source/CustomizableObject/Private/MuCO/DescriptorHash.cpp
+++ b/Engine/Plugins/Mutable/Source/CustomizableObject/Private/MuCO/DescriptorHash.cpp
@@ -46,12 +46,6 @@ FDescriptorHash::FDescriptorHash(const FCustomizableObjectInstanceDescriptor& De
 	
 	Hash = HashCombine(Hash, GetTypeHash(Descriptor.State));
 	Hash = HashCombine(Hash, GetTypeHash(Descriptor.GetBuildParameterRelevancy()));
-
-	for (const TTuple<FName, FMultilayerProjector>& Pair : Descriptor.MultilayerProjectors)
-	{
-		// Hash = HashCombine(Hash, GetTypeHash(Pair.Key)); // Already hashed by the FMultilayerProjector.
-		Hash = HashCombine(Hash, GetTypeHash(Pair.Value));
-	}
 	
 	MinLOD = Descriptor.MinLOD;
 
diff --git a/Engine/Plugins/Mutable/Source/CustomizableObject/Private/MuCO/MultilayerProjector.cpp b/Engine/Plugins/Mutable/Source/CustomizableObject/Private/MuCO/MultilayerProjector.cpp
--- a/Engine/Plugins/Mutable/Source/CustomizableObject/Private/MuCO/MultilayerProjector.cpp
+++ b/Engine/Plugins/Mutable/Source/CustomizableObject/Private/MuCO/MultilayerProjector.cpp
@@ -4,18 +4,16 @@
 
 #include "MuCO/CustomizableObject.h"
 #include "MuCO/CustomizableObjectInstanceDescriptor.h"
+#include "MuCO/CustomizableObjectInstancePrivate.h"
 
 #include UE_INLINE_GENERATED_CPP_BY_NAME(MultilayerProjector)
 
 
-const FString FMultilayerProjector::MULTILAYER_PROJECTOR_PARAMETERS_INVALID = TEXT("Invalid Multilayer Projector Parameters.");
-
-
 void FMultilayerProjectorLayer::Read(const FCustomizableObjectInstanceDescriptor& Descriptor, const FString& ParamName, const int32 LayerIndex)
 {
 	if (!Descriptor.IsMultilayerProjector(ParamName))
 	{
-		ensureAlwaysMsgf(false, TEXT("%s"), *FMultilayerProjector::MULTILAYER_PROJECTOR_PARAMETERS_INVALID);
+		ensureAlwaysMsgf(false, TEXT("%s"), *MULTILAYER_PROJECTOR_PARAMETERS_INVALID);
 		return;
 	}
 
@@ -27,12 +25,12 @@ void FMultilayerProjectorLayer::Read(const FCustomizableObjectInstanceDescriptor
 	}
 
 	{
-		const int32 ImageParamIndex = Descriptor.FindTypedParameterIndex(ParamName + FMultilayerProjector::IMAGE_PARAMETER_POSTFIX, EMutableParameterType::Int);
+		const int32 ImageParamIndex = Descriptor.FindTypedParameterIndex(ParamName + IMAGE_PARAMETER_POSTFIX, EMutableParameterType::Int);
 		Image = Descriptor.GetIntParameters()[ImageParamIndex].ParameterRangeValueNames[LayerIndex];
 	}
 
 	{
-		const int32 OpacityParamIndex = Descriptor.FindTypedParameterIndex(ParamName + FMultilayerProjector::OPACITY_PARAMETER_POSTFIX, EMutableParameterType::Float);
+		const int32 OpacityParamIndex = Descriptor.FindTypedParameterIndex(ParamName + OPACITY_PARAMETER_POSTFIX, EMutableParameterType::Float);
 		Opacity = Descriptor.GetFloatParameters()[OpacityParamIndex].ParameterRangeValues[LayerIndex];
 	}
 }
@@ -42,7 +40,7 @@ void FMultilayerProjectorLayer::Write(FCustomizableObjectInstanceDescriptor& Des
 {
 	if (!Descriptor.IsMultilayerProjector(ParamName))
 	{
-		ensureAlwaysMsgf(false, TEXT("%s"), *FMultilayerProjector::MULTILAYER_PROJECTOR_PARAMETERS_INVALID);
+		ensureAlwaysMsgf(false, TEXT("%s"), *MULTILAYER_PROJECTOR_PARAMETERS_INVALID);
 		return;
 	}
 
@@ -50,8 +48,8 @@ void FMultilayerProjectorLayer::Write(FCustomizableObjectInstanceDescriptor& Des
 
 
 	Descriptor.SetProjectorValue(ParamName, Position, Direction, Up, Scale, Angle, LayerIndex);
-	Descriptor.SetIntParameterSelectedOption(ParamName + FMultilayerProjector::IMAGE_PARAMETER_POSTFIX, Image, LayerIndex);
-	Descriptor.SetFloatParameterSelectedOption(ParamName + FMultilayerProjector::OPACITY_PARAMETER_POSTFIX, Opacity, LayerIndex);
+	Descriptor.SetIntParameterSelectedOption(ParamName + IMAGE_PARAMETER_POSTFIX, Image, LayerIndex);
+	Descriptor.SetFloatParameterSelectedOption(ParamName + OPACITY_PARAMETER_POSTFIX, Opacity, LayerIndex);
 
 }
 
@@ -70,273 +68,3 @@ uint32 GetTypeHash(const FMultilayerProjectorLayer& Key)
 	return Hash;
 }
 
-
-FMultilayerProjectorVirtualLayer::FMultilayerProjectorVirtualLayer(const FMultilayerProjectorLayer& Layer, const bool bEnabled, const int32 Order):
-	FMultilayerProjectorLayer(Layer),
-	bEnabled(bEnabled),
-	Order(Order)
-{
-}
-
-
-const FString FMultilayerProjector::NUM_LAYERS_PARAMETER_POSTFIX = FString("_NumLayers");
-
-const FString FMultilayerProjector::OPACITY_PARAMETER_POSTFIX = FString("_Opacity");
-
-const FString FMultilayerProjector::IMAGE_PARAMETER_POSTFIX = FString("_SelectedImages");
-
-const FString FMultilayerProjector::POSE_PARAMETER_POSTFIX = FString("_SelectedPoses");
-
-TArray<FName> FMultilayerProjector::GetVirtualLayers() const
-{
-	TArray<FName> VirtualLayers;
-	VirtualLayersMapping.GetKeys(VirtualLayers);
-	return VirtualLayers;
-}
-
-
-void FMultilayerProjector::CreateVirtualLayer(FCustomizableObjectInstanceDescriptor& Descriptor, const FName& Id)
-{
-	if (!VirtualLayersMapping.Contains(Id))
-	{
-		const int32 Index = Descriptor.NumProjectorLayers(ParamName);
-		
-		Descriptor.CreateLayer(ParamName, Index);
-		VirtualLayersMapping.Add(Id, Index);
-		VirtualLayersOrder.Add(Id, NEW_VIRTUAL_LAYER_ORDER);
-	}
-}
-
-
-FMultilayerProjectorVirtualLayer FMultilayerProjector::FindOrCreateVirtualLayer(FCustomizableObjectInstanceDescriptor& Descriptor, const FName& Id)
-{
-	FMultilayerProjectorLayer Layer;
-	bool bEnabled;
-	int32 Order;
-
-	if (const int32* Index = VirtualLayersMapping.Find(Id))
-	{
-		if (*Index == VIRTUAL_LAYER_DISABLED)
-		{
-			Layer = DisableVirtualLayers[Id];
-			bEnabled = false;
-		}
-		else
-		{
-			Layer = Descriptor.GetLayer(ParamName, *Index);
-			bEnabled = true;
-		}
-
-		Order = VirtualLayersOrder[Id];
-	}
-	else
-	{
-		const int32 NewIndex = Descriptor.NumProjectorLayers(ParamName);
-		constexpr int32 NewOrder = NEW_VIRTUAL_LAYER_ORDER;
-		
-		Descriptor.CreateLayer(ParamName, NewIndex);
-		VirtualLayersMapping.Add(Id, NewIndex);
-		VirtualLayersOrder.Add(Id, NewOrder);
-
-		Layer = Descriptor.GetLayer(ParamName, NewIndex);
-		bEnabled = true;
-		Order = NewOrder;
-	}
-
-	return FMultilayerProjectorVirtualLayer(Layer, bEnabled, Order);
-}
-
-
-void FMultilayerProjector::RemoveVirtualLayer(FCustomizableObjectInstanceDescriptor& Descriptor, const FName& Id)
-{
-	const int32* Index = VirtualLayersMapping.Find(Id);
-	check(Index); // Virtual Layer not created.
-	
-	if (*Index == VIRTUAL_LAYER_DISABLED)
-	{
-		DisableVirtualLayers.Remove(Id);
-	}
-	else
-	{
-		Descriptor.RemoveLayerAt(ParamName, *Index);
-		
-		for (TMap<FName, int32>::TIterator It =  VirtualLayersMapping.CreateIterator(); It; ++It)
-		{
-			if (It.Key() == Id)
-			{
-				It.RemoveCurrent();
-			}
-			else if (It.Value() > *Index) // Update following Layers.
-			{
-				--It.Value();
-			}
-		}
-	}
-
-	VirtualLayersOrder.Remove(Id);
-}
-
-
-FMultilayerProjectorVirtualLayer FMultilayerProjector::GetVirtualLayer(const FCustomizableObjectInstanceDescriptor& Descriptor, const FName& Id) const
-{
-	const int32* Index = VirtualLayersMapping.Find(Id);
-	check(Index); // Virtual Layer not created.
-
-	const FMultilayerProjectorLayer Layer = Descriptor.GetLayer(ParamName, *Index);
-	const bool bEnabled = *Index != VIRTUAL_LAYER_DISABLED;
-	const int32 Order = VirtualLayersOrder[Id];
-	
-	return FMultilayerProjectorVirtualLayer(Layer, bEnabled, Order);
-}
-
-
-void FMultilayerProjector::UpdateVirtualLayer(FCustomizableObjectInstanceDescriptor& Descriptor, const FName& Id, const FMultilayerProjectorVirtualLayer& Layer)
-{
-	const int32* Index = VirtualLayersMapping.Find(Id);
-	check(Index); // Virtual Layer not created.
-
-	const bool bEnabled = *Index != VIRTUAL_LAYER_DISABLED;
-	
-	if (!bEnabled)
-	{
-		DisableVirtualLayers[Id] = static_cast<FMultilayerProjectorLayer>(Layer); // Update disabled layer.
-		VirtualLayersOrder[Id] = Layer.Order;
-	}
-	else
-	{
-		int32* Order = VirtualLayersOrder.Find(Id);
-		if (*Order != Layer.Order) // Order changed, check if it needs to be moved.
-		{
-			const int32 OldIndex = *Index;
-			int32 NewIndex = CalculateVirtualLayerIndex(Id, Layer.Order);
-			if (OldIndex != NewIndex) // Move required. Could be optimized by moving only the in-between values.
-			{
-				Descriptor.RemoveLayerAt(ParamName, OldIndex);
-				UpdateMappingVirtualLayerDisabled(Id, OldIndex);
-
-				if (OldIndex < NewIndex)
-				{
-					NewIndex -= 1;
-				}
-				
-				Descriptor.CreateLayer(ParamName, NewIndex);
-				UpdateMappingVirtualLayerEnabled(Id, NewIndex);
-			}
-			
-			*Order = Layer.Order;
-		}
-
-		Descriptor.UpdateLayer(ParamName, *Index, static_cast<FMultilayerProjectorLayer>(Layer)); // Update enabled layer.
-	}
-	
-	// Enable or disable virtual layer.
-	if (Layer.bEnabled && !bEnabled)
-	{
-		const int32 NewIndex = CalculateVirtualLayerIndex(Id, VirtualLayersOrder[Id]);
-
-		Descriptor.CreateLayer(ParamName, NewIndex);
-		UpdateMappingVirtualLayerEnabled(Id, NewIndex);
-
-		Descriptor.UpdateLayer(ParamName, NewIndex, static_cast<FMultilayerProjectorLayer>(Layer));
-		
-		DisableVirtualLayers.Remove(Id);
-	}
-	else if (!Layer.bEnabled && bEnabled)
-	{
-		Descriptor.RemoveLayerAt(ParamName, *Index);
-		UpdateMappingVirtualLayerDisabled(Id, *Index);
-		
-		DisableVirtualLayers.Add(Id, Layer);
-	}
-}
-
-
-FMultilayerProjector::FMultilayerProjector(const FName& InParamName) :
-	ParamName(InParamName)
-{
-}
-
-
-int32 FMultilayerProjector::CalculateVirtualLayerIndex(const FName& Id, const int32 InsertOrder) const
-{
-	int32 LayerBeforeIndex = -1;
-	int32 LayerBeforeOrder = -1;
-	
-	for (const TTuple<FName, int>& MappingTuple : VirtualLayersMapping) // Find closest smallest layer.
-	{
-		if (MappingTuple.Value != VIRTUAL_LAYER_DISABLED && MappingTuple.Key != Id)
-		{
-			const int32 LayerOrder = VirtualLayersOrder[MappingTuple.Key];
-			if (LayerOrder <= InsertOrder)
-			{
-				if ((LayerOrder > LayerBeforeOrder) ||
-					(LayerOrder == LayerBeforeOrder && MappingTuple.Value > LayerBeforeIndex))
-				{
-					LayerBeforeIndex = MappingTuple.Value;
-					LayerBeforeOrder = LayerOrder;
-				}
-			}
-		}
-	}
-	
-	return LayerBeforeIndex + 1;
-}
-
-
-void FMultilayerProjector::UpdateMappingVirtualLayerEnabled(const FName& Id, const int32 Index)
-{
-	for (TTuple<FName, int>& Tuple : VirtualLayersMapping)
-	{
-		if (Tuple.Key == Id)
-		{
-			Tuple.Value = Index;
-		}
-		else if (Tuple.Value >= Index) // Update following Layers.
-		{
-			++Tuple.Value;
-		}
-	}
-}
-
-
-void FMultilayerProjector::UpdateMappingVirtualLayerDisabled(const FName& Id, const int32 Index)
-{
-	for (TTuple<FName, int>& Tuple : VirtualLayersMapping)
-	{
-		if (Tuple.Key == Id)
-    	{
-			Tuple.Value = VIRTUAL_LAYER_DISABLED;
-    	}
-		else if (Tuple.Value > Index) // Update following Layers.
-		{
-			--Tuple.Value;
-		}
-	}
-}
-
-
-uint32 GetTypeHash(const FMultilayerProjector& Key)
-{
-	uint32 Hash = GetTypeHash(Key.ParamName);
-
-	for (const TTuple<FName, int32>& Pair : Key.VirtualLayersMapping)
-	{
-		Hash = HashCombine(Hash, GetTypeHash(Pair.Key));
-		Hash = HashCombine(Hash, GetTypeHash(Pair.Value));
-	}
-	
-	for (const TTuple<FName, int32>& Pair : Key.VirtualLayersOrder)
-	{
-		Hash = HashCombine(Hash, GetTypeHash(Pair.Key));
-		Hash = HashCombine(Hash, GetTypeHash(Pair.Value));
-	}
-
-	for (const TTuple<FName, FMultilayerProjectorLayer>& Pair : Key.DisableVirtualLayers)
-	{
-		Hash = HashCombine(Hash, GetTypeHash(Pair.Key));
-		Hash = HashCombine(Hash, GetTypeHash(Pair.Value));
-	}
-	
-	return Hash;
-}
-
diff --git a/Engine/Plugins/Mutable/Source/CustomizableObject/Public/MuCO/CustomizableObjectInstance.h b/Engine/Plugins/Mutable/Source/CustomizableObject/Public/MuCO/CustomizableObjectInstance.h
--- a/Engine/Plugins/Mutable/Source/CustomizableObject/Public/MuCO/CustomizableObjectInstance.h
+++ b/Engine/Plugins/Mutable/Source/CustomizableObject/Public/MuCO/CustomizableObjectInstance.h
@@ -534,15 +534,6 @@ public:
 	// Multilayer Projectors
 	// ------------------------------------------------------------
 	
-	/** Given Multilayer Projector name, create a new Multilayer Projector Helper (if non-existent). See FMultilayerProjector.
-	  * @return true if successfully created (or was already created). */
-	UFUNCTION(BlueprintCallable, Category = CustomizableObjectInstance)
-	bool CreateMultiLayerProjector(const FName& ProjectorParamName);
-	
-	/** Given Multilayer Projector name, remove a Multilayer Projector Helper. See FMultilayerProjector. */
-	UFUNCTION(BlueprintCallable, Category = CustomizableObjectInstance)
-	void RemoveMultilayerProjector(const FName& ProjectorParamName);
-	
 	// Layers
 
 	/** See FMultilayerProjector::NumLayers. */
@@ -564,32 +555,6 @@ public:
 	/** See FMultilayerProjector::UpdateLayer. */
 	UFUNCTION(BlueprintCallable, Category = CustomizableObjectInstance)
 	void MultilayerProjectorUpdateLayer(const FName& ProjectorParamName, int32 Index, const FMultilayerProjectorLayer& Layer);
-
-	// Virtual layers
-
-	/** See FMultilayerProjector::GetVirtualLayers. */
-	UFUNCTION(BlueprintCallable, Category = CustomizableObjectInstance)
-	TArray<FName> MultilayerProjectorGetVirtualLayers(const FName& ProjectorParamName) const;
-	
-	/** See FMultilayerProjector::VirtualLayer. */
-	UFUNCTION(BlueprintCallable, Category = CustomizableObjectInstance)
-	void MultilayerProjectorCreateVirtualLayer(const FName& ProjectorParamName, const FName& Id);
-
-	/** See FMultilayerProjector::FindOrCreateVirtualLayer. */
-	UFUNCTION(BlueprintCallable, Category = CustomizableObjectInstance)
-	FMultilayerProjectorVirtualLayer MultilayerProjectorFindOrCreateVirtualLayer(const FName& ProjectorParamName, const FName& Id);
-
-	/** See FMultilayerProjector::RemoveVirtualLayer. */
-	UFUNCTION(BlueprintCallable, Category = CustomizableObjectInstance)
-	void MultilayerProjectorRemoveVirtualLayer(const FName& ProjectorParamName, const FName& Id);
-
-	/** See FMultilayerProjector::GetVirtualLayer. */
-	UFUNCTION(BlueprintCallable, Category = CustomizableObjectInstance)
-	FMultilayerProjectorVirtualLayer MultilayerProjectorGetVirtualLayer(const FName& ProjectorParamName, const FName& Id) const;
-
-	/** See FMultilayerProjector::UpdateVirtualLayer. */
-	UFUNCTION(BlueprintCallable, Category = CustomizableObjectInstance)
-	void MultilayerProjectorUpdateVirtualLayer(const FName& ProjectorParamName, const FName& Id, const FMultilayerProjectorVirtualLayer& Layer);
 	
 	// ------------------------------------------------------------
 	
@@ -715,9 +680,6 @@ private:
 	UPROPERTY()
 	TArray<FCustomizableObjectProjectorParameterValue> ProjectorParameters_DEPRECATED;
 	
-   	UPROPERTY()
-   	TMap<FName, FMultilayerProjector> MultilayerProjectors_DEPRECATED;
-
 	bool bBuildParameterRelevancy_DEPRECATED = false;
 };
 
diff --git a/Engine/Plugins/Mutable/Source/CustomizableObject/Public/MuCO/MultilayerProjector.h b/Engine/Plugins/Mutable/Source/CustomizableObject/Public/MuCO/MultilayerProjector.h
--- a/Engine/Plugins/Mutable/Source/CustomizableObject/Public/MuCO/MultilayerProjector.h
+++ b/Engine/Plugins/Mutable/Source/CustomizableObject/Public/MuCO/MultilayerProjector.h
@@ -59,123 +59,3 @@ struct FMultilayerProjectorLayer
 
 uint32 GetTypeHash(const FMultilayerProjectorLayer& Key);
 
-
-/** Data structure representing a Multilayer Projector Virtual Layer. */
-USTRUCT(BlueprintType)
-struct FMultilayerProjectorVirtualLayer : public FMultilayerProjectorLayer
-{
-	GENERATED_BODY()
-
-	FMultilayerProjectorVirtualLayer() = default;
-
-	FMultilayerProjectorVirtualLayer(const FMultilayerProjectorLayer& Layer, bool bEnabled, int32 Order);
-
-	/** True if the Virtual Layer is enabled. */
-	bool bEnabled;
-
-	/** Virtual Layer absolute order. Equal values may result in an undefined behaviour. */
-	int32 Order;
-};
-
-
-/** Multilayer Projector Helper. Eases the management of Virtual Layers.
- * 
- * Virtual Layer: Management of Multilayer Projector Layers by name.
- * - Allows to enabled and disabled layers.
- * - Allows to sort layers.
- */
-USTRUCT()
-struct CUSTOMIZABLEOBJECT_API FMultilayerProjector
-{
-	GENERATED_BODY()
-
-	friend FCustomizableObjectInstanceDescriptor;
-	friend FMultilayerProjectorLayer;
-	friend uint32 GetTypeHash(const FMultilayerProjector& Key);
-
-	// Log texts
-	static const FString MULTILAYER_PROJECTOR_PARAMETERS_INVALID;
-
-	// Parameters encoding
-	static const FString NUM_LAYERS_PARAMETER_POSTFIX;
-	static const FString OPACITY_PARAMETER_POSTFIX;
-	static const FString IMAGE_PARAMETER_POSTFIX;
-	static const FString POSE_PARAMETER_POSTFIX;
-	
-	// Constructors
-	FMultilayerProjector() = default;
-
-	explicit FMultilayerProjector(const FName& InParamName);
-
-	// Virtual layers
-
-	/** Get all created Virtual Layers. */
-	TArray<FName> GetVirtualLayers() const;
-	
-	/** Given a new identifier, create a new Virtual Layer (if non-existent).*/
-	void CreateVirtualLayer(FCustomizableObjectInstanceDescriptor& Descriptor, const FName& Id);
-
-	/** Given a new identifier, find or create a new Virtual Layer.*/
-	FMultilayerProjectorVirtualLayer FindOrCreateVirtualLayer(FCustomizableObjectInstanceDescriptor& Descriptor, const FName& Id);
-
-	/** Given its identifier, remove the Virtual Layer. */
-	void RemoveVirtualLayer(FCustomizableObjectInstanceDescriptor& Descriptor, const FName& Id);
-
-	/** Given its identifier, get the Virtual Layer properties. */
-	FMultilayerProjectorVirtualLayer GetVirtualLayer(const FCustomizableObjectInstanceDescriptor& Descriptor, const FName& Id) const;
-
-	/** Given its identifier, update a Virtual Layer properties. */
-	void UpdateVirtualLayer(FCustomizableObjectInstanceDescriptor& Descriptor, const FName& Id, const FMultilayerProjectorVirtualLayer& Layer);
-
-private:
-	/** Multilayer Projector Parameter name. */
-	UPROPERTY()
-	FName ParamName;
-	
-	/** Maps a Virtual Layer to a Layer.
-	 *
- 	 * All enabled Virtual Layers must have an entry pointing to the Layer index.
- 	 * Disabled layers should have VIRTUAL_LAYER_DISABLED as value.
-	 */
-	UPROPERTY()
-	TMap<FName, int32> VirtualLayersMapping;
-
-	/** Maps a Virtual Layer to a Layer.
-	 *
-	 * All created Virtual Layers (enabled or disabled) must have an entry to this map.
-	 */
-	UPROPERTY()
-	TMap<FName, int32> VirtualLayersOrder;
-	
-	/** Since disabling a Virtual Layers removes the Layer, its values are no longer on the Instance Parameters.
-	 * Instead, they are saved in this data structure to be able to reenable it later. */
-	UPROPERTY()
-	TMap<FName, FMultilayerProjectorLayer> DisableVirtualLayers;
-
-	/** Virtual Layer Disabled Index. */
-	static constexpr int32 VIRTUAL_LAYER_DISABLED = -1;
-
-	/** New Virtual Layer default order. */
-	static constexpr int32 NEW_VIRTUAL_LAYER_ORDER = 0;
-
-	/** Taking into consideration the Virtual Layer Order, Calculate where the Layer should be inserted. */
-	int32 CalculateVirtualLayerIndex(const FName& Id, int32 InsertOrder) const;
-	
-	/** Given an inserted Virtual Layer, update the mapping of the Virtual Layers.
- 	 *
- 	 * @param Id Virtual Layer.
-	 * @param Index Index which the layer has been enabled.
-	 */
-	void UpdateMappingVirtualLayerEnabled(const FName& Id, int32 Index);
-	
-	/** Given a removed Virtual Layer, update the mapping of the Virtual Layers.
-	 *
- 	 * @param Id Virtual Layer.	
-	 * @param Index Index which the layer has been disabled.
-	 */
-	void UpdateMappingVirtualLayerDisabled(const FName& Id, int32 Index);
-};
-
-
-uint32 GetTypeHash(const FMultilayerProjector& Key);
-
diff --git a/Engine/Plugins/Mutable/Source/CustomizableObjectEditor/Private/MuCOE/CustomizableInstanceDetails.cpp b/Engine/Plugins/Mutable/Source/CustomizableObjectEditor/Private/MuCOE/CustomizableInstanceDetails.cpp
--- a/Engine/Plugins/Mutable/Source/CustomizableObjectEditor/Private/MuCOE/CustomizableInstanceDetails.cpp
+++ b/Engine/Plugins/Mutable/Source/CustomizableObjectEditor/Private/MuCOE/CustomizableInstanceDetails.cpp
@@ -815,10 +815,10 @@ bool FCustomizableInstanceDetails::IsVisible(int32 ParamIndexInObject)
 	FMutableParamUIMetadata UIMetadata = CustomizableObject->GetParameterUIMetadata(ParamName);
 	const FString* ParentName = UIMetadata.ExtraInformation.Find(FString("__ParentParamName"));
 
-	const bool IsAProjectorParam = ParamName.EndsWith(FMultilayerProjector::NUM_LAYERS_PARAMETER_POSTFIX)
-		|| (ParamName.EndsWith(FMultilayerProjector::IMAGE_PARAMETER_POSTFIX) && CustomizableObject->IsParameterMultidimensional(ParamIndexInObject))
-		|| (ParamName.EndsWith(FMultilayerProjector::OPACITY_PARAMETER_POSTFIX) && CustomizableObject->IsParameterMultidimensional(ParamIndexInObject))
-		|| (ParamName.EndsWith(FMultilayerProjector::POSE_PARAMETER_POSTFIX));
+	const bool IsAProjectorParam = ParamName.EndsWith(NUM_LAYERS_PARAMETER_POSTFIX)
+		|| (ParamName.EndsWith(IMAGE_PARAMETER_POSTFIX) && CustomizableObject->IsParameterMultidimensional(ParamIndexInObject))
+		|| (ParamName.EndsWith(OPACITY_PARAMETER_POSTFIX) && CustomizableObject->IsParameterMultidimensional(ParamIndexInObject))
+		|| (ParamName.EndsWith(POSE_PARAMETER_POSTFIX));
 
 	if (!IsAProjectorParam && ParentName && CustomInstance->GetPrivate()->bShowOnlyRelevantParameters)
 	{
@@ -1598,7 +1598,7 @@ TSharedRef<SWidget> FCustomizableInstanceDetails::GenerateMultidimensionalProjec
 	check(ProjectorParamIndex < ProjectorParameters.Num());
 
 	// Selected Pose UI
-	const FString PoseSwitchEnumParamName = ParamName + FMultilayerProjector::POSE_PARAMETER_POSTFIX;
+	const FString PoseSwitchEnumParamName = ParamName + POSE_PARAMETER_POSTFIX;
 	const int32 PoseSwitchEnumParamIndexInObject = CustomizableObject.FindParameter(PoseSwitchEnumParamName);
 	TSharedPtr<SVerticalBox> ProjectorBox = SNew(SVerticalBox);
 
@@ -1681,8 +1681,8 @@ TSharedRef<SWidget> FCustomizableInstanceDetails::GenerateMultidimensionalProjec
 		];
 	}
 
-	const FString TextureSwitchEnumParamName = ParamName + FMultilayerProjector::IMAGE_PARAMETER_POSTFIX;
-	const FString OpacitySliderParamName = ParamName + FMultilayerProjector::OPACITY_PARAMETER_POSTFIX;
+	const FString TextureSwitchEnumParamName = ParamName + IMAGE_PARAMETER_POSTFIX;
+	const FString OpacitySliderParamName = ParamName + OPACITY_PARAMETER_POSTFIX;
 
 	IDetailGroup* ProjectorGroup = *ParentsGroups.Find(ParamName);
 
diff --git a/Engine/Plugins/Mutable/Source/CustomizableObjectEditor/Private/MuCOE/GenerateMutableSource/GenerateMutableSourceGroupProjector.cpp b/Engine/Plugins/Mutable/Source/CustomizableObjectEditor/Private/MuCOE/GenerateMutableSource/GenerateMutableSourceGroupProjector.cpp
--- a/Engine/Plugins/Mutable/Source/CustomizableObjectEditor/Private/MuCOE/GenerateMutableSource/GenerateMutableSourceGroupProjector.cpp
+++ b/Engine/Plugins/Mutable/Source/CustomizableObjectEditor/Private/MuCOE/GenerateMutableSource/GenerateMutableSourceGroupProjector.cpp
@@ -4,6 +4,7 @@
 
 #include "GenerateMutableSourceImage.h"
 #include "Materials/MaterialInterface.h"
+#include "MuCO/CustomizableObjectInstancePrivate.h"
 #include "MuCO/MultilayerProjector.h"
 #include "MuCOE/CustomizableObjectCompiler.h"
 #include "MuCOE/GraphTraversal.h"
@@ -359,7 +360,7 @@ bool GenerateMutableSourceGroupProjector(const UEdGraphPin* Pin, FMutableGraphGe
 			GroupProjectorTempData.CustomizableObjectNodeGroupProjectorParameter = ProjParamNode;
 
 			mu::NodeScalarParameterPtr NodeScalarParam = new mu::NodeScalarParameter;
-			FString NodeScalarParamName = ProjParamNode->ParameterName + FMultilayerProjector::NUM_LAYERS_PARAMETER_POSTFIX;
+			FString NodeScalarParamName = ProjParamNode->ParameterName + NUM_LAYERS_PARAMETER_POSTFIX;
 			NodeScalarParam->SetName(NodeScalarParamName);
 			NodeScalarParam->SetUid(NumLayersParamUid.ToString());
 			GenerationContext.AddParameterNameUnique(originalGroup, NodeScalarParamName);
@@ -376,7 +377,7 @@ bool GenerateMutableSourceGroupProjector(const UEdGraphPin* Pin, FMutableGraphGe
 
 			
 			mu::NodeScalarEnumParameterPtr PoseEnumParameterNode = new mu::NodeScalarEnumParameter;
-			FString PoseNodeEnumParamName = ProjParamNode->ParameterName + FMultilayerProjector::POSE_PARAMETER_POSTFIX;
+			FString PoseNodeEnumParamName = ProjParamNode->ParameterName + POSE_PARAMETER_POSTFIX;
 			PoseEnumParameterNode->SetName(PoseNodeEnumParamName);
 			PoseEnumParameterNode->SetUid(SelectedPoseParamUid.ToString());
 			GenerationContext.AddParameterNameUnique(originalGroup, PoseNodeEnumParamName);
@@ -389,7 +390,7 @@ bool GenerateMutableSourceGroupProjector(const UEdGraphPin* Pin, FMutableGraphGe
 				FMutableParameterData(ProjParamNode->ParamUIMetadata, EMutableParameterType::Int));
 			
 			mu::NodeScalarParameterPtr OpacityParameterNode = new mu::NodeScalarParameter;
-			FString OpacityParameterNodeName = ProjParamNode->ParameterName + FMultilayerProjector::OPACITY_PARAMETER_POSTFIX;
+			FString OpacityParameterNodeName = ProjParamNode->ParameterName + OPACITY_PARAMETER_POSTFIX;
 			OpacityParameterNode->SetName(OpacityParameterNodeName);
 			OpacityParameterNode->SetUid(OpacityParamUid.ToString());
 			GenerationContext.AddParameterNameUnique(originalGroup, OpacityParameterNodeName);
@@ -449,7 +450,7 @@ bool GenerateMutableSourceGroupProjector(const UEdGraphPin* Pin, FMutableGraphGe
 			}
 		
 			mu::NodeScalarEnumParameterPtr EnumParameterNode = new mu::NodeScalarEnumParameter;
-			FString NodeEnumParamName = ProjParamNode->ParameterName + FMultilayerProjector::IMAGE_PARAMETER_POSTFIX;
+			FString NodeEnumParamName = ProjParamNode->ParameterName + IMAGE_PARAMETER_POSTFIX;
 			EnumParameterNode->SetName(NodeEnumParamName);
 			EnumParameterNode->SetUid(SelectedImageParamUid.ToString());
 			GenerationContext.AddParameterNameUnique(originalGroup, NodeEnumParamName);
