[mutable] Clamped min value for texture targets (Texture Interpolation Node)
- Target count clamped to 2 as "count<2" makes no sense in the context of this node

#rnx

diff --git a/Engine/Plugins/Mutable/Source/CustomizableObjectEditor/Public/MuCOE/Nodes/CustomizableObjectNodeTextureInterpolate.h b/Engine/Plugins/Mutable/Source/CustomizableObjectEditor/Public/MuCOE/Nodes/CustomizableObjectNodeTextureInterpolate.h
--- a/Engine/Plugins/Mutable/Source/CustomizableObjectEditor/Public/MuCOE/Nodes/CustomizableObjectNodeTextureInterpolate.h
+++ b/Engine/Plugins/Mutable/Source/CustomizableObjectEditor/Public/MuCOE/Nodes/CustomizableObjectNodeTextureInterpolate.h
@@ -57,7 +57,7 @@ public:
 		return Count;
 	}
 
-	UPROPERTY(EditAnywhere, Category=CustomizableObject)
+	UPROPERTY(EditAnywhere, Category=CustomizableObject, meta = (ClampMin = 2))
 	int32 NumTargets;
 };
 
