[mutable] Improve tag and morph selection popup menu with a hiearchy of nodes leading to each proposed option.
#rnx

diff --git a/Engine/Plugins/Mutable/Source/CustomizableObjectEditor/Private/MuCOE/Nodes/CustomizableObjectNodeModifierMorphMeshSectionDetails.cpp b/Engine/Plugins/Mutable/Source/CustomizableObjectEditor/Private/MuCOE/Nodes/CustomizableObjectNodeModifierMorphMeshSectionDetails.cpp
--- a/Engine/Plugins/Mutable/Source/CustomizableObjectEditor/Private/MuCOE/Nodes/CustomizableObjectNodeModifierMorphMeshSectionDetails.cpp
+++ b/Engine/Plugins/Mutable/Source/CustomizableObjectEditor/Private/MuCOE/Nodes/CustomizableObjectNodeModifierMorphMeshSectionDetails.cpp
@@ -5,7 +5,10 @@
 #include "MuCOE/CustomizableObjectEditorUtilities.h"
 #include "MuCOE/CustomizableObjectEditorStyle.h"
 #include "MuCOE/GraphTraversal.h"
+#include "MuCOE/Nodes/CustomizableObjectNodeObject.h"
+#include "MuCOE/Nodes/CustomizableObjectNodeObjectGroup.h"
 #include "MuCOE/Nodes/CustomizableObjectNodeMaterial.h"
+#include "MuCOE/Nodes/CustomizableObjectNodeModifierBase.h"
 #include "MuCOE/Nodes/CustomizableObjectNodeModifierMorphMeshSection.h"
 #include "MuCOE/Nodes/CustomizableObjectNodeModifierExtendMeshSection.h"
 #include "MuCOE/Nodes/CustomizableObjectNodeSkeletalMesh.h"
@@ -50,63 +53,7 @@ void FCustomizableObjectNodeModifierMorphMeshSectionDetails::CustomizeDetails( I
 	IDetailCategoryBuilder& MorphCategory = DetailBuilder.EditCategory( "Morph" );
 
 	// Scan for hint morph names
-	{
-		TArray<UCustomizableObjectNode*> CandidateNodes;
-		Node->GetPossiblyModifiedNodes(CandidateNodes);
-
-		// For now just show the first one
-		for (UCustomizableObjectNode* Candidate : CandidateNodes)
-		{
-			USkeletalMesh* SkeletalMesh = nullptr;
-
-			if (UCustomizableObjectNodeMaterialBase* MaterialNode = Cast<UCustomizableObjectNodeMaterialBase>(Candidate))
-			{
-				if (MaterialNode->OutputPin())
-				{
-					const UEdGraphPin* SourceMeshPin = FindMeshBaseSource(*MaterialNode->OutputPin(), false);
-					if (SourceMeshPin)
-					{
-						UCustomizableObjectNodeSkeletalMesh* SkeletalNode = Cast<UCustomizableObjectNodeSkeletalMesh>(SourceMeshPin->GetOwningNode());
-						if (SkeletalNode)
-						{
-							SkeletalMesh = SkeletalNode->SkeletalMesh;
-						}
-					}
-				}
-			}
-
-			else if (UCustomizableObjectNodeModifierExtendMeshSection* ExtendNode = Cast<UCustomizableObjectNodeModifierExtendMeshSection>(Candidate))
-			{
-				if (ExtendNode->OutputPin())
-				{
-					const UEdGraphPin* SourceMeshPin = FindMeshBaseSource(*ExtendNode->OutputPin(), false);
-					if (SourceMeshPin)
-					{
-						UCustomizableObjectNodeSkeletalMesh* SkeletalNode = Cast<UCustomizableObjectNodeSkeletalMesh>(SourceMeshPin->GetOwningNode());
-						if (SkeletalNode)
-						{
-							SkeletalMesh = SkeletalNode->SkeletalMesh;
-						}
-					}
-				}
-			}
-
-			if (SkeletalMesh)
-			{
-				const TArray<TObjectPtr<UMorphTarget>>& Morphs = SkeletalMesh->GetMorphTargets();
-				for (TObjectPtr<UMorphTarget> Morph : Morphs)
-				{
-					if (Morph)
-					{
-						FString MorphTargetName = Morph->GetName();
-						MorphOptionsSource.AddUnique(MakeShared<FString>(MorphTargetName));
-					}
-				}
-			}
-		}
-
-		// TODO: Add all morphs if no candidate is found? Add both in different menu sections?
-	}
+	RefreshMorphOptions();
 
 	TSharedRef<IPropertyHandle> MorphTargetNameProperty = DetailBuilder.GetProperty("MorphTargetName");
 	MorphCategory.AddCustomRow(LOCTEXT("MorphMaterialDetails_Target", "Target"))
@@ -163,5 +110,215 @@ void FCustomizableObjectNodeModifierMorphMeshSectionDetails::OnMorphTargetComboB
 }
 
 
+void FCustomizableObjectNodeModifierMorphMeshSectionDetails::OnRequiredTagsPropertyChanged()
+{
+	FCustomizableObjectNodeModifierBaseDetails::OnRequiredTagsPropertyChanged();
+	RefreshMorphOptions();
+}
+
+
+void FCustomizableObjectNodeModifierMorphMeshSectionDetails::RefreshMorphOptions()
+{
+	MorphOptionsSource.SetNum(0, EAllowShrinking::No);
+
+	TArray<UCustomizableObjectNode*> CandidateNodes;
+	Node->GetPossiblyModifiedNodes(CandidateNodes);
+
+	TMap<UEdGraphNode*, TSharedPtr<SMutableSearchComboBox::FFilteredOption>> AddedOptions;
+
+	for (UCustomizableObjectNode* Candidate : CandidateNodes)
+	{
+		AddMorphsFromNode(Candidate, AddedOptions);
+	}
+
+	// Add all morphs if no candidate is found
+	if (MorphOptionsSource.IsEmpty())
+	{
+		UCustomizableObject* ThisNodeObject = GetRootObject(*Node);
+		UCustomizableObject* RootObject = GraphTraversal::GetRootObject(ThisNodeObject);
+
+		TSet<UCustomizableObject*> AllCustomizableObject;
+		GetAllObjectsInGraph(RootObject, AllCustomizableObject);
+
+		for (const UCustomizableObject* CustObject : AllCustomizableObject)
+		{
+			if (!CustObject)
+			{
+				continue;
+			}
+
+			for (const TObjectPtr<UEdGraphNode>& Candidate : CustObject->GetPrivate()->GetSource()->Nodes)
+			{
+				AddMorphsFromNode(Candidate, AddedOptions);
+			}
+		}
+	}
+}
+
+
+TSharedPtr<SMutableSearchComboBox::FFilteredOption> FCustomizableObjectNodeModifierMorphMeshSectionDetails::AddNodeHierarchyOptions(UEdGraphNode* InNode, TMap<UEdGraphNode*, TSharedPtr<SMutableSearchComboBox::FFilteredOption>>& AddedOptions)
+{
+	TSharedPtr<SMutableSearchComboBox::FFilteredOption> Option;
+	if (TSharedPtr<SMutableSearchComboBox::FFilteredOption>* FoundCached = AddedOptions.Find(InNode))
+	{
+		Option = *FoundCached;
+	}
+
+	if (InNode && !Option)
+	{
+		// Add parents
+		TSharedPtr<SMutableSearchComboBox::FFilteredOption> ParentOption;
+		{
+			// Add this node as placeholder in the cache to prevent infinite loops because of graph loops.
+			AddedOptions.Add(InNode, Option);
+
+			// Pin traversal
+			for (const UEdGraphPin* Pin : InNode->Pins)
+			{
+				if (Pin && Pin->Direction == EEdGraphPinDirection::EGPD_Output
+					&& !Pin->LinkedTo.IsEmpty() && Pin->LinkedTo[0])
+				{
+					UEdGraphNode* ParentNode = Pin->LinkedTo[0]->GetOwningNode();
+
+					ParentOption = AddNodeHierarchyOptions(ParentNode, AddedOptions);
+
+					// We are ok with just one parent
+					if (ParentOption)
+					{
+						break;
+					}
+				}
+			}
+
+			// Node internal references
+			if (!ParentOption)
+			{
+				// Is it an object referencing an external group?
+				if (UCustomizableObjectNodeObject* ObjectNode = Cast<UCustomizableObjectNodeObject>(InNode))
+				{
+					if (ObjectNode->ParentObject)
+					{
+						UEdGraphNode* ExternalParentNode = GetCustomizableObjectExternalNode<UEdGraphNode>(ObjectNode->ParentObject, ObjectNode->ParentObjectGroupId);
+						ParentOption = AddNodeHierarchyOptions(ExternalParentNode, AddedOptions);
+					}
+				}
+			}
+
+			// TODO: Support import/export nodes
+		}
+
+		// Is it a relevant type that we want to show in the hierarchy?
+		if (UCustomizableObjectNodeMaterial* MeshSectionNode = Cast<UCustomizableObjectNodeMaterial>(InNode))
+		{
+			Option = MakeShared<SMutableSearchComboBox::FFilteredOption>();
+			Option->Parent = ParentOption;
+			UMaterialInterface* Material = MeshSectionNode->GetMaterial();
+			Option->DisplayOption = FString::Printf(TEXT("Mesh Section [%s]"), Material ? *Material->GetName() : TEXT("no-material"));
+			MorphOptionsSource.Add(Option.ToSharedRef());
+		}
+
+		else if (UCustomizableObjectNodeObject* ObjectNode = Cast<UCustomizableObjectNodeObject>(InNode))
+		{
+			Option = MakeShared<SMutableSearchComboBox::FFilteredOption>();
+			Option->Parent = ParentOption;
+			Option->DisplayOption = ObjectNode->ObjectName;
+			if (Option->DisplayOption.IsEmpty())
+			{
+				Option->DisplayOption = "Unnamed Object";
+			}
+			MorphOptionsSource.Add(Option.ToSharedRef());
+		}
+
+		else if (UCustomizableObjectNodeObjectGroup* GroupNode = Cast<UCustomizableObjectNodeObjectGroup>(InNode))
+		{
+			Option = MakeShared<SMutableSearchComboBox::FFilteredOption>();
+			Option->Parent = ParentOption;
+			Option->DisplayOption = GroupNode->GroupName;
+			if (Option->DisplayOption.IsEmpty())
+			{
+				Option->DisplayOption = "Unnamed Group";
+			}
+			MorphOptionsSource.Add(Option.ToSharedRef());
+		}
+
+		else if (UCustomizableObjectNodeModifierBase* ModifierNode = Cast<UCustomizableObjectNodeModifierBase>(InNode))
+		{
+			Option = MakeShared<SMutableSearchComboBox::FFilteredOption>();
+			Option->Parent = ParentOption;
+			Option->DisplayOption = ModifierNode->GetNodeTitle(ENodeTitleType::ListView).ToString();
+			MorphOptionsSource.Add(Option.ToSharedRef());
+		}
+
+		// Overwrite in to cache, to prevent loops
+		AddedOptions.Add(InNode, Option);
+
+		// If this node wasn't of interest, maybe its parent was.
+		if (!Option)
+		{
+			Option = ParentOption;
+		}
+	}
+
+	return Option;
+}
+
+
+void FCustomizableObjectNodeModifierMorphMeshSectionDetails::AddMorphsFromNode( UEdGraphNode* Candidate, TMap<UEdGraphNode*, TSharedPtr<SMutableSearchComboBox::FFilteredOption>>& AddedOptions)
+{
+	USkeletalMesh* SkeletalMesh = nullptr;
+
+	if (UCustomizableObjectNodeMaterialBase* MaterialNode = Cast<UCustomizableObjectNodeMaterialBase>(Candidate))
+	{
+		if (MaterialNode->OutputPin())
+		{
+			const UEdGraphPin* SourceMeshPin = FindMeshBaseSource(*MaterialNode->OutputPin(), false);
+			if (SourceMeshPin)
+			{
+				UCustomizableObjectNodeSkeletalMesh* SkeletalNode = Cast<UCustomizableObjectNodeSkeletalMesh>(SourceMeshPin->GetOwningNode());
+				if (SkeletalNode)
+				{
+					SkeletalMesh = SkeletalNode->SkeletalMesh;
+				}
+			}
+		}
+	}
+
+	else if (UCustomizableObjectNodeModifierExtendMeshSection* ExtendNode = Cast<UCustomizableObjectNodeModifierExtendMeshSection>(Candidate))
+	{
+		if (ExtendNode->OutputPin())
+		{
+			const UEdGraphPin* SourceMeshPin = FindMeshBaseSource(*ExtendNode->OutputPin(), false);
+			if (SourceMeshPin)
+			{
+				UCustomizableObjectNodeSkeletalMesh* SkeletalNode = Cast<UCustomizableObjectNodeSkeletalMesh>(SourceMeshPin->GetOwningNode());
+				if (SkeletalNode)
+				{
+					SkeletalMesh = SkeletalNode->SkeletalMesh;
+				}
+			}
+		}
+	}
+
+	if (SkeletalMesh)
+	{
+		const TArray<TObjectPtr<UMorphTarget>>& Morphs = SkeletalMesh->GetMorphTargets();
+		for (TObjectPtr<UMorphTarget> Morph : Morphs)
+		{
+			if (Morph)
+			{
+				TSharedPtr<SMutableSearchComboBox::FFilteredOption> NodeOption = AddNodeHierarchyOptions(Candidate, AddedOptions);
+
+				FString MorphTargetName = Morph->GetName();
+
+				TSharedRef<SMutableSearchComboBox::FFilteredOption> Option = MakeShared<SMutableSearchComboBox::FFilteredOption>();
+				Option->ActualOption = MorphTargetName;
+				Option->DisplayOption = MorphTargetName;
+				Option->Parent = NodeOption;
+				MorphOptionsSource.Add(Option);
+			}
+		}
+	}
+}
+
 
 #undef LOCTEXT_NAMESPACE
diff --git a/Engine/Plugins/Mutable/Source/CustomizableObjectEditor/Private/MuCOE/Nodes/CustomizableObjectNodeModifierMorphMeshSectionDetails.h b/Engine/Plugins/Mutable/Source/CustomizableObjectEditor/Private/MuCOE/Nodes/CustomizableObjectNodeModifierMorphMeshSectionDetails.h
--- a/Engine/Plugins/Mutable/Source/CustomizableObjectEditor/Private/MuCOE/Nodes/CustomizableObjectNodeModifierMorphMeshSectionDetails.h
+++ b/Engine/Plugins/Mutable/Source/CustomizableObjectEditor/Private/MuCOE/Nodes/CustomizableObjectNodeModifierMorphMeshSectionDetails.h
@@ -28,13 +28,27 @@ public:
 	// ILayoutDetails interface
 	void CustomizeDetails( IDetailLayoutBuilder& DetailBuilder ) override;
 
+protected:
+
+	// FCustomizableObjectNodeModifierBaseDetails interface
+	virtual void OnRequiredTagsPropertyChanged() override;
+
 private:
 
 	class UCustomizableObjectNodeModifierMorphMeshSection* Node = nullptr;
 
 	/** */
 	TSharedPtr<SMutableSearchComboBox> MorphCombo;
-	TArray< TSharedPtr<FString> > MorphOptionsSource;
+	TArray< TSharedRef<SMutableSearchComboBox::FFilteredOption> > MorphOptionsSource;
 	void OnMorphTargetComboBoxSelectionChanged(const FText& NewText);
 
+	/** */
+	void AddMorphsFromNode(UEdGraphNode* Candidate, TMap<UEdGraphNode*, TSharedPtr<SMutableSearchComboBox::FFilteredOption>>& AddedOptions);
+		
+	/** */
+	TSharedPtr<SMutableSearchComboBox::FFilteredOption> AddNodeHierarchyOptions(UEdGraphNode* Node, TMap<UEdGraphNode*, TSharedPtr<SMutableSearchComboBox::FFilteredOption>>& AddedOptions);
+
+	/** */
+	void RefreshMorphOptions();
+
 };
diff --git a/Engine/Plugins/Mutable/Source/CustomizableObjectEditor/Private/MuCOE/SMutableSearchComboBox.cpp b/Engine/Plugins/Mutable/Source/CustomizableObjectEditor/Private/MuCOE/SMutableSearchComboBox.cpp
--- a/Engine/Plugins/Mutable/Source/CustomizableObjectEditor/Private/MuCOE/SMutableSearchComboBox.cpp
+++ b/Engine/Plugins/Mutable/Source/CustomizableObjectEditor/Private/MuCOE/SMutableSearchComboBox.cpp
@@ -36,7 +36,6 @@ void SMutableSearchComboBox::Construct(const FArguments& InArgs)
 	const FButtonStyle* const OurButtonStyle = InArgs._ButtonStyle ? InArgs._ButtonStyle : &OurComboButtonStyle.ButtonStyle;
 
 	this->OnSelectionChanged = InArgs._OnSelectionChanged;
-	this->OnGenerateWidget = InArgs._OnGenerateWidget;
 
 	OptionsSource = InArgs._OptionsSource;
 	RefreshOptions();
@@ -62,9 +61,10 @@ void SMutableSearchComboBox::Construct(const FArguments& InArgs)
 
 				+ SVerticalBox::Slot()
 				[
-					SAssignNew(this->ComboListView, SComboListType)
-						.ListItemsSource(&FilteredOptionsSource)
+					SAssignNew(this->ComboTreeView, SComboTreeType)
+						.TreeItemsSource(&FilteredRootOptionsSource)
 						.OnGenerateRow(this, &SMutableSearchComboBox::GenerateMenuItemRow)
+						.OnGetChildren(this, &SMutableSearchComboBox::OnGetChildren)
 						.OnSelectionChanged(this, &SMutableSearchComboBox::OnSelectionChanged_Internal)
 						.OnKeyDownHandler(this, &SMutableSearchComboBox::OnKeyDownHandler)
 						.SelectionMode(ESelectionMode::Single)
@@ -103,7 +103,7 @@ void SMutableSearchComboBox::Construct(const FArguments& InArgs)
 	}
 	else
 	{
-		SetMenuContentWidgetToFocus(ComboListView);
+		SetMenuContentWidgetToFocus(ComboTreeView);
 	}
 }
 
@@ -115,28 +115,19 @@ void SMutableSearchComboBox::RefreshOptions()
 
 	if (SearchText.IsEmpty())
 	{
-		for (const TSharedPtr<FString>& Option : *OptionsSource)
-		{
-			if (Option)
-			{
-				FFilteredOption Data;
-				Data.ActualOption = *Option;
-				Data.DisplayOption = *Option;
-				FilteredOptionsSource.Add(MakeShared<FFilteredOption>(Data));
-			}
-		}
+		FilteredOptionsSource = *OptionsSource;
 	}
 	else
 	{
 		TArray<FString> SearchTokens;
 		SearchText.ToString().ParseIntoArrayWS(SearchTokens);
 
-		for (const TSharedPtr<FString>& Option : *OptionsSource)
+		for (const TSharedRef<FFilteredOption>& Option : *OptionsSource)
 		{
 			bool bAllTokensMatch = true;
 			for (const FString& SearchToken : SearchTokens)
 			{
-				if (Option->Find(SearchToken, ESearchCase::Type::IgnoreCase) == INDEX_NONE)
+				if (Option->DisplayOption.Find(SearchToken, ESearchCase::Type::IgnoreCase) == INDEX_NONE)
 				{
 					bAllTokensMatch = false;
 					break;
@@ -145,62 +136,90 @@ void SMutableSearchComboBox::RefreshOptions()
 
 			if (bAllTokensMatch)
 			{
-				FFilteredOption Data;
-				Data.ActualOption = *Option;
-				Data.DisplayOption = *Option;
-				FilteredOptionsSource.Add(MakeShared<FFilteredOption>(Data));
+				FilteredOptionsSource.Add(Option);
 			}
 		}
 
 		bool bFullMatch = false;
 		FString SearchString = SearchText.ToString();
-		for (const TSharedPtr<FString>& Option : *OptionsSource)		
+		for (const TSharedPtr<FFilteredOption>& Option : *OptionsSource)
 		{
-			if (*Option == SearchString)
+			if (Option->DisplayOption == SearchString)
 			{
 				bFullMatch = true; 
 				break;
 			}
 		}
 
-		if ( bAllowAddNewOptions && !SearchText.IsEmpty() )
+		if ( bAllowAddNewOptions && !bFullMatch && !SearchText.IsEmpty() )
 		{
 			FFilteredOption Data;
 			Data.ActualOption = SearchString;
-			Data.DisplayOption = FString::Printf(TEXT("Add new (%s)"), *SearchString);
-			FilteredOptionsSource.Add(MakeShared<FFilteredOption>(Data));
+			Data.DisplayOption = FString::Printf(TEXT("Add new tag (%s)"), *SearchString);
+			FilteredOptionsSource.Insert(MakeShared<FFilteredOption>(Data),0);
+		}
+
+		// Ensure filtered options parents are added
+		for (int32 OptionIndex = 0; OptionIndex < FilteredOptionsSource.Num(); ++OptionIndex)
+		{
+			TSharedPtr<FFilteredOption> Parent = FilteredOptionsSource[OptionIndex]->Parent;
+			if (Parent)
+			{
+				FilteredOptionsSource.AddUnique(Parent.ToSharedRef());
+
+				while (Parent)
+				{
+					ComboTreeView->SetItemExpansion(Parent.ToSharedRef(), true);
+					Parent = Parent->Parent;
+				}
+			}
 		}
 	}
 
-	if (ComboListView)
+	FilteredRootOptionsSource.SetNum(0, EAllowShrinking::No);
+	FilteredRootOptionsSource.Reserve(FilteredOptionsSource.Num());
+	for (const TSharedRef<FFilteredOption>& Option : FilteredOptionsSource)
 	{
-		ComboListView->RequestListRefresh();
+		if (!Option->Parent.IsValid())
+		{
+			FilteredRootOptionsSource.Add(Option);
+		}
 	}
+
+	if (ComboTreeView)
+	{
+		ComboTreeView->RequestTreeRefresh();
+	}
+}
+
+
+TSharedRef<ITableRow> SMutableSearchComboBox::GenerateMenuItemRow(TSharedRef<FFilteredOption> InItem, const TSharedRef<STableViewBase>& OwnerTable)
+{
+	FSlateColor LabelColor = InItem->ActualOption.IsEmpty() 
+		? FSlateColor::UseSubduedForeground() 
+		: FSlateColor::UseForeground();
+
+	return SNew(SComboRow<TSharedPtr<FString>>, OwnerTable)
+		[
+			SNew(STextBlock)
+				.Text(FText::FromString(InItem->DisplayOption))
+				.Font(IDetailLayoutBuilder::GetDetailFont())
+				.ColorAndOpacity(LabelColor)
+		];
 }
 
 
-TSharedRef<ITableRow> SMutableSearchComboBox::GenerateMenuItemRow(TSharedPtr<FFilteredOption> InItem, const TSharedRef<STableViewBase>& OwnerTable)
+void SMutableSearchComboBox::OnGetChildren(TSharedRef<FFilteredOption> InItem, TArray<TSharedRef<FFilteredOption>>& OutChildren)
 {
-	// TODO
-	//if (OnGenerateWidget.IsBound())
-	//{
-	//	return SNew(SComboRow<TSharedPtr<FString>>, OwnerTable)
-	//		.Style(ItemStyle)
-	//		.Padding(MenuRowPadding)
-	//		[
-	//			OnGenerateWidget.Execute(InItem)
-	//		];
-	//}
-	//else
+	// TODO: Optimize if necessary
+	for (const TSharedRef<FFilteredOption>& Option : FilteredOptionsSource)
 	{
-		// Just add the text
-		return SNew(SComboRow<TSharedPtr<FString>>, OwnerTable)
-			[
-				SNew(STextBlock)
-					.Text(FText::FromString(InItem ? InItem->DisplayOption : FString(TEXT("Invalid item."))))
-					.Font(IDetailLayoutBuilder::GetDetailFont())
-			];
+		if (Option->Parent==InItem)
+		{
+			OutChildren.Add(Option);
+		}
 	}
+
 }
 
 
@@ -212,7 +231,7 @@ void SMutableSearchComboBox::OnMenuOpenChanged(bool bOpen)
 		FSlateApplication::Get().ForEachUser([this](FSlateUser& User)
 			{
 				TSharedRef<SWidget> ThisRef = this->AsShared();
-				if (User.IsWidgetInFocusPath(this->ComboListView))
+				if (User.IsWidgetInFocusPath(this->ComboTreeView))
 				{
 					User.SetFocus(ThisRef);
 				}
@@ -229,10 +248,16 @@ void SMutableSearchComboBox::OnSelectionChanged_Internal(TSharedPtr<FFilteredOpt
 		return;
 	}
 
-	// close combo as long as the selection wasn't from navigation
+	// If the proposed selection is not a valid element (it is a hierarchy label)
+	if (ProposedSelection->ActualOption.IsEmpty())
+	{
+		return;
+	}
+
+	// Close combo as long as the selection wasn't from navigation
 	if (SelectInfo != ESelectInfo::OnNavigation)
 	{		
-		OnSelectionChanged.ExecuteIfBound(ProposedSelection ? FText::FromString(ProposedSelection->ActualOption) : FText());
+		OnSelectionChanged.ExecuteIfBound(FText::FromString(ProposedSelection->ActualOption));
 		this->SetIsOpen(false);
 	}
 }
@@ -250,7 +275,7 @@ void SMutableSearchComboBox::OnSearchTextCommitted(const FText& InText, ETextCom
 {
 	if ((InCommitType == ETextCommit::Type::OnEnter) && FilteredOptionsSource.Num() > 0)
 	{
-		ComboListView->SetSelection(FilteredOptionsSource[0], ESelectInfo::OnKeyPress);
+		ComboTreeView->SetSelection(FilteredOptionsSource[0], ESelectInfo::OnKeyPress);
 	}
 }
 
@@ -260,7 +285,7 @@ FReply SMutableSearchComboBox::OnKeyDownHandler(const FGeometry& MyGeometry, con
 	if (InKeyEvent.GetKey() == EKeys::Enter)
 	{
 		// Select the first selected item on hitting enter
-		TArray<TSharedPtr<FFilteredOption>> SelectedItems = ComboListView->GetSelectedItems();
+		TArray<TSharedRef<FFilteredOption>> SelectedItems = ComboTreeView->GetSelectedItems();
 		if (SelectedItems.Num() > 0)
 		{
 			OnSelectionChanged_Internal(SelectedItems[0], ESelectInfo::OnKeyPress);
diff --git a/Engine/Plugins/Mutable/Source/CustomizableObjectEditor/Private/MuCOE/SMutableSearchComboBox.h b/Engine/Plugins/Mutable/Source/CustomizableObjectEditor/Private/MuCOE/SMutableSearchComboBox.h
--- a/Engine/Plugins/Mutable/Source/CustomizableObjectEditor/Private/MuCOE/SMutableSearchComboBox.h
+++ b/Engine/Plugins/Mutable/Source/CustomizableObjectEditor/Private/MuCOE/SMutableSearchComboBox.h
@@ -5,7 +5,7 @@
 #include "Types/SlateEnums.h"
 #include "Widgets/Input/SComboButton.h"
 #include "Widgets/Input/SEditableTextBox.h"
-#include "Widgets/Views/SListView.h"
+#include "Widgets/Views/STreeView.h"
 #include "Widgets/Views/ITableRow.h"
 
 namespace ESelectInfo { enum Type : int; }
@@ -23,13 +23,15 @@ public:
 	{
 		FString ActualOption;
 		FString DisplayOption;
+
+		TSharedPtr<FFilteredOption> Parent;
 	};
 
 	/** Type of list used for showing menu options. */
-	typedef SListView< TSharedPtr<FFilteredOption> > SComboListType;
+	typedef STreeView< TSharedRef<FFilteredOption> > SComboTreeType;
 
 	/** Delegate type used to generate widgets that represent Options */
-	typedef typename TSlateDelegates< TSharedPtr<FString> >::FOnGenerateWidget FOnGenerateWidget;
+	typedef typename TSlateDelegates< TSharedRef<FString> >::FOnGenerateWidget FOnGenerateWidget;
 
 	SLATE_BEGIN_ARGS(SMutableSearchComboBox)
 		: _Content()
@@ -40,7 +42,6 @@ public:
 		, _ForegroundColor(FSlateColor::UseStyle())
 		, _OptionsSource()
 		, _OnSelectionChanged()
-		, _OnGenerateWidget()
 		, _MenuButtonBrush(nullptr)
 		, _Method()
 		, _AllowAddNewOptions(false)
@@ -61,9 +62,8 @@ public:
 		SLATE_ATTRIBUTE(FMargin, ContentPadding)
 		SLATE_ATTRIBUTE(FSlateColor, ForegroundColor)
 
-		SLATE_ARGUMENT(const TArray< TSharedPtr<FString> >*, OptionsSource)
+		SLATE_ARGUMENT(const TArray< TSharedRef<FFilteredOption> >*, OptionsSource)
 		SLATE_EVENT(FOnTextChanged, OnSelectionChanged)
-		SLATE_EVENT(FOnGenerateWidget, OnGenerateWidget)
 
 		SLATE_ARGUMENT(const FSlateBrush*, MenuButtonBrush)
 
@@ -94,7 +94,10 @@ public:
 private:
 
 	/** Generate a row for the InItem in the combo box's list (passed in as OwnerTable). Do this by calling the user-specified OnGenerateWidget */
-	TSharedRef<ITableRow> GenerateMenuItemRow(TSharedPtr<FFilteredOption> InItem, const TSharedRef<STableViewBase>& OwnerTable);
+	TSharedRef<ITableRow> GenerateMenuItemRow(TSharedRef<FFilteredOption> InItem, const TSharedRef<STableViewBase>& OwnerTable);
+
+	/**  */
+	void OnGetChildren(TSharedRef<FFilteredOption> InItem, TArray<TSharedRef<FFilteredOption>>& OutChildren);
 
 	/** Called if the menu is closed */
 	void OnMenuOpenChanged(bool bOpen);
@@ -122,20 +125,19 @@ private:
 	/** The search field used for the combox box's contents */
 	TSharedPtr<SEditableTextBox> SearchField;
 	/** The ListView that we pop up; visualized the available options. */
-	TSharedPtr< SComboListType > ComboListView;
+	TSharedPtr< SComboTreeType > ComboTreeView;
 	/** The Scrollbar used in the ListView. */
 	TSharedPtr< SScrollBar > CustomScrollbar;
-	/** Delegate to invoke when we need to visualize an option as a widget. */
-	FOnGenerateWidget OnGenerateWidget;
 
 	/** Updated whenever search text is changed */
 	FText SearchText;
 
 	/** Source data for this combo box */
-	const TArray< TSharedPtr<FString> >* OptionsSource;
+	const TArray< TSharedRef<FFilteredOption> >* OptionsSource;
 
 	/** Filtered list that is actually displayed */
-	TArray< TSharedPtr<FFilteredOption> > FilteredOptionsSource;
+	TArray< TSharedRef<FFilteredOption> > FilteredOptionsSource;
+	TArray< TSharedRef<FFilteredOption> > FilteredRootOptionsSource;
 
 	/** Copied to replace the image. */
 	FComboButtonStyle OurComboButtonStyle;
diff --git a/Engine/Plugins/Mutable/Source/CustomizableObjectEditor/Private/MuCOE/SMutableTagListWidget.cpp b/Engine/Plugins/Mutable/Source/CustomizableObjectEditor/Private/MuCOE/SMutableTagListWidget.cpp
--- a/Engine/Plugins/Mutable/Source/CustomizableObjectEditor/Private/MuCOE/SMutableTagListWidget.cpp
+++ b/Engine/Plugins/Mutable/Source/CustomizableObjectEditor/Private/MuCOE/SMutableTagListWidget.cpp
@@ -6,6 +6,9 @@
 #include "MuCOE/CustomizableObjectEditorStyle.h"
 #include "MuCOE/GraphTraversal.h"
 #include "MuCOE/Nodes/CustomizableObjectNode.h"
+#include "MuCOE/Nodes/CustomizableObjectNodeObject.h"
+#include "MuCOE/Nodes/CustomizableObjectNodeObjectGroup.h"
+#include "MuCOE/Nodes/CustomizableObjectNodeMaterial.h"
 #include "MuCOE/Nodes/CustomizableObjectNodeModifierBase.h"
 #include "MuCOE/UnrealEditorPortabilityHelpers.h"
 #include "PropertyCustomizationHelpers.h"
@@ -17,11 +20,9 @@
 #include "Engine/SkeletalMesh.h"
 #include "Misc/Attribute.h"
 
-
 #define LOCTEXT_NAMESPACE "CustomizableObjectDetails"
 
 
-
 void SMutableTagListWidget::Construct(const FArguments& InArgs)
 {
 	Node = InArgs._Node;
@@ -64,7 +65,7 @@ void SMutableTagListWidget::Construct(const FArguments& InArgs)
 		[
 			SAssignNew(this->TagListWidget, SListView<TSharedPtr<FTagUIData>>)
 				.ListItemsSource(&CurrentTagsSource)
-				.OnGenerateRow(this, &SMutableTagListWidget::GenerateTagMenuItemRow)
+				.OnGenerateRow(this, &SMutableTagListWidget::GenerateTagListItemRow)
 				.SelectionMode(ESelectionMode::None)
 		]
 
@@ -86,9 +87,116 @@ void SMutableTagListWidget::Construct(const FArguments& InArgs)
 }
 
 
+TSharedPtr<SMutableSearchComboBox::FFilteredOption> SMutableTagListWidget::AddNodeHierarchyOptions(UEdGraphNode* InNode, TMap<UEdGraphNode*, TSharedPtr<SMutableSearchComboBox::FFilteredOption>>& AddedOptions )
+{
+	TSharedPtr<SMutableSearchComboBox::FFilteredOption> Option;
+	if (TSharedPtr<SMutableSearchComboBox::FFilteredOption>* FoundCached = AddedOptions.Find(InNode))
+	{
+		Option = *FoundCached;
+	}
+
+	if (InNode && !Option)
+	{
+		// Add parents
+		TSharedPtr<SMutableSearchComboBox::FFilteredOption> ParentOption;
+		{
+			// Add this node as placeholder in the cache to prevent infinite loops because of graph loops.
+			AddedOptions.Add(InNode, Option);
+
+			// Pin traversal
+			for (const UEdGraphPin* Pin : InNode->Pins)
+			{
+				if (Pin && Pin->Direction == EEdGraphPinDirection::EGPD_Output
+					&& !Pin->LinkedTo.IsEmpty() && Pin->LinkedTo[0])
+				{
+					UEdGraphNode* ParentNode = Pin->LinkedTo[0]->GetOwningNode();
+
+					ParentOption = AddNodeHierarchyOptions(ParentNode, AddedOptions);
+
+					// We are ok with just one parent
+					if (ParentOption)
+					{
+						break;
+					}
+				}
+			}
+
+			// Node internal references
+			if (!ParentOption)
+			{
+				// Is it an object referencing an external group?
+				if (UCustomizableObjectNodeObject* ObjectNode = Cast<UCustomizableObjectNodeObject>(InNode))
+				{
+					if (ObjectNode->ParentObject)
+					{
+						UEdGraphNode* ExternalParentNode = GetCustomizableObjectExternalNode<UEdGraphNode>(ObjectNode->ParentObject, ObjectNode->ParentObjectGroupId);
+						ParentOption = AddNodeHierarchyOptions(ExternalParentNode, AddedOptions);
+					}
+				}
+			}
+
+			// TODO: Support import/export nodes
+		}
+
+		// Is it a relevant type that we want to show in the hierarchy?
+		if (UCustomizableObjectNodeMaterial* MeshSectionNode = Cast<UCustomizableObjectNodeMaterial>(InNode))
+		{
+			Option = MakeShared<SMutableSearchComboBox::FFilteredOption>();
+			Option->Parent = ParentOption;
+			UMaterialInterface* Material = MeshSectionNode->GetMaterial();
+			Option->DisplayOption = FString::Printf(TEXT("Mesh Section [%s]"), Material?*Material->GetName():TEXT("no-material"));
+			TagComboOptionsSource.Add(Option.ToSharedRef());
+		}
+
+		else if (UCustomizableObjectNodeObject* ObjectNode = Cast<UCustomizableObjectNodeObject>(InNode))
+		{
+			Option = MakeShared<SMutableSearchComboBox::FFilteredOption>();
+			Option->Parent = ParentOption;
+			Option->DisplayOption = ObjectNode->ObjectName;
+			if (Option->DisplayOption.IsEmpty())
+			{
+				Option->DisplayOption = "Unnamed Object";
+			}
+			TagComboOptionsSource.Add(Option.ToSharedRef());
+		}
+
+		else if (UCustomizableObjectNodeObjectGroup* GroupNode = Cast<UCustomizableObjectNodeObjectGroup>(InNode))
+		{
+			Option = MakeShared<SMutableSearchComboBox::FFilteredOption>();
+			Option->Parent = ParentOption;
+			Option->DisplayOption = GroupNode->GroupName;
+			if (Option->DisplayOption.IsEmpty())
+			{
+				Option->DisplayOption = "Unnamed Group";
+			}
+			TagComboOptionsSource.Add(Option.ToSharedRef());
+		}
+
+		else if (UCustomizableObjectNodeModifierBase* ModifierNode = Cast<UCustomizableObjectNodeModifierBase>(InNode))
+		{
+			Option = MakeShared<SMutableSearchComboBox::FFilteredOption>();
+			Option->Parent = ParentOption;
+			Option->DisplayOption = ModifierNode->GetNodeTitle(ENodeTitleType::ListView).ToString();
+			TagComboOptionsSource.Add(Option.ToSharedRef());
+		}
+
+		// Overwrite in to cache, to prevent loops
+		AddedOptions.Add(InNode, Option);
+
+		// If this node wasn't of interest, maybe its parent was.
+		if (!Option)
+		{
+			Option = ParentOption;
+		}
+	}
+
+	return Option;
+}
+
+
 void SMutableTagListWidget::RefreshOptions()
 {
-	// Tag combo options
+	// Collect all possible tags for the "add" combo menu
 	{
 		TagComboOptionsSource.SetNum(0, EAllowShrinking::No);
 
@@ -99,51 +207,56 @@ void SMutableTagListWidget::RefreshOptions()
 		TSet<UCustomizableObject*> AllCustomizableObject;
 		GetAllObjectsInGraph(RootObject, AllCustomizableObject);
 
+		TMap<UEdGraphNode*, TSharedPtr<SMutableSearchComboBox::FFilteredOption>> AddedOptions;
+
 		for (const UCustomizableObject* CustObject : AllCustomizableObject)
 		{
-			if (CustObject)
+			if (!CustObject)
+			{
+				continue;
+			}
+
+
+			for (const TObjectPtr<UEdGraphNode>& CandidateNode : CustObject->GetPrivate()->GetSource()->Nodes)
 			{
-				for (const TObjectPtr<UEdGraphNode>& CandidateNode : CustObject->GetPrivate()->GetSource()->Nodes)
+				UCustomizableObjectNode* Typed = Cast<UCustomizableObjectNode>(CandidateNode);
+				if (!Typed)
+				{
+					continue;
+				}
+
+				TArray<FString>* EnableTags = Typed->GetEnableTags();
+				if (EnableTags)
 				{
-					UCustomizableObjectNode* Typed = Cast<UCustomizableObjectNode>(CandidateNode);
-					if (Typed)
+					TSharedPtr<SMutableSearchComboBox::FFilteredOption> NodeOption = AddNodeHierarchyOptions(Typed,AddedOptions);
+
+					for (const FString& OneTag : *EnableTags)
 					{
-						TArray<FString>* EnableTags = Typed->GetEnableTags();
-						if (EnableTags)
+						if (!OneTag.IsEmpty())
 						{
-							for (const FString& OneTag : *EnableTags)
-							{
-								if (!OneTag.IsEmpty())
-								{
-									bool bContained = TagComboOptionsSource.ContainsByPredicate([&](const TSharedPtr<FString>& Candidate) 
-										{
-											return Candidate.IsValid() && *Candidate == OneTag;
-										});
-									if (!bContained)
-									{
-										TagComboOptionsSource.Add(MakeShared<FString>(OneTag));
-									}
-								}
-							}
+							TSharedRef Option = MakeShared<SMutableSearchComboBox::FFilteredOption>();
+							Option->ActualOption = OneTag;
+							Option->DisplayOption = OneTag;
+							Option->Parent = NodeOption;
+							TagComboOptionsSource.Add(Option);
 						}
 					}
+				}
 
-					UCustomizableObjectNodeModifierBase* TypedModifier = Cast<UCustomizableObjectNodeModifierBase>(CandidateNode);
-					if (TypedModifier)
+				UCustomizableObjectNodeModifierBase* TypedModifier = Cast<UCustomizableObjectNodeModifierBase>(CandidateNode);
+				if (TypedModifier)
+				{
+					for (const FString& OneTag : TypedModifier->RequiredTags)
 					{
-						for (const FString& OneTag : TypedModifier->RequiredTags)
+						TSharedPtr<SMutableSearchComboBox::FFilteredOption> NodeOption = AddNodeHierarchyOptions(Typed, AddedOptions);
+
+						if (!OneTag.IsEmpty())
 						{
-							if (!OneTag.IsEmpty())
-							{
-								bool bContained = TagComboOptionsSource.ContainsByPredicate([&](const TSharedPtr<FString>& Candidate)
-									{
-										return Candidate.IsValid() && *Candidate == OneTag;
-									});
-								if (!bContained)
-								{
-									TagComboOptionsSource.Add(MakeShared<FString>(OneTag));
-								}
-							}
+							TSharedRef Option = MakeShared<SMutableSearchComboBox::FFilteredOption>();
+							Option->ActualOption = OneTag;
+							Option->DisplayOption = OneTag;
+							Option->Parent = NodeOption;
+							TagComboOptionsSource.Add(Option);
 						}
 					}
 				}
@@ -185,7 +298,7 @@ void SMutableTagListWidget::OnTagComboBoxSelectionChanged(const FText& NewText)
 }
 
 
-TSharedRef<ITableRow> SMutableTagListWidget::GenerateTagMenuItemRow(TSharedPtr<FTagUIData> InItem, const TSharedRef<STableViewBase>& OwnerTable)
+TSharedRef<ITableRow> SMutableTagListWidget::GenerateTagListItemRow(TSharedPtr<FTagUIData> InItem, const TSharedRef<STableViewBase>& OwnerTable)
 {
 	if (!InItem)
 	{
diff --git a/Engine/Plugins/Mutable/Source/CustomizableObjectEditor/Private/MuCOE/SMutableTagListWidget.h b/Engine/Plugins/Mutable/Source/CustomizableObjectEditor/Private/MuCOE/SMutableTagListWidget.h
--- a/Engine/Plugins/Mutable/Source/CustomizableObjectEditor/Private/MuCOE/SMutableTagListWidget.h
+++ b/Engine/Plugins/Mutable/Source/CustomizableObjectEditor/Private/MuCOE/SMutableTagListWidget.h
@@ -18,6 +18,7 @@ class IDetailLayoutBuilder;
 class IPropertyHandle;
 class UCustomizableObjectNode;
 class UCustomizableObjectNodeObject;
+class UEdGraphNode;
 
 /** */
 class SMutableTagListWidget : public SCompoundWidget
@@ -53,7 +54,7 @@ private:
 	FText EmptyListText;
 
 	TSharedPtr<SMutableSearchComboBox> TagCombo;
-	TArray< TSharedPtr<FString> > TagComboOptionsSource;
+	TArray< TSharedRef<SMutableSearchComboBox::FFilteredOption> > TagComboOptionsSource;
 	void OnTagComboBoxSelectionChanged(const FText& NewText);
 
 	struct FTagUIData
@@ -65,7 +66,10 @@ private:
 	TArray< TSharedPtr<FTagUIData> > CurrentTagsSource;
 
 	/** */
-	TSharedRef<ITableRow> GenerateTagMenuItemRow(TSharedPtr<FTagUIData> InItem, const TSharedRef<STableViewBase>& OwnerTable);
+	TSharedRef<ITableRow> GenerateTagListItemRow(TSharedPtr<FTagUIData> InItem, const TSharedRef<STableViewBase>& OwnerTable);
+
+	/** */
+	TSharedPtr<SMutableSearchComboBox::FFilteredOption> AddNodeHierarchyOptions(UEdGraphNode* Node, TMap<UEdGraphNode*, TSharedPtr<SMutableSearchComboBox::FFilteredOption>>& AddedOptions);
 
 };
 
