[mutable] Fix compilation issue in Linux.
#rnx

diff --git a/Engine/Plugins/Mutable/Source/CustomizableObjectEditor/Private/MuCOE/SMutableSearchComboBox.cpp b/Engine/Plugins/Mutable/Source/CustomizableObjectEditor/Private/MuCOE/SMutableSearchComboBox.cpp
--- a/Engine/Plugins/Mutable/Source/CustomizableObjectEditor/Private/MuCOE/SMutableSearchComboBox.cpp
+++ b/Engine/Plugins/Mutable/Source/CustomizableObjectEditor/Private/MuCOE/SMutableSearchComboBox.cpp
@@ -142,7 +142,7 @@ void SMutableSearchComboBox::RefreshOptions()
 
 		bool bFullMatch = false;
 		FString SearchString = SearchText.ToString();
-		for (const TSharedPtr<FFilteredOption>& Option : *OptionsSource)
+		for (const TSharedRef<FFilteredOption>& Option : *OptionsSource)
 		{
 			if (Option->DisplayOption == SearchString)
 			{
