[mutable] CIS : Added markers to isolate logs by execution context.
- Added "sections" to our test commandlets to be able to differentiate the origin of the error/log (compilation, instance update, baking...).
- Added new logging category "MutableValidationLog" to prevent collisions with other mutable log categories.

diff --git a/Engine/Plugins/Mutable/Source/MutableValidation/Private/MuV/COIBakingTestCommandlet.cpp b/Engine/Plugins/Mutable/Source/MutableValidation/Private/MuV/COIBakingTestCommandlet.cpp
--- a/Engine/Plugins/Mutable/Source/MutableValidation/Private/MuV/COIBakingTestCommandlet.cpp
+++ b/Engine/Plugins/Mutable/Source/MutableValidation/Private/MuV/COIBakingTestCommandlet.cpp
@@ -3,6 +3,7 @@
 #include "../../Public/MuV/COIBakingTestCommandlet.h"
 
 #include "CustomizableObjectCompilationUtility.h"
+#include "ScopedLogSection.h"
 #include "ValidationUtils.h"
 #include "AssetRegistry/AssetData.h"
 #include "HAL/FileManager.h"
@@ -105,6 +106,8 @@ int32 UCOIBakingTestCommandlet::Main(const FString& Params)
 	
 	// Update the instance
 	{
+		const FScopedLogSection UpdateSection (EMutableLogSection::Update);
+		
 		// If this fail something is very wrong
 		check(TargetInstance);
 
@@ -132,15 +135,16 @@ int32 UCOIBakingTestCommandlet::Main(const FString& Params)
 		// Check the end status of the instance update
 		if (!bWasInstanceUpdateSuccessful)
 		{
-			UE_LOG(LogMutable,Error,TEXT("Failed to succesfully update the target COI. Exitting commandlet."));
+			UE_LOG(LogMutable,Error,TEXT("Failed to successfully update the target COI. Exiting commandlet."));
 			return 1;
 		}
 	}
-	
 
 	// Bake the instance
 	bool bWasBakingSuccessful = false;
 	{
+		const FScopedLogSection InstanceBakeSection (EMutableLogSection::Bake);
+		
 		const FString BakedResourcesFileName = "MuBakedInstances";
 		
 		// If this fail something is very wrong
@@ -195,7 +199,7 @@ int32 UCOIBakingTestCommandlet::Main(const FString& Params)
 			return 1;
 		}
 	}
-
+	
 	if (bWasBakingSuccessful)
 	{
 		UE_LOG(LogMutable, Display, TEXT("Instance Baking operation has been completed succesfully."));
diff --git a/Engine/Plugins/Mutable/Source/MutableValidation/Private/MuV/CustomizableObjectCompilationUtility.cpp b/Engine/Plugins/Mutable/Source/MutableValidation/Private/MuV/CustomizableObjectCompilationUtility.cpp
--- a/Engine/Plugins/Mutable/Source/MutableValidation/Private/MuV/CustomizableObjectCompilationUtility.cpp
+++ b/Engine/Plugins/Mutable/Source/MutableValidation/Private/MuV/CustomizableObjectCompilationUtility.cpp
@@ -2,6 +2,7 @@
 
 #include "CustomizableObjectCompilationUtility.h"
 
+#include "ScopedLogSection.h"
 #include "Commandlets/Commandlet.h"
 #include "MuCO/CustomizableObject.h"
 #include "MuCO/CustomizableObjectPrivate.h"
@@ -16,7 +17,7 @@
 bool FCustomizableObjectCompilationUtility::CompileCustomizableObject(UCustomizableObject* InCustomizableObject, const bool bShouldLogMutableLogs /* = true */, const FCompilationOptions* InCompilationOptionsOverride  /* nullptr */)
 {
 	LLM_SCOPE_BYNAME(TEXT("FCustomizableObjectCompilationUtility/Compile"));
-	
+	const FScopedLogSection CompilationSection (EMutableLogSection::Compilation);	
 	check(InCustomizableObject);
 	
 	CustomizableObject = TStrongObjectPtr(InCustomizableObject);
@@ -173,7 +174,7 @@ bool FCustomizableObjectCompilationUtility::CompileCustomizableObject(UCustomiza
 		}
 #endif
 	}
-	
+
 	// Return the success state of the operation
 	return bCompilationSuccess;
 }
diff --git a/Engine/Plugins/Mutable/Source/MutableValidation/Private/MuV/CustomizableObjectInstanceUpdateUtility.cpp b/Engine/Plugins/Mutable/Source/MutableValidation/Private/MuV/CustomizableObjectInstanceUpdateUtility.cpp
--- a/Engine/Plugins/Mutable/Source/MutableValidation/Private/MuV/CustomizableObjectInstanceUpdateUtility.cpp
+++ b/Engine/Plugins/Mutable/Source/MutableValidation/Private/MuV/CustomizableObjectInstanceUpdateUtility.cpp
@@ -7,12 +7,15 @@
 #include "MuCO/CustomizableObjectSystem.h"
 #include "Engine/TextureStreamingTypes.h"
 #include "SceneTypes.h"
+#include "ScopedLogSection.h"
+#include "ValidationUtils.h"
 #include "Components/SkeletalMeshComponent.h"
 
 
 bool FCustomizableObjectInstanceUpdateUtility::UpdateInstance(UCustomizableObjectInstance* InInstance)
 {
 	LLM_SCOPE_BYNAME(TEXT("FCustomizableObjectInstanceUpdateUtility/UpdateInstance"));
+	const FScopedLogSection UpdateSection (EMutableLogSection::Update);	
 	
 	check (InInstance);
 	check (ComponentsBeingUpdated.IsEmpty());
diff --git a/Engine/Plugins/Mutable/Source/MutableValidation/Private/MuV/ScopedLogSection.cpp b/Engine/Plugins/Mutable/Source/MutableValidation/Private/MuV/ScopedLogSection.cpp
new file mode 100644
--- /dev/null
+++ b/Engine/Plugins/Mutable/Source/MutableValidation/Private/MuV/ScopedLogSection.cpp
@@ -0,0 +1,60 @@
+﻿// Copyright Epic Games, Inc. All Rights Reserved.
+#include "ScopedLogSection.h"
+#include "Logging/LogMacros.h"
+//#include "UObject/Class.h"
+
+DECLARE_LOG_CATEGORY_EXTERN(LogMutableValidation, Log, All);
+DEFINE_LOG_CATEGORY(LogMutableValidation);
+
+
+FScopedLogSection::FScopedLogSection(EMutableLogSection Section)
+{
+	// Two scopes can not be active at the same time (not handled at the parsing level)
+	check(CurrentSection == EMutableLogSection::Undefined);
+	
+	UE_LOG(LogMutableValidation,Log,TEXT(" SECTION START : %s "), *GetLogSectionName(Section));
+	
+	// Set the current section handled by this object
+	CurrentSection = Section;
+}
+
+
+FScopedLogSection::~FScopedLogSection()
+{
+	// You are trying to close a section that is already closed. This should not be possible as we handle it in this objects
+	check(CurrentSection != EMutableLogSection::Undefined);
+	
+	UE_LOG(LogMutableValidation,Log,TEXT(" SECTION END : %s "), *GetLogSectionName(CurrentSection));
+	
+	// Set the current section to none (undefined)
+	CurrentSection = EMutableLogSection::Undefined;
+}
+
+
+	
+FString FScopedLogSection::GetLogSectionName(EMutableLogSection Section) const
+{
+	FString Output (TEXT("unknown"));
+	
+	 switch(Section)
+	 {
+	 case EMutableLogSection::Undefined:
+	 	Output = TEXT("undefined");
+ 		break;
+	 case EMutableLogSection::Compilation:
+	 	Output = TEXT("compilation");
+ 		break;
+	 case EMutableLogSection::Update:
+	 	Output = TEXT("update");
+ 		break;
+	 case EMutableLogSection::Bake:
+	 	Output = TEXT("bake");
+ 		break;
+	 default:
+	 	checkNoEntry();
+	 	break;
+	 }
+
+	return Output;
+}
+
diff --git a/Engine/Plugins/Mutable/Source/MutableValidation/Private/MuV/ScopedLogSection.h b/Engine/Plugins/Mutable/Source/MutableValidation/Private/MuV/ScopedLogSection.h
new file mode 100644
--- /dev/null
+++ b/Engine/Plugins/Mutable/Source/MutableValidation/Private/MuV/ScopedLogSection.h
@@ -0,0 +1,37 @@
+﻿// Copyright Epic Games, Inc. All Rights Reserved.
+#pragma once
+
+#include  "HAL/Platform.h"
+
+class FString;
+
+/**
+ * Names for the possible sections used during the logging of this module commandlets.
+ * Do not change the name of them before consulting the Mutable team
+ */
+enum class EMutableLogSection : uint8
+{
+	Undefined = 0,
+	Compilation,
+	Update,
+	Bake
+};
+
+
+/**
+* Object that handles the logging of scope based log sections that we can later parse out and interpret externally 
+*/
+class FScopedLogSection
+{
+public:
+	FScopedLogSection(EMutableLogSection Section);
+	~FScopedLogSection();
+
+private:
+
+	/** Get the name of the context as a string of characters */
+	FString GetLogSectionName(EMutableLogSection Section) const;
+	
+  	/** Section that this object is currently representing */
+	static inline EMutableLogSection CurrentSection = EMutableLogSection::Undefined;
+};
\ No newline at end of file
