[Mutable] Fixed LoadCompiledDataFromDisk crash.

- Fixed issue which opened all COs in the project.

#rnx

diff --git a/Engine/Plugins/Mutable/Source/CustomizableObject/Internal/MuCO/CustomizableObjectPrivate.h b/Engine/Plugins/Mutable/Source/CustomizableObject/Internal/MuCO/CustomizableObjectPrivate.h
--- a/Engine/Plugins/Mutable/Source/CustomizableObject/Internal/MuCO/CustomizableObjectPrivate.h
+++ b/Engine/Plugins/Mutable/Source/CustomizableObject/Internal/MuCO/CustomizableObjectPrivate.h
@@ -1324,6 +1324,8 @@ public:
 
 	FCompilationOptions GetCompileOptions() const;
 #endif
+
+	void BackwardsCompatibleFixup(int32 CustomizableObjectCustomVersion);
 	
 	/** Cache of generated SkeletalMeshes */
 	FMeshCache MeshCache;
diff --git a/Engine/Plugins/Mutable/Source/CustomizableObject/Internal/MuCO/ICustomizableObjectEditorModule.h b/Engine/Plugins/Mutable/Source/CustomizableObject/Internal/MuCO/ICustomizableObjectEditorModule.h
--- a/Engine/Plugins/Mutable/Source/CustomizableObject/Internal/MuCO/ICustomizableObjectEditorModule.h
+++ b/Engine/Plugins/Mutable/Source/CustomizableObject/Internal/MuCO/ICustomizableObjectEditorModule.h
@@ -16,6 +16,7 @@ class UCustomizableObjectInstance;
 class USkeletalMesh;
 class FExtensibilityManager;
 class FBakeOperationCompletedDelegate;
+class UEdGraph;
 struct FBakingConfiguration;
 struct FCompilationRequest;
 struct FCompilationOptions;
@@ -97,5 +98,9 @@ public:
 	/** Perform a fast compilation pass to get all participating objects.
 	 *  @param bLoadObjects Load any object. If false, no objects will load. If true, only objects strictly required to get the full list of participating objects will load. */
 	virtual TMap<FName, FGuid> GetParticipatingObjects(const UCustomizableObject* Object, bool bLoadObjects, const FCompilationOptions* Options = nullptr) const = 0;
+
+	virtual void BackwardsCompatibleFixup(UEdGraph& Graph, int32 CustomizableObjectCustomVersion) = 0;
+	
+	virtual void PostBackwardsCompatibleFixup(UEdGraph& Graph) = 0;
 };
  
diff --git a/Engine/Plugins/Mutable/Source/CustomizableObject/Private/MuCO/CustomizableObject.cpp b/Engine/Plugins/Mutable/Source/CustomizableObject/Private/MuCO/CustomizableObject.cpp
--- a/Engine/Plugins/Mutable/Source/CustomizableObject/Private/MuCO/CustomizableObject.cpp
+++ b/Engine/Plugins/Mutable/Source/CustomizableObject/Private/MuCO/CustomizableObject.cpp
@@ -314,44 +314,36 @@ bool UCustomizableObjectPrivate::TryLoadCompiledCookDataForPlatform(const ITarge
 void UCustomizableObject::PostLoad()
 {
 	Super::PostLoad();
-	
-	const int32 CustomizableObjectCustomVersion = GetLinkerCustomVersion(FCustomizableObjectCustomVersion::GUID);
 
 #if WITH_EDITOR
-	if (ReferenceSkeletalMesh_DEPRECATED)
-	{
-		ReferenceSkeletalMeshes_DEPRECATED.Add(ReferenceSkeletalMesh_DEPRECATED);
-		ReferenceSkeletalMesh_DEPRECATED = nullptr;
-	}
-
-#if WITH_EDITORONLY_DATA
-	if (CustomizableObjectCustomVersion < FCustomizableObjectCustomVersion::CompilationOptions)
+	if (Source)
 	{
-		GetPrivate()->OptimizationLevel = CompileOptions_DEPRECATED.OptimizationLevel;
-		GetPrivate()->TextureCompression = CompileOptions_DEPRECATED.TextureCompression;
-		GetPrivate()->bUseDiskCompilation = CompileOptions_DEPRECATED.bUseDiskCompilation;
-		GetPrivate()->EmbeddedDataBytesLimit = CompileOptions_DEPRECATED.EmbeddedDataBytesLimit;
-		GetPrivate()->PackagedDataBytesLimit = CompileOptions_DEPRECATED.PackagedDataBytesLimit;
+		Source->ConditionalPostLoad();
 	}
-
-	if (CustomizableObjectCustomVersion < FCustomizableObjectCustomVersion::NewComponentOptions)
+	
+	for (int32 Version = GetLinkerCustomVersion(FCustomizableObjectCustomVersion::GUID) + 1; Version <= FCustomizableObjectCustomVersion::LatestVersion; ++Version)
 	{
-		if (GetPrivate()->MutableMeshComponents_DEPRECATED.IsEmpty())
+		GetPrivate()->BackwardsCompatibleFixup(Version);
+		
+		if (Source)
 		{
-			for (int32 SkeletalMeshIndex = 0; SkeletalMeshIndex < ReferenceSkeletalMeshes_DEPRECATED.Num(); ++SkeletalMeshIndex)
+			if (ICustomizableObjectEditorModule* Module = ICustomizableObjectEditorModule::Get())
 			{
-				FMutableMeshComponentData NewComponent;
-				NewComponent.Name = FName(FString::FromInt(SkeletalMeshIndex));
-				NewComponent.ReferenceSkeletalMesh = ReferenceSkeletalMeshes_DEPRECATED[SkeletalMeshIndex];
-
-				GetPrivate()->MutableMeshComponents_DEPRECATED.Add(NewComponent);
+				// Execute backwards compatible code for all nodes. It requires all nodes to be loaded.
+			
+				Module->BackwardsCompatibleFixup(*Source, Version);
 			}
-
-			ReferenceSkeletalMeshes_DEPRECATED.Empty();
 		}
 	}
-#endif
 
+	if (Source)
+	{
+		if (ICustomizableObjectEditorModule* Module = ICustomizableObjectEditorModule::Get())
+		{
+			Module->PostBackwardsCompatibleFixup(*Source);
+		}
+	}
+	
 	// Register to dirty delegate so we update derived data version ID each time that the package is marked as dirty.
 	if (UPackage* Package = GetOutermost())
 	{
@@ -367,10 +359,7 @@ void UCustomizableObject::PostLoad()
 	if (!IsRunningCookCommandlet())
 	{
 		GetPrivate()->Status.NextState(FCustomizableObjectStatusTypes::EState::Loading);
-
-		ITargetPlatformManagerModule& TargetPlatformManager = GetTargetPlatformManagerRef();
-		const ITargetPlatform* RunningPlatform = TargetPlatformManager.GetRunningTargetPlatform();
-
+	
 		const FAssetRegistryModule& AssetRegistryModule = FModuleManager::LoadModuleChecked<FAssetRegistryModule>("AssetRegistry");
 		if (AssetRegistryModule.Get().IsLoadingAssets())
 		{
@@ -385,6 +374,45 @@ void UCustomizableObject::PostLoad()
 }
 
 
+void UCustomizableObjectPrivate::BackwardsCompatibleFixup(int32 CustomizableObjectCustomVersion)
+{
+#if WITH_EDITOR
+	if (GetPublic()->ReferenceSkeletalMesh_DEPRECATED)
+	{
+		GetPublic()->ReferenceSkeletalMeshes_DEPRECATED.Add(GetPublic()->ReferenceSkeletalMesh_DEPRECATED);
+		GetPublic()->ReferenceSkeletalMesh_DEPRECATED = nullptr;
+	}
+
+#if WITH_EDITORONLY_DATA
+	if (CustomizableObjectCustomVersion == FCustomizableObjectCustomVersion::CompilationOptions)
+	{
+		OptimizationLevel = GetPublic()->CompileOptions_DEPRECATED.OptimizationLevel;
+		TextureCompression = GetPublic()->CompileOptions_DEPRECATED.TextureCompression;
+		bUseDiskCompilation = GetPublic()->CompileOptions_DEPRECATED.bUseDiskCompilation;
+		EmbeddedDataBytesLimit = GetPublic()->CompileOptions_DEPRECATED.EmbeddedDataBytesLimit;
+		PackagedDataBytesLimit = GetPublic()->CompileOptions_DEPRECATED.PackagedDataBytesLimit;
+	}
+
+	if (CustomizableObjectCustomVersion == FCustomizableObjectCustomVersion::NewComponentOptions)
+	{
+		if (MutableMeshComponents_DEPRECATED.IsEmpty())
+		{
+			for (int32 SkeletalMeshIndex = 0; SkeletalMeshIndex < GetPublic()->ReferenceSkeletalMeshes_DEPRECATED.Num(); ++SkeletalMeshIndex)
+			{
+				FMutableMeshComponentData NewComponent;
+				NewComponent.Name = FName(FString::FromInt(SkeletalMeshIndex));
+				NewComponent.ReferenceSkeletalMesh = GetPublic()->ReferenceSkeletalMeshes_DEPRECATED[SkeletalMeshIndex];
+
+				MutableMeshComponents_DEPRECATED.Add(NewComponent);
+			}
+
+			GetPublic()->ReferenceSkeletalMeshes_DEPRECATED.Empty();
+		}
+	}
+#endif
+#endif
+}
+
 bool UCustomizableObjectPrivate::IsLocked() const
 {
 	return bLocked;
diff --git a/Engine/Plugins/Mutable/Source/CustomizableObjectEditor/Private/MuCOE/CustomizableObjectCompiler.cpp b/Engine/Plugins/Mutable/Source/CustomizableObjectEditor/Private/MuCOE/CustomizableObjectCompiler.cpp
--- a/Engine/Plugins/Mutable/Source/CustomizableObjectEditor/Private/MuCOE/CustomizableObjectCompiler.cpp
+++ b/Engine/Plugins/Mutable/Source/CustomizableObjectEditor/Private/MuCOE/CustomizableObjectCompiler.cpp
@@ -366,6 +366,11 @@ void ProcessChildObjectsRecursively(const UCustomizableObject* ParentObject, FMu
 	const FAssetRegistryModule& AssetRegistryModule = FModuleManager::LoadModuleChecked<FAssetRegistryModule>("AssetRegistry");
 	AssetRegistryModule.Get().GetReferencers(*ParentObject->GetOuter()->GetPathName(), ReferencedObjectNames, UE::AssetRegistry::EDependencyCategory::Package, UE::AssetRegistry::EDependencyQuery::Hard);
 
+	if (ReferencedObjectNames.IsEmpty())
+	{
+		return;
+	}
+	
 	// Required to be deterministic.
 	ReferencedObjectNames.Sort([](const FName& A, const FName& B)
 	{
diff --git a/Engine/Plugins/Mutable/Source/CustomizableObjectEditor/Private/MuCOE/CustomizableObjectEditorModule.cpp b/Engine/Plugins/Mutable/Source/CustomizableObjectEditor/Private/MuCOE/CustomizableObjectEditorModule.cpp
--- a/Engine/Plugins/Mutable/Source/CustomizableObjectEditor/Private/MuCOE/CustomizableObjectEditorModule.cpp
+++ b/Engine/Plugins/Mutable/Source/CustomizableObjectEditor/Private/MuCOE/CustomizableObjectEditorModule.cpp
@@ -78,6 +78,7 @@
 #include "MuCOE/CustomizableObjectInstanceBaker.h"
 #include "Editor.h"
 #include "MuCO/CustomizableObjectSystemPrivate.h"
+#include "MuCOE/CustomizableObjectGraph.h"
 #include "Nodes/CustomizableObjectNodeComponentMeshDetails.h"
 
 class AActor;
@@ -738,6 +739,24 @@ TMap<FName, FGuid> FCustomizableObjectEditorModule::GetParticipatingObjects(cons
 }
 
 
+void FCustomizableObjectEditorModule::BackwardsCompatibleFixup(UEdGraph& Graph, int32 CustomizableObjectCustomVersion)
+{
+	if (UCustomizableObjectGraph* COGraph = Cast<UCustomizableObjectGraph>(&Graph))
+	{
+		COGraph->BackwardsCompatibleFixup(CustomizableObjectCustomVersion);
+	}
+}
+
+
+void FCustomizableObjectEditorModule::PostBackwardsCompatibleFixup(UEdGraph& Graph)
+{
+	if (UCustomizableObjectGraph* COGraph = Cast<UCustomizableObjectGraph>(&Graph))
+	{
+		COGraph->PostBackwardsCompatibleFixup();
+	}
+}
+
+
 void FCustomizableObjectEditorModule::CompileCustomizableObject(const TSharedRef<FCompilationRequest>& InCompilationRequest, bool bForceRequest)
 {
 	if (IsRunningGame())
diff --git a/Engine/Plugins/Mutable/Source/CustomizableObjectEditor/Private/MuCOE/CustomizableObjectEditorModule.h b/Engine/Plugins/Mutable/Source/CustomizableObjectEditor/Private/MuCOE/CustomizableObjectEditorModule.h
--- a/Engine/Plugins/Mutable/Source/CustomizableObjectEditor/Private/MuCOE/CustomizableObjectEditorModule.h
+++ b/Engine/Plugins/Mutable/Source/CustomizableObjectEditor/Private/MuCOE/CustomizableObjectEditorModule.h
@@ -46,6 +46,8 @@ public:
 	virtual void BakeCustomizableObjectInstance(UCustomizableObjectInstance* InTargetInstance, const FBakingConfiguration& InBakingConfig) override;
 	virtual USkeletalMesh* GetReferenceSkeletalMesh(const UCustomizableObject& Object, const FName& ComponentName) const override;
 	virtual TMap<FName, FGuid> GetParticipatingObjects(const UCustomizableObject* Object, bool bLoadObjects, const FCompilationOptions* Options = nullptr) const override;
+	virtual void BackwardsCompatibleFixup(UEdGraph& Graph, int32 CustomizableObjectCustomVersion) override;
+	virtual void PostBackwardsCompatibleFixup(UEdGraph& Graph) override;
 
 	/** Request for a given customizable object to be compiled. Async compile requests will be queued and processed sequentially. 
 	 * @param InCompilationRequest - Request to compile an object. 
diff --git a/Engine/Plugins/Mutable/Source/CustomizableObjectEditor/Private/MuCOE/CustomizableObjectGraph.cpp b/Engine/Plugins/Mutable/Source/CustomizableObjectEditor/Private/MuCOE/CustomizableObjectGraph.cpp
--- a/Engine/Plugins/Mutable/Source/CustomizableObjectEditor/Private/MuCOE/CustomizableObjectGraph.cpp
+++ b/Engine/Plugins/Mutable/Source/CustomizableObjectEditor/Private/MuCOE/CustomizableObjectGraph.cpp
@@ -21,9 +21,6 @@ void UCustomizableObjectGraph::PostLoad()
 	Super::PostLoad();
 
 	// TODO UE-222779 If compatibility code has to be executed, loaded the full hierarchy and in sync the BackwardsCompatibleFixup 
-	
-	// Make sure the Object has executed its compatibility code.
-	GetOuter()->ConditionalPostLoad();
 
 	// Make sure all nodes have finished loading.
 	for (UEdGraphNode* Node : Nodes)
@@ -44,38 +41,41 @@ void UCustomizableObjectGraph::PostLoad()
 			Pin->LinkedTo.RemoveAll([](UEdGraphPin* Other) { return Other == nullptr; });
 		}
 	}
+}
+
 
-	// Execute backwards compatible code for all nodes. It requires all nodes to be loaded.
-	for (int32 Version = GetLinkerCustomVersion(FCustomizableObjectCustomVersion::GUID) + 1; Version <= FCustomizableObjectCustomVersion::LatestVersion; ++Version)
+void UCustomizableObjectGraph::BackwardsCompatibleFixup(int32 CustomizableObjectCustomVersion)
+{
+	TArray<TObjectPtr<UEdGraphNode>> NodesCopy = Nodes; // Copy to be able to remove nodes inside the BackwardsCompatibleFixup.
+	for (UEdGraphNode* Node : NodesCopy)
 	{
-		TArray<TObjectPtr<UEdGraphNode>> NodesCopy = Nodes; // Copy to be able to remove nodes inside the BackwardsCompatibleFixup.
-		for (UEdGraphNode* Node : NodesCopy)
+		if (UCustomizableObjectNode* CustomizableObjectNode = Cast<UCustomizableObjectNode>(Node))
 		{
-			if (UCustomizableObjectNode* CustomizableObjectNode = Cast<UCustomizableObjectNode>(Node))
+			TArray<UEdGraphPin*> PinsCopy = CustomizableObjectNode->GetAllPins(); // Copy to be able to remove pins inside the BackwardsCompatibleFixup.
+			for (UEdGraphPin* Pin : PinsCopy)
 			{
-				TArray<UEdGraphPin*> PinsCopy = CustomizableObjectNode->GetAllPins(); // Copy to be able to remove pins inside the BackwardsCompatibleFixup.
-				for (UEdGraphPin* Pin : PinsCopy)
+				if (!Pin)
 				{
-					if (!Pin)
-					{
-						continue;
-					}
-					
-					UCustomizableObjectNodePinData* PinData = CustomizableObjectNode->GetPinData(*Pin);
-					if (!PinData)
-					{
-						continue;
-					}
+					continue;
+				}
 					
-					PinData->BackwardsCompatibleFixup(Version);
+				UCustomizableObjectNodePinData* PinData = CustomizableObjectNode->GetPinData(*Pin);
+				if (!PinData)
+				{
+					continue;
 				}
-				
-				CustomizableObjectNode->BackwardsCompatibleFixup(Version);
+					
+				PinData->BackwardsCompatibleFixup(CustomizableObjectCustomVersion);
 			}
+				
+			CustomizableObjectNode->BackwardsCompatibleFixup(CustomizableObjectCustomVersion);
 		}
 	}
+}
 
 
+void UCustomizableObjectGraph::PostBackwardsCompatibleFixup()
+{
 	// Do any additional work which require nodes to be valid (i.e., have executed BackwardsCompatibleFixup).
 	TArray<TObjectPtr<UEdGraphNode>> NodesCopy = Nodes; // Copy to be able to remove nodes inside the PostBackwardsCompatibleFixup.
 	for (UEdGraphNode* Node : NodesCopy)
diff --git a/Engine/Plugins/Mutable/Source/CustomizableObjectEditor/Public/MuCOE/CustomizableObjectGraph.h b/Engine/Plugins/Mutable/Source/CustomizableObjectEditor/Public/MuCOE/CustomizableObjectGraph.h
--- a/Engine/Plugins/Mutable/Source/CustomizableObjectEditor/Public/MuCOE/CustomizableObjectGraph.h
+++ b/Engine/Plugins/Mutable/Source/CustomizableObjectEditor/Public/MuCOE/CustomizableObjectGraph.h
@@ -31,6 +31,10 @@ public:
 	/** Adds the necessary nodes for a CO to work */
 	void AddEssentialGraphNodes();
 
+	void BackwardsCompatibleFixup(int32 CustomizableObjectCustomVersion);
+
+	void PostBackwardsCompatibleFixup();
+	
 private:
 
 	// Request Node Id Update Map
