[mutable] Fix possible crash with legacy data.
#rnx

diff --git a/Engine/Plugins/Mutable/Source/CustomizableObjectEditor/Private/MuCOE/GenerateMutableSource/GenerateMutableSource.cpp b/Engine/Plugins/Mutable/Source/CustomizableObjectEditor/Private/MuCOE/GenerateMutableSource/GenerateMutableSource.cpp
--- a/Engine/Plugins/Mutable/Source/CustomizableObjectEditor/Private/MuCOE/GenerateMutableSource/GenerateMutableSource.cpp
+++ b/Engine/Plugins/Mutable/Source/CustomizableObjectEditor/Private/MuCOE/GenerateMutableSource/GenerateMutableSource.cpp
@@ -927,14 +927,15 @@ mu::Ptr<mu::NodeObject> GenerateMutableSource(const UEdGraphPin * Pin, FMutableG
 		
 		// Process components.
 		//-------------------------------------------------------------------
-		TArray<mu::Ptr<mu::NodeComponent>> ExplicitComponents;
-
 		const UEdGraphPin* ComponentsPin = TypedNodeObj->ComponentsPin();
-		TArray<UEdGraphPin*> ConnectedComponentPins = FollowInputPinArray(*ComponentsPin);
-		for (const UEdGraphPin* ComponentNodePin : ConnectedComponentPins)
-		{		
-			mu::Ptr<mu::NodeComponent> ComponentNode = GenerateMutableSourceComponent(ComponentNodePin, GenerationContext);
-			ObjectNode->Components.Add(ComponentNode);
+		if (ComponentsPin)
+		{
+			TArray<UEdGraphPin*> ConnectedComponentPins = FollowInputPinArray(*ComponentsPin);
+			for (const UEdGraphPin* ComponentNodePin : ConnectedComponentPins)
+			{
+				mu::Ptr<mu::NodeComponent> ComponentNode = GenerateMutableSourceComponent(ComponentNodePin, GenerationContext);
+				ObjectNode->Components.Add(ComponentNode);
+			}
 		}
 
 		// Process modifiers.
