[Mutable] Removed unused FCustomizableObjectIdentifier.

#rnx

diff --git a/Engine/Plugins/Mutable/Source/CustomizableObject/Internal/MuCO/CustomizableObjectIdentifier.h b/Engine/Plugins/Mutable/Source/CustomizableObject/Internal/MuCO/CustomizableObjectIdentifier.h
--- a/Engine/Plugins/Mutable/Source/CustomizableObject/Internal/MuCO/CustomizableObjectIdentifier.h
+++ b/Engine/Plugins/Mutable/Source/CustomizableObject/Internal/MuCO/CustomizableObjectIdentifier.h
@@ -33,24 +33,3 @@ struct CUSTOMIZABLEOBJECT_API FCustomizableObjectIdPair
 	}
 };
 
-
-USTRUCT(BlueprintType)
-struct CUSTOMIZABLEOBJECT_API FCustomizableObjectIdentifier
-{
-	GENERATED_USTRUCT_BODY()
-
-private:
-	UPROPERTY()
-	FString CustomizableObjectGroupName; // Deprecated
-
-	UPROPERTY()
-	FString CustomizableObjectName; // Deprecated
-
-public:
-	UPROPERTY(Category = CustomizableObject, BlueprintReadOnly, EditDefaultsOnly)
-	FString Guid;
-
-	FCustomizableObjectIdentifier();
-	FCustomizableObjectIdentifier(FString ObjectGroupName, FString ObjectName);
-};
-
diff --git a/Engine/Plugins/Mutable/Source/CustomizableObject/Private/MuCO/CustomizableObjectIdentifier.cpp b/Engine/Plugins/Mutable/Source/CustomizableObject/Private/MuCO/CustomizableObjectIdentifier.cpp
--- a/Engine/Plugins/Mutable/Source/CustomizableObject/Private/MuCO/CustomizableObjectIdentifier.cpp
+++ b/Engine/Plugins/Mutable/Source/CustomizableObject/Private/MuCO/CustomizableObjectIdentifier.cpp
@@ -11,13 +11,3 @@ FCustomizableObjectIdPair::FCustomizableObjectIdPair(FString ObjectGroupName, FS
 {
 }
 
-FCustomizableObjectIdentifier::FCustomizableObjectIdentifier()
-{
-}
-
-FCustomizableObjectIdentifier::FCustomizableObjectIdentifier(FString ObjectGroupName, FString ObjectName)
-	: CustomizableObjectGroupName(ObjectGroupName)
-	, CustomizableObjectName(ObjectName)
-{
-}
-
