[mutable] Fixed output logs from Component nodes not pointing to the Component node but to the root node.
- This seems to be a bug born of the changes we have made lately in how we handle components since before they got handled in the root object (ActualRoot) but that is no longer the case.

#rnx

diff --git a/Engine/Plugins/Mutable/Source/CustomizableObjectEditor/Private/MuCOE/GenerateMutableSource/GenerateMutableSourceComponent.cpp b/Engine/Plugins/Mutable/Source/CustomizableObjectEditor/Private/MuCOE/GenerateMutableSource/GenerateMutableSourceComponent.cpp
--- a/Engine/Plugins/Mutable/Source/CustomizableObjectEditor/Private/MuCOE/GenerateMutableSource/GenerateMutableSourceComponent.cpp
+++ b/Engine/Plugins/Mutable/Source/CustomizableObjectEditor/Private/MuCOE/GenerateMutableSource/GenerateMutableSourceComponent.cpp
@@ -130,13 +130,11 @@ mu::Ptr<mu::NodeComponent> GenerateMutableSourceComponent(const UEdGraphPin* Pin
 	
 	if (const UCustomizableObjectNodeComponentMesh* TypedComponentMesh = Cast<UCustomizableObjectNodeComponentMesh>(Node))
 	{
-		UCustomizableObjectNodeObject* ActualRoot = GenerationContext.Root;
-			
 		FName ComponentName = TypedComponentMesh->ComponentName;
 		
 		if (TypedComponentMesh->ComponentName.IsNone())
 		{
-			GenerationContext.Log(LOCTEXT("EmptyComponentNameError", "Error! Missing name in a component of the Customizable Object."), ActualRoot, EMessageSeverity::Error);
+			GenerationContext.Log(LOCTEXT("EmptyComponentNameError", "Error! Missing name in a component of the Customizable Object."), TypedComponentMesh, EMessageSeverity::Error);
 			return nullptr;
 		}
 
@@ -146,14 +144,14 @@ mu::Ptr<mu::NodeComponent> GenerateMutableSourceComponent(const UEdGraphPin* Pin
 		}))
 		{
 			GenerationContext.Log(FText::Format(LOCTEXT("RepeatedComponentName", "Error! Repeated name [{0}] used in more than one Component"),
-				FText::FromName(ComponentName)), ActualRoot, EMessageSeverity::Error);
+				FText::FromName(ComponentName)), TypedComponentMesh, EMessageSeverity::Error);
 			return nullptr;
 		}
 			
 		USkeletalMesh* RefSkeletalMesh = TypedComponentMesh->ReferenceSkeletalMesh;
 		if (!RefSkeletalMesh)
 		{
-			GenerationContext.Log(LOCTEXT("NoReferenceMeshObjectTab", "Error! Missing reference Skeletal Mesh"), ActualRoot, EMessageSeverity::Error);
+			GenerationContext.Log(LOCTEXT("NoReferenceMeshObjectTab", "Error! Missing reference Skeletal Mesh"), TypedComponentMesh, EMessageSeverity::Error);
 			return nullptr;
 		}
 
@@ -162,7 +160,7 @@ mu::Ptr<mu::NodeComponent> GenerateMutableSourceComponent(const UEdGraphPin* Pin
 		{
 			FText Msg = FText::Format(LOCTEXT("NoReferenceSkeleton", "Error! Missing skeleton in the reference mesh [{0}]"), FText::FromString(GenerationContext.CustomizableObjectWithCycle->GetPathName()));
 
-			GenerationContext.Log(Msg, ActualRoot, EMessageSeverity::Error);
+			GenerationContext.Log(Msg, TypedComponentMesh, EMessageSeverity::Error);
 			return nullptr;
 		}
 
@@ -177,7 +175,7 @@ mu::Ptr<mu::NodeComponent> GenerateMutableSourceComponent(const UEdGraphPin* Pin
 		// Add reference meshes to the participating objects
 		GenerationContext.AddParticipatingObject(*RefSkeletalMesh);
 		
-		// Ensure that the CO has a valid AutoLODStrategy on the ActualRoot.
+		// Ensure that the CO has a valid AutoLODStrategy on the Component node.
 		if (TypedComponentMesh->AutoLODStrategy == ECustomizableObjectAutomaticLODStrategy::Inherited)
 		{
 			GenerationContext.Log(LOCTEXT("RootInheritsFromParent", "Error! Component LOD Strategy can't be set to 'Inherit from parent object'"), TypedComponentMesh, EMessageSeverity::Error);
