[Mutable] Fix crash when selecting layouts in a SkeletalMesh Node.

#rnx

diff --git a/Engine/Plugins/Mutable/Source/CustomizableObjectEditor/Private/MuCOE/CustomizableObjectLayout.cpp b/Engine/Plugins/Mutable/Source/CustomizableObjectEditor/Private/MuCOE/CustomizableObjectLayout.cpp
--- a/Engine/Plugins/Mutable/Source/CustomizableObjectEditor/Private/MuCOE/CustomizableObjectLayout.cpp
+++ b/Engine/Plugins/Mutable/Source/CustomizableObjectEditor/Private/MuCOE/CustomizableObjectLayout.cpp
@@ -102,6 +102,7 @@ void UCustomizableObjectLayout::GenerateAutomaticBlocksFromUVs()
 	FCompilationOptions Options = Object->GetPrivate()->GetCompileOptions();
 
 	FMutableGraphGenerationContext GenerationContext(Object, &Compiler, Options);
+	GenerationContext.bSkipParticipatingObjectsPass = true;
 
 	bool bOutWasEmpty = false;
 	mu::Ptr<mu::NodeLayout> LayoutNode = CreateMutableLayoutNode(GenerationContext, this, false, bOutWasEmpty );
