[Mutable] Fix ProjectorConstant, ProjectorParameter, and ModifierClipMorph Node bone selection.

- Component index was used to select a reference skeleton, now a name is used.

#rnx

diff --git a/Engine/Plugins/Mutable/Source/CustomizableObject/Internal/MuCO/CustomizableObjectCustomVersion.h b/Engine/Plugins/Mutable/Source/CustomizableObject/Internal/MuCO/CustomizableObjectCustomVersion.h
--- a/Engine/Plugins/Mutable/Source/CustomizableObject/Internal/MuCO/CustomizableObjectCustomVersion.h
+++ b/Engine/Plugins/Mutable/Source/CustomizableObject/Internal/MuCO/CustomizableObjectCustomVersion.h
@@ -162,6 +162,8 @@ struct CUSTOMIZABLEOBJECT_API FCustomizableObjectCustomVersion
 		
 		FixAutomaticBlocksStrategyLegacyNodes,
 
+		SnapToBoneComponentIndexToName,
+
 		// -----<new versions can be added above this line>-------------------------------------------------
 		VersionPlusOne,
 		LatestVersion = VersionPlusOne - 1
diff --git a/Engine/Plugins/Mutable/Source/CustomizableObjectEditor/Private/MuCOE/Nodes/CustomizableObjectNodeMeshClipMorph.cpp b/Engine/Plugins/Mutable/Source/CustomizableObjectEditor/Private/MuCOE/Nodes/CustomizableObjectNodeMeshClipMorph.cpp
--- a/Engine/Plugins/Mutable/Source/CustomizableObjectEditor/Private/MuCOE/Nodes/CustomizableObjectNodeMeshClipMorph.cpp
+++ b/Engine/Plugins/Mutable/Source/CustomizableObjectEditor/Private/MuCOE/Nodes/CustomizableObjectNodeMeshClipMorph.cpp
@@ -174,6 +174,11 @@ void UCustomizableObjectNodeModifierClipMorph::BackwardsCompatibleFixup(int32 Cu
 		RequiredTags = Tags_DEPRECATED;
 		Tags_DEPRECATED.Empty();
 	}
+
+	if (CustomizableObjectCustomVersion == FCustomizableObjectCustomVersion::SnapToBoneComponentIndexToName)
+	{
+		ReferenceSkeletonComponent = FName(FString::FromInt(ReferenceSkeletonIndex_DEPRECATED));
+	}
 }
 
 
diff --git a/Engine/Plugins/Mutable/Source/CustomizableObjectEditor/Private/MuCOE/Nodes/CustomizableObjectNodeModifierClipMorphDetails.cpp b/Engine/Plugins/Mutable/Source/CustomizableObjectEditor/Private/MuCOE/Nodes/CustomizableObjectNodeModifierClipMorphDetails.cpp
--- a/Engine/Plugins/Mutable/Source/CustomizableObjectEditor/Private/MuCOE/Nodes/CustomizableObjectNodeModifierClipMorphDetails.cpp
+++ b/Engine/Plugins/Mutable/Source/CustomizableObjectEditor/Private/MuCOE/Nodes/CustomizableObjectNodeModifierClipMorphDetails.cpp
@@ -46,10 +46,10 @@ void FCustomizableObjectNodeModifierClipMorphDetails::CustomizeDetails(IDetailLa
 	IDetailCategoryBuilder& MeshClipParametersCategory = DetailBuilder.EditCategory("MeshClipParameters");
 	DetailBuilder.HideProperty("bInvertNormal");
 
-	TSharedPtr<IPropertyHandle> ReferenceSkeletonIndexPropertyHandle = DetailBuilder.GetProperty("ReferenceSkeletonIndex");
-	const FSimpleDelegate OnReferenceSkeletonIndexChangedDelegate = 
-		FSimpleDelegate::CreateSP(this, &FCustomizableObjectNodeModifierClipMorphDetails::OnReferenceSkeletonIndexChanged);
-	ReferenceSkeletonIndexPropertyHandle->SetOnPropertyValueChanged(OnReferenceSkeletonIndexChangedDelegate);
+	TSharedPtr<IPropertyHandle> ReferenceSkeletonIndexPropertyHandle = DetailBuilder.GetProperty("ReferenceSkeletonComponent");
+	const FSimpleDelegate OnReferenceSkeletonComponentChangedDelegate = 
+			FSimpleDelegate::CreateSP(this, &FCustomizableObjectNodeModifierClipMorphDetails::OnReferenceSkeletonComponentChanged);
+	ReferenceSkeletonIndexPropertyHandle->SetOnPropertyValueChanged(OnReferenceSkeletonComponentChangedDelegate);
 
 	if (Node)
 	{
@@ -60,26 +60,8 @@ void FCustomizableObjectNodeModifierClipMorphDetails::CustomizeDetails(IDetailLa
 
 		if (CustomizableObject)
 		{
-			if (CustomizableObject->IsChildObject())
-			{
-				bool bMultipleBaseObjects;
-				UCustomizableObjectNodeObject* RootNode = GetRootNode(CustomizableObject, bMultipleBaseObjects);
-
-				if (RootNode && !bMultipleBaseObjects)
-				{
-					TArray<UCustomizableObject*> VisitedObjects;
-					UCustomizableObject* Parent = GetFullGraphRootObject(RootNode, VisitedObjects);
-
-					if (Parent)
-					{
-						SkeletalMesh = VisitedObjects.Last()->GetRefSkeletalMesh(Node->ReferenceSkeletonIndex);
-					}
-				}
-			}
-			else
-			{
-				SkeletalMesh = CustomizableObject->GetRefSkeletalMesh(Node->ReferenceSkeletonIndex);
-			}
+			const FName ReferenceComponentName = Node->ReferenceSkeletonComponent;
+			SkeletalMesh = CustomizableObject->GetComponentMeshReferenceSkeletalMesh(ReferenceComponentName);
 		}
 
 		if (SkeletalMesh)
@@ -339,7 +321,7 @@ ECheckBoxState FCustomizableObjectNodeModifierClipMorphDetails::GetInvertNormalC
 }
 
 
-void FCustomizableObjectNodeModifierClipMorphDetails::OnReferenceSkeletonIndexChanged()
+void FCustomizableObjectNodeModifierClipMorphDetails::OnReferenceSkeletonComponentChanged()
 {
 	if (Node)
 	{
diff --git a/Engine/Plugins/Mutable/Source/CustomizableObjectEditor/Private/MuCOE/Nodes/CustomizableObjectNodeModifierClipMorphDetails.h b/Engine/Plugins/Mutable/Source/CustomizableObjectEditor/Private/MuCOE/Nodes/CustomizableObjectNodeModifierClipMorphDetails.h
--- a/Engine/Plugins/Mutable/Source/CustomizableObjectEditor/Private/MuCOE/Nodes/CustomizableObjectNodeModifierClipMorphDetails.h
+++ b/Engine/Plugins/Mutable/Source/CustomizableObjectEditor/Private/MuCOE/Nodes/CustomizableObjectNodeModifierClipMorphDetails.h
@@ -36,5 +36,5 @@ private:
 	void OnInvertNormalCheckboxChanged(ECheckBoxState CheckBoxState);
 	ECheckBoxState GetInvertNormalCheckBoxState() const;
 
-	void OnReferenceSkeletonIndexChanged();
+	void OnReferenceSkeletonComponentChanged();
 };
diff --git a/Engine/Plugins/Mutable/Source/CustomizableObjectEditor/Private/MuCOE/Nodes/CustomizableObjectNodeProjectorConstant.cpp b/Engine/Plugins/Mutable/Source/CustomizableObjectEditor/Private/MuCOE/Nodes/CustomizableObjectNodeProjectorConstant.cpp
--- a/Engine/Plugins/Mutable/Source/CustomizableObjectEditor/Private/MuCOE/Nodes/CustomizableObjectNodeProjectorConstant.cpp
+++ b/Engine/Plugins/Mutable/Source/CustomizableObjectEditor/Private/MuCOE/Nodes/CustomizableObjectNodeProjectorConstant.cpp
@@ -59,6 +59,11 @@ void UCustomizableObjectNodeProjectorConstant::BackwardsCompatibleFixup(int32 Cu
 	{
 		Value.ProjectionType = ProjectionType_DEPRECATED;
 	}
+
+	if (CustomizableObjectCustomVersion == FCustomizableObjectCustomVersion::SnapToBoneComponentIndexToName)
+	{
+		ReferenceSkeletonComponent = FName(FString::FromInt(ReferenceSkeletonIndex_DEPRECATED));
+	}
 }
 
 
diff --git a/Engine/Plugins/Mutable/Source/CustomizableObjectEditor/Private/MuCOE/Nodes/CustomizableObjectNodeProjectorParameter.cpp b/Engine/Plugins/Mutable/Source/CustomizableObjectEditor/Private/MuCOE/Nodes/CustomizableObjectNodeProjectorParameter.cpp
--- a/Engine/Plugins/Mutable/Source/CustomizableObjectEditor/Private/MuCOE/Nodes/CustomizableObjectNodeProjectorParameter.cpp
+++ b/Engine/Plugins/Mutable/Source/CustomizableObjectEditor/Private/MuCOE/Nodes/CustomizableObjectNodeProjectorParameter.cpp
@@ -48,6 +48,11 @@ void UCustomizableObjectNodeProjectorParameter::BackwardsCompatibleFixup(int32 C
 	{
 		DefaultValue.ProjectionType = ProjectionType_DEPRECATED;
 	}
+
+	if (CustomizableObjectCustomVersion == FCustomizableObjectCustomVersion::SnapToBoneComponentIndexToName)
+	{
+		ReferenceSkeletonComponent = FName(FString::FromInt(ReferenceSkeletonIndex_DEPRECATED));
+	}
 }
 
 
diff --git a/Engine/Plugins/Mutable/Source/CustomizableObjectEditor/Private/MuCOE/Nodes/CustomizableObjectNodeProjectorParameterDetails.cpp b/Engine/Plugins/Mutable/Source/CustomizableObjectEditor/Private/MuCOE/Nodes/CustomizableObjectNodeProjectorParameterDetails.cpp
--- a/Engine/Plugins/Mutable/Source/CustomizableObjectEditor/Private/MuCOE/Nodes/CustomizableObjectNodeProjectorParameterDetails.cpp
+++ b/Engine/Plugins/Mutable/Source/CustomizableObjectEditor/Private/MuCOE/Nodes/CustomizableObjectNodeProjectorParameterDetails.cpp
@@ -87,10 +87,10 @@ void FCustomizableObjectNodeProjectorParameterDetails::CustomizeDetails( IDetail
 		];
 	}
 
-	TSharedPtr<IPropertyHandle> ReferenceSkeletonIndexPropertyHandle = DetailBuilder.GetProperty("ReferenceSkeletonIndex");
-	const FSimpleDelegate OnReferenceSkeletonIndexChangedDelegate =
-		FSimpleDelegate::CreateSP(this, &FCustomizableObjectNodeProjectorParameterDetails::OnReferenceSkeletonIndexChanged);
-	ReferenceSkeletonIndexPropertyHandle->SetOnPropertyValueChanged(OnReferenceSkeletonIndexChangedDelegate);
+	TSharedPtr<IPropertyHandle> ReferenceSkeletonIndexPropertyHandle = DetailBuilder.GetProperty("ReferenceSkeletonComponent");
+	const FSimpleDelegate OnReferenceSkeletonComponentChangedDelegate =
+		FSimpleDelegate::CreateSP(this, &FCustomizableObjectNodeProjectorParameterDetails::OnReferenceSkeletonComponentChanged);
+	ReferenceSkeletonIndexPropertyHandle->SetOnPropertyValueChanged(OnReferenceSkeletonComponentChangedDelegate);
 
 	IDetailCategoryBuilder& SnapToBoneCategory = DetailBuilder.EditCategory("ProjectorSnapToBone");
 	DetailBuilder.HideProperty("ProjectorBone");
@@ -100,26 +100,26 @@ void FCustomizableObjectNodeProjectorParameterDetails::CustomizeDetails( IDetail
 
 	UCustomizableObject* CustomizableObject = nullptr;
 	FName ProjectorBoneName;
-	int32 ReferenceIndex = 0;
+	FName ReferenceComponent;
 
 	if (NodeConstant != nullptr)
 	{
 		CustomizableObject = Cast<UCustomizableObject>(NodeConstant->GetCustomizableObjectGraph()->GetOuter());
 		ProjectorBoneName = NodeConstant->ProjectorBone;
-		ReferenceIndex = NodeConstant->ReferenceSkeletonIndex;
+		ReferenceComponent = NodeConstant->ReferenceSkeletonComponent;
 	}
 	else if (NodeParameter != nullptr)
 	{
 		CustomizableObject = Cast<UCustomizableObject>(NodeParameter->GetCustomizableObjectGraph()->GetOuter());
 		ProjectorBoneName = NodeParameter->ProjectorBone;
-		ReferenceIndex = NodeParameter->ReferenceSkeletonIndex;
+		ReferenceComponent = NodeParameter->ReferenceSkeletonComponent;
 	}
 
 	//UCustomizableObject* CustomizableObject = Cast<UCustomizableObject>(Node->GetCustomizableObjectGraph()->GetOuter());
 
 	if (CustomizableObject)
 	{
-		SkeletalMesh = CustomizableObject->GetRefSkeletalMesh(ReferenceIndex);
+		SkeletalMesh = CustomizableObject->GetComponentMeshReferenceSkeletalMesh(ReferenceComponent);
 	}
 
 	if (SkeletalMesh)
@@ -294,7 +294,7 @@ void FCustomizableObjectNodeProjectorParameterDetails::OnBoneComboBoxSelectionCh
 }
 
 
-void FCustomizableObjectNodeProjectorParameterDetails::OnReferenceSkeletonIndexChanged()
+void FCustomizableObjectNodeProjectorParameterDetails::OnReferenceSkeletonComponentChanged()
 {
 	if (NodeConstant != nullptr)
 	{
diff --git a/Engine/Plugins/Mutable/Source/CustomizableObjectEditor/Private/MuCOE/Nodes/CustomizableObjectNodeProjectorParameterDetails.h b/Engine/Plugins/Mutable/Source/CustomizableObjectEditor/Private/MuCOE/Nodes/CustomizableObjectNodeProjectorParameterDetails.h
--- a/Engine/Plugins/Mutable/Source/CustomizableObjectEditor/Private/MuCOE/Nodes/CustomizableObjectNodeProjectorParameterDetails.h
+++ b/Engine/Plugins/Mutable/Source/CustomizableObjectEditor/Private/MuCOE/Nodes/CustomizableObjectNodeProjectorParameterDetails.h
@@ -35,5 +35,5 @@ private:
 	class IDetailLayoutBuilder* DetailBuilderPtr;
 
 	void OnBoneComboBoxSelectionChanged(TSharedPtr<FString> Selection, ESelectInfo::Type SelectInfo, TSharedRef<IPropertyHandle> BoneProperty);
-	void OnReferenceSkeletonIndexChanged();
+	void OnReferenceSkeletonComponentChanged();
 };
diff --git a/Engine/Plugins/Mutable/Source/CustomizableObjectEditor/Public/MuCOE/Nodes/CustomizableObjectNodeModifierClipMorph.h b/Engine/Plugins/Mutable/Source/CustomizableObjectEditor/Public/MuCOE/Nodes/CustomizableObjectNodeModifierClipMorph.h
--- a/Engine/Plugins/Mutable/Source/CustomizableObjectEditor/Public/MuCOE/Nodes/CustomizableObjectNodeModifierClipMorph.h
+++ b/Engine/Plugins/Mutable/Source/CustomizableObjectEditor/Public/MuCOE/Nodes/CustomizableObjectNodeModifierClipMorph.h
@@ -28,8 +28,10 @@ public:
 	UPROPERTY()
 	TArray<FString> Tags_DEPRECATED;
 
+	uint32 ReferenceSkeletonIndex_DEPRECATED = 0;
+
 	UPROPERTY(EditAnywhere, Category = MeshToClipAndMorph)
-	uint32 ReferenceSkeletonIndex = 0;
+	FName ReferenceSkeletonComponent;
 
 	UPROPERTY(EditAnywhere, Category = MeshClipParameters, meta = (DisplayName="Morph Start Offset", ToolTip="Offset from the origin of the selected bone to the actual start of the morph."))
 	FVector StartOffset;
diff --git a/Engine/Plugins/Mutable/Source/CustomizableObjectEditor/Public/MuCOE/Nodes/CustomizableObjectNodeProjectorConstant.h b/Engine/Plugins/Mutable/Source/CustomizableObjectEditor/Public/MuCOE/Nodes/CustomizableObjectNodeProjectorConstant.h
--- a/Engine/Plugins/Mutable/Source/CustomizableObjectEditor/Public/MuCOE/Nodes/CustomizableObjectNodeProjectorConstant.h
+++ b/Engine/Plugins/Mutable/Source/CustomizableObjectEditor/Public/MuCOE/Nodes/CustomizableObjectNodeProjectorConstant.h
@@ -32,8 +32,11 @@ public:
 	UPROPERTY(EditAnywhere, Category = CustomizableObject, Meta = (DisplayName = "Projection Angle (degrees)"))
 	float ProjectionAngle;
 
+	UPROPERTY()
+	uint32 ReferenceSkeletonIndex_DEPRECATED = 0;
+
 	UPROPERTY(EditAnywhere, Category = ProjectorSnapToBone)
-	uint32 ReferenceSkeletonIndex = 0;
+	FName ReferenceSkeletonComponent;
 
 	UPROPERTY(EditAnywhere, Category = ProjectorSnapToBone)
 	FName ProjectorBone;
diff --git a/Engine/Plugins/Mutable/Source/CustomizableObjectEditor/Public/MuCOE/Nodes/CustomizableObjectNodeProjectorParameter.h b/Engine/Plugins/Mutable/Source/CustomizableObjectEditor/Public/MuCOE/Nodes/CustomizableObjectNodeProjectorParameter.h
--- a/Engine/Plugins/Mutable/Source/CustomizableObjectEditor/Public/MuCOE/Nodes/CustomizableObjectNodeProjectorParameter.h
+++ b/Engine/Plugins/Mutable/Source/CustomizableObjectEditor/Public/MuCOE/Nodes/CustomizableObjectNodeProjectorParameter.h
@@ -32,8 +32,11 @@ public:
 	UPROPERTY(EditAnywhere, Category = CustomizableObject, Meta = (DisplayName = "Projection Angle (degrees)"))
 	float ProjectionAngle = 360.0f;
 
+	UPROPERTY()
+	uint32 ReferenceSkeletonIndex_DEPRECATED = 0;
+
 	UPROPERTY(EditAnywhere, Category = ProjectorSnapToBone)
-	uint32 ReferenceSkeletonIndex = 0;
+	FName ReferenceSkeletonComponent;
 
 	UPROPERTY(EditAnywhere, Category = ProjectorSnapToBone)
 	FName ProjectorBone;
