[mutable] Fixed Instance Preview Settings "Add new Profile" not opening the level selection window.
- The default behavior for Windows may have changed making the spawning of the window not enough to display it.

diff --git a/Engine/Plugins/Mutable/Source/CustomizableObjectEditor/Private/MuCOE/SCustomizableObjectEditorAdvancedPreviewSettings.cpp b/Engine/Plugins/Mutable/Source/CustomizableObjectEditor/Private/MuCOE/SCustomizableObjectEditorAdvancedPreviewSettings.cpp
--- a/Engine/Plugins/Mutable/Source/CustomizableObjectEditor/Private/MuCOE/SCustomizableObjectEditorAdvancedPreviewSettings.cpp
+++ b/Engine/Plugins/Mutable/Source/CustomizableObjectEditor/Private/MuCOE/SCustomizableObjectEditorAdvancedPreviewSettings.cpp
@@ -35,7 +35,7 @@ void SLevelSelectWindow::Construct(const FArguments& InArgs)
 	FContentBrowserModule& ContentBrowserModule = FModuleManager::LoadModuleChecked<FContentBrowserModule>("ContentBrowser");
 
 	SWindow::Construct(SWindow::FArguments()
-		.Title(LOCTEXT("SNewProfileWindow_Title", "Select a level asset with UObjects coresponding the the filled names"))
+		.Title(LOCTEXT("SNewProfileWindow_Title", "Select a level asset with UObjects corresponding the the filled names"))
 		.SupportsMinimize(false)
 		.SupportsMaximize(false)
 		.ClientSize(FVector2D(450, 450))
@@ -470,6 +470,7 @@ FReply SCustomizableObjectEditorAdvancedPreviewSettings::ShowAddProfileWindow()
 		.CustomizableObjectEditor(CustomizableObjectEditor)
 		.DefaultSettings(DefaultSettings);
 	
+	LevelSelectWindow->ShowModal();
 	return FReply::Handled();
 }
 
