[mutable] Allow creation of tags with the same name than objects and groups.
#rnx

diff --git a/Engine/Plugins/Mutable/Source/CustomizableObjectEditor/Private/MuCOE/SMutableSearchComboBox.cpp b/Engine/Plugins/Mutable/Source/CustomizableObjectEditor/Private/MuCOE/SMutableSearchComboBox.cpp
--- a/Engine/Plugins/Mutable/Source/CustomizableObjectEditor/Private/MuCOE/SMutableSearchComboBox.cpp
+++ b/Engine/Plugins/Mutable/Source/CustomizableObjectEditor/Private/MuCOE/SMutableSearchComboBox.cpp
@@ -144,7 +144,7 @@ void SMutableSearchComboBox::RefreshOptions()
 		FString SearchString = SearchText.ToString();
 		for (const TSharedRef<FFilteredOption>& Option : *OptionsSource)
 		{
-			if (Option->DisplayOption == SearchString)
+			if (Option->DisplayOption == SearchString && !Option->ActualOption.IsEmpty())
 			{
 				bFullMatch = true; 
 				break;
