[Mutable] Do not setup the camera until there is something to display.

#rnx

diff --git a/Engine/Plugins/Mutable/Source/CustomizableObjectEditor/Private/MuCOE/CustomizableObjectEditorViewportClient.cpp b/Engine/Plugins/Mutable/Source/CustomizableObjectEditor/Private/MuCOE/CustomizableObjectEditorViewportClient.cpp
--- a/Engine/Plugins/Mutable/Source/CustomizableObjectEditor/Private/MuCOE/CustomizableObjectEditorViewportClient.cpp
+++ b/Engine/Plugins/Mutable/Source/CustomizableObjectEditor/Private/MuCOE/CustomizableObjectEditorViewportClient.cpp
@@ -776,12 +776,12 @@ void FCustomizableObjectEditorViewportClient::OnInstanceUpdate(UCustomizableObje
 			return;
 		}
 		
-		bIsCameraSetup = true;
-		
 		FVector Center;
 		FVector Extents;
 		Actor.Get()->GetActorBounds(false, Center, Extents, true);
 
+		bIsCameraSetup = Extents.X * Extents.Y * Extents.Z > 0.0;
+		
 		static FRotator CustomOrbitRotation(-33.75, -135, 0);
 		FVector CustomOrbitZoom(0, Extents.GetMax() * 2.5 / (75.0 * PI / 360.0), 0);
 
