[Mutable] Fix clip morph bone selection and invert normal gizmo update.

- For bone selection, the Origin and Normal fields where set after triggering the gizmo update on Property change.
- Changing Invert Normal would not trigger a gizmo reconstruction as no property change was done.

#rnx

diff --git a/Engine/Plugins/Mutable/Source/CustomizableObjectEditor/Private/MuCOE/Nodes/CustomizableObjectNodeModifierClipMorphDetails.cpp b/Engine/Plugins/Mutable/Source/CustomizableObjectEditor/Private/MuCOE/Nodes/CustomizableObjectNodeModifierClipMorphDetails.cpp
--- a/Engine/Plugins/Mutable/Source/CustomizableObjectEditor/Private/MuCOE/Nodes/CustomizableObjectNodeModifierClipMorphDetails.cpp
+++ b/Engine/Plugins/Mutable/Source/CustomizableObjectEditor/Private/MuCOE/Nodes/CustomizableObjectNodeModifierClipMorphDetails.cpp
@@ -114,6 +114,7 @@ void FCustomizableObjectNodeModifierClipMorphDetails::CustomizeDetails(IDetailLa
 				]
 			];
 
+			TSharedRef<IPropertyHandle> InvertPlaneProperty = DetailBuilder.GetProperty("bInvertNormal");
 			MeshClipParametersCategory.AddCustomRow(LOCTEXT("ClipMorphDetails_PlaneNormal", "Invert plane normal"))
 			[
 				SNew(SProperty, BoneProperty)
@@ -134,7 +135,7 @@ void FCustomizableObjectNodeModifierClipMorphDetails::CustomizeDetails(IDetailLa
 						+SHorizontalBox::Slot().HAlign(HAlign_Left)
 						[
 							SNew(SCheckBox)
-							.OnCheckStateChanged(this, &FCustomizableObjectNodeModifierClipMorphDetails::OnInvertNormalCheckboxChanged)
+							.OnCheckStateChanged(this, &FCustomizableObjectNodeModifierClipMorphDetails::OnInvertNormalCheckboxChanged, InvertPlaneProperty)
 							.IsChecked(this, &FCustomizableObjectNodeModifierClipMorphDetails::GetInvertNormalCheckBoxState)
 							.ToolTipText(LOCTEXT("ClipMorphDetails_InvertNormal_Tooltip", "Invert normal direction of the clip plane"))
 						]
@@ -263,14 +264,16 @@ void FCustomizableObjectNodeModifierClipMorphDetails::OnBoneComboBoxSelectionCha
 				Direction = (ChildLocation - Location).GetSafeNormal();
 			}
 			 
-			BoneProperty->SetValue(*BoneComboOptions[OptionIndex].Get());
-			
 			if (Node)
 			{
 				Node->Origin = Location;
 				Node->Normal = Direction;
 			}
 
+			// Set the bone property after Node Origin and Normal update, otherwise the 
+			// viewport gizmo will be constructed with the old values.
+			BoneProperty->SetValue(*BoneComboOptions[OptionIndex].Get());
+			
 			return;
 		}
 	}
@@ -279,34 +282,22 @@ void FCustomizableObjectNodeModifierClipMorphDetails::OnBoneComboBoxSelectionCha
 }
 
 
-void FCustomizableObjectNodeModifierClipMorphDetails::OnInvertNormalCheckboxChanged(ECheckBoxState CheckBoxState)
+void FCustomizableObjectNodeModifierClipMorphDetails::OnInvertNormalCheckboxChanged(ECheckBoxState CheckBoxState, TSharedRef<IPropertyHandle> InvertPlaneProperty)
 {
 	if (Node == nullptr)
 	{
 		return;
 	}
 
-	switch (CheckBoxState)
-	{
-		case ECheckBoxState::Checked:
-		{
-			Node->bInvertNormal = true;
-			Node->Normal *= -1.0f;
-			break;
-		}
-		case ECheckBoxState::Unchecked:
-		{
-			Node->bInvertNormal = false;
-			Node->Normal *= -1.0f;
-			break;
-		}
-	}
-
 	if (Node->bLocalStartOffset)
 	{
-		Node->StartOffset.Z *= -1;
-		Node->StartOffset.X *= -1;
+		Node->StartOffset.Z *= -1.0f;
+		Node->StartOffset.X *= -1.0f;
 	}
+
+	Node->Normal *= -1.0f;
+
+	InvertPlaneProperty->SetValue(CheckBoxState == ECheckBoxState::Checked);
 }
 
 
diff --git a/Engine/Plugins/Mutable/Source/CustomizableObjectEditor/Private/MuCOE/Nodes/CustomizableObjectNodeModifierClipMorphDetails.h b/Engine/Plugins/Mutable/Source/CustomizableObjectEditor/Private/MuCOE/Nodes/CustomizableObjectNodeModifierClipMorphDetails.h
--- a/Engine/Plugins/Mutable/Source/CustomizableObjectEditor/Private/MuCOE/Nodes/CustomizableObjectNodeModifierClipMorphDetails.h
+++ b/Engine/Plugins/Mutable/Source/CustomizableObjectEditor/Private/MuCOE/Nodes/CustomizableObjectNodeModifierClipMorphDetails.h
@@ -33,7 +33,7 @@ private:
 
 	void OnBoneComboBoxSelectionChanged(TSharedPtr<FString> Selection, ESelectInfo::Type SelectInfo, TSharedRef<IPropertyHandle> BoneProperty);
 
-	void OnInvertNormalCheckboxChanged(ECheckBoxState CheckBoxState);
+	void OnInvertNormalCheckboxChanged(ECheckBoxState CheckBoxState, TSharedRef<IPropertyHandle> InvertPlaneProperty);
 	ECheckBoxState GetInvertNormalCheckBoxState() const;
 
 	void OnReferenceSkeletonComponentChanged();
