[Mutable] Fixed states not working.

We were not detecting partial compilations correctly.

#rnx

diff --git a/Engine/Plugins/Mutable/Source/CustomizableObjectEditor/Private/MuCOE/CustomizableObjectCompiler.cpp b/Engine/Plugins/Mutable/Source/CustomizableObjectEditor/Private/MuCOE/CustomizableObjectCompiler.cpp
--- a/Engine/Plugins/Mutable/Source/CustomizableObjectEditor/Private/MuCOE/CustomizableObjectCompiler.cpp
+++ b/Engine/Plugins/Mutable/Source/CustomizableObjectEditor/Private/MuCOE/CustomizableObjectCompiler.cpp
@@ -663,7 +663,7 @@ mu::Ptr<mu::NodeObject> GenerateMutableRoot(
 		GenerationContext.TableToParamNames = Object->GetPrivate()->GetModelResources().TableToParamNames;
 	}
 
-	GenerationContext.bPartialCompilation = !Root->ParentObject;
+	GenerationContext.bPartialCompilation = Root->ParentObject != nullptr;
 
 	// Generate the object expression
 	UE_LOG(LogMutable, Verbose, TEXT("PROFILE: [ %16.8f ] GenerateMutableSource start."), FPlatformTime::Seconds());
