[Mutable] Sort Participating Objects for DDC key.

Also standarize FName case.

#rnx

diff --git a/Engine/Plugins/Mutable/Source/CustomizableObject/Private/MuCO/CustomizableObject.cpp b/Engine/Plugins/Mutable/Source/CustomizableObject/Private/MuCO/CustomizableObject.cpp
--- a/Engine/Plugins/Mutable/Source/CustomizableObject/Private/MuCO/CustomizableObject.cpp
+++ b/Engine/Plugins/Mutable/Source/CustomizableObject/Private/MuCO/CustomizableObject.cpp
@@ -2898,11 +2898,20 @@ TArray<uint8> UCustomizableObjectPrivate::BuildDerivedDataKey(FCompilationOption
 	// Participating objects hash
 	if (ICustomizableObjectEditorModule* Module = ICustomizableObjectEditorModule::Get())
 	{
-		TMap<FName, FGuid> ParticipatingObjects = Module->GetParticipatingObjects(GetPublic(), true, &Options);
-		for (TTuple<FName, FGuid>& Tuple : ParticipatingObjects)
+		TArray<TTuple<FName, FGuid>> ParticipatingObjects = Module->GetParticipatingObjects(GetPublic(), true, &Options).Array();
+		ParticipatingObjects.Sort([](const TTuple<FName, FGuid>& A, const TTuple<FName, FGuid>& B)
 		{
-			Ar << Tuple.Get<0>();
-			Ar << Tuple.Get<1>();
+			return A.Get<0>().LexicalLess(B.Get<0>()) && A.Get<1>() < B.Get<1>();
+		});
+		
+		for (const TTuple<FName, FGuid>& Tuple : ParticipatingObjects)
+		{
+			FString Key = Tuple.Get<0>().ToString();
+			Key.ToLowerInline();
+			Ar << Key;
+
+			FGuid Id = Tuple.Get<1>();
+			Ar << Id;
 		}
 	}
 
