[Mutable] Fix Failed to find Parameter.

- Update the Instance details after compiling a CO. If not the details will point to removed parameters.

#rnx

diff --git a/Engine/Plugins/Mutable/Source/CustomizableObjectEditor/Private/MuCOE/CustomizableInstanceDetails.cpp b/Engine/Plugins/Mutable/Source/CustomizableObjectEditor/Private/MuCOE/CustomizableInstanceDetails.cpp
--- a/Engine/Plugins/Mutable/Source/CustomizableObjectEditor/Private/MuCOE/CustomizableInstanceDetails.cpp
+++ b/Engine/Plugins/Mutable/Source/CustomizableObjectEditor/Private/MuCOE/CustomizableInstanceDetails.cpp
@@ -78,7 +78,7 @@ void FCustomizableInstanceDetails::CustomizeDetails(const TSharedPtr<IDetailLayo
 		CustomInstance->GetPrivate()->SaveParametersToProfile(CustomInstance->GetPrivate()->SelectedProfileIndex);
 	}
 
-	// Delegate to refresh the detils when the instance has finished the Update
+	// Delegate to refresh the details when the instance has finished the Update
 	CustomInstance->UpdatedNativeDelegate.AddSP(this, &FCustomizableInstanceDetails::InstanceUpdated);
 	CustomInstance->GetPrivate()->OnInstanceTransactedDelegate.AddSP(this, &FCustomizableInstanceDetails::OnInstanceTransacted);
 
@@ -102,6 +102,8 @@ void FCustomizableInstanceDetails::CustomizeDetails(const TSharedPtr<IDetailLayo
 		return;
 	}
 
+	CustomizableObject->GetPostCompileDelegate().AddSP(this, &FCustomizableInstanceDetails::ObjectCompiled);
+
 	TArray<UObject*> Private;
 	Private.Add(CustomInstance->GetPrivate());
 
@@ -303,6 +305,15 @@ void FCustomizableInstanceDetails::InstanceUpdated(UCustomizableObjectInstance*
 }
 
 
+void FCustomizableInstanceDetails::ObjectCompiled()
+{
+	if (!bUpdatingSlider)
+	{
+		Refresh();
+	}
+}
+
+
 // STATE SELECTOR -----------------------------------------------------------------------------------------------------------------
 
 TSharedRef<SWidget> FCustomizableInstanceDetails::GenerateStateSelector()
@@ -1174,7 +1185,7 @@ TSharedRef<SWidget> FCustomizableInstanceDetails::GenerateFloatWidget(const UCus
 
 float FCustomizableInstanceDetails::GetFloatParameterValue(FString ParamName, int32 RangeIndex) const
 {
-	if (CustomInstance->GetCustomizableObject()->GetPrivate()->IsLocked())
+	if (CustomInstance->GetCustomizableObject()->GetPrivate()->IsLocked()) // TODO Move, if necessary, to GetFloatParameterSelectedOption. UE-224815
 	{
 		// Prevent crashing if polling the float value during CO compilation
 		return -1.f;
diff --git a/Engine/Plugins/Mutable/Source/CustomizableObjectEditor/Private/MuCOE/CustomizableInstanceDetails.h b/Engine/Plugins/Mutable/Source/CustomizableObjectEditor/Private/MuCOE/CustomizableInstanceDetails.h
--- a/Engine/Plugins/Mutable/Source/CustomizableObjectEditor/Private/MuCOE/CustomizableInstanceDetails.h
+++ b/Engine/Plugins/Mutable/Source/CustomizableObjectEditor/Private/MuCOE/CustomizableInstanceDetails.h
@@ -53,6 +53,8 @@ private:
 	// Callback to regenerate the details when the instance has finished an update
 	void InstanceUpdated(UCustomizableObjectInstance* Instance) const;
 
+	void ObjectCompiled();
+
 	// State Selector 
 	// Generates The StateSelector Widget
 	TSharedRef<SWidget> GenerateStateSelector();
