[mutable] Fix possible crashes with incomplete graph compilations involving texture variations.
#rnx

diff --git a/Engine/Plugins/Mutable/Source/MutableTools/Private/MuT/ASTOpImageLayer.cpp b/Engine/Plugins/Mutable/Source/MutableTools/Private/MuT/ASTOpImageLayer.cpp
--- a/Engine/Plugins/Mutable/Source/MutableTools/Private/MuT/ASTOpImageLayer.cpp
+++ b/Engine/Plugins/Mutable/Source/MutableTools/Private/MuT/ASTOpImageLayer.cpp
@@ -181,9 +181,14 @@ namespace mu
 	{
 		Ptr<ASTOp> at;
 
-		auto baseAt = base.child();
-		auto blendAt = blend.child();
-		auto maskAt = mask.child();
+		Ptr<ASTOp> baseAt = base.child();
+		Ptr<ASTOp> blendAt = blend.child();
+		Ptr<ASTOp> maskAt = mask.child();
+
+		if (!baseAt)
+		{
+			return at;
+		}
 
 		// Convert to image layer color if blend is plain
 		if (!at && blendAt && blendAt->GetOpType() == OP_TYPE::IM_PLAINCOLOUR)
@@ -531,9 +536,14 @@ namespace mu
 
 		// Layer effects may be worth sinking down switches and conditionals, to be able
 		// to apply extra optimisations
-		auto baseAt = base.child();
-		auto blendAt = blend.child();
-		auto maskAt = mask.child();
+		Ptr<ASTOp> baseAt = base.child();
+		Ptr<ASTOp> blendAt = blend.child();
+		Ptr<ASTOp> maskAt = mask.child();
+
+		if (!baseAt)
+		{
+			return at;
+		}
 
 		// Promote conditions from the base
 		OP_TYPE baseType = baseAt->GetOpType();
diff --git a/Engine/Plugins/Mutable/Source/MutableTools/Private/MuT/ASTOpImagePixelFormat.cpp b/Engine/Plugins/Mutable/Source/MutableTools/Private/MuT/ASTOpImagePixelFormat.cpp
--- a/Engine/Plugins/Mutable/Source/MutableTools/Private/MuT/ASTOpImagePixelFormat.cpp
+++ b/Engine/Plugins/Mutable/Source/MutableTools/Private/MuT/ASTOpImagePixelFormat.cpp
@@ -126,7 +126,7 @@ namespace mu
 		//bool isBlockFormat = GetImageFormatData( format ).PixelsPerBlockX!=0;
 
 		// The instruction can be sunk
-		OP_TYPE sourceType = sourceAt->GetOpType();
+		OP_TYPE sourceType = sourceAt ? sourceAt->GetOpType() : OP_TYPE::NONE;
 		switch (sourceType)
 		{
 		case OP_TYPE::IM_PIXELFORMAT:
@@ -260,7 +260,10 @@ namespace mu
 			}
 		}
 
-		res = Source.child()->GetImageDesc(returnBestOption, context);
+		if (Source.child())
+		{
+			res = Source.child()->GetImageDesc(returnBestOption, context);
+		}
 
 		if (FormatIfAlpha != EImageFormat::IF_NONE
 			&&
diff --git a/Engine/Plugins/Mutable/Source/MutableTools/Private/MuT/CodeGenerator_Image.cpp b/Engine/Plugins/Mutable/Source/MutableTools/Private/MuT/CodeGenerator_Image.cpp
--- a/Engine/Plugins/Mutable/Source/MutableTools/Private/MuT/CodeGenerator_Image.cpp
+++ b/Engine/Plugins/Mutable/Source/MutableTools/Private/MuT/CodeGenerator_Image.cpp
@@ -1840,7 +1840,12 @@ namespace mu
         }
 
         // Make sure all options are the same format and size
-        FImageDesc desc = currentOp->GetImageDesc( true );
+		FImageDesc desc;
+		if (currentOp)
+		{
+			desc = currentOp->GetImageDesc(true);
+		}
+
         if ( desc.m_format == EImageFormat::IF_NONE )
         {
             // TODO: Look for the most generic of the options?
