[Mutable] Enable DDC for cook.

#rnx

diff --git a/Engine/Plugins/Mutable/Source/CustomizableObject/Private/MuCO/CustomizableObject.cpp b/Engine/Plugins/Mutable/Source/CustomizableObject/Private/MuCO/CustomizableObject.cpp
--- a/Engine/Plugins/Mutable/Source/CustomizableObject/Private/MuCO/CustomizableObject.cpp
+++ b/Engine/Plugins/Mutable/Source/CustomizableObject/Private/MuCO/CustomizableObject.cpp
@@ -78,7 +78,7 @@ TAutoConsoleVariable<bool> CVarMutableUseBulkData(
 
 TAutoConsoleVariable<int32> CVarMutableDerivedDataCacheUsage(
 	TEXT("mutable.DerivedDataCacheUsage"),
-	0,
+	2,
 	TEXT("Derived data cache access for cooked data.")
 	TEXT("0 - None. Disables access to the cache.")
 	TEXT("1 - Local. Allow cache requests to query and store records and values in local caches.")
diff --git a/Engine/Plugins/Mutable/Source/CustomizableObjectEditor/Private/MuCOE/CustomizableObjectCompileRunnable.cpp b/Engine/Plugins/Mutable/Source/CustomizableObjectEditor/Private/MuCOE/CustomizableObjectCompileRunnable.cpp
--- a/Engine/Plugins/Mutable/Source/CustomizableObjectEditor/Private/MuCOE/CustomizableObjectCompileRunnable.cpp
+++ b/Engine/Plugins/Mutable/Source/CustomizableObjectEditor/Private/MuCOE/CustomizableObjectCompileRunnable.cpp
@@ -506,7 +506,6 @@ void FCustomizableObjectSaveDDRunnable::StoreCachedPlatformDataInDDC(bool& bStor
 			});
 
 		RequestOwner.Wait();
-		check(bStoredSuccessfully);
 	}
 }
 
