[mutable] Use tag search widget to "variation" node's UI.
Include variation nodes in tag searches.
#rnx

diff --git a/Engine/Plugins/Mutable/Source/CustomizableObjectEditor/Private/MuCOE/CustomizableObjectEditorModule.cpp b/Engine/Plugins/Mutable/Source/CustomizableObjectEditor/Private/MuCOE/CustomizableObjectEditorModule.cpp
--- a/Engine/Plugins/Mutable/Source/CustomizableObjectEditor/Private/MuCOE/CustomizableObjectEditorModule.cpp
+++ b/Engine/Plugins/Mutable/Source/CustomizableObjectEditor/Private/MuCOE/CustomizableObjectEditorModule.cpp
@@ -66,6 +66,9 @@
 #include "MuCOE/Nodes/CustomizableObjectNodeTable.h"
 #include "MuCOE/Nodes/CustomizableObjectNodeTableDetails.h"
 #include "MuCOE/Nodes/CustomizableObjectNodeComponentMeshDetails.h"
+#include "MuCOE/Nodes/CustomizableObjectNodeVariation.h"
+#include "MuCOE/Nodes/CustomizableObjectNodeTextureVariation.h"
+#include "MuCOE/Widgets/CustomizableObjectVariationCustomization.h"
 #include "MuCOE/Widgets/CustomizableObjectLODReductionSettings.h"
 #include "PropertyEditorModule.h"
 #include "AssetRegistry/AssetRegistryModule.h"
@@ -271,11 +274,13 @@ void FCustomizableObjectEditorModule::StartupModule()
 	RegisterCustomDetails(PropertyModule, UCustomizableObjectInstance::StaticClass(), FOnGetDetailCustomizationInstance::CreateStatic(&FCustomizableInstanceDetails::MakeInstance));
 	RegisterCustomDetails(PropertyModule, UCustomSettings::StaticClass(), FOnGetDetailCustomizationInstance::CreateStatic(&FCustomizableObjectCustomSettingsDetails::MakeInstance));
 
-	//Custom properties
+	// Custom properties
 	PropertyModule.RegisterCustomPropertyTypeLayout("CustomizableObjectIdentifier", FOnGetPropertyTypeCustomizationInstance::CreateStatic(&FCustomizableObjectIdentifierCustomization::MakeInstance));
 	PropertyModule.RegisterCustomPropertyTypeLayout(FMeshReshapeBoneReference::StaticStruct()->GetFName(), FOnGetPropertyTypeCustomizationInstance::CreateStatic(&FMeshReshapeBonesReferenceCustomization::MakeInstance));
 	PropertyModule.RegisterCustomPropertyTypeLayout(FBoneToRemove::StaticStruct()->GetFName(), FOnGetPropertyTypeCustomizationInstance::CreateStatic(&FCustomizableObjectLODReductionSettings::MakeInstance));
 	PropertyModule.RegisterCustomPropertyTypeLayout(NAME_StrProperty, FOnGetPropertyTypeCustomizationInstance::CreateStatic(&FCustomizableObjectStateParameterSelector::MakeInstance), MakeShared<FStatePropertyTypeIdentifier>());
+	PropertyModule.RegisterCustomPropertyTypeLayout(FCustomizableObjectVariation::StaticStruct()->GetFName(), FOnGetPropertyTypeCustomizationInstance::CreateStatic(&FCustomizableObjectVariationCustomization::MakeInstance));
+	PropertyModule.RegisterCustomPropertyTypeLayout(FCustomizableObjectTextureVariation::StaticStruct()->GetFName(), FOnGetPropertyTypeCustomizationInstance::CreateStatic(&FCustomizableObjectVariationCustomization::MakeInstance));
 
 	PropertyModule.NotifyCustomizationModuleChanged();
 
diff --git a/Engine/Plugins/Mutable/Source/CustomizableObjectEditor/Private/MuCOE/SMutableTagListWidget.cpp b/Engine/Plugins/Mutable/Source/CustomizableObjectEditor/Private/MuCOE/SMutableTagListWidget.cpp
--- a/Engine/Plugins/Mutable/Source/CustomizableObjectEditor/Private/MuCOE/SMutableTagListWidget.cpp
+++ b/Engine/Plugins/Mutable/Source/CustomizableObjectEditor/Private/MuCOE/SMutableTagListWidget.cpp
@@ -10,6 +10,12 @@
 #include "MuCOE/Nodes/CustomizableObjectNodeObjectGroup.h"
 #include "MuCOE/Nodes/CustomizableObjectNodeMaterial.h"
 #include "MuCOE/Nodes/CustomizableObjectNodeModifierBase.h"
+#include "MuCOE/Nodes/CustomizableObjectNodeVariation.h"
+#include "MuCOE/Nodes/CustomizableObjectNodeTextureVariation.h"
+#include "MuCOE/Nodes/CustomizableObjectNodeFloatVariation.h"
+#include "MuCOE/Nodes/CustomizableObjectNodeColorVariation.h"
+#include "MuCOE/Nodes/CustomizableObjectNodeMeshVariation.h"
+#include "MuCOE/Nodes/CustomizableObjectNodeComponentVariation.h"
 #include "MuCOE/UnrealEditorPortabilityHelpers.h"
 #include "PropertyCustomizationHelpers.h"
 #include "Widgets/Input/STextComboBox.h"
@@ -23,71 +29,28 @@
 #define LOCTEXT_NAMESPACE "CustomizableObjectDetails"
 
 
-void SMutableTagListWidget::Construct(const FArguments& InArgs)
+void SMutableTagComboBox::Construct(const FArguments& InArgs)
 {
 	Node = InArgs._Node;
-	TagArray = InArgs._TagArray;
-	EmptyListText = InArgs._EmptyListText;
 
-	OnTagListChangedDelegate = InArgs._OnTagListChanged;
-
 	RefreshOptions();
 
-	ChildSlot
-	[
-		SNew(SVerticalBox)
-
-		// Header with the "add tag" UI
-		+ SVerticalBox::Slot()
-		.AutoHeight()
-		.HAlign(HAlign_Fill)
+	SMutableSearchComboBox::Construct(
+		SMutableSearchComboBox::FArguments()
+		.OptionsSource(&TagComboOptionsSource)
+		.AllowAddNewOptions(true)
+		.ButtonStyle(FAppStyle::Get(), "NoBorder")
+		.MenuButtonBrush(InArgs._MenuButtonBrush)
+		.OnSelectionChanged(InArgs._OnSelectionChanged)
+		.Content()
 		[
-			SNew(SHorizontalBox)
-
-				+ SHorizontalBox::Slot()
-				.FillWidth(10.0f)
-				.HAlign(HAlign_Fill)
-				.VAlign(VAlign_Center)
-				[
-					SAssignNew(this->TagCombo, SMutableSearchComboBox)
-						.ButtonStyle(FAppStyle::Get(), "NoBorder")
-						.MenuButtonBrush(FAppStyle::GetBrush(TEXT("Icons.PlusCircle")))
-						.OptionsSource(&TagComboOptionsSource)
-						.AllowAddNewOptions(true)
-						.OnSelectionChanged(this, &SMutableTagListWidget::OnTagComboBoxSelectionChanged)
-				]
+			InArgs._Content.Widget
 		]
-
-		// List of tags
-		+ SVerticalBox::Slot()
-		.VAlign(VAlign_Fill)
-		.HAlign(HAlign_Fill)
-		[
-			SAssignNew(this->TagListWidget, SListView<TSharedPtr<FTagUIData>>)
-				.ListItemsSource(&CurrentTagsSource)
-				.OnGenerateRow(this, &SMutableTagListWidget::GenerateTagListItemRow)
-				.SelectionMode(ESelectionMode::None)
-		]
-
-		// Shown only if there are no tags defined
-		+ SVerticalBox::Slot()
-		.AutoHeight()
-		.HAlign(HAlign_Fill)
-		[
-			SNew(STextBlock)
-				.Text(EmptyListText)
-				.AutoWrapText(true)
-				.Font(IDetailLayoutBuilder::GetDetailFont())
-				.Visibility_Lambda([&]()
-					{
-						return (!TagArray || TagArray->IsEmpty()) ? EVisibility::Visible : EVisibility::Collapsed;
-					})
-		]
-	];
+	);
 }
 
 
-TSharedPtr<SMutableSearchComboBox::FFilteredOption> SMutableTagListWidget::AddNodeHierarchyOptions(UEdGraphNode* InNode, TMap<UEdGraphNode*, TSharedPtr<SMutableSearchComboBox::FFilteredOption>>& AddedOptions )
+TSharedPtr<SMutableSearchComboBox::FFilteredOption> SMutableTagComboBox::AddNodeHierarchyOptions(UEdGraphNode* InNode, TMap<UEdGraphNode*, TSharedPtr<SMutableSearchComboBox::FFilteredOption>>& AddedOptions)
 {
 	TSharedPtr<SMutableSearchComboBox::FFilteredOption> Option;
 	if (TSharedPtr<SMutableSearchComboBox::FFilteredOption>* FoundCached = AddedOptions.Find(InNode))
@@ -144,7 +107,7 @@
 			Option = MakeShared<SMutableSearchComboBox::FFilteredOption>();
 			Option->Parent = ParentOption;
 			UMaterialInterface* Material = MeshSectionNode->GetMaterial();
-			Option->DisplayOption = FString::Printf(TEXT("Mesh Section [%s]"), Material?*Material->GetName():TEXT("no-material"));
+			Option->DisplayOption = FString::Printf(TEXT("Mesh Section [%s]"), Material ? *Material->GetName() : TEXT("no-material"));
 			TagComboOptionsSource.Add(Option.ToSharedRef());
 		}
 
@@ -194,67 +157,105 @@
 }
 
 
-void SMutableTagListWidget::RefreshOptions()
+void SMutableTagComboBox::RefreshOptions()
 {
-	// Collect all possible tags for the "add" combo menu
+	TagComboOptionsSource.SetNum(0, EAllowShrinking::No);
+
+	// Scan all potential receivers
+	UCustomizableObject* ThisNodeObject = GetRootObject(*Node);
+	UCustomizableObject* RootObject = GraphTraversal::GetRootObject(ThisNodeObject);
+
+	TSet<UCustomizableObject*> AllCustomizableObject;
+	GetAllObjectsInGraph(RootObject, AllCustomizableObject);
+
+	TMap<UEdGraphNode*, TSharedPtr<SMutableSearchComboBox::FFilteredOption>> AddedOptions;
+
+	for (const UCustomizableObject* CustObject : AllCustomizableObject)
 	{
-		TagComboOptionsSource.SetNum(0, EAllowShrinking::No);
+		if (!CustObject)
+		{
+			continue;
+		}
 
-		// Scan all potential receivers
-		UCustomizableObject* ThisNodeObject = GetRootObject(*Node);
-		UCustomizableObject* RootObject = GraphTraversal::GetRootObject(ThisNodeObject);
-
-		TSet<UCustomizableObject*> AllCustomizableObject;
-		GetAllObjectsInGraph(RootObject, AllCustomizableObject);
-
-		TMap<UEdGraphNode*, TSharedPtr<SMutableSearchComboBox::FFilteredOption>> AddedOptions;
-
-		for (const UCustomizableObject* CustObject : AllCustomizableObject)
+		for (const TObjectPtr<UEdGraphNode>& CandidateNode : CustObject->GetPrivate()->GetSource()->Nodes)
 		{
-			if (!CustObject)
+			UCustomizableObjectNode* Typed = Cast<UCustomizableObjectNode>(CandidateNode);
+			if (!Typed)
 			{
 				continue;
 			}
 
+			TArray<FString>* EnableTags = Typed->GetEnableTags();
+			if (EnableTags)
+			{
+				TSharedPtr<SMutableSearchComboBox::FFilteredOption> NodeOption = AddNodeHierarchyOptions(Typed, AddedOptions);
 
-			for (const TObjectPtr<UEdGraphNode>& CandidateNode : CustObject->GetPrivate()->GetSource()->Nodes)
+				for (const FString& OneTag : *EnableTags)
+				{
+					if (!OneTag.IsEmpty())
+					{
+						TSharedRef Option = MakeShared<SMutableSearchComboBox::FFilteredOption>();
+						Option->ActualOption = OneTag;
+						Option->DisplayOption = OneTag;
+						Option->Parent = NodeOption;
+						TagComboOptionsSource.Add(Option);
+					}
+				}
+			}
+
+			UCustomizableObjectNodeModifierBase* TypedModifier = Cast<UCustomizableObjectNodeModifierBase>(CandidateNode);
+			if (TypedModifier)
 			{
-				UCustomizableObjectNode* Typed = Cast<UCustomizableObjectNode>(CandidateNode);
-				if (!Typed)
+				for (const FString& OneTag : TypedModifier->RequiredTags)
 				{
-					continue;
+					TSharedPtr<SMutableSearchComboBox::FFilteredOption> NodeOption = AddNodeHierarchyOptions(Typed, AddedOptions);
+
+					if (!OneTag.IsEmpty())
+					{
+						TSharedRef Option = MakeShared<SMutableSearchComboBox::FFilteredOption>();
+						Option->ActualOption = OneTag;
+						Option->DisplayOption = OneTag;
+						Option->Parent = NodeOption;
+						TagComboOptionsSource.Add(Option);
+					}
 				}
+			}
 
-				TArray<FString>* EnableTags = Typed->GetEnableTags();
-				if (EnableTags)
+			// Generic variation nodes
+			{
+				UCustomizableObjectNodeVariation* VariationNode = Cast<UCustomizableObjectNodeVariation>(CandidateNode);
+				if (VariationNode)
 				{
-					TSharedPtr<SMutableSearchComboBox::FFilteredOption> NodeOption = AddNodeHierarchyOptions(Typed,AddedOptions);
+					for (const FCustomizableObjectVariation& Var : VariationNode->VariationsData)
+					{
+						TSharedPtr<SMutableSearchComboBox::FFilteredOption> NodeOption = AddNodeHierarchyOptions(VariationNode, AddedOptions);
 
-					for (const FString& OneTag : *EnableTags)
-					{
-						if (!OneTag.IsEmpty())
+						if (!Var.Tag.IsEmpty())
 						{
 							TSharedRef Option = MakeShared<SMutableSearchComboBox::FFilteredOption>();
-							Option->ActualOption = OneTag;
-							Option->DisplayOption = OneTag;
+							Option->ActualOption = Var.Tag;
+							Option->DisplayOption = Var.Tag;
 							Option->Parent = NodeOption;
 							TagComboOptionsSource.Add(Option);
 						}
 					}
 				}
+			}
 
-				UCustomizableObjectNodeModifierBase* TypedModifier = Cast<UCustomizableObjectNodeModifierBase>(CandidateNode);
-				if (TypedModifier)
+			// Texture variation nodes
+			{
+				UCustomizableObjectNodeTextureVariation* VariationNode = Cast<UCustomizableObjectNodeTextureVariation>(CandidateNode);
+				if (VariationNode)
 				{
-					for (const FString& OneTag : TypedModifier->RequiredTags)
+					for (const FCustomizableObjectTextureVariation& Var : VariationNode->Variations)
 					{
-						TSharedPtr<SMutableSearchComboBox::FFilteredOption> NodeOption = AddNodeHierarchyOptions(Typed, AddedOptions);
+						TSharedPtr<SMutableSearchComboBox::FFilteredOption> NodeOption = AddNodeHierarchyOptions(VariationNode, AddedOptions);
 
-						if (!OneTag.IsEmpty())
+						if (!Var.Tag.IsEmpty())
 						{
 							TSharedRef Option = MakeShared<SMutableSearchComboBox::FFilteredOption>();
-							Option->ActualOption = OneTag;
-							Option->DisplayOption = OneTag;
+							Option->ActualOption = Var.Tag;
+							Option->DisplayOption = Var.Tag;
 							Option->Parent = NodeOption;
 							TagComboOptionsSource.Add(Option);
 						}
@@ -262,8 +263,77 @@
 				}
 			}
 		}
+	}
+}
 
-		// TODO: add material nodes and options to create tags for them, show hierarchy, ...
+
+void SMutableTagListWidget::Construct(const FArguments& InArgs)
+{
+	Node = InArgs._Node;
+	TagArray = InArgs._TagArray;
+	EmptyListText = InArgs._EmptyListText;
+
+	OnTagListChangedDelegate = InArgs._OnTagListChanged;
+
+	RefreshOptions();
+
+	ChildSlot
+	[
+		SNew(SVerticalBox)
+
+		// Header with the "add tag" UI
+		+ SVerticalBox::Slot()
+		.AutoHeight()
+		.HAlign(HAlign_Fill)
+		[
+			SNew(SHorizontalBox)
+
+				+ SHorizontalBox::Slot()
+				.FillWidth(10.0f)
+				.HAlign(HAlign_Fill)
+				.VAlign(VAlign_Center)
+				[
+					SAssignNew(this->TagCombo, SMutableTagComboBox)
+						.Node(Node)
+						.MenuButtonBrush( FAppStyle::GetBrush(TEXT("Icons.PlusCircle")) )
+						.OnSelectionChanged(this, &SMutableTagListWidget::OnTagComboBoxSelectionChanged)
+				]
+		]
+
+		// List of tags
+		+ SVerticalBox::Slot()
+		.VAlign(VAlign_Fill)
+		.HAlign(HAlign_Fill)
+		[
+			SAssignNew(this->TagListWidget, SListView<TSharedPtr<FTagUIData>>)
+				.ListItemsSource(&CurrentTagsSource)
+				.OnGenerateRow(this, &SMutableTagListWidget::GenerateTagListItemRow)
+				.SelectionMode(ESelectionMode::None)
+		]
+
+		// Shown only if there are no tags defined
+		+ SVerticalBox::Slot()
+		.AutoHeight()
+		.HAlign(HAlign_Fill)
+		[
+			SNew(STextBlock)
+				.Text(EmptyListText)
+				.AutoWrapText(true)
+				.Font(IDetailLayoutBuilder::GetDetailFont())
+				.Visibility_Lambda([&]()
+					{
+						return (!TagArray || TagArray->IsEmpty()) ? EVisibility::Visible : EVisibility::Collapsed;
+					})
+		]
+	];
+}
+
+
+void SMutableTagListWidget::RefreshOptions()
+{
+	if (TagCombo)
+	{
+		TagCombo->RefreshOptions();
 	}
 
 	// Current Tags
diff --git a/Engine/Plugins/Mutable/Source/CustomizableObjectEditor/Private/MuCOE/SMutableTagListWidget.h b/Engine/Plugins/Mutable/Source/CustomizableObjectEditor/Private/MuCOE/SMutableTagListWidget.h
--- a/Engine/Plugins/Mutable/Source/CustomizableObjectEditor/Private/MuCOE/SMutableTagListWidget.h
+++ b/Engine/Plugins/Mutable/Source/CustomizableObjectEditor/Private/MuCOE/SMutableTagListWidget.h
@@ -20,6 +20,41 @@ class UCustomizableObjectNode;
 class UCustomizableObjectNodeObject;
 class UEdGraphNode;
 
+
+/** */
+class SMutableTagComboBox : public SMutableSearchComboBox
+{
+public:
+
+	SLATE_BEGIN_ARGS(SMutableTagComboBox)
+		: _Node(nullptr)
+		{}
+		/** Slot for this button's content (optional) */
+		SLATE_DEFAULT_SLOT(FArguments, Content)
+		SLATE_ARGUMENT(UCustomizableObjectNode*, Node)
+		SLATE_ARGUMENT(const FSlateBrush*, MenuButtonBrush)
+		SLATE_EVENT(FOnTextChanged, OnSelectionChanged)
+
+	SLATE_END_ARGS()
+
+	/** */
+	void Construct(const FArguments& InArgs);
+
+	/** */
+	void RefreshOptions();
+
+private:
+
+	class UCustomizableObjectNode* Node = nullptr;
+
+	TArray< TSharedRef<SMutableSearchComboBox::FFilteredOption> > TagComboOptionsSource;
+
+	/** */
+	TSharedPtr<SMutableSearchComboBox::FFilteredOption> AddNodeHierarchyOptions(UEdGraphNode* Node, TMap<UEdGraphNode*, TSharedPtr<SMutableSearchComboBox::FFilteredOption>>& AddedOptions);
+
+};
+
+
 /** */
 class SMutableTagListWidget : public SCompoundWidget
 {
@@ -53,7 +88,7 @@ private:
 	TArray<FString>* TagArray = nullptr;
 	FText EmptyListText;
 
-	TSharedPtr<SMutableSearchComboBox> TagCombo;
+	TSharedPtr<SMutableTagComboBox> TagCombo;
 	TArray< TSharedRef<SMutableSearchComboBox::FFilteredOption> > TagComboOptionsSource;
 	void OnTagComboBoxSelectionChanged(const FText& NewText);
 
@@ -68,8 +103,5 @@ private:
 	/** */
 	TSharedRef<ITableRow> GenerateTagListItemRow(TSharedPtr<FTagUIData> InItem, const TSharedRef<STableViewBase>& OwnerTable);
 
-	/** */
-	TSharedPtr<SMutableSearchComboBox::FFilteredOption> AddNodeHierarchyOptions(UEdGraphNode* Node, TMap<UEdGraphNode*, TSharedPtr<SMutableSearchComboBox::FFilteredOption>>& AddedOptions);
-
 };
 
diff --git a/Engine/Plugins/Mutable/Source/CustomizableObjectEditor/Private/MuCOE/Widgets/CustomizableObjectVariationCustomization.cpp b/Engine/Plugins/Mutable/Source/CustomizableObjectEditor/Private/MuCOE/Widgets/CustomizableObjectVariationCustomization.cpp
new file mode 100644
--- /dev/null
+++ b/Engine/Plugins/Mutable/Source/CustomizableObjectEditor/Private/MuCOE/Widgets/CustomizableObjectVariationCustomization.cpp
@@ -0,0 +1,97 @@
+// Copyright Epic Games, Inc. All Rights Reserved.
+
+#include "MuCOE/Widgets/CustomizableObjectVariationCustomization.h"
+
+#include "DetailLayoutBuilder.h"
+#include "IDetailsView.h"
+#include "MuCO/CustomizableObjectPrivate.h"
+#include "MuCOE/CustomizableObjectEditorUtilities.h"
+#include "MuCOE/CustomizableObjectGraph.h"
+#include "MuCOE/GraphTraversal.h"
+#include "MuCOE/Nodes/CustomizableObjectNode.h"
+#include "MuCOE/Nodes/CustomizableObjectNodeVariation.h"
+#include "MuCOE/SMutableTagListWidget.h"
+#include "PropertyCustomizationHelpers.h"
+#include "ScopedTransaction.h"
+#include "SSearchableComboBox.h"
+#include "Widgets/Input/SComboButton.h"
+#include "Widgets/Input/STextComboBox.h"
+#include "Widgets/Text/STextBlock.h"
+
+
+#define LOCTEXT_NAMESPACE "CustomizableObjectDetails"
+
+
+TSharedRef<IPropertyTypeCustomization> FCustomizableObjectVariationCustomization::MakeInstance()
+{
+	return MakeShared<FCustomizableObjectVariationCustomization>();
+}
+
+
+void FCustomizableObjectVariationCustomization::CustomizeHeader(TSharedRef<IPropertyHandle> StructPropertyHandle, FDetailWidgetRow& InHeaderRow, IPropertyTypeCustomizationUtils& InCustomizationUtils)
+{
+	TArray<UObject*> OuterObjects;
+	StructPropertyHandle->GetOuterObjects(OuterObjects);
+	PropertyHandle = StructPropertyHandle;
+
+	constexpr bool bRecurse = false;
+	TagPropertyHandle = StructPropertyHandle->GetChildHandle(GET_MEMBER_NAME_CHECKED(FCustomizableObjectVariation, Tag), bRecurse);
+	check(TagPropertyHandle);
+
+	if (OuterObjects.Num())
+	{
+		BaseObjectNode = Cast<UCustomizableObjectNode>(OuterObjects[0]);
+		
+		if (!BaseObjectNode.IsValid())
+		{
+			return;
+		}
+	}
+
+	InHeaderRow
+		.NameContent().HAlign(EHorizontalAlignment::HAlign_Fill).VAlign(EVerticalAlignment::VAlign_Center)
+		[
+			StructPropertyHandle->CreatePropertyNameWidget()
+		]
+		.ValueContent().MinDesiredWidth(300.0f)
+		[
+			SNew(SBorder)
+			.BorderBackgroundColor(FLinearColor::Transparent)
+			[
+				SNew(SMutableTagComboBox)				
+					.Node(BaseObjectNode.Get())
+					.MenuButtonBrush(FAppStyle::GetBrush(TEXT("Icons.Search")))
+					.OnSelectionChanged_Lambda([&](const FText& NewText)
+						{
+							TagPropertyHandle->SetValue(NewText.ToString());
+						})
+					[
+						SNew(SEditableTextBox)
+							.Font(IDetailLayoutBuilder::GetDetailFont())
+							.Text_Lambda([&]() 
+								{ 
+									FString Value;
+									if (TagPropertyHandle)
+									{
+										TagPropertyHandle->GetValue(Value);
+									}
+									return FText::FromString(Value); 
+								})
+							.OnTextCommitted_Lambda([&](const FText& NewText, ETextCommit::Type)
+								{
+									TagPropertyHandle->SetValue(NewText.ToString());
+								})
+					]
+			]
+		]
+		.OverrideResetToDefault(FResetToDefaultOverride::Create(FSimpleDelegate::CreateSP(this, &FCustomizableObjectVariationCustomization::ResetSelectedParameterButtonClicked)));
+}
+
+
+void FCustomizableObjectVariationCustomization::ResetSelectedParameterButtonClicked()
+{
+	TagPropertyHandle->SetValue(FString());
+}
+
+
+#undef LOCTEXT_NAMESPACE
diff --git a/Engine/Plugins/Mutable/Source/CustomizableObjectEditor/Private/MuCOE/Widgets/CustomizableObjectVariationCustomization.h b/Engine/Plugins/Mutable/Source/CustomizableObjectEditor/Private/MuCOE/Widgets/CustomizableObjectVariationCustomization.h
new file mode 100644
--- /dev/null
+++ b/Engine/Plugins/Mutable/Source/CustomizableObjectEditor/Private/MuCOE/Widgets/CustomizableObjectVariationCustomization.h
@@ -0,0 +1,38 @@
+// Copyright Epic Games, Inc. All Rights Reserved.
+
+#pragma once
+
+#include "IPropertyTypeCustomization.h"
+
+class UCustomizableObjectNodeComponentMesh;
+class IPropertyHandle;
+class IPropertyHandleArray;
+class SBoneSelectionWidget;
+class UCustomizableObjectNode;
+
+struct FReferenceSkeleton;
+
+enum class ECheckBoxState : uint8;
+
+class FCustomizableObjectVariationCustomization : public IPropertyTypeCustomization
+{
+public:
+	static TSharedRef<IPropertyTypeCustomization> MakeInstance();
+
+	// IPropertyTypeCustomization interface
+	void CustomizeHeader(TSharedRef<IPropertyHandle> StructPropertyHandle, class FDetailWidgetRow& HeaderRow, IPropertyTypeCustomizationUtils& StructCustomizationUtils) override;
+	void CustomizeChildren(TSharedRef<IPropertyHandle> StructPropertyHandle, class IDetailChildrenBuilder& ChildBuilder, IPropertyTypeCustomizationUtils& StructCustomizationUtils) override {};
+
+private:
+
+	/** Weak pointer to the Customizable Object node that contains this property */
+	TWeakObjectPtr<UCustomizableObjectNode> BaseObjectNode;
+
+	/** Runtime Parameter Name property of a State */
+	TSharedPtr<IPropertyHandle> PropertyHandle;
+	TSharedPtr<IPropertyHandle> TagPropertyHandle;
+
+	/** Callback of the function OverrideResetToDefault */
+	void ResetSelectedParameterButtonClicked();
+
+};
diff --git a/Engine/Plugins/Mutable/Source/CustomizableObjectEditor/Public/MuCOE/Nodes/CustomizableObjectNodeVariation.h b/Engine/Plugins/Mutable/Source/CustomizableObjectEditor/Public/MuCOE/Nodes/CustomizableObjectNodeVariation.h
--- a/Engine/Plugins/Mutable/Source/CustomizableObjectEditor/Public/MuCOE/Nodes/CustomizableObjectNodeVariation.h
+++ b/Engine/Plugins/Mutable/Source/CustomizableObjectEditor/Public/MuCOE/Nodes/CustomizableObjectNodeVariation.h
@@ -32,7 +32,7 @@ class CUSTOMIZABLEOBJECTEDITOR_API UCustomizableObjectNodeVariation : public UCu
 {
 	GENERATED_BODY()
 
-protected:
+public:
 	UPROPERTY(EditAnywhere, Category = CustomizableObject)
 	TArray<FCustomizableObjectVariation> VariationsData; // The variable name can not be Variations due issues with the on to UObject Serialization system
 
