[mutable] Fixed issue with "Clip Mesh with Plane and Morph" node only working for the first object node that compiled it.
- Using previously implemented solution to handle these cases.
- I would recommend fixing the inherent problem in our node caching system in the future.
- This is not a rushed solution, but a patch, so this gets into the preview 1 of the 5.5 mutable beta plugin

diff --git a/Engine/Plugins/Mutable/Source/CustomizableObjectEditor/Private/MuCOE/GenerateMutableSource/GenerateMutableSourceModifier.cpp b/Engine/Plugins/Mutable/Source/CustomizableObjectEditor/Private/MuCOE/GenerateMutableSource/GenerateMutableSourceModifier.cpp
--- a/Engine/Plugins/Mutable/Source/CustomizableObjectEditor/Private/MuCOE/GenerateMutableSource/GenerateMutableSourceModifier.cpp
+++ b/Engine/Plugins/Mutable/Source/CustomizableObjectEditor/Private/MuCOE/GenerateMutableSource/GenerateMutableSourceModifier.cpp
@@ -64,6 +64,10 @@ mu::Ptr<mu::NodeModifier> GenerateMutableSourceModifier(const UEdGraphPin * Pin,
 			EMutableMeshConversionFlags::IgnorePhysics;
 		GenerationContext.MeshGenerationFlags.Push(ModifiersMeshFlags);
 
+		// This modifier can be connected to multiple nodes at the same time and, when that happens and if the cache is being used, only the first node to be processed does work. 
+		// By not caching the mutable node we avoid this from even happening
+		bDoNotAddToGeneratedCache = true;
+		
 		mu::Ptr<mu::NodeModifierMeshClipMorphPlane> ClipNode = new mu::NodeModifierMeshClipMorphPlane();
 		Result = ClipNode;
 
