[mutable] Fixed Text duplicated key

#rnx

diff --git a/Engine/Plugins/Mutable/Source/CustomizableObjectEditor/Private/MuCOE/SCustomizableObjectEditorAdvancedPreviewSettings.cpp b/Engine/Plugins/Mutable/Source/CustomizableObjectEditor/Private/MuCOE/SCustomizableObjectEditorAdvancedPreviewSettings.cpp
--- a/Engine/Plugins/Mutable/Source/CustomizableObjectEditor/Private/MuCOE/SCustomizableObjectEditorAdvancedPreviewSettings.cpp
+++ b/Engine/Plugins/Mutable/Source/CustomizableObjectEditor/Private/MuCOE/SCustomizableObjectEditorAdvancedPreviewSettings.cpp
@@ -35,7 +35,7 @@ void SLevelSelectWindow::Construct(const FArguments& InArgs)
 	FContentBrowserModule& ContentBrowserModule = FModuleManager::LoadModuleChecked<FContentBrowserModule>("ContentBrowser");
 
 	SWindow::Construct(SWindow::FArguments()
-		.Title(LOCTEXT("SNewProfileWindow_Title", "Select a level asset with UObjects corresponding the the filled names"))
+		.Title(LOCTEXT("SNewProfileWindow_Title_1", "Select a level asset with UObjects corresponding the the filled names"))
 		.SupportsMinimize(false)
 		.SupportsMaximize(false)
 		.ClientSize(FVector2D(450, 450))
@@ -205,7 +205,7 @@ void SUObjectSelecWindow::Construct(const FArguments& InArgs)
 	}
 
 	SWindow::Construct(SWindow::FArguments()
-		.Title(LOCTEXT("SNewProfileWindow_Title", "Select a level asset with UObjects coresponding the the filled names"))
+		.Title(LOCTEXT("SNewProfileWindow_Title_2", "Select a level asset with UObjects corresponding the the filled names"))
 		.SupportsMinimize(false)
 		.SupportsMaximize(false)
 		.ClientSize(FVector2D(450, 450))
