[Mutable] Add a SourceIDs to identify and group Roms and use the SourceDataDescriptor to propagate them.
Set unique row IDs to table rows to generate deterministic SourceIDs.

#rnx

diff --git a/Engine/Plugins/Mutable/Source/CustomizableObjectEditor/Private/MuCOE/GenerateMutableSource/GenerateMutableSource.h b/Engine/Plugins/Mutable/Source/CustomizableObjectEditor/Private/MuCOE/GenerateMutableSource/GenerateMutableSource.h
--- a/Engine/Plugins/Mutable/Source/CustomizableObjectEditor/Private/MuCOE/GenerateMutableSource/GenerateMutableSource.h
+++ b/Engine/Plugins/Mutable/Source/CustomizableObjectEditor/Private/MuCOE/GenerateMutableSource/GenerateMutableSource.h
@@ -597,6 +597,7 @@ public:
 
 		// Stores the names of the rows that will be compiled
 		TArray<FName> RowNames;
+		TArray<uint32> RowIds;
 
 		// Compilation Restrictions:
 		// If there is a bool column in the table, checked rows will not be compiled
diff --git a/Engine/Plugins/Mutable/Source/CustomizableObjectEditor/Private/MuCOE/GenerateMutableSource/GenerateMutableSourceGroupProjector.cpp b/Engine/Plugins/Mutable/Source/CustomizableObjectEditor/Private/MuCOE/GenerateMutableSource/GenerateMutableSourceGroupProjector.cpp
--- a/Engine/Plugins/Mutable/Source/CustomizableObjectEditor/Private/MuCOE/GenerateMutableSource/GenerateMutableSourceGroupProjector.cpp
+++ b/Engine/Plugins/Mutable/Source/CustomizableObjectEditor/Private/MuCOE/GenerateMutableSource/GenerateMutableSourceGroupProjector.cpp
@@ -504,6 +504,9 @@ bool GenerateMutableSourceGroupProjector(const UEdGraphPin* Pin, FMutableGraphGe
 						int32 NumMipsBeyondMin = FMath::Max(0, TotalMips - GenerationContext.Options.MinDiskMips);
 						int32 HighResMipsForThisImage = FMath::Min(NumMipsBeyondMin, GenerationContext.Options.NumHighResImageMips);
 						ImageNode->SourceDataDescriptor.SourceHighResMips = HighResMipsForThisImage;
+
+						const FString TextureName = GetNameSafe(Texture);
+						ImageNode->SourceDataDescriptor.SourceId = CityHash32(reinterpret_cast<const char*>(*TextureName), TextureName.Len() * sizeof(FString::ElementType));
 					}
 				}
 				else
diff --git a/Engine/Plugins/Mutable/Source/CustomizableObjectEditor/Private/MuCOE/GenerateMutableSource/GenerateMutableSourceImage.cpp b/Engine/Plugins/Mutable/Source/CustomizableObjectEditor/Private/MuCOE/GenerateMutableSource/GenerateMutableSourceImage.cpp
--- a/Engine/Plugins/Mutable/Source/CustomizableObjectEditor/Private/MuCOE/GenerateMutableSource/GenerateMutableSourceImage.cpp
+++ b/Engine/Plugins/Mutable/Source/CustomizableObjectEditor/Private/MuCOE/GenerateMutableSource/GenerateMutableSourceImage.cpp
@@ -127,6 +127,9 @@ mu::NodeImagePtr GenerateMutableSourceImage(const UEdGraphPin* Pin, FMutableGrap
 				int32 HighResMipsForThisImage = FMath::Min(NumMipsBeyondMin, GenerationContext.Options.NumHighResImageMips);
 				ConstantImageNode->SourceDataDescriptor.SourceHighResMips = HighResMipsForThisImage;
 
+				const FString TextureName = GetNameSafe(BaseTexture);
+				ConstantImageNode->SourceDataDescriptor.SourceId = CityHash32(reinterpret_cast<const char*>(*TextureName), TextureName.Len() * sizeof(FString::ElementType));
+
 				ImageNode = ConstantImageNode;
 
 				GenerationContext.GeneratedImages.Add(ImageKey, ImageNode);
@@ -227,6 +230,9 @@ mu::NodeImagePtr GenerateMutableSourceImage(const UEdGraphPin* Pin, FMutableGrap
 			int32 NumMipsBeyondMin = FMath::Max(0, TotalMips - int32(MipsToSkip) - GenerationContext.Options.MinDiskMips);
 			int32 HighResMipsForThisImage = FMath::Min(NumMipsBeyondMin, GenerationContext.Options.NumHighResImageMips);
 			ImageNode->SourceDataDescriptor.SourceHighResMips = HighResMipsForThisImage;
+
+			const FString TextureName = GetNameSafe(Texture);
+			ImageNode->SourceDataDescriptor.SourceId = CityHash32(reinterpret_cast<const char*>(*TextureName), TextureName.Len() * sizeof(FString::ElementType));
 		}
 		else
 		{
@@ -991,6 +997,9 @@ mu::NodeImagePtr GenerateMutableSourceImage(const UEdGraphPin* Pin, FMutableGrap
 		}
 		else
 		{
+			const FString TableName = TypedNodeTable->Table ? GetNameSafe(TypedNodeTable->Table) : GetNameSafe(TypedNodeTable->Structure);
+			const uint32 TableId = CityHash32(reinterpret_cast<const char*>(*TableName), TableName.Len() * sizeof(FString::ElementType));
+
 			// This node will add a checker texture in case of error
 			mu::Ptr<mu::NodeImageConstant> EmptyNode = new mu::NodeImageConstant();
 			Result = EmptyNode;
@@ -1113,6 +1122,7 @@ mu::NodeImagePtr GenerateMutableSourceImage(const UEdGraphPin* Pin, FMutableGrap
 								int32 NumMipsBeyondMin = FMath::Max(0, TotalMips - GenerationContext.Options.MinDiskMips);
 								int32 HighResMipsForThisImage = FMath::Min(NumMipsBeyondMin, GenerationContext.Options.NumHighResImageMips);
 								ImageTableNode->SourceDataDescriptor.SourceHighResMips = HighResMipsForThisImage;
+								ImageTableNode->SourceDataDescriptor.SourceId = TableId; // It will be combined with the RowId when generating constants
 							}
 						}
 					}
diff --git a/Engine/Plugins/Mutable/Source/CustomizableObjectEditor/Private/MuCOE/GenerateMutableSource/GenerateMutableSourceMesh.cpp b/Engine/Plugins/Mutable/Source/CustomizableObjectEditor/Private/MuCOE/GenerateMutableSource/GenerateMutableSourceMesh.cpp
--- a/Engine/Plugins/Mutable/Source/CustomizableObjectEditor/Private/MuCOE/GenerateMutableSource/GenerateMutableSourceMesh.cpp
+++ b/Engine/Plugins/Mutable/Source/CustomizableObjectEditor/Private/MuCOE/GenerateMutableSource/GenerateMutableSourceMesh.cpp
@@ -2696,7 +2696,8 @@ TArray<TTuple<USkeletalMesh*, TSoftClassPtr<UAnimInstance>>> GetSkeletalMeshesIn
 
 		if (DataTable)
 		{
-			for (const FName& RowName : GetRowsToCompile(*DataTable, *TableNode, GenerationContext))
+			TArray<uint32> RowIds;
+			for (const FName& RowName : GetRowsToCompile(*DataTable, *TableNode, GenerationContext, RowIds))
 			{
 				USkeletalMesh* SkeletalMesh = Cast<USkeletalMesh>(GenerationContext.LoadObject(TableNode->GetSkeletalMeshAt(SourceMeshPin, DataTable, RowName)));
 				TSoftClassPtr<UAnimInstance> MeshAnimInstance = TableNode->GetAnimInstanceAt(SourceMeshPin, DataTable, RowName);
@@ -3158,10 +3159,15 @@ mu::NodeMeshPtr GenerateMorphMesh(const UEdGraphPin* Pin,
 	
 	if (const UCustomizableObjectNodeTable* TypedNodeTable = Cast<UCustomizableObjectNodeTable>(Pin->GetOwningNode()))
 	{
+		const FString TableName = TypedNodeTable->Table ? GetNameSafe(TypedNodeTable->Table) : GetNameSafe(TypedNodeTable->Structure);
+		const uint32 TableId = CityHash32(reinterpret_cast<const char*>(*TableName), TableName.Len() * sizeof(FString::ElementType));
+
+
 		UDataTable* DataTable = GetDataTable(TypedNodeTable, GenerationContext);
 
 		// Generate a new Column for each morph
-		const TArray<FName>& RowNames = GetRowsToCompile(*DataTable, *TypedNodeTable, GenerationContext);
+		TArray<uint32> RowIds;
+		const TArray<FName>& RowNames = GetRowsToCompile(*DataTable, *TypedNodeTable, GenerationContext, RowIds);
 		int32 NumRows = RowNames.Num();
 
 		// Should exist
@@ -3174,6 +3180,7 @@ mu::NodeMeshPtr GenerateMorphMesh(const UEdGraphPin* Pin,
 		for (int32 RowIndex = 0; RowIndex < NumRows; ++RowIndex)
 		{
 			const FName RowName = RowNames[RowIndex];
+			const uint32 RowId = RowIds[RowIndex];
 
 			ColumnIndex = Table->FindColumn(ColumnName);
 
@@ -3183,7 +3190,7 @@ mu::NodeMeshPtr GenerateMorphMesh(const UEdGraphPin* Pin,
 			}
 
 			mu::MeshPtr MorphedSourceTableMesh = BuildMorphedMutableMesh(Pin, TypedNodeMorphs[MorphIndex].MorphTargetName, GenerationContext, bOnlyConnectedLOD, RowName);
-			Table->SetCell(ColumnIndex, RowIndex, MorphedSourceTableMesh.get());
+			Table->SetCell(ColumnIndex, RowId, MorphedSourceTableMesh.get());
 		}
 
 		if (ColumnIndex > INDEX_NONE)
@@ -3195,6 +3202,8 @@ mu::NodeMeshPtr GenerateMorphMesh(const UEdGraphPin* Pin,
 			MorphedSourceMeshNodeTable->SetColumn(ColumnName);
 			MorphedSourceMeshNodeTable->SetParameterName(TypedNodeTable->ParameterName);
 			MorphedSourceMeshNodeTable->SetMessageContext(MorphNode);
+			MorphedSourceMeshNodeTable->SourceDataDescriptor.SourceId = TableId;
+			MorphedSourceMeshNodeTable->SourceDataDescriptor.SourceHighResMips = 0;
 
 			mu::NodeMeshMakeMorphPtr Morph = new mu::NodeMeshMakeMorph;
 			Morph->SetBase(BaseSourceMesh.get());
@@ -3541,6 +3550,11 @@ mu::NodeMeshPtr GenerateMutableSourceMesh(const UEdGraphPin* Pin,
 				}
 			}
 
+
+			const FString MeshName = GetNameSafe(TypedNodeSkel->SkeletalMesh);
+			MeshNode->SourceDataDescriptor.SourceId = CityHash32(reinterpret_cast<const char*>(*MeshName), MeshName.Len() * sizeof(FString::ElementType));
+			MeshNode->SourceDataDescriptor.SourceHighResMips = 0;
+
 			// Applying Mesh Morph Nodes
 			if (GenerationContext.MeshMorphStack.Num())
 			{
@@ -3630,6 +3644,11 @@ mu::NodeMeshPtr GenerateMutableSourceMesh(const UEdGraphPin* Pin,
 
 				MeshNode->SetLayout(0, LayoutNode);
 				LayoutNode->SetMessageContext(Node);  // We need it here because we create multiple nodes.
+
+
+				const FString MeshName = GetNameSafe(TypedNodeStatic->StaticMesh);
+				MeshNode->SourceDataDescriptor.SourceId = CityHash32(reinterpret_cast<const char*>(*MeshName), MeshName.Len() * sizeof(FString::ElementType));
+				MeshNode->SourceDataDescriptor.SourceHighResMips = 0;
 			}
 			else
 			{
@@ -4147,6 +4166,9 @@ mu::NodeMeshPtr GenerateMutableSourceMesh(const UEdGraphPin* Pin,
 
 	else if (const UCustomizableObjectNodeTable* TypedNodeTable = Cast<UCustomizableObjectNodeTable>(Node))
 	{
+		const FString TableName = TypedNodeTable->Table ? GetNameSafe(TypedNodeTable->Table) : GetNameSafe(TypedNodeTable->Structure);
+		const uint32 TableId = CityHash32(reinterpret_cast<const char*>(*TableName), TableName.Len() * sizeof(FString::ElementType));
+
 		mu::NodeMeshConstantPtr EmptyNode = new mu::NodeMeshConstant();
 		Result = EmptyNode;
 		bool bSuccess = true;
@@ -4236,6 +4258,8 @@ mu::NodeMeshPtr GenerateMutableSourceMesh(const UEdGraphPin* Pin,
 						MeshTableNode->SetParameterName(TypedNodeTable->ParameterName);
 						MeshTableNode->SetNoneOption(TypedNodeTable->bAddNoneOption);
 						MeshTableNode->SetDefaultRowName(TypedNodeTable->DefaultRowName.ToString());
+						MeshTableNode->SourceDataDescriptor.SourceId = TableId;
+						MeshTableNode->SourceDataDescriptor.SourceHighResMips = 0;
 
 						// Pose Assets do not need this part of the code
 						if (Pin->PinType.PinCategory == Schema->PC_Mesh)
diff --git a/Engine/Plugins/Mutable/Source/CustomizableObjectEditor/Private/MuCOE/GenerateMutableSource/GenerateMutableSourceSurface.cpp b/Engine/Plugins/Mutable/Source/CustomizableObjectEditor/Private/MuCOE/GenerateMutableSource/GenerateMutableSourceSurface.cpp
--- a/Engine/Plugins/Mutable/Source/CustomizableObjectEditor/Private/MuCOE/GenerateMutableSource/GenerateMutableSourceSurface.cpp
+++ b/Engine/Plugins/Mutable/Source/CustomizableObjectEditor/Private/MuCOE/GenerateMutableSource/GenerateMutableSourceSurface.cpp
@@ -626,6 +626,9 @@ mu::Ptr<mu::NodeSurface> GenerateMutableSourceSurface(const UEdGraphPin * Pin, F
 										int32 HighResMipsForThisImage = FMath::Min(NumMipsBeyondMin, GenerationContext.Options.NumHighResImageMips);
 										ConstImageNode->SourceDataDescriptor.SourceHighResMips = HighResMipsForThisImage;
 
+										const FString TextureName = GetNameSafe(Texture2D);
+										ConstImageNode->SourceDataDescriptor.SourceId = CityHash32(reinterpret_cast<const char*>(*TextureName), TextureName.Len() * sizeof(FString::ElementType));
+
 										return Result;
 									}
 									else
@@ -722,6 +725,9 @@ mu::Ptr<mu::NodeSurface> GenerateMutableSourceSurface(const UEdGraphPin * Pin, F
 									int32 NumMipsBeyondMin = FMath::Max(0, TotalMips - int32(MipsToSkip) - GenerationContext.Options.MinDiskMips);
 									int32 HighResMipsForThisImage = FMath::Min(NumMipsBeyondMin, GenerationContext.Options.NumHighResImageMips);
 									CompositeNormalImage->SourceDataDescriptor.SourceHighResMips = HighResMipsForThisImage;
+
+									const FString TextureName = GetNameSafe(ReferenceCompositeNormalTexture);
+									CompositeNormalImage->SourceDataDescriptor.SourceId = CityHash32(reinterpret_cast<const char*>(*TextureName), TextureName.Len() * sizeof(FString::ElementType));
 								}
 
 								LastImage = CompositedImage;
diff --git a/Engine/Plugins/Mutable/Source/CustomizableObjectEditor/Private/MuCOE/GenerateMutableSource/GenerateMutableSourceTable.cpp b/Engine/Plugins/Mutable/Source/CustomizableObjectEditor/Private/MuCOE/GenerateMutableSource/GenerateMutableSourceTable.cpp
--- a/Engine/Plugins/Mutable/Source/CustomizableObjectEditor/Private/MuCOE/GenerateMutableSource/GenerateMutableSourceTable.cpp
+++ b/Engine/Plugins/Mutable/Source/CustomizableObjectEditor/Private/MuCOE/GenerateMutableSource/GenerateMutableSourceTable.cpp
@@ -31,7 +31,7 @@
 #define LOCTEXT_NAMESPACE "CustomizableObjectEditor"
 
 
-bool FillTableColumn(const UCustomizableObjectNodeTable* TableNode, mu::TablePtr MutableTable, const FString& ColumnName, const FString& RowName, const int32 RowIdx, uint8* CellData, const FProperty* ColumnProperty,
+bool FillTableColumn(const UCustomizableObjectNodeTable* TableNode, mu::TablePtr MutableTable, const FString& ColumnName, const FString& RowName, const uint32 RowId, uint8* CellData, const FProperty* ColumnProperty,
 	const int LODIndexConnected, const int32 SectionIndexConnected, int32 LODIndex, int32 SectionIndex, const bool bOnlyConnectedLOD, FMutableGraphGenerationContext& GenerationContext)
 {
 	int32 CurrentColumn;
@@ -62,7 +62,7 @@ bool FillTableColumn(const UCustomizableObjectNodeTable* TableNode, mu::TablePtr
 				}
 
 				mu::MeshPtr EmptySkeletalMesh = nullptr;
-				MutableTable->SetCell(CurrentColumn, RowIdx, EmptySkeletalMesh.get());
+				MutableTable->SetCell(CurrentColumn, RowId, EmptySkeletalMesh.get());
 
 				return true;
 			}
@@ -321,7 +321,7 @@ bool FillTableColumn(const UCustomizableObjectNodeTable* TableNode, mu::TablePtr
 					AddTagToMutableMeshUnique(*MutableMesh, MeshTag);
 				}
 
-				MutableTable->SetCell(CurrentColumn, RowIdx, MutableMesh.get(), SkeletalMesh);
+				MutableTable->SetCell(CurrentColumn, RowId, MutableMesh.get(), SkeletalMesh);
 			}
 			else
 			{
@@ -395,7 +395,7 @@ bool FillTableColumn(const UCustomizableObjectNodeTable* TableNode, mu::TablePtr
 
 			if (MutableMesh)
 			{
-				MutableTable->SetCell(CurrentColumn, RowIdx, MutableMesh.get(), StaticMesh);
+				MutableTable->SetCell(CurrentColumn, RowId, MutableMesh.get(), StaticMesh);
 			}
 			else
 			{
@@ -437,7 +437,7 @@ bool FillTableColumn(const UCustomizableObjectNodeTable* TableNode, mu::TablePtr
 
 			const bool bIsPassthroughTexture = TableNode->GetColumnImageMode(PinName) == ETableTextureType::PASSTHROUGH_TEXTURE;
 			mu::Ptr<mu::ResourceProxyMemory<mu::Image>> Proxy = new mu::ResourceProxyMemory<mu::Image>(GenerateImageConstant(Texture, GenerationContext, bIsPassthroughTexture));
-			MutableTable->SetCell(CurrentColumn, RowIdx, Proxy.get());
+			MutableTable->SetCell(CurrentColumn, RowId, Proxy.get());
 		}
 
 		else if (SoftObjectProperty->PropertyClass->IsChildOf(UMaterialInterface::StaticClass()))
@@ -511,7 +511,7 @@ bool FillTableColumn(const UCustomizableObjectNodeTable* TableNode, mu::TablePtr
 				}
 
 				int32 ReferenceMaterialId = GenerationContext.ReferencedMaterials.AddUnique(MaterialInstance);
-				MutableTable->SetCell(CurrentColumn, RowIdx, (float)ReferenceMaterialId);
+				MutableTable->SetCell(CurrentColumn, RowId, (float)ReferenceMaterialId);
 
 				return true;
 			}
@@ -571,7 +571,7 @@ bool FillTableColumn(const UCustomizableObjectNodeTable* TableNode, mu::TablePtr
 
 				bool bIsPassthroughTexture = false;
 				mu::Ptr<mu::ResourceProxyMemory<mu::Image>> Proxy = new mu::ResourceProxyMemory<mu::Image>(GenerateImageConstant(ParameterTexture, GenerationContext, bIsPassthroughTexture));
-				MutableTable->SetCell(ColumnIndex, RowIdx, Proxy.get());
+				MutableTable->SetCell(ColumnIndex, RowId, Proxy.get());
 
 				return true;
 			}
@@ -613,7 +613,7 @@ bool FillTableColumn(const UCustomizableObjectNodeTable* TableNode, mu::TablePtr
 					MutableMesh->SetBonePose(i, BoneId, (FTransform3f)ArrayTransform[i], mu::EBoneUsageFlags::Skinning);
 				}
 
-				MutableTable->SetCell(CurrentColumn, RowIdx, MutableMesh.get());
+				MutableTable->SetCell(CurrentColumn, RowId, MutableMesh.get());
 			}
 		}
 
@@ -640,7 +640,7 @@ bool FillTableColumn(const UCustomizableObjectNodeTable* TableNode, mu::TablePtr
 
 			// Setting cell value
 			FLinearColor Value = *(FLinearColor*)CellData;
-			MutableTable->SetCell(CurrentColumn, RowIdx, Value);
+			MutableTable->SetCell(CurrentColumn, RowId, Value);
 		}
 		
 		else
@@ -661,7 +661,7 @@ bool FillTableColumn(const UCustomizableObjectNodeTable* TableNode, mu::TablePtr
 
 		// Setting cell value
 		float Value = FloatNumProperty->GetFloatingPointPropertyValue(CellData);
-		MutableTable->SetCell(CurrentColumn, RowIdx, Value);
+		MutableTable->SetCell(CurrentColumn, RowId, Value);
 	}
 
 	else if (const FNumericProperty* DoubleNumProperty = CastField<FDoubleProperty>(ColumnProperty))
@@ -675,7 +675,7 @@ bool FillTableColumn(const UCustomizableObjectNodeTable* TableNode, mu::TablePtr
 	
 		// Setting cell value
 		float Value = DoubleNumProperty->GetFloatingPointPropertyValue(CellData);
-		MutableTable->SetCell(CurrentColumn, RowIdx, Value);
+		MutableTable->SetCell(CurrentColumn, RowId, Value);
 	}
 
 	else
@@ -858,10 +858,44 @@ void RestrictRowContentByVersion( TArray<FName>& InOutRowNames, const UDataTable
 }
 
 
-TArray<FName> GetRowsToCompile(const UDataTable& DataTable, const UCustomizableObjectNodeTable& TableNode, FMutableGraphGenerationContext& GenerationContext)
+void GenerateUniqueRowIds(const TArray<FName>& RowNames, TArray<uint32>& OutRowIds)
+{
+	const int32 NumRows = RowNames.Num();
+
+	OutRowIds.SetNum(NumRows);
+
+	for (int32 RowIndex = 0; RowIndex < NumRows; ++RowIndex)
+	{
+		const FString& RowName = RowNames[RowIndex].ToString();
+
+		uint32 RowId = CityHash32(reinterpret_cast<const char*>(*RowName), RowName.Len() * sizeof(FString::ElementType));
+
+		// Ensure Row Id is unique 
+		bool bIsUnique = false;
+		while (!bIsUnique)
+		{
+			bIsUnique = true;
+			for (int32 RowIdIndex = 0; RowIdIndex < RowIndex; ++RowIdIndex)
+			{
+				if (OutRowIds[RowIdIndex] == RowId)
+				{
+					bIsUnique = false;
+					++RowId;
+					break;
+				}
+			}
+		}
+
+		OutRowIds[RowIndex] = RowId;
+	}
+}
+
+
+TArray<FName> GetRowsToCompile(const UDataTable& DataTable, const UCustomizableObjectNodeTable& TableNode, FMutableGraphGenerationContext& GenerationContext, TArray<uint32>& OutRowIds)
 {
 	if (FMutableGraphGenerationContext::FGeneratedDataTablesData* Result = GenerationContext.GeneratedTables.Find(DataTable.GetName()))
 	{
+		OutRowIds = Result->RowIds;
 		return Result->RowNames;
 	}
 	else
@@ -874,6 +908,8 @@ TArray<FName> GetRowsToCompile(const UDataTable& DataTable, const UCustomizableO
 			RestrictRowContentByVersion(RowNames, DataTable, TableNode, GenerationContext);
 		}
 
+		GenerateUniqueRowIds(RowNames, OutRowIds);
+
 		return RowNames;
 	}
 }
@@ -901,7 +937,8 @@ bool GenerateTableColumn(const UCustomizableObjectNodeTable* TableNode, const UE
 	GenerationContext.AddParticipatingObject(*DataTable);
 
 	// Getting names of the rows to access the information
-	TArray<FName> RowNames = GetRowsToCompile(*DataTable, *TableNode, GenerationContext);
+	TArray<uint32> RowIds;
+	TArray<FName> RowNames = GetRowsToCompile(*DataTable, *TableNode, GenerationContext, RowIds);
 
 	// Pre-pass to request async loading of all data. This seems to be slightly faster because it avoids the sync after every separate load.
 	if (!GenerationContext.bParticipatingObjectsPass)
@@ -937,7 +974,7 @@ bool GenerateTableColumn(const UCustomizableObjectNodeTable* TableNode, const UE
 	{
 		if (uint8* CellData = GetCellData(RowNames[RowIndex], *DataTable, *ColumnProperty))
 		{
-			bool bCellGenerated = FillTableColumn(TableNode, MutableTable, DataTableColumnName, RowNames[RowIndex].ToString(), RowIndex, CellData, ColumnProperty,
+			bool bCellGenerated = FillTableColumn(TableNode, MutableTable, DataTableColumnName, RowNames[RowIndex].ToString(), RowIds[RowIndex], CellData, ColumnProperty,
 				LODIndexConnected, SectionIndexConnected, LODIndex, SectionIndex, bOnlyConnectedLOD, GenerationContext);
 
 			if (!bCellGenerated)
@@ -956,7 +993,8 @@ void GenerateTableParameterUIData(const UDataTable* DataTable, const UCustomizab
 	// Checking if the parameter name already exists
 	GenerationContext.AddParameterNameUnique(TableNode, TableNode->ParameterName);
 
-	TArray<FName> RowNames = GetRowsToCompile(*DataTable, *TableNode, GenerationContext);
+	TArray<uint32> RowIds;
+	TArray<FName> RowNames = GetRowsToCompile(*DataTable, *TableNode, GenerationContext, RowIds);
 
 	for (const FName& Name : RowNames)
 	{
@@ -1103,15 +1141,16 @@ mu::TablePtr GenerateMutableSourceTable(const UDataTable* DataTable, const UCust
 	if (const UScriptStruct* TableStruct = DataTable->GetRowStruct())
 	{
 		// Getting Table and row names to access the information
-		TArray<FName> RowNames = GetRowsToCompile(*DataTable, *TableNode, GenerationContext);
+		TArray<uint32> RowIds;
+		TArray<FName> RowNames = GetRowsToCompile(*DataTable, *TableNode, GenerationContext, RowIds);
 
 		// Adding and filling Name Column
 		MutableTable->AddColumn("Name", mu::ETableColumnType::String);
 
-		for (int32 NameIndex = 0; NameIndex < RowNames.Num(); ++NameIndex)
+		for (int32 RowIndex = 0; RowIndex < RowNames.Num(); ++RowIndex)
 		{
-			MutableTable->AddRow(NameIndex);
-			MutableTable->SetCell(0, NameIndex, RowNames[NameIndex].ToString());
+			MutableTable->AddRow(RowIds[RowIndex]);
+			MutableTable->SetCell(0, RowIds[RowIndex], RowNames[RowIndex].ToString());
 		}
 
 		// Generating Parameter Metadata for new table parameters
@@ -1122,6 +1161,7 @@ mu::TablePtr GenerateMutableSourceTable(const UDataTable* DataTable, const UCust
 		GeneratedTable.bDisableCheckedRows = TableNode->bDisableCheckedRows;
 		GeneratedTable.VersionColumn = TableNode->VersionColumn;
 		GeneratedTable.RowNames = RowNames;
+		GeneratedTable.RowIds = RowIds;
 		GeneratedTable.ReferenceNode = TableNode;
 
 		// Add table to cache
diff --git a/Engine/Plugins/Mutable/Source/CustomizableObjectEditor/Private/MuCOE/GenerateMutableSource/GenerateMutableSourceTable.h b/Engine/Plugins/Mutable/Source/CustomizableObjectEditor/Private/MuCOE/GenerateMutableSource/GenerateMutableSourceTable.h
--- a/Engine/Plugins/Mutable/Source/CustomizableObjectEditor/Private/MuCOE/GenerateMutableSource/GenerateMutableSourceTable.h
+++ b/Engine/Plugins/Mutable/Source/CustomizableObjectEditor/Private/MuCOE/GenerateMutableSource/GenerateMutableSourceTable.h
@@ -51,7 +51,7 @@ bool GenerateTableColumn(const UCustomizableObjectNodeTable* TableNode, const UE
  * @param GenerationContext 
  * @return  */
 bool FillTableColumn(const UCustomizableObjectNodeTable* TableNode, mu::TablePtr MutableTable, const FString& ColumnName,
-	const FString& RowName, int32 RowIdx, uint8* CellData, const FProperty* ColumnProperty,
+	const FString& RowName, uint32 RowId, uint8* CellData, const FProperty* ColumnProperty,
 	int32 LODIndexConnected, int32 SectionIndexConnected, int32 LODIndex, int32 SectionIndex,
 	bool bOnlyConnectedLOD,
 	FMutableGraphGenerationContext& GenerationContext);
@@ -66,8 +66,9 @@ void GenerateTableParameterUIData(const UDataTable* DataTable, const UCustomizab
 UDataTable* GetDataTable(const UCustomizableObjectNodeTable* TableNode, FMutableGraphGenerationContext& GenerationContext);
 
 /** Gets all the rows of a data table that are going to be compiled. Some rows can be disabled with a bool column or an asset version system 
+	@OutRowIds Deterministic Ids for each row. Used to Add, Get, and Set cell data
 	@return Array with all the names of the rows that are going to be compiled */
-TArray<FName> GetRowsToCompile(const UDataTable& DataTable, const UCustomizableObjectNodeTable& TableNode, FMutableGraphGenerationContext& GenerationContext);
+TArray<FName> GetRowsToCompile(const UDataTable& DataTable, const UCustomizableObjectNodeTable& TableNode, FMutableGraphGenerationContext& GenerationContext, TArray<uint32>& OutRowIds);
 
 /** Generates a Data Table from the Data Tables referenced in a Script Struct */
 UDataTable* GenerateDataTableFromStruct(const UCustomizableObjectNodeTable* TableNode, FMutableGraphGenerationContext& GenerationContext);
diff --git a/Engine/Plugins/Mutable/Source/MutableRuntime/Internal/MuR/ModelPrivate.h b/Engine/Plugins/Mutable/Source/MutableRuntime/Internal/MuR/ModelPrivate.h
--- a/Engine/Plugins/Mutable/Source/MutableRuntime/Internal/MuR/ModelPrivate.h
+++ b/Engine/Plugins/Mutable/Source/MutableRuntime/Internal/MuR/ModelPrivate.h
@@ -29,8 +29,11 @@ namespace mu
     //! Data stored for a rom even if it is not loaded
     struct FRomData
     {
+		// TODO: Remove. Use Rom index instead to find the file
         //! This is used to identify a ROM file. It is usually a hash from its data.
-        uint32 Id=0;
+        uint32 Id=0; 
+
+		uint32 SourceId=0;
 
 		//! Size of the rom
 		uint32 Size=0;
@@ -43,6 +46,10 @@ namespace mu
 
 		/** Properties of the rom data. */
 		ERomFlags Flags = ERomFlags::None;
+
+		// TODO: Store the offset here and delete the FModelStreamableBlock map
+		//! Offset in file
+		// uint64 Offset = 0;
     };
 
     MUTABLE_DEFINE_POD_SERIALISABLE(FRomData);
diff --git a/Engine/Plugins/Mutable/Source/MutableRuntime/Private/MuR/Model.cpp b/Engine/Plugins/Mutable/Source/MutableRuntime/Private/MuR/Model.cpp
--- a/Engine/Plugins/Mutable/Source/MutableRuntime/Private/MuR/Model.cpp
+++ b/Engine/Plugins/Mutable/Source/MutableRuntime/Private/MuR/Model.cpp
@@ -435,12 +435,22 @@ namespace mu
     	return m_pD->m_program.m_roms[Index].Id;
     }
 
+	uint32 Model::GetRomSourceId(int32 Index) const
+	{
+		return m_pD->m_program.m_roms[Index].SourceId;
+	}
+
 
     uint32 Model::GetRomSize(int32 Index) const
     {
     	return m_pD->m_program.m_roms[Index].Size;
     }
 
+	uint16 Model::GetRomType(int32 Index) const
+	{
+		return m_pD->m_program.m_roms[Index].ResourceType;
+	}
+
 	ERomFlags Model::GetRomFlags(int32 Index) const
 	{
 		return m_pD->m_program.m_roms[Index].Flags;
diff --git a/Engine/Plugins/Mutable/Source/MutableRuntime/Public/MuR/Model.h b/Engine/Plugins/Mutable/Source/MutableRuntime/Public/MuR/Model.h
--- a/Engine/Plugins/Mutable/Source/MutableRuntime/Public/MuR/Model.h
+++ b/Engine/Plugins/Mutable/Source/MutableRuntime/Public/MuR/Model.h
@@ -156,8 +156,12 @@ namespace mu
 
 		uint32 GetRomId(int32 Index) const;
 
+		uint32 GetRomSourceId(int32 Index) const;
+
 		uint32 GetRomSize(int32 Index) const;
 
+		uint16 GetRomType(int32 Index) const;
+
 		ERomFlags GetRomFlags(int32 Index) const;
     	
 		//-----------------------------------------------------------------------------------------
diff --git a/Engine/Plugins/Mutable/Source/MutableTools/Internal/MuT/AST.h b/Engine/Plugins/Mutable/Source/MutableTools/Internal/MuT/AST.h
--- a/Engine/Plugins/Mutable/Source/MutableTools/Internal/MuT/AST.h
+++ b/Engine/Plugins/Mutable/Source/MutableTools/Internal/MuT/AST.h
@@ -4,7 +4,7 @@
 
 #include "MuT/Platform.h"
 #include "MuT/Compiler.h"
-#include "MuT/NodeImageConstant.h"
+#include "MuT/NodeImage.h"
 #include "MuR/Image.h"
 #include "MuR/Mesh.h"
 #include "MuR/ModelPrivate.h"
@@ -295,6 +295,11 @@ namespace mu
 			* It must have the same size than the Program::ConstantImages array.
 			*/
 			TArray<FSourceDataDescriptor> SourceImagePerConstant;
+
+			/** Source data descriptor for every mesh constant that has been generated.
+			* It must have the same size than the Program::ConstantMeshes array.
+			*/
+			TArray<FSourceDataDescriptor> SourceMeshPerConstant;
 		};
 
 		FAdditionalData AdditionalData;
diff --git a/Engine/Plugins/Mutable/Source/MutableTools/Private/MuT/AST.cpp b/Engine/Plugins/Mutable/Source/MutableTools/Private/MuT/AST.cpp
--- a/Engine/Plugins/Mutable/Source/MutableTools/Private/MuT/AST.cpp
+++ b/Engine/Plugins/Mutable/Source/MutableTools/Private/MuT/AST.cpp
@@ -1704,7 +1704,30 @@ FSourceDataDescriptor ASTOpFixed::GetSourceDataDescriptor(FGetSourceDataDescript
 	case OP_TYPE::IM_COLOURMAP:
 		Result = GetSourceDataDescriptor(op.args.ImageColourMap.base, Context);
 		break;
-
+	case OP_TYPE::ME_APPLYLAYOUT:
+		Result = GetSourceDataDescriptor(op.args.MeshApplyLayout.mesh, Context);
+		break;	
+	case OP_TYPE::ME_MERGE:
+		Result = GetSourceDataDescriptor(op.args.MeshMerge.base, Context);
+		break;	
+	case OP_TYPE::ME_INTERPOLATE:
+		Result = GetSourceDataDescriptor(op.args.MeshInterpolate.base, Context);
+		break;	
+	case OP_TYPE::ME_MASKDIFF:
+		Result = GetSourceDataDescriptor(op.args.MeshMaskDiff.source, Context);
+		break;	
+	case OP_TYPE::ME_CLIPMORPHPLANE:
+		Result = GetSourceDataDescriptor(op.args.MeshClipMorphPlane.source, Context);
+		break;	
+	case OP_TYPE::ME_CLIPWITHMESH:
+		Result = GetSourceDataDescriptor(op.args.MeshClipWithMesh.source, Context);
+		break;	
+	case OP_TYPE::ME_SETSKELETON:
+		Result = GetSourceDataDescriptor(op.args.MeshSetSkeleton.source, Context);
+		break;	
+	case OP_TYPE::ME_PROJECT:
+		Result = GetSourceDataDescriptor(op.args.MeshProject.mesh, Context);
+		break;
 	default:
 		// This should probably be implemented
 		ensure(false);
diff --git a/Engine/Plugins/Mutable/Source/MutableTools/Private/MuT/ASTOpConstantResource.cpp b/Engine/Plugins/Mutable/Source/MutableTools/Private/MuT/ASTOpConstantResource.cpp
--- a/Engine/Plugins/Mutable/Source/MutableTools/Private/MuT/ASTOpConstantResource.cpp
+++ b/Engine/Plugins/Mutable/Source/MutableTools/Private/MuT/ASTOpConstantResource.cpp
@@ -435,6 +435,10 @@ namespace mu
 				if (!IndexPtr)
 				{
 					args.value = program.AddConstant(MeshData.get());
+
+					int32 DataDescIndex = Options->AdditionalData.SourceMeshPerConstant.Add(SourceDataDescriptor);
+					check(DataDescIndex == args.value);
+
 					Options->MeshConstantMap.Add(MeshData, int32(args.value));
 				}
 				else
diff --git a/Engine/Plugins/Mutable/Source/MutableTools/Private/MuT/ASTOpMeshAddTags.cpp b/Engine/Plugins/Mutable/Source/MutableTools/Private/MuT/ASTOpMeshAddTags.cpp
--- a/Engine/Plugins/Mutable/Source/MutableTools/Private/MuT/ASTOpMeshAddTags.cpp
+++ b/Engine/Plugins/Mutable/Source/MutableTools/Private/MuT/ASTOpMeshAddTags.cpp
@@ -78,4 +78,15 @@ namespace mu
 		}
 	}
 
+
+	FSourceDataDescriptor ASTOpMeshAddTags::GetSourceDataDescriptor(FGetSourceDataDescriptorContext* Context) const
+	{
+		if (Source)
+		{
+			return Source->GetSourceDataDescriptor(Context);
+		}
+
+		return {};
+	}
+
 }
diff --git a/Engine/Plugins/Mutable/Source/MutableTools/Private/MuT/ASTOpMeshAddTags.h b/Engine/Plugins/Mutable/Source/MutableTools/Private/MuT/ASTOpMeshAddTags.h
--- a/Engine/Plugins/Mutable/Source/MutableTools/Private/MuT/ASTOpMeshAddTags.h
+++ b/Engine/Plugins/Mutable/Source/MutableTools/Private/MuT/ASTOpMeshAddTags.h
@@ -36,6 +36,7 @@ namespace mu
 		virtual bool IsEqual(const ASTOp& otherUntyped) const override;
 		virtual Ptr<ASTOp> Clone(MapChildFuncRef mapChild) const override;
 		virtual void Link(FProgram& program, FLinkerOptions*) override;
+		virtual FSourceDataDescriptor GetSourceDataDescriptor(FGetSourceDataDescriptorContext*) const override;
 	};
 
 }
diff --git a/Engine/Plugins/Mutable/Source/MutableTools/Private/MuT/ASTOpMeshApplyPose.cpp b/Engine/Plugins/Mutable/Source/MutableTools/Private/MuT/ASTOpMeshApplyPose.cpp
--- a/Engine/Plugins/Mutable/Source/MutableTools/Private/MuT/ASTOpMeshApplyPose.cpp
+++ b/Engine/Plugins/Mutable/Source/MutableTools/Private/MuT/ASTOpMeshApplyPose.cpp
@@ -118,4 +118,14 @@ namespace mu
 
 		return NewOp;
 	}
+
+	FSourceDataDescriptor ASTOpMeshApplyPose::GetSourceDataDescriptor(FGetSourceDataDescriptorContext* Context) const
+	{
+		if (base)
+		{
+			return base->GetSourceDataDescriptor(Context);
+		}
+
+		return {};
+	}
 }
diff --git a/Engine/Plugins/Mutable/Source/MutableTools/Private/MuT/ASTOpMeshApplyPose.h b/Engine/Plugins/Mutable/Source/MutableTools/Private/MuT/ASTOpMeshApplyPose.h
--- a/Engine/Plugins/Mutable/Source/MutableTools/Private/MuT/ASTOpMeshApplyPose.h
+++ b/Engine/Plugins/Mutable/Source/MutableTools/Private/MuT/ASTOpMeshApplyPose.h
@@ -36,6 +36,7 @@ struct FProgram;
 		void ForEachChild(const TFunctionRef<void(ASTChild&)>) override;
 		void Link(FProgram& program, FLinkerOptions* Options) override;
 		Ptr<ASTOp> OptimiseSink(const FModelOptimizationOptions&, FOptimizeSinkContext&) const override;
+		virtual FSourceDataDescriptor GetSourceDataDescriptor(FGetSourceDataDescriptorContext*) const override;
 	};
 
 
diff --git a/Engine/Plugins/Mutable/Source/MutableTools/Private/MuT/ASTOpMeshApplyShape.cpp b/Engine/Plugins/Mutable/Source/MutableTools/Private/MuT/ASTOpMeshApplyShape.cpp
--- a/Engine/Plugins/Mutable/Source/MutableTools/Private/MuT/ASTOpMeshApplyShape.cpp
+++ b/Engine/Plugins/Mutable/Source/MutableTools/Private/MuT/ASTOpMeshApplyShape.cpp
@@ -114,4 +114,15 @@ namespace mu
 
 	}
 
+
+	FSourceDataDescriptor ASTOpMeshApplyShape::GetSourceDataDescriptor(FGetSourceDataDescriptorContext* Context) const
+	{
+		if (Mesh)
+		{
+			return Mesh->GetSourceDataDescriptor(Context);
+		}
+
+		return {};
+	}
+
 }
diff --git a/Engine/Plugins/Mutable/Source/MutableTools/Private/MuT/ASTOpMeshApplyShape.h b/Engine/Plugins/Mutable/Source/MutableTools/Private/MuT/ASTOpMeshApplyShape.h
--- a/Engine/Plugins/Mutable/Source/MutableTools/Private/MuT/ASTOpMeshApplyShape.h
+++ b/Engine/Plugins/Mutable/Source/MutableTools/Private/MuT/ASTOpMeshApplyShape.h
@@ -39,6 +39,7 @@ struct FProgram;
 		Ptr<ASTOp> Clone(MapChildFuncRef mapChild) const override;
 		void ForEachChild(const TFunctionRef<void(ASTChild&)>) override;
 		void Link(FProgram& program, FLinkerOptions* Options) override;
+		virtual FSourceDataDescriptor GetSourceDataDescriptor(FGetSourceDataDescriptorContext*) const override;
 	};
 
 }
diff --git a/Engine/Plugins/Mutable/Source/MutableTools/Private/MuT/ASTOpMeshBindShape.cpp b/Engine/Plugins/Mutable/Source/MutableTools/Private/MuT/ASTOpMeshBindShape.cpp
--- a/Engine/Plugins/Mutable/Source/MutableTools/Private/MuT/ASTOpMeshBindShape.cpp
+++ b/Engine/Plugins/Mutable/Source/MutableTools/Private/MuT/ASTOpMeshBindShape.cpp
@@ -470,4 +470,16 @@ namespace mu
 		return NewOp;
 	}
 
+
+	//-------------------------------------------------------------------------------------------------
+	FSourceDataDescriptor ASTOpMeshBindShape::GetSourceDataDescriptor(FGetSourceDataDescriptorContext* Context) const
+	{
+		if (Mesh)
+		{
+			return Mesh->GetSourceDataDescriptor(Context);
+		}
+
+		return {};
+	}
+
 }
diff --git a/Engine/Plugins/Mutable/Source/MutableTools/Private/MuT/ASTOpMeshBindShape.h b/Engine/Plugins/Mutable/Source/MutableTools/Private/MuT/ASTOpMeshBindShape.h
--- a/Engine/Plugins/Mutable/Source/MutableTools/Private/MuT/ASTOpMeshBindShape.h
+++ b/Engine/Plugins/Mutable/Source/MutableTools/Private/MuT/ASTOpMeshBindShape.h
@@ -52,6 +52,7 @@ namespace mu
 		void ForEachChild(const TFunctionRef<void(ASTChild&)>) override;
 		void Link(FProgram& program, FLinkerOptions* Options) override;
 		Ptr<ASTOp> OptimiseSink(const FModelOptimizationOptions&, FOptimizeSinkContext&) const override;
+		virtual FSourceDataDescriptor GetSourceDataDescriptor(FGetSourceDataDescriptorContext*) const override;
 
 	};
 
diff --git a/Engine/Plugins/Mutable/Source/MutableTools/Private/MuT/ASTOpMeshClipDeform.cpp b/Engine/Plugins/Mutable/Source/MutableTools/Private/MuT/ASTOpMeshClipDeform.cpp
--- a/Engine/Plugins/Mutable/Source/MutableTools/Private/MuT/ASTOpMeshClipDeform.cpp
+++ b/Engine/Plugins/Mutable/Source/MutableTools/Private/MuT/ASTOpMeshClipDeform.cpp
@@ -91,4 +91,15 @@ namespace mu
 
 	}
 
+
+	FSourceDataDescriptor ASTOpMeshClipDeform::GetSourceDataDescriptor(FGetSourceDataDescriptorContext* Context) const
+	{
+		if (Mesh)
+		{
+			return Mesh->GetSourceDataDescriptor(Context);
+		}
+
+		return {};
+	}
+
 }
diff --git a/Engine/Plugins/Mutable/Source/MutableTools/Private/MuT/ASTOpMeshClipDeform.h b/Engine/Plugins/Mutable/Source/MutableTools/Private/MuT/ASTOpMeshClipDeform.h
--- a/Engine/Plugins/Mutable/Source/MutableTools/Private/MuT/ASTOpMeshClipDeform.h
+++ b/Engine/Plugins/Mutable/Source/MutableTools/Private/MuT/ASTOpMeshClipDeform.h
@@ -38,6 +38,7 @@ struct FProgram;
 		Ptr<ASTOp> Clone(MapChildFuncRef mapChild) const override;
 		void ForEachChild(const TFunctionRef<void(ASTChild&)>) override;
 		void Link(FProgram& program, FLinkerOptions* Options) override;
+		virtual FSourceDataDescriptor GetSourceDataDescriptor(FGetSourceDataDescriptorContext*) const override;
 	};
 
 
diff --git a/Engine/Plugins/Mutable/Source/MutableTools/Private/MuT/ASTOpMeshClipMorphPlane.cpp b/Engine/Plugins/Mutable/Source/MutableTools/Private/MuT/ASTOpMeshClipMorphPlane.cpp
--- a/Engine/Plugins/Mutable/Source/MutableTools/Private/MuT/ASTOpMeshClipMorphPlane.cpp
+++ b/Engine/Plugins/Mutable/Source/MutableTools/Private/MuT/ASTOpMeshClipMorphPlane.cpp
@@ -117,6 +117,7 @@ namespace mu
 
 	}
 
+
 	mu::Ptr<ASTOp> ASTOpMeshClipMorphPlane::OptimiseSink(const FModelOptimizationOptions&, FOptimizeSinkContext&) const
 	{
 		Ptr<ASTOp> NewOp;
@@ -155,4 +156,15 @@ namespace mu
 		return NewOp;
 	}
 
+
+	FSourceDataDescriptor ASTOpMeshClipMorphPlane::GetSourceDataDescriptor(FGetSourceDataDescriptorContext* Context) const
+	{
+		if (source)
+		{
+			return source->GetSourceDataDescriptor(Context);
+		}
+
+		return {};
+	}
+
 }
diff --git a/Engine/Plugins/Mutable/Source/MutableTools/Private/MuT/ASTOpMeshClipMorphPlane.h b/Engine/Plugins/Mutable/Source/MutableTools/Private/MuT/ASTOpMeshClipMorphPlane.h
--- a/Engine/Plugins/Mutable/Source/MutableTools/Private/MuT/ASTOpMeshClipMorphPlane.h
+++ b/Engine/Plugins/Mutable/Source/MutableTools/Private/MuT/ASTOpMeshClipMorphPlane.h
@@ -44,6 +44,7 @@ struct FProgram;
 		virtual void ForEachChild(const TFunctionRef<void(ASTChild&)>) override;
 		virtual void Link(FProgram& program, FLinkerOptions* Options) override;
 		virtual mu::Ptr<ASTOp> OptimiseSink(const FModelOptimizationOptions&, FOptimizeSinkContext&) const override;
+		virtual FSourceDataDescriptor GetSourceDataDescriptor(FGetSourceDataDescriptorContext*) const override;
 	};
 
 
diff --git a/Engine/Plugins/Mutable/Source/MutableTools/Private/MuT/ASTOpMeshDifference.cpp b/Engine/Plugins/Mutable/Source/MutableTools/Private/MuT/ASTOpMeshDifference.cpp
--- a/Engine/Plugins/Mutable/Source/MutableTools/Private/MuT/ASTOpMeshDifference.cpp
+++ b/Engine/Plugins/Mutable/Source/MutableTools/Private/MuT/ASTOpMeshDifference.cpp
@@ -361,4 +361,16 @@ namespace mu
 		return NewOp;
 	}
 
+
+	//-------------------------------------------------------------------------------------------------
+	FSourceDataDescriptor ASTOpMeshDifference::GetSourceDataDescriptor(FGetSourceDataDescriptorContext* Context) const
+	{
+		if (Base)
+		{
+			return Base->GetSourceDataDescriptor(Context);
+		}
+
+		return {};
+	}
+
 }
diff --git a/Engine/Plugins/Mutable/Source/MutableTools/Private/MuT/ASTOpMeshDifference.h b/Engine/Plugins/Mutable/Source/MutableTools/Private/MuT/ASTOpMeshDifference.h
--- a/Engine/Plugins/Mutable/Source/MutableTools/Private/MuT/ASTOpMeshDifference.h
+++ b/Engine/Plugins/Mutable/Source/MutableTools/Private/MuT/ASTOpMeshDifference.h
@@ -53,6 +53,7 @@ namespace mu
 		Ptr<ASTOp> Clone(MapChildFuncRef mapChild) const override;
 		void Link(FProgram& program, FLinkerOptions*) override;
 		Ptr<ASTOp> OptimiseSink(const FModelOptimizationOptions&, FOptimizeSinkContext&) const override;
+		virtual FSourceDataDescriptor GetSourceDataDescriptor(FGetSourceDataDescriptorContext*) const override;
 
 	};
 
diff --git a/Engine/Plugins/Mutable/Source/MutableTools/Private/MuT/ASTOpMeshExtractLayoutBlocks.cpp b/Engine/Plugins/Mutable/Source/MutableTools/Private/MuT/ASTOpMeshExtractLayoutBlocks.cpp
--- a/Engine/Plugins/Mutable/Source/MutableTools/Private/MuT/ASTOpMeshExtractLayoutBlocks.cpp
+++ b/Engine/Plugins/Mutable/Source/MutableTools/Private/MuT/ASTOpMeshExtractLayoutBlocks.cpp
@@ -103,6 +103,17 @@ namespace mu
 	}
 
 
+	FSourceDataDescriptor ASTOpMeshExtractLayoutBlocks::GetSourceDataDescriptor(FGetSourceDataDescriptorContext* Context) const
+	{
+		if (Source)
+		{
+			return Source->GetSourceDataDescriptor(Context);
+		}
+
+		return {};
+	}
+
+
 	//---------------------------------------------------------------------------------------------
 	mu::Ptr<ASTOp> Sink_MeshExtractLayoutBlocksAST::Apply(const ASTOpMeshExtractLayoutBlocks* root)
 	{
diff --git a/Engine/Plugins/Mutable/Source/MutableTools/Private/MuT/ASTOpMeshExtractLayoutBlocks.h b/Engine/Plugins/Mutable/Source/MutableTools/Private/MuT/ASTOpMeshExtractLayoutBlocks.h
--- a/Engine/Plugins/Mutable/Source/MutableTools/Private/MuT/ASTOpMeshExtractLayoutBlocks.h
+++ b/Engine/Plugins/Mutable/Source/MutableTools/Private/MuT/ASTOpMeshExtractLayoutBlocks.h
@@ -44,6 +44,7 @@ namespace mu
 		void Assert() override;
 		void Link(FProgram& program, FLinkerOptions* Options) override;
 		mu::Ptr<ASTOp> OptimiseSink(const FModelOptimizationOptions&, FOptimizeSinkContext&) const;
+		virtual FSourceDataDescriptor GetSourceDataDescriptor(FGetSourceDataDescriptorContext*) const override;
 
 	};
 
diff --git a/Engine/Plugins/Mutable/Source/MutableTools/Private/MuT/ASTOpMeshFormat.cpp b/Engine/Plugins/Mutable/Source/MutableTools/Private/MuT/ASTOpMeshFormat.cpp
--- a/Engine/Plugins/Mutable/Source/MutableTools/Private/MuT/ASTOpMeshFormat.cpp
+++ b/Engine/Plugins/Mutable/Source/MutableTools/Private/MuT/ASTOpMeshFormat.cpp
@@ -110,6 +110,17 @@ mu::Ptr<ASTOp> ASTOpMeshFormat::OptimiseSink(const FModelOptimizationOptions& op
 }
 
 
+FSourceDataDescriptor ASTOpMeshFormat::GetSourceDataDescriptor(FGetSourceDataDescriptorContext* Context) const
+{
+	if (Source)
+	{
+		return Source->GetSourceDataDescriptor(Context);
+	}
+
+	return {};
+}
+
+
 //---------------------------------------------------------------------------------------------
 //---------------------------------------------------------------------------------------------
 //---------------------------------------------------------------------------------------------
@@ -275,6 +286,7 @@ mu::Ptr<ASTOp> Sink_MeshFormatAST::Visit(const mu::Ptr<ASTOp>& at, const ASTOpMe
 		mu::Ptr<ASTOpConstantResource> NewFormatConstant = new ASTOpConstantResource();
 		NewFormatConstant->Type = OP_TYPE::ME_CONSTANT;
 		NewFormatConstant->SetValue(pTargetMorphFormat, nullptr);
+		NewFormatConstant->SourceDataDescriptor = at->GetSourceDataDescriptor();
 
 		if (NewOp->Target)
 		{
@@ -319,6 +331,7 @@ mu::Ptr<ASTOp> Sink_MeshFormatAST::Visit(const mu::Ptr<ASTOp>& at, const ASTOpMe
 		mu::Ptr<ASTOpConstantResource> TargetMorphFormatOp = new ASTOpConstantResource();
 		TargetMorphFormatOp->Type = OP_TYPE::ME_CONSTANT;
 		TargetMorphFormatOp->SetValue(pTargetMorphFormat, nullptr );
+		TargetMorphFormatOp->SourceDataDescriptor = at->GetSourceDataDescriptor();
 
 		for (int32 t = 0; t < MUTABLE_OP_MAX_INTERPOLATE_COUNT - 1; ++t)
 		{
diff --git a/Engine/Plugins/Mutable/Source/MutableTools/Private/MuT/ASTOpMeshFormat.h b/Engine/Plugins/Mutable/Source/MutableTools/Private/MuT/ASTOpMeshFormat.h
--- a/Engine/Plugins/Mutable/Source/MutableTools/Private/MuT/ASTOpMeshFormat.h
+++ b/Engine/Plugins/Mutable/Source/MutableTools/Private/MuT/ASTOpMeshFormat.h
@@ -34,6 +34,7 @@ namespace mu
 		void ForEachChild(const TFunctionRef<void(ASTChild&)>) override;
 		void Link(FProgram& program, FLinkerOptions* Options) override;
 		Ptr<ASTOp> OptimiseSink(const FModelOptimizationOptions& options, FOptimizeSinkContext& context) const override;
+		virtual FSourceDataDescriptor GetSourceDataDescriptor(FGetSourceDataDescriptorContext*) const override;
 
 	};
 
diff --git a/Engine/Plugins/Mutable/Source/MutableTools/Private/MuT/ASTOpMeshGeometryOperation.cpp b/Engine/Plugins/Mutable/Source/MutableTools/Private/MuT/ASTOpMeshGeometryOperation.cpp
--- a/Engine/Plugins/Mutable/Source/MutableTools/Private/MuT/ASTOpMeshGeometryOperation.cpp
+++ b/Engine/Plugins/Mutable/Source/MutableTools/Private/MuT/ASTOpMeshGeometryOperation.cpp
@@ -90,4 +90,15 @@ namespace mu
 
 	}
 
+
+	FSourceDataDescriptor ASTOpMeshGeometryOperation::GetSourceDataDescriptor(FGetSourceDataDescriptorContext* Context) const
+	{
+		if (meshA)
+		{
+			return meshA->GetSourceDataDescriptor(Context);
+		}
+
+		return {};
+	}
+
 }
diff --git a/Engine/Plugins/Mutable/Source/MutableTools/Private/MuT/ASTOpMeshGeometryOperation.h b/Engine/Plugins/Mutable/Source/MutableTools/Private/MuT/ASTOpMeshGeometryOperation.h
--- a/Engine/Plugins/Mutable/Source/MutableTools/Private/MuT/ASTOpMeshGeometryOperation.h
+++ b/Engine/Plugins/Mutable/Source/MutableTools/Private/MuT/ASTOpMeshGeometryOperation.h
@@ -37,6 +37,7 @@ namespace mu
 		Ptr<ASTOp> Clone(MapChildFuncRef mapChild) const override;
 		void ForEachChild(const TFunctionRef<void(ASTChild&)>) override;
 		void Link(FProgram& program, FLinkerOptions* Options) override;
+		virtual FSourceDataDescriptor GetSourceDataDescriptor(FGetSourceDataDescriptorContext*) const override;
 	};
 
 
diff --git a/Engine/Plugins/Mutable/Source/MutableTools/Private/MuT/ASTOpMeshMaskClipMesh.cpp b/Engine/Plugins/Mutable/Source/MutableTools/Private/MuT/ASTOpMeshMaskClipMesh.cpp
--- a/Engine/Plugins/Mutable/Source/MutableTools/Private/MuT/ASTOpMeshMaskClipMesh.cpp
+++ b/Engine/Plugins/Mutable/Source/MutableTools/Private/MuT/ASTOpMeshMaskClipMesh.cpp
@@ -364,5 +364,17 @@ namespace mu
 		return at;
 	}
 
+
+	//-------------------------------------------------------------------------------------------------
+	FSourceDataDescriptor ASTOpMeshMaskClipMesh::GetSourceDataDescriptor(FGetSourceDataDescriptorContext* Context) const
+	{
+		if (source)
+		{
+			return source->GetSourceDataDescriptor(Context);
+		}
+
+		return {};
+	}
+
 }
  
\ No newline at end of file
diff --git a/Engine/Plugins/Mutable/Source/MutableTools/Private/MuT/ASTOpMeshMaskClipMesh.h b/Engine/Plugins/Mutable/Source/MutableTools/Private/MuT/ASTOpMeshMaskClipMesh.h
--- a/Engine/Plugins/Mutable/Source/MutableTools/Private/MuT/ASTOpMeshMaskClipMesh.h
+++ b/Engine/Plugins/Mutable/Source/MutableTools/Private/MuT/ASTOpMeshMaskClipMesh.h
@@ -36,6 +36,7 @@ struct FProgram;
 		void ForEachChild(const TFunctionRef<void(ASTChild&)>) override;
 		void Link(FProgram& program, FLinkerOptions* Options) override;
 		Ptr<ASTOp> OptimiseSink(const FModelOptimizationOptions&, FOptimizeSinkContext&) const override;
+		virtual FSourceDataDescriptor GetSourceDataDescriptor(FGetSourceDataDescriptorContext*) const override;
 
 	};
 
diff --git a/Engine/Plugins/Mutable/Source/MutableTools/Private/MuT/ASTOpMeshMaskClipUVMask.cpp b/Engine/Plugins/Mutable/Source/MutableTools/Private/MuT/ASTOpMeshMaskClipUVMask.cpp
--- a/Engine/Plugins/Mutable/Source/MutableTools/Private/MuT/ASTOpMeshMaskClipUVMask.cpp
+++ b/Engine/Plugins/Mutable/Source/MutableTools/Private/MuT/ASTOpMeshMaskClipUVMask.cpp
@@ -372,5 +372,17 @@ namespace mu
 		return at;
 	}
 
+
+	//-------------------------------------------------------------------------------------------------
+	FSourceDataDescriptor ASTOpMeshMaskClipUVMask::GetSourceDataDescriptor(FGetSourceDataDescriptorContext* Context) const
+	{
+		if (Source)
+		{
+			return Source->GetSourceDataDescriptor(Context);
+		}
+
+		return {};
+	}
+
 }
  
\ No newline at end of file
diff --git a/Engine/Plugins/Mutable/Source/MutableTools/Private/MuT/ASTOpMeshMaskClipUVMask.h b/Engine/Plugins/Mutable/Source/MutableTools/Private/MuT/ASTOpMeshMaskClipUVMask.h
--- a/Engine/Plugins/Mutable/Source/MutableTools/Private/MuT/ASTOpMeshMaskClipUVMask.h
+++ b/Engine/Plugins/Mutable/Source/MutableTools/Private/MuT/ASTOpMeshMaskClipUVMask.h
@@ -36,6 +36,7 @@ namespace mu
 		void ForEachChild(const TFunctionRef<void(ASTChild&)>) override;
 		void Link(FProgram& program, FLinkerOptions* Options) override;
 		Ptr<ASTOp> OptimiseSink(const FModelOptimizationOptions&, FOptimizeSinkContext&) const override;
+		virtual FSourceDataDescriptor GetSourceDataDescriptor(FGetSourceDataDescriptorContext*) const override;
 
 	};
 
diff --git a/Engine/Plugins/Mutable/Source/MutableTools/Private/MuT/ASTOpMeshMorph.cpp b/Engine/Plugins/Mutable/Source/MutableTools/Private/MuT/ASTOpMeshMorph.cpp
--- a/Engine/Plugins/Mutable/Source/MutableTools/Private/MuT/ASTOpMeshMorph.cpp
+++ b/Engine/Plugins/Mutable/Source/MutableTools/Private/MuT/ASTOpMeshMorph.cpp
@@ -157,4 +157,16 @@ namespace mu
 		return NewOp;
 	}
 
+
+	//-------------------------------------------------------------------------------------------------
+	FSourceDataDescriptor ASTOpMeshMorph::GetSourceDataDescriptor(FGetSourceDataDescriptorContext* Context) const
+	{
+		if (Base)
+		{
+			return Base->GetSourceDataDescriptor(Context);
+		}
+
+		return {};
+	}
+
 }
diff --git a/Engine/Plugins/Mutable/Source/MutableTools/Private/MuT/ASTOpMeshMorph.h b/Engine/Plugins/Mutable/Source/MutableTools/Private/MuT/ASTOpMeshMorph.h
--- a/Engine/Plugins/Mutable/Source/MutableTools/Private/MuT/ASTOpMeshMorph.h
+++ b/Engine/Plugins/Mutable/Source/MutableTools/Private/MuT/ASTOpMeshMorph.h
@@ -42,6 +42,7 @@ namespace mu
 		Ptr<ASTOp> Clone(MapChildFuncRef mapChild) const override;
 		void Link(FProgram& program, FLinkerOptions*) override;
 		Ptr<ASTOp> OptimiseSink(const FModelOptimizationOptions&, FOptimizeSinkContext&) const override;
+		virtual FSourceDataDescriptor GetSourceDataDescriptor(FGetSourceDataDescriptorContext*) const override;
 
 	};
 
diff --git a/Engine/Plugins/Mutable/Source/MutableTools/Private/MuT/ASTOpMeshMorphReshape.cpp b/Engine/Plugins/Mutable/Source/MutableTools/Private/MuT/ASTOpMeshMorphReshape.cpp
--- a/Engine/Plugins/Mutable/Source/MutableTools/Private/MuT/ASTOpMeshMorphReshape.cpp
+++ b/Engine/Plugins/Mutable/Source/MutableTools/Private/MuT/ASTOpMeshMorphReshape.cpp
@@ -89,4 +89,15 @@ namespace mu
 
 	}
 
+
+	FSourceDataDescriptor ASTOpMeshMorphReshape::GetSourceDataDescriptor(FGetSourceDataDescriptorContext* Context) const
+	{
+		if (Morph)
+		{
+			return Morph->GetSourceDataDescriptor(Context);
+		}
+
+		return {};
+	}
+
 }
diff --git a/Engine/Plugins/Mutable/Source/MutableTools/Private/MuT/ASTOpMeshMorphReshape.h b/Engine/Plugins/Mutable/Source/MutableTools/Private/MuT/ASTOpMeshMorphReshape.h
--- a/Engine/Plugins/Mutable/Source/MutableTools/Private/MuT/ASTOpMeshMorphReshape.h
+++ b/Engine/Plugins/Mutable/Source/MutableTools/Private/MuT/ASTOpMeshMorphReshape.h
@@ -34,6 +34,7 @@ struct FProgram;
 		Ptr<ASTOp> Clone(MapChildFuncRef mapChild) const override;
 		void ForEachChild(const TFunctionRef<void(ASTChild&)>) override;
 		void Link(FProgram& program, FLinkerOptions* Options) override;
+		virtual FSourceDataDescriptor GetSourceDataDescriptor(FGetSourceDataDescriptorContext*) const override;
 	};
 
 
diff --git a/Engine/Plugins/Mutable/Source/MutableTools/Private/MuT/ASTOpMeshOptimizeSkinning.cpp b/Engine/Plugins/Mutable/Source/MutableTools/Private/MuT/ASTOpMeshOptimizeSkinning.cpp
--- a/Engine/Plugins/Mutable/Source/MutableTools/Private/MuT/ASTOpMeshOptimizeSkinning.cpp
+++ b/Engine/Plugins/Mutable/Source/MutableTools/Private/MuT/ASTOpMeshOptimizeSkinning.cpp
@@ -83,4 +83,16 @@ namespace mu
 		}
 	}
 
+
+	//-------------------------------------------------------------------------------------------------
+	FSourceDataDescriptor ASTOpMeshOptimizeSkinning::GetSourceDataDescriptor(FGetSourceDataDescriptorContext* Context) const
+	{
+		if (source)
+		{
+			return source->GetSourceDataDescriptor(Context);
+		}
+
+		return {};
+	}
+
 }
diff --git a/Engine/Plugins/Mutable/Source/MutableTools/Private/MuT/ASTOpMeshOptimizeSkinning.h b/Engine/Plugins/Mutable/Source/MutableTools/Private/MuT/ASTOpMeshOptimizeSkinning.h
--- a/Engine/Plugins/Mutable/Source/MutableTools/Private/MuT/ASTOpMeshOptimizeSkinning.h
+++ b/Engine/Plugins/Mutable/Source/MutableTools/Private/MuT/ASTOpMeshOptimizeSkinning.h
@@ -33,6 +33,7 @@ namespace mu
 		bool IsEqual(const ASTOp& otherUntyped) const override;
 		Ptr<ASTOp> Clone(MapChildFuncRef mapChild) const override;
 		void Link(FProgram& program, FLinkerOptions*) override;
+		virtual FSourceDataDescriptor GetSourceDataDescriptor(FGetSourceDataDescriptorContext*) const override;
 	};
 
 
diff --git a/Engine/Plugins/Mutable/Source/MutableTools/Private/MuT/ASTOpMeshRemoveMask.cpp b/Engine/Plugins/Mutable/Source/MutableTools/Private/MuT/ASTOpMeshRemoveMask.cpp
--- a/Engine/Plugins/Mutable/Source/MutableTools/Private/MuT/ASTOpMeshRemoveMask.cpp
+++ b/Engine/Plugins/Mutable/Source/MutableTools/Private/MuT/ASTOpMeshRemoveMask.cpp
@@ -278,4 +278,16 @@ namespace mu
 		return at;
 	}
 
+
+	//-------------------------------------------------------------------------------------------------
+	FSourceDataDescriptor ASTOpMeshRemoveMask::GetSourceDataDescriptor(FGetSourceDataDescriptorContext* Context) const
+	{
+		if (source)
+		{
+			return source->GetSourceDataDescriptor(Context);
+		}
+
+		return {};
+	}
+
 }
diff --git a/Engine/Plugins/Mutable/Source/MutableTools/Private/MuT/ASTOpMeshRemoveMask.h b/Engine/Plugins/Mutable/Source/MutableTools/Private/MuT/ASTOpMeshRemoveMask.h
--- a/Engine/Plugins/Mutable/Source/MutableTools/Private/MuT/ASTOpMeshRemoveMask.h
+++ b/Engine/Plugins/Mutable/Source/MutableTools/Private/MuT/ASTOpMeshRemoveMask.h
@@ -40,6 +40,7 @@ namespace mu
 		Ptr<ASTOp> Clone(MapChildFuncRef mapChild) const override;
 		void Link(FProgram& program, FLinkerOptions*) override;
 		Ptr<ASTOp> OptimiseSink(const FModelOptimizationOptions&, FOptimizeSinkContext&) const override;
+		virtual FSourceDataDescriptor GetSourceDataDescriptor(FGetSourceDataDescriptorContext*) const override;
 
 		// Own interface
 		void AddRemove(const Ptr<ASTOp>& condition, const Ptr<ASTOp>& mask);
diff --git a/Engine/Plugins/Mutable/Source/MutableTools/Private/MuT/ASTOpMeshTransform.cpp b/Engine/Plugins/Mutable/Source/MutableTools/Private/MuT/ASTOpMeshTransform.cpp
--- a/Engine/Plugins/Mutable/Source/MutableTools/Private/MuT/ASTOpMeshTransform.cpp
+++ b/Engine/Plugins/Mutable/Source/MutableTools/Private/MuT/ASTOpMeshTransform.cpp
@@ -80,4 +80,15 @@ namespace mu
 
 	}
 
+
+	FSourceDataDescriptor ASTOpMeshTransform::GetSourceDataDescriptor(FGetSourceDataDescriptorContext* Context) const
+	{
+		if (source)
+		{
+			return source->GetSourceDataDescriptor(Context);
+		}
+
+		return {};
+	}
+
 }
diff --git a/Engine/Plugins/Mutable/Source/MutableTools/Private/MuT/ASTOpMeshTransform.h b/Engine/Plugins/Mutable/Source/MutableTools/Private/MuT/ASTOpMeshTransform.h
--- a/Engine/Plugins/Mutable/Source/MutableTools/Private/MuT/ASTOpMeshTransform.h
+++ b/Engine/Plugins/Mutable/Source/MutableTools/Private/MuT/ASTOpMeshTransform.h
@@ -31,12 +31,14 @@ struct FProgram;
 		ASTOpMeshTransform(const ASTOpMeshTransform&) = delete;
 		~ASTOpMeshTransform();
 
+		// ASTOp interface
 		OP_TYPE GetOpType() const override { return OP_TYPE::ME_TRANSFORM; }
 		uint64 Hash() const override;
 		void ForEachChild(const TFunctionRef<void(ASTChild&)>) override;
 		bool IsEqual(const ASTOp& otherUntyped) const override;
 		Ptr<ASTOp> Clone(MapChildFuncRef mapChild) const override;
 		void Link(FProgram& program, FLinkerOptions* Options) override;
+		virtual FSourceDataDescriptor GetSourceDataDescriptor(FGetSourceDataDescriptorContext*) const override;
 
 	};
 
diff --git a/Engine/Plugins/Mutable/Source/MutableTools/Private/MuT/CodeGenerator.cpp b/Engine/Plugins/Mutable/Source/MutableTools/Private/MuT/CodeGenerator.cpp
--- a/Engine/Plugins/Mutable/Source/MutableTools/Private/MuT/CodeGenerator.cpp
+++ b/Engine/Plugins/Mutable/Source/MutableTools/Private/MuT/CodeGenerator.cpp
@@ -52,6 +52,7 @@
 #include "MuT/NodeComponentSwitch.h"
 #include "MuT/NodeComponentVariation.h"
 #include "MuT/NodeImage.h"
+#include "MuT/NodeImageConstant.h"
 #include "MuT/NodeImageFormat.h"
 #include "MuT/NodeImageFormatPrivate.h"
 #include "MuT/NodeImageMipmap.h"
@@ -344,20 +345,22 @@ namespace mu
 
 			// Add every row
 			int32 RowCount = CacheKey.Table->GetPrivate()->Rows.Num();
-			for (int32 i = 0; i < RowCount; ++i)
+			check(RowCount < MAX_int16) // max FIntValueDesc allows
+
+			for (int32 RowIndex = 0; RowIndex < RowCount; ++RowIndex)
 			{
 				FParameterDesc::FIntValueDesc value;
-				value.m_value = (int16)CacheKey.Table->GetPrivate()->Rows[i].Id;
+				value.m_value = RowIndex;
 
 				if (nameCol > -1)
 				{
-					value.m_name = CacheKey.Table->GetPrivate()->Rows[i].Values[nameCol].String;
+					value.m_name = CacheKey.Table->GetPrivate()->Rows[RowIndex].Values[nameCol].String;
 				}
 
 				param.m_possibleValues.Add(value);
 
 				// Set the first row as the default one (if there is none option)
-				if (i == 0 && !bAddNoneOption)
+				if (RowIndex == 0 && !bAddNoneOption)
 				{
 					param.m_defaultValue.Set<ParamIntType>(value.m_value);
 				}
@@ -2111,6 +2114,7 @@ namespace mu
 					Ptr<ASTOpConstantResource> TargetMeshOp = new ASTOpConstantResource;
 					TargetMeshOp->Type = OP_TYPE::ME_CONSTANT;
 					TargetMeshOp->SetValue(TargetMesh->Clone(), CompilerOptions->OptimisationOptions.DiskCacheContext);
+					TargetMeshOp->SourceDataDescriptor = OriginalMeshNode->SourceDataDescriptor;
 
 					// Morph generation through mesh diff
 					Ptr<ASTOpMeshDifference> diffAd;
@@ -2230,6 +2234,7 @@ namespace mu
 				Ptr<ASTOpConstantResource> UVMeshOp = new ASTOpConstantResource();
 				UVMeshOp->Type = OP_TYPE::ME_CONSTANT;
 				UVMeshOp->SetValue(OriginalMesh->Clone(), CompilerOptions->OptimisationOptions.DiskCacheContext);
+				UVMeshOp->SourceDataDescriptor = OriginalMeshNode->SourceDataDescriptor;
 
 				const NodeModifierMeshClipWithUVMask* TypedClipNode = static_cast<const NodeModifierMeshClipWithUVMask*>(m.Node);
 
diff --git a/Engine/Plugins/Mutable/Source/MutableTools/Private/MuT/CodeGenerator.h b/Engine/Plugins/Mutable/Source/MutableTools/Private/MuT/CodeGenerator.h
--- a/Engine/Plugins/Mutable/Source/MutableTools/Private/MuT/CodeGenerator.h
+++ b/Engine/Plugins/Mutable/Source/MutableTools/Private/MuT/CodeGenerator.h
@@ -893,11 +893,12 @@ namespace mu
 		SwitchOp->variable = Variable;
 		SwitchOp->def = GenerateDefaultTableValue(TYPE);
 
-		for (int32 i = 0; i < NumRows; ++i)
+		for (int32 RowIndex = 0; RowIndex < NumRows; ++RowIndex)
         {
-            check(NodeTable->GetPrivate()->Rows[i].Id <= 0xFFFF);
-            auto Condition = (uint16)NodeTable->GetPrivate()->Rows[i].Id;
-            Ptr<ASTOp> Branch = GenerateOption( node, ColIndex, (int)i, ErrorLog.get() );
+            check(RowIndex <= 0xFFFF);
+			auto Condition = (uint16)RowIndex;
+
+            Ptr<ASTOp> Branch = GenerateOption( node, ColIndex, (int)RowIndex, ErrorLog.get() );
 
 			if (Branch || TYPE != ETableColumnType::Mesh)
 			{
diff --git a/Engine/Plugins/Mutable/Source/MutableTools/Private/MuT/CodeGenerator_Image.cpp b/Engine/Plugins/Mutable/Source/MutableTools/Private/MuT/CodeGenerator_Image.cpp
--- a/Engine/Plugins/Mutable/Source/MutableTools/Private/MuT/CodeGenerator_Image.cpp
+++ b/Engine/Plugins/Mutable/Source/MutableTools/Private/MuT/CodeGenerator_Image.cpp
@@ -1455,6 +1455,7 @@ namespace mu
                 // will remove the faces that are not in the layout block we are generating.
                 Ptr<ASTOpConstantResource> cop = new ASTOpConstantResource();
                 cop->Type = OP_TYPE::ME_CONSTANT;
+				cop->SourceDataDescriptor = CurrentMeshToProjectOp->GetSourceDataDescriptor();
 				Ptr<Mesh> FormatMeshResult = new Mesh();
 				CreateMeshOptimisedForWrappingProjection(FormatMeshResult.get(), node.m_layout);
 
@@ -1485,6 +1486,7 @@ namespace mu
                 // Reformat the mesh to a more efficient format for this operation
                 Ptr<ASTOpConstantResource> cop = new ASTOpConstantResource();
                 cop->Type = OP_TYPE::ME_CONSTANT;
+				cop->SourceDataDescriptor = CurrentMeshToProjectOp->GetSourceDataDescriptor();
 
 				Ptr<Mesh> FormatMeshResult = new Mesh();
                 CreateMeshOptimisedForProjection(FormatMeshResult.get(), node.m_layout);
@@ -1907,6 +1909,11 @@ namespace mu
 					// TODO: We probably want to get the data tags from the table row.
 					ImageConst->SourceDataDescriptor = InNode->SourceDataDescriptor;
 
+					// Combine the SourceId of the node with the RowId to generate one shared between all resources from this row.
+					// Hash collisions are allowed, since it is used to group resources, not to differentiate them.
+					const uint32 RowId = node.Table->GetPrivate()->Rows[row].Id;
+					ImageConst->SourceDataDescriptor.SourceId = HashCombine(InNode->SourceDataDescriptor.SourceId, RowId);
+
 					FImageGenerationResult Result;
 					GenerateImage(Options, Result, ImageConst);
 					ImageOp = Result.op;
diff --git a/Engine/Plugins/Mutable/Source/MutableTools/Private/MuT/CodeGenerator_Mesh.cpp b/Engine/Plugins/Mutable/Source/MutableTools/Private/MuT/CodeGenerator_Mesh.cpp
--- a/Engine/Plugins/Mutable/Source/MutableTools/Private/MuT/CodeGenerator_Mesh.cpp
+++ b/Engine/Plugins/Mutable/Source/MutableTools/Private/MuT/CodeGenerator_Mesh.cpp
@@ -1007,6 +1007,13 @@ namespace mu
 
 					TargetOptions.OverrideContext = node.Table->GetPrivate()->Rows[row].Values[colIndex].ErrorContext;
 
+					pCell->SourceDataDescriptor = node.SourceDataDescriptor;
+
+					// Combine the SourceId of the node with the RowId to generate one shared between all resources from this row.
+					// Hash collisions are allowed, since it is used to group resources, not to differentiate them.
+					const uint32 RowId = node.Table->GetPrivate()->Rows[row].Id;
+					pCell->SourceDataDescriptor.SourceId = HashCombine(node.SourceDataDescriptor.SourceId, RowId);
+
 					GenerateMesh(TargetOptions, BranchResults, pCell);
 
 					if (!bFirstRowGenerated)
@@ -1117,6 +1124,7 @@ namespace mu
 
         Ptr<ASTOpConstantResource> ConstantOp = new ASTOpConstantResource();
 		ConstantOp->Type = OP_TYPE::ME_CONSTANT;
+		ConstantOp->SourceDataDescriptor = InNode->SourceDataDescriptor;
 		OutResult.BaseMeshOp = ConstantOp;
 		OutResult.MeshOp = ConstantOp;
 		OutResult.GeneratedLayouts.Empty();
@@ -1357,6 +1365,7 @@ namespace mu
             Ptr<ASTOpConstantResource> cop = new ASTOpConstantResource();
             cop->Type = OP_TYPE::ME_CONSTANT;
             cop->SetValue( FormatMesh, CompilerOptions->OptimisationOptions.DiskCacheContext );
+			cop->SourceDataDescriptor = baseResult.BaseMeshOp->GetSourceDataDescriptor();
             op->Format = cop;
 
             OutResult.MeshOp = op;
diff --git a/Engine/Plugins/Mutable/Source/MutableTools/Private/MuT/CodeOptimiser.cpp b/Engine/Plugins/Mutable/Source/MutableTools/Private/MuT/CodeOptimiser.cpp
--- a/Engine/Plugins/Mutable/Source/MutableTools/Private/MuT/CodeOptimiser.cpp
+++ b/Engine/Plugins/Mutable/Source/MutableTools/Private/MuT/CodeOptimiser.cpp
@@ -481,7 +481,7 @@ namespace mu
 			pSystem->GetPrivate()->ImagePixelFormatOverride = ImOp.FormatImageOverride;
 
 			FSourceDataDescriptor SourceDataDescriptor;
-			if (DataType == DT_IMAGE)
+			if (DataType == DT_IMAGE || DataType == DT_MESH)
 			{
 				SourceDataDescriptor = SourceCloned->GetSourceDataDescriptor();
 				check(!SourceDataDescriptor.IsInvalid());
@@ -513,10 +513,11 @@ namespace mu
 
 				if (pMesh)
 				{
-					mu::Ptr<ASTOpConstantResource> constantOp = new ASTOpConstantResource();
-					constantOp->Type = OP_TYPE::ME_CONSTANT;
-					constantOp->SetValue( pMesh, DiskCacheContext );
-					Result = constantOp;
+					mu::Ptr<ASTOpConstantResource> ConstantOp = new ASTOpConstantResource();
+					ConstantOp->SourceDataDescriptor = SourceDataDescriptor;
+					ConstantOp->Type = OP_TYPE::ME_CONSTANT;
+					ConstantOp->SetValue( pMesh, DiskCacheContext );
+					Result = ConstantOp;
 				  }
 				break;
 			}
@@ -744,7 +745,7 @@ namespace mu
 						bool bCanBeGenerated = true;
 
 						// Check source data incompatiblities: when generating constants don't mix data that has different source descriptors (tags and other properties).
-						if (DataType == DT_IMAGE)
+						if (DataType == DT_IMAGE || DataType == DT_MESH)
 						{
 							FSourceDataDescriptor SourceDescriptor = SubgraphRoot->GetSourceDataDescriptor();
 							if (SourceDescriptor.IsInvalid())
@@ -1339,6 +1340,7 @@ namespace mu
 							mu::Ptr<ASTOpConstantResource> newOp = new ASTOpConstantResource();
 							newOp->Type = OP_TYPE::ME_CONSTANT;
 							newOp->SetValue(NewMesh, options.DiskCacheContext);
+							newOp->SourceDataDescriptor = at->GetSourceDataDescriptor();
 
 							ASTOp::Replace(at, newOp);
 						}
diff --git a/Engine/Plugins/Mutable/Source/MutableTools/Private/MuT/Compiler.cpp b/Engine/Plugins/Mutable/Source/MutableTools/Private/MuT/Compiler.cpp
--- a/Engine/Plugins/Mutable/Source/MutableTools/Private/MuT/Compiler.cpp
+++ b/Engine/Plugins/Mutable/Source/MutableTools/Private/MuT/Compiler.cpp
@@ -547,6 +547,23 @@ namespace mu
 			}
 		}
 
+		{
+			MUTABLE_CPUPROFILER_SCOPE(GenerateRoms_ImageSourceIds);
+
+			for (int32 ImageIndex = 0; ImageIndex < Program.ConstantImages.Num(); ++ImageIndex)
+			{
+				const FImageLODRange& LODRange = Program.ConstantImages[ImageIndex];
+
+				uint32 SourceId = AdditionalData.SourceImagePerConstant[ImageIndex].SourceId;
+
+				for (int32 LODRangeIndex = 0; LODRangeIndex < LODRange.LODCount; ++LODRangeIndex)
+				{
+					int32 RomIndex = Program.ConstantImageLODIndices[LODRange.FirstIndex + LODRangeIndex];
+					RomDatas[RomIndex].SourceId = SourceId;
+				}
+			}
+		}
+
 		{
 			MUTABLE_CPUPROFILER_SCOPE(GenerateRoms_ImageIdsUnique);
 
@@ -619,6 +636,7 @@ namespace mu
 			for (int32 ResourceIndex = 0; ResourceIndex < Program.ConstantMeshes.Num(); ++ResourceIndex)
 			{
 				FRomData& RomData = RomDatas[ResourceIndex];
+				RomData.SourceId = AdditionalData.SourceMeshPerConstant[ResourceIndex].SourceId;
 
 				// If the resource uses less memory than the threshold, don't save it in a separate rom.
 				if (int32(RomData.Size) <= EmbeddedDataBytesLimit)
diff --git a/Engine/Plugins/Mutable/Source/MutableTools/Private/MuT/CompilerPrivate.h b/Engine/Plugins/Mutable/Source/MutableTools/Private/MuT/CompilerPrivate.h
--- a/Engine/Plugins/Mutable/Source/MutableTools/Private/MuT/CompilerPrivate.h
+++ b/Engine/Plugins/Mutable/Source/MutableTools/Private/MuT/CompilerPrivate.h
@@ -103,6 +103,11 @@ namespace mu
 			* It must have the same size than the Program::ConstantImages array.
 			*/
 			TArray<FSourceDataDescriptor> SourceImagePerConstant;
+
+			/** Source data descriptor for every mesh constant that has been generated.
+			* It must have the same size than the Program::ConstantMeshes array.
+			*/
+			TArray<FSourceDataDescriptor> SourceMeshPerConstant;
 		};
 
 		//!
diff --git a/Engine/Plugins/Mutable/Source/MutableTools/Public/MuT/NodeImage.h b/Engine/Plugins/Mutable/Source/MutableTools/Public/MuT/NodeImage.h
--- a/Engine/Plugins/Mutable/Source/MutableTools/Public/MuT/NodeImage.h
+++ b/Engine/Plugins/Mutable/Source/MutableTools/Public/MuT/NodeImage.h
@@ -15,9 +15,9 @@ namespace mu
 	typedef Ptr<NodeImage> NodeImagePtr;
 	typedef Ptr<const NodeImage> NodeImagePtrConst;
 
-	/** Data related to the a source image that is necessary to classify the final image fragments
-	* that are derived from this source.
-	*/
+	/** Data related to the a source image that is necessary to classify the final image and mesh fragments
+	 * that are derived from this source.
+	 */
 	struct FSourceDataDescriptor
 	{
 		enum ESpecialValues
@@ -35,6 +35,9 @@ namespace mu
 		/** Source tags that mark this data and prevent it from mixing with other data at compile time. */
 		TArray<FString> Tags;
 
+		/** Source Id */
+		uint32 SourceId = MAX_uint32;
+
 		inline bool operator==(const FSourceDataDescriptor& Other) const
 		{
 			return SourceHighResMips == Other.SourceHighResMips && Tags == Other.Tags;
@@ -55,6 +58,7 @@ namespace mu
 			if (IsInvalid() || Other.IsInvalid())
 			{
 				SourceHighResMips = EInvalid;
+				SourceId = MAX_uint32;
 				Tags.Empty();
 				return;
 			}
@@ -67,13 +71,15 @@ namespace mu
 			if (IsNeutral())
 			{
 				SourceHighResMips = Other.SourceHighResMips;
+				SourceId = Other.SourceId;
 				Tags = Other.Tags;
 				return;
 			}
 
-			if (!(*this==Other))
+			if (!(*this == Other))
 			{
 				SourceHighResMips = EInvalid;
+				SourceId = MAX_uint32;
 				Tags.Empty();
 				return;
 			}
diff --git a/Engine/Plugins/Mutable/Source/MutableTools/Public/MuT/NodeMeshConstant.h b/Engine/Plugins/Mutable/Source/MutableTools/Public/MuT/NodeMeshConstant.h
--- a/Engine/Plugins/Mutable/Source/MutableTools/Public/MuT/NodeMeshConstant.h
+++ b/Engine/Plugins/Mutable/Source/MutableTools/Public/MuT/NodeMeshConstant.h
@@ -6,6 +6,7 @@
 #include "MuR/RefCounted.h"
 #include "MuT/Node.h"
 #include "MuT/NodeMesh.h"
+#include "MuT/NodeImage.h"
 
 
 namespace mu
@@ -23,6 +24,10 @@ namespace mu
 	//! It allows to define the layouts for the texture channels of the constant mesh
 	class MUTABLETOOLS_API NodeMeshConstant : public NodeMesh
 	{
+	public:
+
+		FSourceDataDescriptor SourceDataDescriptor;
+
 	public:
 
 		NodeMeshConstant();
diff --git a/Engine/Plugins/Mutable/Source/MutableTools/Public/MuT/NodeMeshTable.h b/Engine/Plugins/Mutable/Source/MutableTools/Public/MuT/NodeMeshTable.h
--- a/Engine/Plugins/Mutable/Source/MutableTools/Public/MuT/NodeMeshTable.h
+++ b/Engine/Plugins/Mutable/Source/MutableTools/Public/MuT/NodeMeshTable.h
@@ -7,6 +7,7 @@
 #include "MuT/Node.h"
 #include "MuT/NodeMesh.h"
 #include "MuT/NodeLayout.h"
+#include "MuT/NodeImage.h"
 #include "MuT/Table.h"
 
 
@@ -32,6 +33,9 @@ namespace mu
 
 		TArray<Ptr<NodeLayout>> Layouts;
 
+		/** */
+		FSourceDataDescriptor SourceDataDescriptor;
+
 	public:
 
 		// Node interface
