[Mutable] Bump CO current supported version to discard serialized data after grouping roms by source.

#rnx

diff --git a/Engine/Plugins/Mutable/Source/CustomizableObject/Internal/MuCO/CustomizableObjectPrivate.h b/Engine/Plugins/Mutable/Source/CustomizableObject/Internal/MuCO/CustomizableObjectPrivate.h
--- a/Engine/Plugins/Mutable/Source/CustomizableObject/Internal/MuCO/CustomizableObjectPrivate.h
+++ b/Engine/Plugins/Mutable/Source/CustomizableObject/Internal/MuCO/CustomizableObjectPrivate.h
@@ -1474,6 +1474,8 @@ public:
 		AddedFaceCullStrategyToSomeOperations,
 
 		DDCParticipatingObjects,
+
+		GroupRomsBySource,
 		
 		// -----<new versions can be added above this line>--------
 		LastCustomizableObjectVersion
