#Mutable Bump CO current supported version after backout

diff --git a/Engine/Plugins/Mutable/Source/CustomizableObject/Internal/MuCO/CustomizableObjectPrivate.h b/Engine/Plugins/Mutable/Source/CustomizableObject/Internal/MuCO/CustomizableObjectPrivate.h
--- a/Engine/Plugins/Mutable/Source/CustomizableObject/Internal/MuCO/CustomizableObjectPrivate.h
+++ b/Engine/Plugins/Mutable/Source/CustomizableObject/Internal/MuCO/CustomizableObjectPrivate.h
@@ -1477,6 +1477,8 @@ public:
 
 		GroupRomsBySource,
 		
+		RemovedGroupRomsBySource,
+		
 		// -----<new versions can be added above this line>--------
 		LastCustomizableObjectVersion
 	};
