[mutable] Fix broken texture projections when applied to meshes with multiple layout blocks.
#rnx

diff --git a/Engine/Plugins/Mutable/Source/CustomizableObjectEditor/Private/MuCOE/GenerateMutableSource/GenerateMutableSourceImage.cpp b/Engine/Plugins/Mutable/Source/CustomizableObjectEditor/Private/MuCOE/GenerateMutableSource/GenerateMutableSourceImage.cpp
--- a/Engine/Plugins/Mutable/Source/CustomizableObjectEditor/Private/MuCOE/GenerateMutableSource/GenerateMutableSourceImage.cpp
+++ b/Engine/Plugins/Mutable/Source/CustomizableObjectEditor/Private/MuCOE/GenerateMutableSource/GenerateMutableSourceImage.cpp
@@ -12,6 +12,7 @@
 #include "MuCOE/GraphTraversal.h"
 #include "MuCOE/Nodes/CustomizableObjectNodeMesh.h"
 #include "MuCOE/Nodes/CustomizableObjectNodeTable.h"
+#include "MuCOE/Nodes/CustomizableObjectNodeMaterial.h"
 #include "MuCOE/Nodes/CustomizableObjectNodeTexture.h"
 #include "MuCOE/Nodes/CustomizableObjectNodePassThroughTexture.h"
 #include "MuCOE/Nodes/CustomizableObjectNodeTextureBinarise.h"
@@ -692,7 +693,12 @@ mu::NodeImagePtr GenerateMutableSourceImage(const UEdGraphPin* Pin, FMutableGrap
 		if (const UEdGraphPin* ConnectedPin = FollowInputPin(*TypedNodeProject->MeshPin()))
 		{
 			FMutableGraphMeshGenerationData DummyMeshData;
+
+			FLayoutGenerationFlags LayoutGenerationFlags;
+			LayoutGenerationFlags.TexturePinModes.Init(EPinMode::Mutable, TEXSTREAM_MAX_NUM_UVCHANNELS);
+			GenerationContext.LayoutGenerationFlags.Push(LayoutGenerationFlags);
 			mu::NodeMeshPtr MeshNode = GenerateMutableSourceMesh(ConnectedPin, GenerationContext, DummyMeshData, false, false);
+			GenerationContext.LayoutGenerationFlags.Pop();
 			ImageNode->SetMesh(MeshNode);
 		}
 
