[mutable] Removed Performance Reporter utility
- Removed subclasses part of the utility
- Removed non API facing methods part of the mutable used by the utility.

#rnx

diff --git a/Engine/Plugins/Mutable/Source/CustomizableObject/Internal/MuCO/CustomizableObjectInstancePrivate.h b/Engine/Plugins/Mutable/Source/CustomizableObject/Internal/MuCO/CustomizableObjectInstancePrivate.h
--- a/Engine/Plugins/Mutable/Source/CustomizableObject/Internal/MuCO/CustomizableObjectInstancePrivate.h
+++ b/Engine/Plugins/Mutable/Source/CustomizableObject/Internal/MuCO/CustomizableObjectInstancePrivate.h
@@ -273,8 +273,6 @@ public:
 	// Copy data generated in the mutable thread over to the instance and initializes additional data required during the update
 	void PrepareForUpdate(const TSharedRef<FUpdateContextPrivate>& OperationData);
 
-	int32 GetNumLODsAvailable() const { return NumLODsAvailable; }
-	
 	// The following method is basically copied from PostEditChangeProperty and/or SkeletalMesh.cpp to be able to replicate PostEditChangeProperty without the editor
 	void PostEditChangePropertyWithoutEditor();
 	
@@ -415,11 +413,6 @@ public:
 
 	TMap<FString, bool> ParamNameToExpandedMap; // Used to check whether a mutable param is expanded in the editor to show its child params
 
-#if WITH_EDITOR
-	/** During editor, always remember the duration of the last update in the mutable runtime, for profiling. */
-	int32 LastUpdateMutableRuntimeCycles = 0;
-#endif
-
 	bool bShowOnlyRuntimeParameters = true;
 	bool bShowOnlyRelevantParameters = true;
 	bool bShowUISections = false;
diff --git a/Engine/Plugins/Mutable/Source/CustomizableObject/Internal/MuCO/CustomizableObjectSystemPrivate.h b/Engine/Plugins/Mutable/Source/CustomizableObject/Internal/MuCO/CustomizableObjectSystemPrivate.h
--- a/Engine/Plugins/Mutable/Source/CustomizableObject/Internal/MuCO/CustomizableObjectSystemPrivate.h
+++ b/Engine/Plugins/Mutable/Source/CustomizableObject/Internal/MuCO/CustomizableObjectSystemPrivate.h
@@ -604,9 +604,6 @@ public:
 	int64 UpdateStartBytes = 0;
 	int64 UpdateEndPeakBytes = 0;
 	int64 UpdateEndRealPeakBytes = 0;
-	
-	/** Used for profiling in the editor. */
-	uint32 MutableRuntimeCycles = 0;
 
 	/** If a InstanceUsage is in this set it means that its AttachParent has been modified (USkeletalMesh changed, UMaterial changed...). */
 	TSet<TWeakObjectPtr<UCustomizableObjectInstanceUsage>> AttachedParentUpdated;
diff --git a/Engine/Plugins/Mutable/Source/CustomizableObject/Private/MuCO/CustomizableObjectInstance.cpp b/Engine/Plugins/Mutable/Source/CustomizableObject/Private/MuCO/CustomizableObjectInstance.cpp
--- a/Engine/Plugins/Mutable/Source/CustomizableObject/Private/MuCO/CustomizableObjectInstance.cpp
+++ b/Engine/Plugins/Mutable/Source/CustomizableObject/Private/MuCO/CustomizableObjectInstance.cpp
@@ -6870,12 +6870,6 @@ int32 UCustomizableObjectInstance::GetMinLODToLoad() const
 }
 
 
-int32 UCustomizableObjectInstance::GetNumLODsAvailable() const
-{
-	return GetPrivate()->GetNumLODsAvailable();
-}
-
-
 int32 UCustomizableObjectInstance::GetCurrentMinLOD() const
 {
 	return GetPrivate()->CommittedDescriptor.GetMinLod();
diff --git a/Engine/Plugins/Mutable/Source/CustomizableObject/Private/MuCO/CustomizableObjectSystem.cpp b/Engine/Plugins/Mutable/Source/CustomizableObject/Private/MuCO/CustomizableObjectSystem.cpp
--- a/Engine/Plugins/Mutable/Source/CustomizableObject/Private/MuCO/CustomizableObjectSystem.cpp
+++ b/Engine/Plugins/Mutable/Source/CustomizableObject/Private/MuCO/CustomizableObjectSystem.cpp
@@ -1564,7 +1564,6 @@ namespace impl
 	void Task_Mutable_GetMeshes_GetImage_Loop(
 		const TSharedRef<FUpdateContextPrivate>& OperationData,
 		double StartTime,
-		uint32 StartCycles,
 		const TSharedRef<TArray<FGetImageData>>& GetImagesData,
 		int32 GetImageIndex);
 	
@@ -1580,7 +1579,6 @@ namespace impl
 	void Task_Mutable_GetMeshes_GetMesh_Loop(
 		const TSharedRef<FUpdateContextPrivate>& OperationData,
 		double StartTime,
-		uint32 StartCycles,
 		const TSharedRef<TArray<FGetMeshData>>& GetMeshesData,
 		int32 GetMeshIndex);
 
@@ -1590,7 +1588,6 @@ namespace impl
 	void Task_Mutable_GetImages_GetImage(
 		const TSharedRef<FUpdateContextPrivate>& OperationData,
 		double StartTime,
-		uint32 StartCycles,
 		const TSharedPtr<TArray<mu::FResourceID>>& ImagesInThisInstance,
 		int32 ImageIndex,
 		UE::Tasks::TTask<mu::FImageDesc> GetImageDescTask);
@@ -1600,7 +1597,6 @@ namespace impl
 	void Task_Mutable_GetImages_Loop(
 		const TSharedRef<FUpdateContextPrivate>& OperationData,
 		double StartTime,
-		uint32 StartCycles,
 		const TSharedPtr<TArray<mu::FResourceID>>& ImagesInThisInstance,
 		int32 ImageIndex);
 
@@ -2088,8 +2084,7 @@ namespace impl
 	/** End of the GetMeshes tasks. */
 	void Task_Mutable_GetMeshes_End(
 		const TSharedRef<FUpdateContextPrivate>& OperationData,
-		double StartTime,
-		uint32 StartCycles)
+		double StartTime)
 	{
 		MUTABLE_CPUPROFILER_SCOPE(Task_Mutable_GetMeshes_End)
 
@@ -2119,10 +2114,6 @@ namespace impl
 			check(NewEntry.Name != NAME_None);
 		}
 		
-#if WITH_EDITOR
-		const uint32 EndCycles = FPlatformTime::Cycles();
-		OperationData->MutableRuntimeCycles = EndCycles - StartCycles;
-#endif
 		OperationData->TaskGetMeshTime = FPlatformTime::Seconds() - StartTime;
 
 		TRACE_END_REGION(UE_TASK_MUTABLE_GETMESHES_REGION);
@@ -2133,7 +2124,6 @@ namespace impl
 	void Task_Mutable_GetMeshes_GetImage_Post(
 		const TSharedRef<FUpdateContextPrivate>& OperationData,
 		double StartTime,
-		uint32 StartCycles,
 		const TSharedRef<TArray<FGetImageData>>& GetImagesData,
 		int32 GetImageIndex,
 		UE::Tasks::TTask<mu::Ptr<const mu::Image>> GetImageTask)
@@ -2164,7 +2154,7 @@ namespace impl
 			UE_LOG(LogMutable, Error, TEXT("Referenced image [%d] was not stored in the resource array."), ReferenceID);
 		}
 			
-		Task_Mutable_GetMeshes_GetImage_Loop(OperationData, StartTime, StartCycles, GetImagesData, ++GetImageIndex);
+		Task_Mutable_GetMeshes_GetImage_Loop(OperationData, StartTime, GetImagesData, ++GetImageIndex);
 	}
 
 	
@@ -2172,7 +2162,6 @@ namespace impl
 	void Task_Mutable_GetMeshes_GetImage_Loop(
 		const TSharedRef<FUpdateContextPrivate>& OperationData,
 		double StartTime,
-		uint32 StartCycles,
 		const TSharedRef<TArray<FGetImageData>>& GetImagesData,
 		int32 GetImageIndex)
 	{
@@ -2180,7 +2169,7 @@ namespace impl
 
 		if (GetImageIndex >= GetImagesData->Num()) 
 		{
-			Task_Mutable_GetMeshes_End(OperationData, StartTime, StartCycles);
+			Task_Mutable_GetMeshes_End(OperationData, StartTime);
 			return;
 		}
 
@@ -2190,7 +2179,7 @@ namespace impl
 
 		UE::Tasks::AddNested(UE::Tasks::Launch(TEXT("Task_Mutable_GetMeshes_GetImage_Post"), [=]()
 		{
-			Task_Mutable_GetMeshes_GetImage_Post(OperationData, StartTime, StartCycles, GetImagesData, GetImageIndex, GetImageTask);
+			Task_Mutable_GetMeshes_GetImage_Post(OperationData, StartTime, GetImagesData, GetImageIndex, GetImageTask);
 		},
 		GetImageTask));
 	}
@@ -2199,8 +2188,7 @@ namespace impl
 	/** Gather all GetImages that have to be called. */
 	void Task_Mutable_GetMeshes_GetImages(
 		const TSharedRef<FUpdateContextPrivate>& OperationData,
-		double StartTime,
-		uint32 StartCycles)
+		double StartTime)
 	{
 		MUTABLE_CPUPROFILER_SCOPE(Task_Mutable_GetMeshes_GetImages)
 		
@@ -2388,7 +2376,7 @@ namespace impl
 			}	
 		}
 
-		Task_Mutable_GetMeshes_GetImage_Loop(OperationData, StartTime, StartCycles, GetImagesData, 0);
+		Task_Mutable_GetMeshes_GetImage_Loop(OperationData, StartTime, GetImagesData, 0);
 	}
 
 
@@ -2396,7 +2384,6 @@ namespace impl
 	void Task_MutableGetMeshes_GetMesh_Post(
 		const TSharedRef<FUpdateContextPrivate>& OperationData,
 		double StartTime,
-		uint32 StartCycles,
 		const TSharedRef<TArray<FGetMeshData>>& GetMeshesData,
 		int32 GetMeshIndex,
 		UE::Tasks::TTask<mu::Ptr<const mu::Mesh>> GetMeshTask)
@@ -2429,7 +2416,7 @@ namespace impl
 			}
 		}
 			
-		Task_Mutable_GetMeshes_GetMesh_Loop(OperationData, StartTime, StartCycles, GetMeshesData, ++GetMeshIndex);
+		Task_Mutable_GetMeshes_GetMesh_Loop(OperationData, StartTime, GetMeshesData, ++GetMeshIndex);
 	}
 	
 
@@ -2437,7 +2424,6 @@ namespace impl
 	void Task_Mutable_GetMeshes_GetMesh_Loop(
 		const TSharedRef<FUpdateContextPrivate>& OperationData,
 		double StartTime,
-		uint32 StartCycles,
 		const TSharedRef<TArray<FGetMeshData>>& GetMeshesData,
 		int32 GetMeshIndex)
 	{
@@ -2445,7 +2431,7 @@ namespace impl
 
 		if (GetMeshIndex >= GetMeshesData->Num())
 		{
-			Task_Mutable_GetMeshes_GetImages(OperationData, StartTime, StartCycles);
+			Task_Mutable_GetMeshes_GetImages(OperationData, StartTime);
 			return;
 		}
 		
@@ -2454,7 +2440,7 @@ namespace impl
 		
 		UE::Tasks::AddNested(UE::Tasks::Launch(TEXT("Task_MutableGetMeshes_GetMesh_Post"), [=]()
 		{
-			Task_MutableGetMeshes_GetMesh_Post(OperationData, StartTime, StartCycles, GetMeshesData, GetMeshIndex, GetMeshTask);
+			Task_MutableGetMeshes_GetMesh_Post(OperationData, StartTime, GetMeshesData, GetMeshIndex, GetMeshTask);
 		},
 		GetMeshTask));
 	}
@@ -2470,7 +2456,6 @@ namespace impl
 			TRACE_BEGIN_REGION(UE_TASK_MUTABLE_GETMESHES_REGION);
 
 			const double StartTime = FPlatformTime::Seconds();
-			const uint32 StartCycles = FPlatformTime::Cycles();
 
 			check(OperationData->Parameters);
 			OperationData->InstanceUpdateData.Clear();
@@ -2499,7 +2484,7 @@ namespace impl
 			if (!Instance)
 			{
 				UE_LOG(LogMutable, Warning, TEXT("An Instace update has failed."));
-				Task_Mutable_GetMeshes_End(OperationData, StartTime, StartCycles);
+				Task_Mutable_GetMeshes_End(OperationData, StartTime);
 				return;
 			}
 
@@ -2555,7 +2540,7 @@ namespace impl
 				}
 			}
 
-			Task_Mutable_GetMeshes_GetMesh_Loop(OperationData, StartTime, StartCycles, GetMeshesData, 0);
+			Task_Mutable_GetMeshes_GetMesh_Loop(OperationData, StartTime, GetMeshesData, 0);
 		}
 	}
 
@@ -2567,15 +2552,13 @@ namespace impl
 
 	
 		/** End of the GetImages tasks. */
-	void Task_Mutable_GetImages_End(const TSharedRef<FUpdateContextPrivate>& OperationData, double StartCycles, double StartTime)
+	void Task_Mutable_GetImages_End(const TSharedRef<FUpdateContextPrivate>& OperationData, double StartTime)
 	{
 		MUTABLE_CPUPROFILER_SCOPE(Task_Mutable_GetImages_End)
 		
 		// TODO: Not strictly mutable: move to another worker thread task to free mutable access?
 		Subtask_Mutable_PrepareTextures(OperationData);
-
-		const uint32 EndCycles = FPlatformTime::Cycles();
-		OperationData->MutableRuntimeCycles += EndCycles - StartCycles;
+		
 		OperationData->TaskGetImagesTime = FPlatformTime::Seconds() - StartTime;
 
 		TRACE_END_REGION(UE_TASK_MUTABLE_GETIMAGES_REGION);
@@ -2587,7 +2570,6 @@ namespace impl
 	void Task_Mutable_GetImages_GetImageDesc(
 		const TSharedRef<FUpdateContextPrivate>& OperationData,
 		double StartTime,
-		uint32 StartCycles,
 		const TSharedPtr<TArray<mu::FResourceID>>& ImagesInThisInstance,
 		int32 ImageIndex)
 	{
@@ -2600,7 +2582,7 @@ namespace impl
 		
 		UE::Tasks::AddNested(UE::Tasks::Launch(TEXT("Task_Mutable_GetImages_GetImage"), [=]()
 		{
-			Task_Mutable_GetImages_GetImage(OperationData, StartTime, StartCycles, ImagesInThisInstance, ImageIndex, GetImageDescTask);
+			Task_Mutable_GetImages_GetImage(OperationData, StartTime, ImagesInThisInstance, ImageIndex, GetImageDescTask);
 		},
 		GetImageDescTask));
 	}
@@ -2610,7 +2592,6 @@ namespace impl
 	void Task_Mutable_GetImages_GetImage_Post(
 		const TSharedRef<FUpdateContextPrivate>& OperationData,
     	double StartTime,
-    	uint32 StartCycles,
     	const TSharedPtr<TArray<mu::FResourceID>>& ImagesInThisInstance,
     	int32 ImageIndex,
     	UE::Tasks::TTask<mu::Ptr<const mu::Image>> GetImageTask,
@@ -2669,7 +2650,7 @@ namespace impl
 
 		ImagesInThisInstance->Add(Image.ImageID);
 
-		Task_Mutable_GetImages_Loop(OperationData, StartTime, StartCycles, ImagesInThisInstance, ++ImageIndex);
+		Task_Mutable_GetImages_Loop(OperationData, StartTime, ImagesInThisInstance, ++ImageIndex);
 	}
 
 
@@ -2677,7 +2658,6 @@ namespace impl
 	void Task_Mutable_GetImages_GetImage(
 		const TSharedRef<FUpdateContextPrivate>& OperationData,
 		double StartTime,
-		uint32 StartCycles,
 		const TSharedPtr<TArray<mu::FResourceID>>& ImagesInThisInstance,
 		int32 ImageIndex,
 		UE::Tasks::TTask<mu::FImageDesc> GetImageDescTask)
@@ -2714,7 +2694,7 @@ namespace impl
 		{
 			UE_LOG(LogMutable, VeryVerbose, TEXT("Texture resource with id [%llu] is cached."), Image.ImageID);
 
-			Task_Mutable_GetImages_Loop(OperationData, StartTime, StartCycles, ImagesInThisInstance, ++ImageIndex);
+			Task_Mutable_GetImages_Loop(OperationData, StartTime, ImagesInThisInstance, ++ImageIndex);
 			return;
 		}
 		
@@ -2737,7 +2717,7 @@ namespace impl
 			mu::Ptr<const mu::Image> NewImage = new mu::Image(MipSizeX, MipSizeY, FullLODCount - MipsToSkip, ImageDesc.m_format, mu::EInitializationType::Black);
 
 			UE::Tasks::TTask<mu::Ptr<const mu::Image>> DummyTask = UE::Tasks::MakeCompletedTask<mu::Ptr<const mu::Image>>(NewImage);
-			Task_Mutable_GetImages_GetImage_Post(OperationData, StartTime, StartCycles, ImagesInThisInstance, ImageIndex, DummyTask, MipSizeX, MipSizeY, FullLODCount, MipsToSkip);
+			Task_Mutable_GetImages_GetImage_Post(OperationData, StartTime, ImagesInThisInstance, ImageIndex, DummyTask, MipSizeX, MipSizeY, FullLODCount, MipsToSkip);
 		}
 		else
 		{
@@ -2745,7 +2725,7 @@ namespace impl
 			
 			UE::Tasks::AddNested(UE::Tasks::Launch(TEXT("Task_Mutable_GetImages_GetImage_Post"), [=]()
 			{
-				Task_Mutable_GetImages_GetImage_Post(OperationData, StartTime, StartCycles, ImagesInThisInstance, ImageIndex, GetImageTask, MipSizeX, MipSizeY, FullLODCount, MipsToSkip);
+				Task_Mutable_GetImages_GetImage_Post(OperationData, StartTime, ImagesInThisInstance, ImageIndex, GetImageTask, MipSizeX, MipSizeY, FullLODCount, MipsToSkip);
 			},
 			GetImageTask));
 		}
@@ -2756,7 +2736,6 @@ namespace impl
 	void Task_Mutable_GetImages_Loop(
 		const TSharedRef<FUpdateContextPrivate>& OperationData,
 		double StartTime,
-		uint32 StartCycles,
 		const TSharedPtr<TArray<mu::FResourceID>>& ImagesInThisInstance,
 		int32 ImageIndex)
 	{
@@ -2768,13 +2747,13 @@ namespace impl
 			const FInstanceUpdateData::FImage& Image = OperationData->InstanceUpdateData.Images[ImageIndex];
 			if (!Image.bIsPassThrough)
 			{
-				Task_Mutable_GetImages_GetImageDesc(OperationData, StartTime, StartCycles, ImagesInThisInstance, ImageIndex);
+				Task_Mutable_GetImages_GetImageDesc(OperationData, StartTime, ImagesInThisInstance, ImageIndex);
 				return;
 			}
 		}
 
 		// If not image needs to be processed, go to end directly
-		Task_Mutable_GetImages_End(OperationData, StartCycles, StartTime);
+		Task_Mutable_GetImages_End(OperationData, StartTime);
 	}
 
 
@@ -2787,10 +2766,9 @@ namespace impl
 			TRACE_BEGIN_REGION(UE_TASK_MUTABLE_GETIMAGES_REGION);
 
 			const double StartTime = FPlatformTime::Seconds();		
-			const uint32 StartCycles = FPlatformTime::Cycles();
 
 			const TSharedPtr<TArray<mu::FResourceID>> ImagesInThisInstance = MakeShared<TArray<mu::FResourceID>>();
-			Task_Mutable_GetImages_Loop(OperationData, StartTime, StartCycles, ImagesInThisInstance, 0);
+			Task_Mutable_GetImages_Loop(OperationData, StartTime, ImagesInThisInstance, 0);
 		}
 		
 	}
@@ -2948,10 +2926,6 @@ namespace impl
 		{
 			UCustomizableInstancePrivate* CustomizableInstancePrivateData = CustomizableObjectInstance->GetPrivate();
 
-#if WITH_EDITOR
-			CustomizableObjectInstance->GetPrivate()->LastUpdateMutableRuntimeCycles = OperationData->MutableRuntimeCycles;
-#endif
-
 			// Convert Step
 			//-------------------------------------------------------------
 
diff --git a/Engine/Plugins/Mutable/Source/CustomizableObject/Public/MuCO/CustomizableObjectInstance.h b/Engine/Plugins/Mutable/Source/CustomizableObject/Public/MuCO/CustomizableObjectInstance.h
--- a/Engine/Plugins/Mutable/Source/CustomizableObject/Public/MuCO/CustomizableObjectInstance.h
+++ b/Engine/Plugins/Mutable/Source/CustomizableObject/Public/MuCO/CustomizableObjectInstance.h
@@ -616,9 +616,7 @@ public:
 
 	/** Return the min LOD that will be used in the next Instance update. */
 	int32 GetMinLODToLoad() const;
-
-	int32 GetNumLODsAvailable() const;
-
+	
 	/** Return the min LOD this Instance is using. */
 	int32 GetCurrentMinLOD() const;
 
diff --git a/Engine/Plugins/Mutable/Source/CustomizableObjectEditor/Private/MuCOE/CustomizableObjectEditor.cpp b/Engine/Plugins/Mutable/Source/CustomizableObjectEditor/Private/MuCOE/CustomizableObjectEditor.cpp
--- a/Engine/Plugins/Mutable/Source/CustomizableObjectEditor/Private/MuCOE/CustomizableObjectEditor.cpp
+++ b/Engine/Plugins/Mutable/Source/CustomizableObjectEditor/Private/MuCOE/CustomizableObjectEditor.cpp
@@ -52,7 +52,6 @@
 #include "MuCOE/Nodes/CustomizableObjectNodeStaticMesh.h"
 #include "MuCOE/Nodes/CustomizableObjectNodeTable.h"
 #include "MuCOE/SCustomizableObjectEditorAdvancedPreviewSettings.h"
-#include "MuCOE/SCustomizableObjectEditorPerformanceReport.h"
 #include "MuCOE/SCustomizableObjectEditorTagExplorer.h"
 #include "MuCOE/SCustomizableObjectEditorTextureAnalyzer.h"
 #include "MuCOE/SCustomizableObjectEditorViewport.h"
@@ -88,7 +87,6 @@ const FName FCustomizableObjectEditor::GraphTabId( TEXT( "CustomizableObjectEdit
 const FName FCustomizableObjectEditor::GraphNodePropertiesTabId( TEXT( "CustomizableObjectEditor_GraphNodeProperties" ) );
 const FName FCustomizableObjectEditor::AdvancedPreviewSettingsTabId(TEXT("CustomizableObjectEditor_AdvancedPreviewSettings"));
 const FName FCustomizableObjectEditor::TextureAnalyzerTabId(TEXT("CustomizableObjectEditor_TextureAnalyzer"));
-const FName FCustomizableObjectEditor::PerformanceReportTabId(TEXT("CustomizableObjectEditor_PerformanceReport"));
 const FName FCustomizableObjectEditor::PerformanceAnalyzerTabId(TEXT("CustomizableObjectEditor_MewPerformanceReport"));
 const FName FCustomizableObjectEditor::TagExplorerTabId(TEXT("CustomizableObjectEditor_TagExplorer"));
 const FName FCustomizableObjectEditor::ObjectDebuggerTabId(TEXT("CustomizableObjectEditor_ObjectDebugger"));
@@ -136,10 +134,6 @@ void FCustomizableObjectEditor::RegisterTabSpawners(const TSharedRef<class FTabM
 		.SetDisplayName(LOCTEXT("TextureAnalyzer", "Texture Analyzer"))
 		.SetGroup(WorkspaceMenuCategoryRef);
 
-	InTabManager->RegisterTabSpawner(PerformanceReportTabId, FOnSpawnTab::CreateSP(this, &FCustomizableObjectEditor::SpawnTab_PerformanceReport))
-		.SetDisplayName(LOCTEXT("PerformanceReport", "Performance Report"))
-		.SetGroup(WorkspaceMenuCategoryRef);
-
 	InTabManager->RegisterTabSpawner(PerformanceAnalyzerTabId, FOnSpawnTab::CreateSP(this, &FCustomizableObjectEditor::SpawnTab_PerformanceAnalyzer))
 		.SetDisplayName(LOCTEXT("PerformanceAnalyzer", "Performance Analyzer"))
 		.SetGroup(WorkspaceMenuCategoryRef);
@@ -157,10 +151,9 @@ void FCustomizableObjectEditor::UnregisterTabSpawners(const TSharedRef<class FTa
 	InTabManager->UnregisterTabSpawner( InstancePropertiesTabId );
 	InTabManager->UnregisterTabSpawner( GraphTabId );
 	InTabManager->UnregisterTabSpawner( GraphNodePropertiesTabId );
-	InTabManager->UnregisterTabSpawner(AdvancedPreviewSettingsTabId);
-	InTabManager->UnregisterTabSpawner(TextureAnalyzerTabId);
-	InTabManager->UnregisterTabSpawner(PerformanceReportTabId);
-	InTabManager->UnregisterTabSpawner(PerformanceAnalyzerTabId);
+	InTabManager->UnregisterTabSpawner( AdvancedPreviewSettingsTabId );
+	InTabManager->UnregisterTabSpawner( TextureAnalyzerTabId );
+	InTabManager->UnregisterTabSpawner( PerformanceAnalyzerTabId );
 }	
 
 
@@ -660,14 +653,7 @@ void FCustomizableObjectEditor::BindCommands()
 		FExecuteAction::CreateSP(this, &FCustomizableObjectEditor::OpenTextureAnalyzerTab),
 		FCanExecuteAction(),
 		FIsActionChecked());
-
-	// Performance Report
-	ToolkitCommands->MapAction(
-		Commands.PerformanceReport,
-		FExecuteAction::CreateSP(this, &FCustomizableObjectEditor::OpenPerformanceReportTab),
-		FCanExecuteAction(),
-		FIsActionChecked());
-
+	
 	// Performance Analyzer
 	ToolkitCommands->MapAction(
 		Commands.PerformanceAnalyzer,
@@ -1216,7 +1202,6 @@ void FCustomizableObjectEditor::ExtendToolbar()
 			
 			ToolbarBuilder.BeginSection("Information");
 			ToolbarBuilder.AddToolBarButton(FCustomizableObjectEditorCommands::Get().TextureAnalyzer);
-			ToolbarBuilder.AddToolBarButton(FCustomizableObjectEditorCommands::Get().PerformanceReport);
 			ToolbarBuilder.AddToolBarButton(FCustomizableObjectEditorCommands::Get().PerformanceAnalyzer);
 			ToolbarBuilder.EndSection();
 		}
@@ -2247,16 +2232,12 @@ void FCustomizableObjectEditor::OpenTextureAnalyzerTab()
 }
 
 
-void FCustomizableObjectEditor::OpenPerformanceReportTab()
-{
-	TabManager->TryInvokeTab(PerformanceReportTabId);
-}
-
 void FCustomizableObjectEditor::OpenPerformanceAnalyzerTab()
 {
 	TabManager->TryInvokeTab(PerformanceAnalyzerTabId);
 }
 
+
 TSharedRef<SDockTab> FCustomizableObjectEditor::SpawnTab_TextureAnalyzer(const FSpawnTabArgs& Args)
 {
 	check(Args.GetTabId() == TextureAnalyzerTabId);
@@ -2269,24 +2250,6 @@ TSharedRef<SDockTab> FCustomizableObjectEditor::SpawnTab_TextureAnalyzer(const F
 }
 
 
-TSharedRef<SDockTab> FCustomizableObjectEditor::SpawnTab_PerformanceReport(const FSpawnTabArgs& Args)
-{
-	check(Args.GetTabId() == PerformanceReportTabId);
-	check(CustomizableObject);
-
-	if (!PerformanceReport.IsValid())
-	{
-		PerformanceReport = SNew(SCustomizableObjecEditorPerformanceReport).CustomizableObject(CustomizableObject);
-	}
-
-	return SNew(SDockTab)
-	.Label(LOCTEXT("Performance Report", "Performance Report"))
-	[
-		PerformanceReport.ToSharedRef()
-	];
-}
-
-
 TSharedRef<SDockTab> FCustomizableObjectEditor::SpawnTab_PerformanceAnalyzer(const FSpawnTabArgs& Args)
 {
 	check(Args.GetTabId() == PerformanceAnalyzerTabId);
diff --git a/Engine/Plugins/Mutable/Source/CustomizableObjectEditor/Private/MuCOE/CustomizableObjectEditor.h b/Engine/Plugins/Mutable/Source/CustomizableObjectEditor/Private/MuCOE/CustomizableObjectEditor.h
--- a/Engine/Plugins/Mutable/Source/CustomizableObjectEditor/Private/MuCOE/CustomizableObjectEditor.h
+++ b/Engine/Plugins/Mutable/Source/CustomizableObjectEditor/Private/MuCOE/CustomizableObjectEditor.h
@@ -275,7 +275,6 @@ private:
 	TSharedRef<SDockTab> SpawnTab_GraphNodeProperties(const FSpawnTabArgs& Args);
 	TSharedRef<SDockTab> SpawnTab_AdvancedPreviewSettings(const FSpawnTabArgs& Args);
 	TSharedRef<SDockTab> SpawnTab_TextureAnalyzer(const FSpawnTabArgs& Args);
-	TSharedRef<SDockTab> SpawnTab_PerformanceReport(const FSpawnTabArgs& Args);
 	TSharedRef<SDockTab> SpawnTab_PerformanceAnalyzer(const FSpawnTabArgs& Args);
 	TSharedRef<SDockTab> SpawnTab_TagExplorer(const FSpawnTabArgs& Args);
 	
@@ -356,9 +355,6 @@ private:
 	/** Open the Texture Analyzer tab */
 	void OpenTextureAnalyzerTab();
 
-	/** Open the Performance Report tab */
-	void OpenPerformanceReportTab();
-
 	/** Open the Performance Analyzer tab */
 	void OpenPerformanceAnalyzerTab();
 	
@@ -390,7 +386,6 @@ public:
 	static const FName SystemPropertiesTabId;
 	static const FName AdvancedPreviewSettingsTabId;
 	static const FName TextureAnalyzerTabId;
-	static const FName PerformanceReportTabId;
 	static const FName PerformanceAnalyzerTabId;
 	static const FName TagExplorerTabId;
 	static const FName ObjectDebuggerTabId;
@@ -436,9 +431,6 @@ private:
 	/** Texture Analyzer table widget which shows the information of the transient textures used in the customizable object instance */
 	TSharedPtr<class SCustomizableObjecEditorTextureAnalyzer> TextureAnalyzer;
 
-	/** Performance report widget to test and analyze the current customizable object resource demands */
-	TSharedPtr<class SCustomizableObjecEditorPerformanceReport> PerformanceReport;
-
 	/** New performance analyzer widget */
 	TSharedPtr<SCustomizableObjectEditorPerformanceAnalyzer> PerformanceAnalyzer;
 	
diff --git a/Engine/Plugins/Mutable/Source/CustomizableObjectEditor/Private/MuCOE/CustomizableObjectEditorActions.cpp b/Engine/Plugins/Mutable/Source/CustomizableObjectEditor/Private/MuCOE/CustomizableObjectEditorActions.cpp
--- a/Engine/Plugins/Mutable/Source/CustomizableObjectEditor/Private/MuCOE/CustomizableObjectEditorActions.cpp
+++ b/Engine/Plugins/Mutable/Source/CustomizableObjectEditor/Private/MuCOE/CustomizableObjectEditorActions.cpp
@@ -28,7 +28,6 @@ void FCustomizableObjectEditorCommands::RegisterCommands()
 	UI_COMMAND(CompileOptions_UseDiskCompilation, "Enable compiling using the disk as memory.", "This is very slow but supports compiling huge objects. It requires a lot of free space in the OS disk.", EUserInterfaceActionType::ToggleButton, FInputChord());
 	UI_COMMAND(Debug, "Debug", "Open the CustomizableObject debugger tab for this object.", EUserInterfaceActionType::Button, FInputChord());
 
-	UI_COMMAND(PerformanceReport, "Performance Report", "Open the Performance Report window to set up and perform all tests relevant to Customizable Objects and access worst cases data.", EUserInterfaceActionType::Button, FInputChord());
 	UI_COMMAND(PerformanceAnalyzer, "Performance Analyzer", "Open the Performance Analyzer window to set up and perform all tests relevant to Customizable Objects.", EUserInterfaceActionType::Button, FInputChord());
 	UI_COMMAND(TextureAnalyzer, "Texture Memory Analyzer", "Open the Texture Analyzer window to check all the information of the textures created by Mutable.", EUserInterfaceActionType::Button, FInputChord());
 
diff --git a/Engine/Plugins/Mutable/Source/CustomizableObjectEditor/Private/MuCOE/SCustomizableObjectEditorPerformanceReport.cpp b/Engine/Plugins/Mutable/Source/CustomizableObjectEditor/Private/MuCOE/SCustomizableObjectEditorPerformanceReport.cpp
deleted file mode 100644
diff --git a/Engine/Plugins/Mutable/Source/CustomizableObjectEditor/Private/MuCOE/SCustomizableObjectEditorPerformanceReport.h b/Engine/Plugins/Mutable/Source/CustomizableObjectEditor/Private/MuCOE/SCustomizableObjectEditorPerformanceReport.h
deleted file mode 100644
diff --git a/Engine/Plugins/Mutable/Source/CustomizableObjectEditor/Private/MuCOE/StatePerformanceTesting.cpp b/Engine/Plugins/Mutable/Source/CustomizableObjectEditor/Private/MuCOE/StatePerformanceTesting.cpp
deleted file mode 100644
diff --git a/Engine/Plugins/Mutable/Source/CustomizableObjectEditor/Private/MuCOE/StatePerformanceTesting.h b/Engine/Plugins/Mutable/Source/CustomizableObjectEditor/Private/MuCOE/StatePerformanceTesting.h
deleted file mode 100644
diff --git a/Engine/Plugins/Mutable/Source/CustomizableObjectEditor/Private/MuCOE/StressTest.cpp b/Engine/Plugins/Mutable/Source/CustomizableObjectEditor/Private/MuCOE/StressTest.cpp
deleted file mode 100644
diff --git a/Engine/Plugins/Mutable/Source/CustomizableObjectEditor/Private/MuCOE/StressTest.h b/Engine/Plugins/Mutable/Source/CustomizableObjectEditor/Private/MuCOE/StressTest.h
deleted file mode 100644
diff --git a/Engine/Plugins/Mutable/Source/CustomizableObjectEditor/Public/MuCOE/CustomizableObjectEditorActions.h b/Engine/Plugins/Mutable/Source/CustomizableObjectEditor/Public/MuCOE/CustomizableObjectEditorActions.h
--- a/Engine/Plugins/Mutable/Source/CustomizableObjectEditor/Public/MuCOE/CustomizableObjectEditorActions.h
+++ b/Engine/Plugins/Mutable/Source/CustomizableObjectEditor/Public/MuCOE/CustomizableObjectEditorActions.h
@@ -20,7 +20,6 @@ public:
 	TSharedPtr< FUICommandInfo > CompileOptions_UseDiskCompilation;
 	TSharedPtr< FUICommandInfo > Debug;
 
-	TSharedPtr< FUICommandInfo > PerformanceReport;
 	TSharedPtr< FUICommandInfo > PerformanceAnalyzer;
 	TSharedPtr< FUICommandInfo > ResetPerformanceReportOptions;
 	TSharedPtr< FUICommandInfo > TextureAnalyzer;
