[mutable] Add concept of internal tags. These are tags some nodes like MeshSection and ExtendMeshSeciton have so that they can be referred to without the need to manually add a new tag to them. They are shown with a friendlier name in the UI.
#rnx

diff --git a/Engine/Plugins/Mutable/Source/CustomizableObjectEditor/Private/MuCOE/GenerateMutableSource/GenerateMutableSourceModifier.cpp b/Engine/Plugins/Mutable/Source/CustomizableObjectEditor/Private/MuCOE/GenerateMutableSource/GenerateMutableSourceModifier.cpp
--- a/Engine/Plugins/Mutable/Source/CustomizableObjectEditor/Private/MuCOE/GenerateMutableSource/GenerateMutableSourceModifier.cpp
+++ b/Engine/Plugins/Mutable/Source/CustomizableObjectEditor/Private/MuCOE/GenerateMutableSource/GenerateMutableSourceModifier.cpp
@@ -387,6 +387,7 @@ mu::Ptr<mu::NodeModifier> GenerateMutableSourceModifier(const UEdGraphPin * Pin,
 		}
 
 		SurfNode->EnableTags = TypedNodeExt->Tags;
+		SurfNode->EnableTags.AddUnique(TypedNodeExt->GetInternalTag());
 
 		GenerationContext.MeshGenerationFlags.Pop();
 		GenerationContext.FromLOD = 0;
diff --git a/Engine/Plugins/Mutable/Source/CustomizableObjectEditor/Private/MuCOE/GenerateMutableSource/GenerateMutableSourceSurface.cpp b/Engine/Plugins/Mutable/Source/CustomizableObjectEditor/Private/MuCOE/GenerateMutableSource/GenerateMutableSourceSurface.cpp
--- a/Engine/Plugins/Mutable/Source/CustomizableObjectEditor/Private/MuCOE/GenerateMutableSource/GenerateMutableSourceSurface.cpp
+++ b/Engine/Plugins/Mutable/Source/CustomizableObjectEditor/Private/MuCOE/GenerateMutableSource/GenerateMutableSourceSurface.cpp
@@ -975,6 +975,8 @@ mu::Ptr<mu::NodeSurface> GenerateMutableSourceSurface(const UEdGraphPin * Pin, F
 			{
 				SurfNode->Tags.AddUnique(Tag);
 			}
+
+			SurfNode->Tags.AddUnique( TypedNodeMat->GetInternalTag() );
 		}
 
 		TArray<mu::Ptr<mu::NodeSurfaceNew>>* ArraySurfaceNodePtr = GenerationContext.MapMaterialNodeToMutableSurfaceNodeArray.Find(TypedNodeMat->GetMaterialNode());
diff --git a/Engine/Plugins/Mutable/Source/CustomizableObjectEditor/Private/MuCOE/Nodes/CustomizableObjectNode.cpp b/Engine/Plugins/Mutable/Source/CustomizableObjectEditor/Private/MuCOE/Nodes/CustomizableObjectNode.cpp
--- a/Engine/Plugins/Mutable/Source/CustomizableObjectEditor/Private/MuCOE/Nodes/CustomizableObjectNode.cpp
+++ b/Engine/Plugins/Mutable/Source/CustomizableObjectEditor/Private/MuCOE/Nodes/CustomizableObjectNode.cpp
@@ -24,6 +24,9 @@ struct FPropertyChangedEvent;
 
 #define LOCTEXT_NAMESPACE "CustomizableObjectEditor"
 
+/** This is used to generate automatic unique tags for some nodes. */
+#define MUTABLE_INTERNAL_TAG_PREFIX		"MutableInternalTag_"
+
 
 UCustomizableObjectGraph* UCustomizableObjectNode::GetCustomizableObjectGraph() const
 {
@@ -432,6 +435,93 @@ TArray<FString>* UCustomizableObjectNode::GetEnableTags()
 }
 
 
+FString UCustomizableObjectNode::GetInternalTag() const
+{
+	return FString::Printf(TEXT("%s%s"), TEXT(MUTABLE_INTERNAL_TAG_PREFIX), *NodeGuid.ToString());
+}
+
+
+FGuid UCustomizableObjectNode::GetInternalTagNodeId(const FString& Tag)
+{
+	FGuid TempId;
+
+	bool bCorrect = Tag.StartsWith(TEXT(MUTABLE_INTERNAL_TAG_PREFIX));
+	if (bCorrect)
+	{
+		int32 PrefixSize = FString(TEXT(MUTABLE_INTERNAL_TAG_PREFIX)).Len();
+		FString IdString = Tag.RightChop(PrefixSize);
+		bCorrect = FGuid::Parse(IdString, TempId);
+	}
+
+	return bCorrect ? TempId : FGuid();
+}
+
+
+bool UCustomizableObjectNode::IsInternalTag(const FString& Tag)
+{
+	return GetInternalTagNodeId(Tag).IsValid();
+}
+
+
+bool UCustomizableObjectNode::FindNodeForInternalTag(const FString& Tag, UCustomizableObjectNode*& OutNode, UCustomizableObject*& OutObject)
+{
+	OutNode = nullptr;
+	OutObject = nullptr;
+
+	FGuid NodeId = GetInternalTagNodeId(Tag);
+	if (!NodeId.IsValid())
+	{
+		return false;
+	}
+
+	// Scan all potential receivers
+	UCustomizableObject* ThisNodeObject = GetRootObject(*this);
+	UCustomizableObject* RootObject = GraphTraversal::GetRootObject(ThisNodeObject);
+
+	TSet<UCustomizableObject*> AllCustomizableObject;
+	GetAllObjectsInGraph(RootObject, AllCustomizableObject);
+
+	for (UCustomizableObject* CustObject : AllCustomizableObject)
+	{
+		if (CustObject)
+		{
+			for (const TObjectPtr<UEdGraphNode>& CandidateNode : CustObject->GetPrivate()->GetSource()->Nodes)
+			{
+				if (CandidateNode->NodeGuid == NodeId)
+				{
+					OutNode = Cast<UCustomizableObjectNode>(CandidateNode);
+					OutObject = CustObject;
+					return true;
+				}
+			}
+		}
+	}
+
+	return false;
+}
+
+
+FString UCustomizableObjectNode::GetInternalTagDisplayName()
+{
+	ensure(false);
+	return FString();
+}
+
+
+FString UCustomizableObjectNode::GetTagDisplayName(const FString& InTag)
+{
+	UCustomizableObjectNode* InternalTagNode = nullptr;
+	UCustomizableObject* InternalTagObject = nullptr;
+	bool bIsInternal = FindNodeForInternalTag(InTag, InternalTagNode, InternalTagObject);
+	if (bIsInternal && InternalTagNode)
+	{
+		return InternalTagNode->GetInternalTagDisplayName();
+	}
+
+	return InTag;
+}
+
+
 UCustomizableObjectNodeRemapPins* UCustomizableObjectNode::CreateRemapPinsDefault() const
 {
 	return CreateRemapPinsByName();
diff --git a/Engine/Plugins/Mutable/Source/CustomizableObjectEditor/Private/MuCOE/Nodes/CustomizableObjectNodeMaterial.cpp b/Engine/Plugins/Mutable/Source/CustomizableObjectEditor/Private/MuCOE/Nodes/CustomizableObjectNodeMaterial.cpp
--- a/Engine/Plugins/Mutable/Source/CustomizableObjectEditor/Private/MuCOE/Nodes/CustomizableObjectNodeMaterial.cpp
+++ b/Engine/Plugins/Mutable/Source/CustomizableObjectEditor/Private/MuCOE/Nodes/CustomizableObjectNodeMaterial.cpp
@@ -363,6 +363,12 @@ bool UCustomizableObjectNodeMaterial::HasPinViewer() const
 }
 
 
+FString UCustomizableObjectNodeMaterial::GetInternalTagDisplayName()
+{
+	return FString::Printf(TEXT("Mesh Section [%s]"), Material ? *Material->GetName() : TEXT("no-material"));
+}
+
+
 UCustomizableObjectNodeRemapPinsByName* UCustomizableObjectNodeMaterial::CreateRemapPinsDefault() const
 {
 	return NewObject<UCustomizableObjectNodeMaterialRemapPinsByName>();
diff --git a/Engine/Plugins/Mutable/Source/CustomizableObjectEditor/Private/MuCOE/Nodes/CustomizableObjectNodeMeshSectionDetails.cpp b/Engine/Plugins/Mutable/Source/CustomizableObjectEditor/Private/MuCOE/Nodes/CustomizableObjectNodeMeshSectionDetails.cpp
--- a/Engine/Plugins/Mutable/Source/CustomizableObjectEditor/Private/MuCOE/Nodes/CustomizableObjectNodeMeshSectionDetails.cpp
+++ b/Engine/Plugins/Mutable/Source/CustomizableObjectEditor/Private/MuCOE/Nodes/CustomizableObjectNodeMeshSectionDetails.cpp
@@ -75,6 +75,7 @@ void FCustomizableObjectNodeMeshSectionDetails::CustomizeDetails( IDetailLayoutB
 				SAssignNew(this->TagListWidget, SMutableTagListWidget)
 					.Node(Node)
 					.TagArray(&Node->Tags)
+					.AllowInternalTags(false)
 					.EmptyListText(LOCTEXT("MeshSectionDetails_NoTags", "No tags enabled by this mesh section."))
 					.OnTagListChanged( this, &FCustomizableObjectNodeMeshSectionDetails::OnEnableTagsPropertyChanged)
 			];
diff --git a/Engine/Plugins/Mutable/Source/CustomizableObjectEditor/Private/MuCOE/Nodes/CustomizableObjectNodeModifierExtendMeshSection.cpp b/Engine/Plugins/Mutable/Source/CustomizableObjectEditor/Private/MuCOE/Nodes/CustomizableObjectNodeModifierExtendMeshSection.cpp
--- a/Engine/Plugins/Mutable/Source/CustomizableObjectEditor/Private/MuCOE/Nodes/CustomizableObjectNodeModifierExtendMeshSection.cpp
+++ b/Engine/Plugins/Mutable/Source/CustomizableObjectEditor/Private/MuCOE/Nodes/CustomizableObjectNodeModifierExtendMeshSection.cpp
@@ -212,4 +212,11 @@ TArray<FString>* UCustomizableObjectNodeModifierExtendMeshSection::GetEnableTags
 }
 
 
+FString UCustomizableObjectNodeModifierExtendMeshSection::GetInternalTagDisplayName()
+{
+	UMaterialInterface* Material = ReferenceMaterial;
+	return FString::Printf(TEXT("Extend Mesh [%s]"), Material ? *Material->GetName() : TEXT("no-material"));
+}
+
+
 #undef LOCTEXT_NAMESPACE
diff --git a/Engine/Plugins/Mutable/Source/CustomizableObjectEditor/Private/MuCOE/Nodes/CustomizableObjectNodeModifierExtendMeshSectionDetails.cpp b/Engine/Plugins/Mutable/Source/CustomizableObjectEditor/Private/MuCOE/Nodes/CustomizableObjectNodeModifierExtendMeshSectionDetails.cpp
--- a/Engine/Plugins/Mutable/Source/CustomizableObjectEditor/Private/MuCOE/Nodes/CustomizableObjectNodeModifierExtendMeshSectionDetails.cpp
+++ b/Engine/Plugins/Mutable/Source/CustomizableObjectEditor/Private/MuCOE/Nodes/CustomizableObjectNodeModifierExtendMeshSectionDetails.cpp
@@ -75,6 +75,7 @@ void FCustomizableObjectNodeModifierExtendMeshSectionDetails::CustomizeDetails(
 				SAssignNew(this->EnableTagListWidget, SMutableTagListWidget)
 					.Node(Node)
 					.TagArray(&Node->Tags)
+					.AllowInternalTags(false)
 					.EmptyListText(LOCTEXT("ExtendMeshSectionDetails_NoTags", "No tags enabled by this extended mesh section."))
 					.OnTagListChanged(this, &FCustomizableObjectNodeModifierExtendMeshSectionDetails::OnEnableTagsPropertyChanged)
 			];
diff --git a/Engine/Plugins/Mutable/Source/CustomizableObjectEditor/Private/MuCOE/Nodes/CustomizableObjectNodeTextureVariation.cpp b/Engine/Plugins/Mutable/Source/CustomizableObjectEditor/Private/MuCOE/Nodes/CustomizableObjectNodeTextureVariation.cpp
--- a/Engine/Plugins/Mutable/Source/CustomizableObjectEditor/Private/MuCOE/Nodes/CustomizableObjectNodeTextureVariation.cpp
+++ b/Engine/Plugins/Mutable/Source/CustomizableObjectEditor/Private/MuCOE/Nodes/CustomizableObjectNodeTextureVariation.cpp
@@ -41,8 +41,8 @@ void UCustomizableObjectNodeTextureVariation::AllocateDefaultPins(UCustomizableO
 		UEdGraphPin* VariationPin = CustomCreatePin(EGPD_Input, Schema->PC_Image, FName(*PinName), false);
 		VariationPin->bDefaultValueIsIgnored = true;
 
-		FString FriendlyName = FString::Printf(TEXT("Variation %d [%s]"), VariationIndex, *Variations[VariationIndex].Tag);
-		VariationPin->PinFriendlyName = FText::FromString(FriendlyName);
+		FString TagName = GetTagDisplayName(Variations[VariationIndex].Tag);
+		VariationPin->PinFriendlyName = FText::Format(LOCTEXT("Variation_Pin_FriendlyName", "Variation {0} [{1}]"), VariationIndex, FText::FromString(TagName));
 	}
 
 	UEdGraphPin* DefaultVariation = CustomCreatePin(EGPD_Input, Schema->PC_Image, FName(TEXT("Default")), false);
diff --git a/Engine/Plugins/Mutable/Source/CustomizableObjectEditor/Private/MuCOE/Nodes/CustomizableObjectNodeVariation.cpp b/Engine/Plugins/Mutable/Source/CustomizableObjectEditor/Private/MuCOE/Nodes/CustomizableObjectNodeVariation.cpp
--- a/Engine/Plugins/Mutable/Source/CustomizableObjectEditor/Private/MuCOE/Nodes/CustomizableObjectNodeVariation.cpp
+++ b/Engine/Plugins/Mutable/Source/CustomizableObjectEditor/Private/MuCOE/Nodes/CustomizableObjectNodeVariation.cpp
@@ -40,7 +40,8 @@ void UCustomizableObjectNodeVariation::AllocateDefaultPins(UCustomizableObjectNo
 		const FName PinName = FName(FString::Printf( TEXT("Variation %d"), VariationIndex));
 		UEdGraphPin* VariationPin = CustomCreatePin(EGPD_Input, Category, PinName, bIsInputPinArray);
 
-		VariationPin->PinFriendlyName = FText::Format(LOCTEXT("Variation_Pin_FriendlyName", "Variation {0} [{1}]"), VariationIndex, FText::FromString(*VariationsData[VariationIndex].Tag));
+		FString TagName = GetTagDisplayName(VariationsData[VariationIndex].Tag);
+		VariationPin->PinFriendlyName = FText::Format(LOCTEXT("Variation_Pin_FriendlyName", "Variation {0} [{1}]"), VariationIndex, FText::FromString(TagName));
 		
 		VariationsPins[VariationIndex] = VariationPin;
 	}
diff --git a/Engine/Plugins/Mutable/Source/CustomizableObjectEditor/Private/MuCOE/SMutableTagListWidget.cpp b/Engine/Plugins/Mutable/Source/CustomizableObjectEditor/Private/MuCOE/SMutableTagListWidget.cpp
--- a/Engine/Plugins/Mutable/Source/CustomizableObjectEditor/Private/MuCOE/SMutableTagListWidget.cpp
+++ b/Engine/Plugins/Mutable/Source/CustomizableObjectEditor/Private/MuCOE/SMutableTagListWidget.cpp
@@ -10,6 +10,7 @@
 #include "MuCOE/Nodes/CustomizableObjectNodeObjectGroup.h"
 #include "MuCOE/Nodes/CustomizableObjectNodeMaterial.h"
 #include "MuCOE/Nodes/CustomizableObjectNodeModifierBase.h"
+#include "MuCOE/Nodes/CustomizableObjectNodeModifierExtendMeshSection.h"
 #include "MuCOE/Nodes/CustomizableObjectNodeVariation.h"
 #include "MuCOE/Nodes/CustomizableObjectNodeTextureVariation.h"
 #include "MuCOE/Nodes/CustomizableObjectNodeFloatVariation.h"
@@ -32,6 +33,7 @@
 void SMutableTagComboBox::Construct(const FArguments& InArgs)
 {
 	Node = InArgs._Node;
+	bAllowInternalTags = InArgs._AllowInternalTags;
 
 	RefreshOptions();
 
@@ -42,6 +44,7 @@ void SMutableTagComboBox::Construct(const FArguments& InArgs)
 		.ButtonStyle(FAppStyle::Get(), "NoBorder")
 		.MenuButtonBrush(InArgs._MenuButtonBrush)
 		.OnSelectionChanged(InArgs._OnSelectionChanged)
+		.ForegroundColor(FSlateColor::UseSubduedForeground())
 		.Content()
 		[
 			InArgs._Content.Widget
@@ -106,8 +109,24 @@ TSharedPtr<SMutableSearchComboBox::FFilteredOption> SMutableTagComboBox::AddNode
 		{
 			Option = MakeShared<SMutableSearchComboBox::FFilteredOption>();
 			Option->Parent = ParentOption;
-			UMaterialInterface* Material = MeshSectionNode->GetMaterial();
-			Option->DisplayOption = FString::Printf(TEXT("Mesh Section [%s]"), Material ? *Material->GetName() : TEXT("no-material"));
+			Option->DisplayOption = MeshSectionNode->GetInternalTagDisplayName();
+			if (bAllowInternalTags)
+			{
+				Option->ActualOption = MeshSectionNode->GetInternalTag();
+			}
+			TagComboOptionsSource.Add(Option.ToSharedRef());
+		}
+
+		else if (UCustomizableObjectNodeModifierExtendMeshSection* ExtendMeshSectionNode = Cast<UCustomizableObjectNodeModifierExtendMeshSection>(InNode))
+		{
+			Option = MakeShared<SMutableSearchComboBox::FFilteredOption>();
+			Option->Parent = ParentOption;
+			UMaterialInterface* Material = ExtendMeshSectionNode->ReferenceMaterial;
+			Option->DisplayOption = ExtendMeshSectionNode->GetInternalTagDisplayName();
+			if (bAllowInternalTags)
+			{
+				Option->ActualOption = ExtendMeshSectionNode->GetInternalTag();
+			}
 			TagComboOptionsSource.Add(Option.ToSharedRef());
 		}
 
@@ -192,14 +211,21 @@ void SMutableTagComboBox::RefreshOptions()
 
 				for (const FString& OneTag : *EnableTags)
 				{
-					if (!OneTag.IsEmpty())
+					if (OneTag.IsEmpty())
 					{
-						TSharedRef Option = MakeShared<SMutableSearchComboBox::FFilteredOption>();
-						Option->ActualOption = OneTag;
-						Option->DisplayOption = OneTag;
-						Option->Parent = NodeOption;
-						TagComboOptionsSource.Add(Option);
+						continue;
+					}
+
+					if (!bAllowInternalTags && UCustomizableObjectNode::IsInternalTag(OneTag))
+					{
+						continue;
 					}
+
+					TSharedRef Option = MakeShared<SMutableSearchComboBox::FFilteredOption>();
+					Option->ActualOption = OneTag;
+					Option->DisplayOption = OneTag;
+					Option->Parent = NodeOption;
+					TagComboOptionsSource.Add(Option);
 				}
 			}
 
@@ -210,14 +236,21 @@ void SMutableTagComboBox::RefreshOptions()
 				{
 					TSharedPtr<SMutableSearchComboBox::FFilteredOption> NodeOption = AddNodeHierarchyOptions(Typed, AddedOptions);
 
-					if (!OneTag.IsEmpty())
+					if (OneTag.IsEmpty())
 					{
-						TSharedRef Option = MakeShared<SMutableSearchComboBox::FFilteredOption>();
-						Option->ActualOption = OneTag;
-						Option->DisplayOption = OneTag;
-						Option->Parent = NodeOption;
-						TagComboOptionsSource.Add(Option);
+						continue;
+					}
+
+					if (!bAllowInternalTags && UCustomizableObjectNode::IsInternalTag(OneTag))
+					{
+						continue;
 					}
+
+					TSharedRef Option = MakeShared<SMutableSearchComboBox::FFilteredOption>();
+					Option->ActualOption = OneTag;
+					Option->DisplayOption = OneTag;
+					Option->Parent = NodeOption;
+					TagComboOptionsSource.Add(Option);
 				}
 			}
 
@@ -230,14 +263,21 @@ void SMutableTagComboBox::RefreshOptions()
 					{
 						TSharedPtr<SMutableSearchComboBox::FFilteredOption> NodeOption = AddNodeHierarchyOptions(VariationNode, AddedOptions);
 
-						if (!Var.Tag.IsEmpty())
+						if (Var.Tag.IsEmpty())
+						{
+							continue;
+						}
+
+						if (!bAllowInternalTags && UCustomizableObjectNode::IsInternalTag(Var.Tag))
 						{
-							TSharedRef Option = MakeShared<SMutableSearchComboBox::FFilteredOption>();
-							Option->ActualOption = Var.Tag;
-							Option->DisplayOption = Var.Tag;
-							Option->Parent = NodeOption;
-							TagComboOptionsSource.Add(Option);
+							continue;
 						}
+
+						TSharedRef Option = MakeShared<SMutableSearchComboBox::FFilteredOption>();
+						Option->ActualOption = Var.Tag;
+						Option->DisplayOption = Var.Tag;
+						Option->Parent = NodeOption;
+						TagComboOptionsSource.Add(Option);
 					}
 				}
 			}
@@ -251,14 +291,21 @@ void SMutableTagComboBox::RefreshOptions()
 					{
 						TSharedPtr<SMutableSearchComboBox::FFilteredOption> NodeOption = AddNodeHierarchyOptions(VariationNode, AddedOptions);
 
-						if (!Var.Tag.IsEmpty())
+						if (Var.Tag.IsEmpty())
+						{
+							continue;
+						}
+
+						if (!bAllowInternalTags && UCustomizableObjectNode::IsInternalTag(Var.Tag))
 						{
-							TSharedRef Option = MakeShared<SMutableSearchComboBox::FFilteredOption>();
-							Option->ActualOption = Var.Tag;
-							Option->DisplayOption = Var.Tag;
-							Option->Parent = NodeOption;
-							TagComboOptionsSource.Add(Option);
+							continue;
 						}
+
+						TSharedRef Option = MakeShared<SMutableSearchComboBox::FFilteredOption>();
+						Option->ActualOption = Var.Tag;
+						Option->DisplayOption = Var.Tag;
+						Option->Parent = NodeOption;
+						TagComboOptionsSource.Add(Option);
 					}
 				}
 			}
@@ -296,6 +343,7 @@ void SMutableTagListWidget::Construct(const FArguments& InArgs)
 					SAssignNew(this->TagCombo, SMutableTagComboBox)
 						.Node(Node)
 						.MenuButtonBrush( FAppStyle::GetBrush(TEXT("Icons.PlusCircle")) )
+						.AllowInternalTags( InArgs._AllowInternalTags )
 						.OnSelectionChanged(this, &SMutableTagListWidget::OnTagComboBoxSelectionChanged)
 				]
 		]
@@ -344,7 +392,19 @@ void SMutableTagListWidget::RefreshOptions()
 		for (const FString& OneTag : *TagArray)
 		{
 			TSharedPtr<FTagUIData> Data = MakeShared<FTagUIData>();
-			Data->DisplayName = OneTag;
+
+			UCustomizableObjectNode* InternalTagNode = nullptr;
+			UCustomizableObject* InternalTagObject = nullptr;
+			bool bIsInternal = Node->FindNodeForInternalTag( OneTag, InternalTagNode, InternalTagObject );
+			if (bIsInternal && InternalTagNode)
+			{
+				Data->DisplayName = InternalTagNode->GetInternalTagDisplayName();
+			}
+			else
+			{
+				Data->DisplayName = OneTag;
+			}
+
 			Data->Tag = OneTag;
 			CurrentTagsSource.Add(Data);
 		}
@@ -372,7 +432,7 @@ TSharedRef<ITableRow> SMutableTagListWidget::GenerateTagListItemRow(TSharedPtr<F
 {
 	if (!InItem)
 	{
-		return SNew(SComboRow<TSharedPtr<FString>>, OwnerTable)
+		return SNew(STableRow<TSharedPtr<FString>>, OwnerTable)
 			[
 				SNew(STextBlock)
 					.Text(FText::FromString("No item."))
@@ -380,13 +440,15 @@ TSharedRef<ITableRow> SMutableTagListWidget::GenerateTagListItemRow(TSharedPtr<F
 			];
 	}
 
-	return SNew(SComboRow<TSharedPtr<FString>>, OwnerTable)
+	return SNew(STableRow<TSharedPtr<FString>>, OwnerTable)
+		.Padding(0)
 		[
 			SNew(SHorizontalBox)
 
 				+ SHorizontalBox::Slot()
 				.HAlign(HAlign_Fill)
 				.VAlign(VAlign_Center)
+				.Padding(4)
 				[
 					SNew(STextBlock)
 						.Text(FText::FromString(InItem->DisplayName))
@@ -395,8 +457,10 @@ TSharedRef<ITableRow> SMutableTagListWidget::GenerateTagListItemRow(TSharedPtr<F
 
 				+ SHorizontalBox::Slot()
 				.AutoWidth()
+				.Padding(0)
 				[
 					SNew(SButton)
+						.ContentPadding(2)
 						.ButtonStyle(FAppStyle::Get(), "NoBorder")
 						.ToolTipText(LOCTEXT("RemoveModifierTag", "Remove this tag from the modifier."))
 						.OnClicked_Lambda([this, InItem]()
@@ -414,6 +478,7 @@ TSharedRef<ITableRow> SMutableTagListWidget::GenerateTagListItemRow(TSharedPtr<F
 						[
 							SNew(SImage)
 								.Image(FAppStyle::GetBrush(TEXT("Icons.MinusCircle")))
+								.ColorAndOpacity(FSlateColor::UseSubduedForeground())
 						]
 				]
 		];
diff --git a/Engine/Plugins/Mutable/Source/CustomizableObjectEditor/Private/MuCOE/SMutableTagListWidget.h b/Engine/Plugins/Mutable/Source/CustomizableObjectEditor/Private/MuCOE/SMutableTagListWidget.h
--- a/Engine/Plugins/Mutable/Source/CustomizableObjectEditor/Private/MuCOE/SMutableTagListWidget.h
+++ b/Engine/Plugins/Mutable/Source/CustomizableObjectEditor/Private/MuCOE/SMutableTagListWidget.h
@@ -28,11 +28,13 @@ public:
 
 	SLATE_BEGIN_ARGS(SMutableTagComboBox)
 		: _Node(nullptr)
+		, _AllowInternalTags(true)
 		{}
 		/** Slot for this button's content (optional) */
 		SLATE_DEFAULT_SLOT(FArguments, Content)
 		SLATE_ARGUMENT(UCustomizableObjectNode*, Node)
 		SLATE_ARGUMENT(const FSlateBrush*, MenuButtonBrush)
+		SLATE_ARGUMENT(bool, AllowInternalTags)
 		SLATE_EVENT(FOnTextChanged, OnSelectionChanged)
 
 	SLATE_END_ARGS()
@@ -49,6 +51,8 @@ private:
 
 	TArray< TSharedRef<SMutableSearchComboBox::FFilteredOption> > TagComboOptionsSource;
 
+	bool bAllowInternalTags = true;
+
 	/** */
 	TSharedPtr<SMutableSearchComboBox::FFilteredOption> AddNodeHierarchyOptions(UEdGraphNode* Node, TMap<UEdGraphNode*, TSharedPtr<SMutableSearchComboBox::FFilteredOption>>& AddedOptions);
 
@@ -63,11 +67,13 @@ public:
 	SLATE_BEGIN_ARGS(SMutableTagListWidget)
 		: _Node(nullptr)
 		, _TagArray(nullptr)
+		, _AllowInternalTags(true)
 		{}
 
 		SLATE_ARGUMENT(UCustomizableObjectNode*, Node)
 		SLATE_ARGUMENT(TArray<FString>*, TagArray)
 		SLATE_ARGUMENT(FText, EmptyListText)
+		SLATE_ARGUMENT(bool, AllowInternalTags)
 
 		SLATE_EVENT(FSimpleDelegate, OnTagListChanged)
 
@@ -89,7 +95,6 @@ private:
 	FText EmptyListText;
 
 	TSharedPtr<SMutableTagComboBox> TagCombo;
-	TArray< TSharedRef<SMutableSearchComboBox::FFilteredOption> > TagComboOptionsSource;
 	void OnTagComboBoxSelectionChanged(const FText& NewText);
 
 	struct FTagUIData
diff --git a/Engine/Plugins/Mutable/Source/CustomizableObjectEditor/Private/MuCOE/Widgets/CustomizableObjectVariationCustomization.cpp b/Engine/Plugins/Mutable/Source/CustomizableObjectEditor/Private/MuCOE/Widgets/CustomizableObjectVariationCustomization.cpp
--- a/Engine/Plugins/Mutable/Source/CustomizableObjectEditor/Private/MuCOE/Widgets/CustomizableObjectVariationCustomization.cpp
+++ b/Engine/Plugins/Mutable/Source/CustomizableObjectEditor/Private/MuCOE/Widgets/CustomizableObjectVariationCustomization.cpp
@@ -75,6 +75,9 @@ void FCustomizableObjectVariationCustomization::CustomizeHeader(TSharedRef<IProp
 									{
 										TagPropertyHandle->GetValue(Value);
 									}
+
+									Value = BaseObjectNode->GetTagDisplayName(Value);
+
 									return FText::FromString(Value); 
 								})
 							.OnTextCommitted_Lambda([&](const FText& NewText, ETextCommit::Type)
diff --git a/Engine/Plugins/Mutable/Source/CustomizableObjectEditor/Public/MuCOE/Nodes/CustomizableObjectNode.h b/Engine/Plugins/Mutable/Source/CustomizableObjectEditor/Public/MuCOE/Nodes/CustomizableObjectNode.h
--- a/Engine/Plugins/Mutable/Source/CustomizableObjectEditor/Public/MuCOE/Nodes/CustomizableObjectNode.h
+++ b/Engine/Plugins/Mutable/Source/CustomizableObjectEditor/Public/MuCOE/Nodes/CustomizableObjectNode.h
@@ -100,7 +100,7 @@ public:
 
 	/** Allocates the default pins using the empty remap pins action. Usually called from CreateNode. */
 	virtual void AllocateDefaultPins() override final; // Final. Override AllocateDefaultPins(UCustomizableObjectNodeRemapPins* RemapPins) instead
-	
+
 	/** Reconstructs the node using its default remap pins action. */
 	virtual void ReconstructNode() override final; // Final. Override ReconstructNode(UCustomizableObjectNodeRemapPins* RemapPins) instead
 	virtual void PostInitProperties() override;
@@ -120,12 +120,12 @@ public:
 	 * Since this function is used BackwardsCompatibleFixup, it should not have major changes (or they will break existing fixups).
 	 * If changes are needed, duplicate this function so old fixups remain unchanged. */
 	void FixupReconstructPins(UCustomizableObjectNodeRemapPins* RemapPinsAction, TFunction<void(UCustomizableObjectNode*, UCustomizableObjectNodeRemapPins*)> AllocateDefaultPins);
-	
+
 	/** Add post load work here.
 	 * When called, it is guaranteed that all nodes in this graph will have executed the BackwardsCompatibleFixup function.
 	 * Notice that no compatibility code should not go here. Use BackwardsCompatibleFixup instead. */
 	virtual void PostBackwardsCompatibleFixup() {}
-	
+
 	/** Virtual implementation of RemovePin. Allows to do work before removing a pin.
 	 * Use this function instead of RemovePin. RemovePin does not removes possible attached PinData. */
 	virtual bool CustomRemovePin(UEdGraphPin& Pin);
@@ -133,7 +133,7 @@ public:
 	/**
 	 * Subclasses should override this to return true and set OutCategory if this node should be
 	 * auto-added to the right-click context menu in the graph editor.
-	 * 
+	 *
 	 * Some nodes are added manually in the graph editor code and don't need to do this.
 	 */
 	virtual bool ShouldAddToContextMenu(FText& OutCategory) const;
@@ -141,7 +141,7 @@ public:
 	void GetInputPins(TArray<class UEdGraphPin*>& OutInputPins) const;
 	void GetOutputPins(TArray<class UEdGraphPin*>& OutOutputPins) const;
 	UEdGraphPin* GetOutputPin(int32 OutputIndex) const;
-	
+
 	virtual bool IsNodeOutDatedAndNeedsRefresh() { return false; }
 	virtual FString GetRefreshMessage() const { return "Refresh Node."; }
 	void SetRefreshNodeWarning();
@@ -159,17 +159,17 @@ public:
 	{
 		TArrayView<const float> UnassignedUVs;
 	};
-	
-	/** 
-	 * Check if two pins can be connected. 
-	 * 
+
+	/**
+	 * Check if two pins can be connected.
+	 *
 	 * @param	InOwnedInputPin	Input pin which belongs to this node.
 	 * @param	InOutputPin		Output pin which belongs to another node. If the node connects to iself, it could belong to this node.
 	 * @param	bOutIsOtherNodeBlocklisted		Is the other node of a type we are not allowed to connect?
 	 * @param	bOutArePinsCompatible		Does InOutputPin pin share the same type as the InOwnedInputPin?
 	 * @return	True if the pin types match and if the other node is not one of the Blocklisted types.
 	 */
-	virtual bool CanConnect( const UEdGraphPin* InOwnedInputPin, const UEdGraphPin* InOutputPin, bool& bOutIsOtherNodeBlocklisted, bool& bOutArePinsCompatible) const;
+	virtual bool CanConnect(const UEdGraphPin* InOwnedInputPin, const UEdGraphPin* InOutputPin, bool& bOutIsOtherNodeBlocklisted, bool& bOutArePinsCompatible) const;
 
 	// Used during compilation process to cache the node for all LODs, or generate it specifically for each of them.
 	virtual bool IsAffectedByLOD() const;
@@ -177,6 +177,26 @@ public:
 	/** Return an array of tags that this node will enable and apply to its data. Null if none, or of it doesn't apply. */
 	virtual TArray<FString>* GetEnableTags();
 
+	/** Return the unique internal tag that can identify this node across all objects.
+	* Internal tags are automatically added to some nodes that can be then referred to from other nodes.
+	* They are usually shown differently in the UI, to represent them as tags that will be defined in a single node.
+	*/
+	FString GetInternalTag() const;
+
+	/** Return a valid GUID if a tag represents an internal tag. */
+	static FGuid GetInternalTagNodeId(const FString& Tag);
+	static bool IsInternalTag(const FString& Tag);
+
+	/** Find the object and node that a given tag refers to, if it belongs to the same CO hierarchy that this node. */
+	bool FindNodeForInternalTag(const FString& Tag, UCustomizableObjectNode*& OutNode, UCustomizableObject*& OutObject);
+
+	/** Return a non-unique, non-persistent UI-ready text to represent the interna tag of this node. It includes some node information for the user to identify the tag. */
+	virtual FString GetInternalTagDisplayName();
+
+	/** Return a UI-ready text to represent the given tag. For most cases this is just the tag, but if it is an internal tag
+	* of any node related to this one (in the same CO hierarchy) a more descriptive (non-persitent) name is build with node information.
+	*/
+	FString GetTagDisplayName(const FString& Tag);
 
 	// Get the CustomizableObject graph that owns this node
 	class UCustomizableObjectGraph* GetCustomizableObjectGraph() const;
diff --git a/Engine/Plugins/Mutable/Source/CustomizableObjectEditor/Public/MuCOE/Nodes/CustomizableObjectNodeMaterial.h b/Engine/Plugins/Mutable/Source/CustomizableObjectEditor/Public/MuCOE/Nodes/CustomizableObjectNodeMaterial.h
--- a/Engine/Plugins/Mutable/Source/CustomizableObjectEditor/Public/MuCOE/Nodes/CustomizableObjectNodeMaterial.h
+++ b/Engine/Plugins/Mutable/Source/CustomizableObjectEditor/Public/MuCOE/Nodes/CustomizableObjectNodeMaterial.h
@@ -133,7 +133,7 @@ public:
 	virtual FString GetRefreshMessage() const override;
 	virtual TSharedPtr<IDetailsView> CustomizePinDetails(const UEdGraphPin& Pin) const override;
 	virtual TArray<FString>* GetEnableTags() override;
-
+	virtual FString GetInternalTagDisplayName() override;
 	
 	// UCustomizableObjectNodeMaterialBase interface
 	virtual TArray<UCustomizableObjectLayout*> GetLayouts() const override;
diff --git a/Engine/Plugins/Mutable/Source/CustomizableObjectEditor/Public/MuCOE/Nodes/CustomizableObjectNodeModifierExtendMeshSection.h b/Engine/Plugins/Mutable/Source/CustomizableObjectEditor/Public/MuCOE/Nodes/CustomizableObjectNodeModifierExtendMeshSection.h
--- a/Engine/Plugins/Mutable/Source/CustomizableObjectEditor/Public/MuCOE/Nodes/CustomizableObjectNodeModifierExtendMeshSection.h
+++ b/Engine/Plugins/Mutable/Source/CustomizableObjectEditor/Public/MuCOE/Nodes/CustomizableObjectNodeModifierExtendMeshSection.h
@@ -43,6 +43,7 @@ public:
 	virtual bool IsSingleOutputNode() const override;
 	virtual bool CustomRemovePin(UEdGraphPin& Pin) override;
 	virtual TArray<FString>* GetEnableTags() override;
+	virtual FString GetInternalTagDisplayName() override;
 
 	// Own interface
 	UEdGraphPin* AddMeshPin() const;
