[Mutable] Private UCustomizableSkeletalComponent and UCustomizableInstanceUsage.

Moved non-public functions and members to the respective private classes.

#rnx

diff --git a/Engine/Plugins/Mutable/Source/CustomizableObject/Internal/MuCO/CustomizableObjectInstanceUsagePrivate.h b/Engine/Plugins/Mutable/Source/CustomizableObject/Internal/MuCO/CustomizableObjectInstanceUsagePrivate.h
new file mode 100644
--- /dev/null
+++ b/Engine/Plugins/Mutable/Source/CustomizableObject/Internal/MuCO/CustomizableObjectInstanceUsagePrivate.h
@@ -0,0 +1,78 @@
+﻿// Copyright Epic Games, Inc. All Rights Reserved.
+
+#pragma once
+
+#include "Tickable.h"
+
+#include "CustomizableObjectInstanceUsagePrivate.generated.h"
+
+class UCustomizableObjectInstance;
+class UCustomizableObjectInstanceUsage;
+class UCustomizableSkeletalComponent;
+class UPhysicsAsset;
+class USkeletalMesh;
+class USkeletalMeshComponent;
+
+UCLASS()
+class CUSTOMIZABLEOBJECT_API UCustomizableObjectInstanceUsagePrivate : public UObject, public FTickableGameObject
+{
+	GENERATED_BODY()
+
+public:
+	// FTickableGameObject interface
+	virtual void Tick(float DeltaTime) override;
+	virtual ETickableTickType GetTickableTickType() const override;
+	virtual TStatId GetStatId() const override;
+	virtual bool IsTickableWhenPaused() const override;
+	virtual bool IsTickableInEditor() const override;
+	virtual bool IsTickable() const override;
+	
+	// Own interface
+	
+	/** Common end point of all updates. Even those which failed. */
+	void Callbacks() const;
+
+	void SetSkeletalMeshAndOverrideMaterials(USkeletalMeshComponent& Parent, USkeletalMesh* SkeletalMesh, const UCustomizableObjectInstance& CustomizableObjectInstance, bool* bOutSkeletalMeshUpdated, bool* bOutMaterialsUpdated);
+
+#if WITH_EDITOR
+	// Used to generate instances outside the CustomizableObject editor and PIE
+	void UpdateDistFromComponentToLevelEditorCamera(const FVector& CameraPosition);
+	void EditorUpdateComponent();
+#endif
+
+	USkeletalMesh* GetSkeletalMesh() const;
+
+	USkeletalMesh* GetAttachedSkeletalMesh() const;
+
+	void SetSkeletalMesh(USkeletalMesh* SkeletalMesh, bool* bOutSkeletalMeshUpdated = nullptr, bool* bOutMaterialsUpdated = nullptr);
+
+	void SetPhysicsAsset(UPhysicsAsset* PhysicsAsset, bool* bOutPhysicsAssetUpdated = nullptr);
+
+	void UpdateDistFromComponentToPlayer(const AActor* const Pawn, bool bForceEvenIfNotBegunPlay = false);
+
+	/** Set to true to replace the SkeletalMesh of the parent component by the ReferenceSkeletalMesh or the generated SkeletalMesh. */
+	void SetPendingSetSkeletalMesh(bool bIsActive);
+	
+	bool GetPendingSetSkeletalMesh() const;
+
+	UCustomizableSkeletalComponent* GetCustomizableSkeletalComponent();
+
+	void SetCustomizableSkeletalComponent(UCustomizableSkeletalComponent* Component);
+	
+	UCustomizableObjectInstanceUsage* GetPublic();
+
+	const UCustomizableObjectInstanceUsage* GetPublic() const;
+
+	// Returns true if the NetMode of the associated UCustomizableSkeletalComponent (or the associated SkeletalMeshComponent if the former does not exist) is equal to InNetMode
+	bool IsNetMode(ENetMode InNetMode) const;
+	
+	// Used to replace the SkeletalMesh of the parent component by the ReferenceSkeletalMesh or the generated SkeletalMesh 
+	bool bUsedPendingSetSkeletalMesh = false;
+
+	// Used to avoid replacing the SkeletalMesh of the parent component by the ReferenceSkeletalMesh if bPendingSetSkeletalMesh is true
+	UPROPERTY()
+	bool bUsedSkipSetReferenceSkeletalMesh = false;
+
+	UPROPERTY()
+	bool bUsedSkipSetSkeletalMeshOnAttach = false;
+};
diff --git a/Engine/Plugins/Mutable/Source/CustomizableObject/Internal/MuCO/CustomizableSkeletalComponentPrivate.h b/Engine/Plugins/Mutable/Source/CustomizableObject/Internal/MuCO/CustomizableSkeletalComponentPrivate.h
new file mode 100644
--- /dev/null
+++ b/Engine/Plugins/Mutable/Source/CustomizableObject/Internal/MuCO/CustomizableSkeletalComponentPrivate.h
@@ -0,0 +1,49 @@
+﻿// Copyright Epic Games, Inc. All Rights Reserved.
+
+#pragma once
+
+#include "CustomizableSkeletalComponentPrivate.generated.h"
+
+class UCustomizableSkeletalComponent;
+class UCustomizableObjectInstanceUsage;
+class UPhysicsAsset;
+class USkeletalMesh;
+
+
+UCLASS(Blueprintable, BlueprintType, ClassGroup = (CustomizableObject), meta = (BlueprintSpawnableComponent))
+class CUSTOMIZABLEOBJECT_API UCustomizableSkeletalComponentPrivate : public UObject
+{
+	GENERATED_BODY()
+
+public:
+	void CreateCustomizableObjectInstanceUsage();
+
+	/** Common end point of all updates. Even those which failed. */
+	void Callbacks() const;
+	
+	USkeletalMesh* GetSkeletalMesh() const;
+
+	USkeletalMesh* GetAttachedSkeletalMesh() const;
+	
+	void SetSkeletalMesh(USkeletalMesh* SkeletalMesh);
+	
+	void SetPhysicsAsset(UPhysicsAsset* PhysicsAsset);
+
+#if WITH_EDITOR
+	void EditorUpdateComponent();
+#endif
+
+	UCustomizableSkeletalComponent* GetPublic();
+	
+	const UCustomizableSkeletalComponent* GetPublic() const;
+	
+	// Used to replace the SkeletalMesh of the parent component by the ReferenceSkeletalMesh or the generated SkeletalMesh 
+	bool bPendingSetSkeletalMesh = false;
+
+	UPROPERTY()
+	bool bSkipSkipSetSkeletalMeshOnAttach = false;
+
+	UPROPERTY(Transient)
+	TObjectPtr<UCustomizableObjectInstanceUsage> CustomizableObjectInstanceUsage;
+};
+
diff --git a/Engine/Plugins/Mutable/Source/CustomizableObject/Private/MuCO/CustomizableInstanceLODManagement.cpp b/Engine/Plugins/Mutable/Source/CustomizableObject/Private/MuCO/CustomizableInstanceLODManagement.cpp
--- a/Engine/Plugins/Mutable/Source/CustomizableObject/Private/MuCO/CustomizableInstanceLODManagement.cpp
+++ b/Engine/Plugins/Mutable/Source/CustomizableObject/Private/MuCO/CustomizableInstanceLODManagement.cpp
@@ -2,6 +2,7 @@
 
 #include "MuCO/CustomizableInstanceLODManagement.h"
 
+#include "MuCO/CustomizableObjectInstanceUsagePrivate.h"
 #include "Camera/CameraActor.h"
 #include "Engine/SkeletalMesh.h"
 #include "Engine/SkinnedAssetCommon.h"
@@ -57,6 +58,15 @@ static TAutoConsoleVariable<bool> CVarOnlyUpdateCloseCustomizableObjects(
 	TEXT("If true, only CustomizableObjects within a predefined distance to the view centers will be generated"),
 	ECVF_Scalability);
 
+
+#if WITH_EDITOR
+void UCustomizableInstanceLODManagementBase::EditorUpdateComponent(UCustomizableObjectInstanceUsage* InstanceUsage)
+{
+	InstanceUsage->GetPrivate()->EditorUpdateComponent();
+}
+#endif
+
+
 UCustomizableInstanceLODManagement::UCustomizableInstanceLODManagement() : UCustomizableInstanceLODManagementBase()
 {
 	CloseCustomizableObjectsDist = 2000.f;
@@ -75,7 +85,7 @@ void UpdatePawnToInstancesDistances(const class UCustomizableObjectInstance* Onl
 	for (TObjectIterator<UCustomizableObjectInstanceUsage> CustomizableObjectInstanceUsage; CustomizableObjectInstanceUsage; ++CustomizableObjectInstanceUsage)
 	{
 #if WITH_EDITOR
-		if (IsValid(*CustomizableObjectInstanceUsage) && CustomizableObjectInstanceUsage->IsNetMode(NM_DedicatedServer))
+		if (IsValid(*CustomizableObjectInstanceUsage) && CustomizableObjectInstanceUsage->GetPrivate()->IsNetMode(NM_DedicatedServer))
 		{
 			continue;
 		}
@@ -83,7 +93,7 @@ void UpdatePawnToInstancesDistances(const class UCustomizableObjectInstance* Onl
 
 		if (IsValid(*CustomizableObjectInstanceUsage) && (OnlyForInstance == nullptr || CustomizableObjectInstanceUsage->GetCustomizableObjectInstance() == OnlyForInstance))
 		{
-			CustomizableObjectInstanceUsage->UpdateDistFromComponentToPlayer(ViewCenter.IsValid() ? ViewCenter.Get() : nullptr, OnlyForInstance != nullptr);
+			CustomizableObjectInstanceUsage->GetPrivate()->UpdateDistFromComponentToPlayer(ViewCenter.IsValid() ? ViewCenter.Get() : nullptr, OnlyForInstance != nullptr);
 		}
 	}
 }
@@ -95,14 +105,14 @@ void UpdateCameraToInstancesDistance(const FVector CameraPosition)
 	for (TObjectIterator<UCustomizableObjectInstanceUsage> CustomizableObjectInstanceUsage; CustomizableObjectInstanceUsage; ++CustomizableObjectInstanceUsage)
 	{
 #if WITH_EDITOR
-		if (IsValid(*CustomizableObjectInstanceUsage) && CustomizableObjectInstanceUsage->IsNetMode(NM_DedicatedServer))
+		if (IsValid(*CustomizableObjectInstanceUsage) && CustomizableObjectInstanceUsage->GetPrivate()->IsNetMode(NM_DedicatedServer))
 		{
 			continue;
 		}
 #endif
 		if (IsValid(*CustomizableObjectInstanceUsage) && !CustomizableObjectInstanceUsage->IsTemplate())
 		{
-			CustomizableObjectInstanceUsage->UpdateDistFromComponentToLevelEditorCamera(CameraPosition);
+			CustomizableObjectInstanceUsage->GetPrivate()->UpdateDistFromComponentToLevelEditorCamera(CameraPosition);
 		}
 	}
 }
@@ -127,7 +137,7 @@ void UCustomizableInstanceLODManagement::UpdateInstanceDistsAndLODs(FMutableInst
 			for (TObjectIterator<UCustomizableObjectInstanceUsage> CustomizableObjectInstanceUsage; CustomizableObjectInstanceUsage; ++CustomizableObjectInstanceUsage)
 			{
 #if WITH_EDITOR
-				if (IsValid(*CustomizableObjectInstanceUsage) && CustomizableObjectInstanceUsage->IsNetMode(NM_DedicatedServer))
+				if (IsValid(*CustomizableObjectInstanceUsage) && CustomizableObjectInstanceUsage->GetPrivate()->IsNetMode(NM_DedicatedServer))
 				{
 					continue;
 				}
@@ -173,7 +183,7 @@ void UCustomizableInstanceLODManagement::UpdateInstanceDistsAndLODs(FMutableInst
 						// Blueprint instances
 						else if (WorldType == EWorldType::EditorPreview)
 						{
-							CustomizableObjectInstanceUsage->EditorUpdateComponent();
+							CustomizableObjectInstanceUsage->GetPrivate()->EditorUpdateComponent();
 						}
 					}
 	#endif // WITH_EDITOR
@@ -303,7 +313,7 @@ void UCustomizableInstanceLODManagement::UpdateInstanceDistsAndLODs(FMutableInst
 		for (TObjectIterator<UCustomizableObjectInstanceUsage> CustomizableObjectInstanceUsage; CustomizableObjectInstanceUsage; ++CustomizableObjectInstanceUsage)
 		{
 #if WITH_EDITOR
-			if (IsValid(*CustomizableObjectInstanceUsage) && CustomizableObjectInstanceUsage->IsNetMode(NM_DedicatedServer))
+			if (IsValid(*CustomizableObjectInstanceUsage) && CustomizableObjectInstanceUsage->GetPrivate()->IsNetMode(NM_DedicatedServer))
 			{
 				continue;
 			}
@@ -336,7 +346,7 @@ void UCustomizableInstanceLODManagement::UpdateInstanceDistsAndLODs(FMutableInst
 
 				if (WorldType == EWorldType::EditorPreview || (!World && !bAttachParentActor))
 				{
-					CustomizableObjectInstanceUsage->EditorUpdateComponent();
+					CustomizableObjectInstanceUsage->GetPrivate()->EditorUpdateComponent();
 					continue;
 				}
 
diff --git a/Engine/Plugins/Mutable/Source/CustomizableObject/Private/MuCO/CustomizableObjectInstance.cpp b/Engine/Plugins/Mutable/Source/CustomizableObject/Private/MuCO/CustomizableObjectInstance.cpp
--- a/Engine/Plugins/Mutable/Source/CustomizableObject/Private/MuCO/CustomizableObjectInstance.cpp
+++ b/Engine/Plugins/Mutable/Source/CustomizableObject/Private/MuCO/CustomizableObjectInstance.cpp
@@ -11,6 +11,7 @@
 #include "BoneControllers/AnimNode_RigidBody.h"
 #include "ClothConfig.h"
 #include "ClothingAsset.h"
+#include "MuCO/CustomizableObjectInstanceUsagePrivate.h"
 #include "Engine/SkeletalMesh.h"
 #include "Engine/SkeletalMeshSocket.h"
 #include "Engine/SkeletalMeshLODSettings.h"
@@ -2977,7 +2978,7 @@ void UCustomizableInstancePrivate::SetDefaultSkeletalMesh(bool bSetEmptyMesh) co
 		}
 
 #if WITH_EDITOR
-		if (CustomizableObjectInstanceUsage->IsNetMode(NM_DedicatedServer))
+		if (CustomizableObjectInstanceUsage->GetPrivate()->IsNetMode(NM_DedicatedServer))
 		{
 			continue;
 		}
@@ -2998,7 +2999,7 @@ void UCustomizableInstancePrivate::SetDefaultSkeletalMesh(bool bSetEmptyMesh) co
 			SkeletalMesh = ModelResources.ReferenceSkeletalMeshesData[ObjectComponentIndex].SoftSkeletalMesh.LoadSynchronous();
 		}
 		
-		CustomizableObjectInstanceUsage->SetSkeletalMesh(SkeletalMesh);
+		CustomizableObjectInstanceUsage->GetPrivate()->SetSkeletalMesh(SkeletalMesh);
 	}
 }
 
@@ -6733,7 +6734,7 @@ void UCustomizableInstancePrivate::BuildMaterials(const TSharedRef<FUpdateContex
 			}
 
 #if WITH_EDITOR
-			if (CustomizableObjectInstanceUsage->IsNetMode(NM_DedicatedServer))
+			if (CustomizableObjectInstanceUsage->GetPrivate()->IsNetMode(NM_DedicatedServer))
 			{
 				continue;
 			}
diff --git a/Engine/Plugins/Mutable/Source/CustomizableObject/Private/MuCO/CustomizableObjectInstanceUsage.cpp b/Engine/Plugins/Mutable/Source/CustomizableObject/Private/MuCO/CustomizableObjectInstanceUsage.cpp
--- a/Engine/Plugins/Mutable/Source/CustomizableObject/Private/MuCO/CustomizableObjectInstanceUsage.cpp
+++ b/Engine/Plugins/Mutable/Source/CustomizableObject/Private/MuCO/CustomizableObjectInstanceUsage.cpp
@@ -1,11 +1,13 @@
 // Copyright Epic Games, Inc. All Rights Reserved.
 
 #include "MuCO/CustomizableObjectInstanceUsage.h"
+#include "MuCO/CustomizableObjectInstanceUsagePrivate.h"
+
 #include "MuCO/CustomizableSkeletalComponent.h"
 #include "MuCO/CustomizableObjectSystem.h"
 #include "MuCO/CustomizableObjectInstancePrivate.h"
 #include "MuCO/UnrealPortabilityHelpers.h"
-
+#include "MuCO/CustomizableSkeletalComponentPrivate.h"
 #include "AnimationRuntime.h"
 #include "Components/SkeletalMeshComponent.h"
 #include "Engine/SkeletalMesh.h"
@@ -20,7 +22,7 @@
 #include "Stats/Stats.h"
 
 
-ETickableTickType UCustomizableObjectInstanceUsage::GetTickableTickType() const
+ETickableTickType UCustomizableObjectInstanceUsagePrivate::GetTickableTickType() const
 { 
 	return (HasAnyFlags(RF_ClassDefaultObject))
 		? ETickableTickType::Never
@@ -28,25 +30,31 @@ ETickableTickType UCustomizableObjectInstanceUsage::GetTickableTickType() const
 }
 
 
-void UCustomizableObjectInstanceUsage::Callbacks() const
+void UCustomizableObjectInstanceUsagePrivate::Callbacks() const
 {
 	for (const UCustomizableObjectExtension* Extension : ICustomizableObjectModule::Get().GetRegisteredExtensions())
 	{
-		Extension->OnCustomizableObjectInstanceUsageUpdated(this);
+		Extension->OnCustomizableObjectInstanceUsageUpdated(GetPublic());
 	}
 	
-	if (CustomizableSkeletalComponent)
+	if (GetPublic()->CustomizableSkeletalComponent)
 	{
-		CustomizableSkeletalComponent->UpdatedDelegate.ExecuteIfBound();
+		GetPublic()->CustomizableSkeletalComponent->UpdatedDelegate.ExecuteIfBound();
 
-		if (UpdatedDelegate.IsBound() && CustomizableSkeletalComponent->UpdatedDelegate.IsBound())
+		if (GetPublic()->UpdatedDelegate.IsBound() && GetPublic()->CustomizableSkeletalComponent->UpdatedDelegate.IsBound())
 		{
 			UE_LOG(LogMutable, Error, TEXT("The UpdatedDelegate is bound both in the UCustomizableObjectInstanceUsage and in its parent CustomizableSkeletalComponent. Only one should be bound."));
 			ensure(false);
 		}
 	}
 	
-	UpdatedDelegate.ExecuteIfBound();
+	GetPublic()->UpdatedDelegate.ExecuteIfBound();
+}
+
+
+UCustomizableObjectInstanceUsage::UCustomizableObjectInstanceUsage()
+{
+	Private = CreateDefaultSubobject<UCustomizableObjectInstanceUsagePrivate>(FName("Private"));
 }
 
 
@@ -136,11 +144,11 @@ FName UCustomizableObjectInstanceUsage::GetComponentName() const
 }
 
 
-void UCustomizableObjectInstanceUsage::SetPendingSetSkeletalMesh(bool bIsActive)
+void UCustomizableObjectInstanceUsagePrivate::SetPendingSetSkeletalMesh(bool bIsActive)
 {
-	if (CustomizableSkeletalComponent)
+	if (GetPublic()->CustomizableSkeletalComponent)
 	{
-		CustomizableSkeletalComponent->bPendingSetSkeletalMesh = bIsActive;
+		GetPublic()->CustomizableSkeletalComponent->GetPrivate()->bPendingSetSkeletalMesh = bIsActive;
 	}
 	else
 	{
@@ -149,11 +157,11 @@ void UCustomizableObjectInstanceUsage::SetPendingSetSkeletalMesh(bool bIsActive)
 }
 
 
-bool UCustomizableObjectInstanceUsage::GetPendingSetSkeletalMesh() const
+bool UCustomizableObjectInstanceUsagePrivate::GetPendingSetSkeletalMesh() const
 {
-	if (CustomizableSkeletalComponent)
+	if (GetPublic()->CustomizableSkeletalComponent)
 	{
-		return CustomizableSkeletalComponent->bPendingSetSkeletalMesh;
+		return GetPublic()->CustomizableSkeletalComponent->GetPrivate()->bPendingSetSkeletalMesh;
 	}
 	else
 	{
@@ -162,15 +170,45 @@ bool UCustomizableObjectInstanceUsage::GetPendingSetSkeletalMesh() const
 }
 
 
-void UCustomizableObjectInstanceUsage::SetSkipSetReferenceSkeletalMesh(bool bIsActive)
+UCustomizableSkeletalComponent* UCustomizableObjectInstanceUsagePrivate::GetCustomizableSkeletalComponent()
+{
+	return GetPublic()->CustomizableSkeletalComponent;
+}
+
+
+void UCustomizableObjectInstanceUsagePrivate::SetCustomizableSkeletalComponent(UCustomizableSkeletalComponent* Component)
+{
+	GetPublic()->CustomizableSkeletalComponent = Component;
+}
+
+
+UCustomizableObjectInstanceUsage* UCustomizableObjectInstanceUsagePrivate::GetPublic()
+{
+	UCustomizableObjectInstanceUsage* Public = StaticCast<UCustomizableObjectInstanceUsage*>(GetOuter());
+	check(Public);
+
+	return Public;
+}
+
+
+const UCustomizableObjectInstanceUsage* UCustomizableObjectInstanceUsagePrivate::GetPublic() const
+{
+	UCustomizableObjectInstanceUsage* Public = StaticCast<UCustomizableObjectInstanceUsage*>(GetOuter());
+	check(Public);
+
+	return Public;
+}
+
+
+void UCustomizableObjectInstanceUsage::SetSkipSetReferenceSkeletalMesh(bool bSkip)
 {
 	if (CustomizableSkeletalComponent)
 	{
-		CustomizableSkeletalComponent->bSkipSetReferenceSkeletalMesh = bIsActive;
+		CustomizableSkeletalComponent->SetSkipSetReferenceSkeletalMesh(bSkip);
 	}
 	else
 	{
-		bUsedSkipSetReferenceSkeletalMesh = bIsActive;
+		GetPrivate()->bUsedSkipSetReferenceSkeletalMesh = bSkip;
 	}
 }
 
@@ -179,11 +217,37 @@ bool UCustomizableObjectInstanceUsage::GetSkipSetReferenceSkeletalMesh() const
 {
 	if (CustomizableSkeletalComponent)
 	{
-		return CustomizableSkeletalComponent->bSkipSetReferenceSkeletalMesh;
+		return CustomizableSkeletalComponent->GetSkipSetReferenceSkeletalMesh();
 	}
 	else
 	{
-		return bUsedSkipSetReferenceSkeletalMesh;
+		return GetPrivate()->bUsedSkipSetReferenceSkeletalMesh;
+	}
+}
+
+
+void UCustomizableObjectInstanceUsage::SetSkipSetSkeletalMeshOnAttach(bool bSkip)
+{
+	if (CustomizableSkeletalComponent)
+	{
+		CustomizableSkeletalComponent->SetSkipSetSkeletalMeshOnAttach(bSkip);
+	}
+	else
+	{
+		GetPrivate()->bUsedSkipSetSkeletalMeshOnAttach = bSkip;
+	}
+}
+
+
+bool UCustomizableObjectInstanceUsage::GetSkipSetSkeletalMeshOnAttach() const
+{
+	if (CustomizableSkeletalComponent)
+	{
+		return CustomizableSkeletalComponent->GetSkipSetReferenceSkeletalMesh();
+	}
+	else
+	{
+		return GetPrivate()->bUsedSkipSetSkeletalMeshOnAttach;
 	}
 }
 
@@ -207,7 +271,7 @@ void UCustomizableObjectInstanceUsage::AttachTo(USkeletalMeshComponent* Skeletal
 		}
 
 		// To mimic the behavior of UCustomizableSkeletalComponent::OnAttachmentChanged()
-		SetPendingSetSkeletalMesh(true);
+		GetPrivate()->SetPendingSetSkeletalMesh(true);
 	}
 }
 
@@ -227,11 +291,11 @@ USkeletalMeshComponent* UCustomizableObjectInstanceUsage::GetAttachParent() cons
 }
 
 
-USkeletalMesh* UCustomizableObjectInstanceUsage::GetSkeletalMesh() const
+USkeletalMesh* UCustomizableObjectInstanceUsagePrivate::GetSkeletalMesh() const
 {
-	UCustomizableObjectInstance* CustomizableObjectInstance = GetCustomizableObjectInstance();
+	UCustomizableObjectInstance* CustomizableObjectInstance = GetPublic()->GetCustomizableObjectInstance();
 
-	return CustomizableObjectInstance ? CustomizableObjectInstance->GetComponentMeshSkeletalMesh(GetComponentName()) : nullptr;
+	return CustomizableObjectInstance ? CustomizableObjectInstance->GetComponentMeshSkeletalMesh(GetPublic()->GetComponentName()) : nullptr;
 }
 
 
@@ -272,7 +336,7 @@ bool RequiresReinitPose(USkeletalMesh* CurrentSkeletalMesh, USkeletalMesh* Skele
 }
 
 
-void UCustomizableObjectInstanceUsage::SetSkeletalMeshAndOverrideMaterials(USkeletalMeshComponent& Parent, USkeletalMesh* SkeletalMesh, 
+void UCustomizableObjectInstanceUsagePrivate::SetSkeletalMeshAndOverrideMaterials(USkeletalMeshComponent& Parent, USkeletalMesh* SkeletalMesh, 
 	const UCustomizableObjectInstance& CustomizableObjectInstance, bool* bOutSkeletalMeshUpdated, bool* bOutMaterialsUpdated)
 {
 	if (SkeletalMesh != Parent.GetSkeletalMeshAsset())
@@ -313,7 +377,7 @@ void UCustomizableObjectInstanceUsage::SetSkeletalMeshAndOverrideMaterials(USkel
 
 	if (bUseOverrideMaterials)
 	{
-		if (FCustomizableInstanceComponentData* ComponentData = CustomizableObjectInstance.GetPrivate()->GetComponentData(GetComponentName()))
+		if (FCustomizableInstanceComponentData* ComponentData = CustomizableObjectInstance.GetPrivate()->GetComponentData(GetPublic()->GetComponentName()))
 		{
 			for (int32 Index = 0; Index < ComponentData->OverrideMaterials.Num(); ++Index)
 			{
@@ -329,10 +393,10 @@ void UCustomizableObjectInstanceUsage::SetSkeletalMeshAndOverrideMaterials(USkel
 }
 
 
-void UCustomizableObjectInstanceUsage::SetSkeletalMesh(USkeletalMesh* SkeletalMesh, bool* bOutSkeletalMeshUpdated, bool* bOutMaterialsUpdated)
+void UCustomizableObjectInstanceUsagePrivate::SetSkeletalMesh(USkeletalMesh* SkeletalMesh, bool* bOutSkeletalMeshUpdated, bool* bOutMaterialsUpdated)
 {
-	USkeletalMeshComponent* Parent = Cast<USkeletalMeshComponent>(GetAttachParent());
-	const UCustomizableObjectInstance* CustomizableObjectInstance = GetCustomizableObjectInstance();
+	USkeletalMeshComponent* Parent = Cast<USkeletalMeshComponent>(GetPublic()->GetAttachParent());
+	const UCustomizableObjectInstance* CustomizableObjectInstance = GetPublic()->GetCustomizableObjectInstance();
 
 	if (Parent && CustomizableObjectInstance)
 	{
@@ -340,9 +404,9 @@ void UCustomizableObjectInstanceUsage::SetSkeletalMesh(USkeletalMesh* SkeletalMe
 	}
 }
 
-void UCustomizableObjectInstanceUsage::SetPhysicsAsset(UPhysicsAsset* PhysicsAsset, bool* bOutPhysicsAssetUpdated)
+void UCustomizableObjectInstanceUsagePrivate::SetPhysicsAsset(UPhysicsAsset* PhysicsAsset, bool* bOutPhysicsAssetUpdated)
 {
-	USkeletalMeshComponent* Parent = Cast<USkeletalMeshComponent>(GetAttachParent());
+	USkeletalMeshComponent* Parent = Cast<USkeletalMeshComponent>(GetPublic()->GetAttachParent());
 
 	if (Parent && Parent->GetWorld() &&
 		PhysicsAsset != Parent->GetPhysicsAsset())
@@ -357,9 +421,9 @@ void UCustomizableObjectInstanceUsage::SetPhysicsAsset(UPhysicsAsset* PhysicsAss
 }
 
 
-USkeletalMesh* UCustomizableObjectInstanceUsage::GetAttachedSkeletalMesh() const
+USkeletalMesh* UCustomizableObjectInstanceUsagePrivate::GetAttachedSkeletalMesh() const
 {
-	USkeletalMeshComponent* Parent = Cast<USkeletalMeshComponent>(GetAttachParent());
+	USkeletalMeshComponent* Parent = Cast<USkeletalMeshComponent>(GetPublic()->GetAttachParent());
 
 	if (Parent)
 	{
@@ -392,15 +456,8 @@ void UCustomizableObjectInstanceUsage::UpdateSkeletalMeshAsyncResult(FInstanceUp
 }
 
 
-void UCustomizableObjectInstanceUsage::BeginDestroy()
-{
-	Super::BeginDestroy();
-}
-
-
 #if WITH_EDITOR
-
-void UCustomizableObjectInstanceUsage::UpdateDistFromComponentToLevelEditorCamera(const FVector& CameraPosition)
+void UCustomizableObjectInstanceUsagePrivate::UpdateDistFromComponentToLevelEditorCamera(const FVector& CameraPosition)
 {
 	// We want instances in the editor to be generated
 	if (!GetWorld() || GetWorld()->WorldType != EWorldType::Editor)
@@ -408,11 +465,11 @@ void UCustomizableObjectInstanceUsage::UpdateDistFromComponentToLevelEditorCamer
 		return;
 	}
 
-	UCustomizableObjectInstance* CustomizableObjectInstance = GetCustomizableObjectInstance();
+	UCustomizableObjectInstance* CustomizableObjectInstance = GetPublic()->GetCustomizableObjectInstance();
 
 	if (CustomizableObjectInstance)
 	{
-		USkeletalMeshComponent* SkeletalMeshComponent = GetAttachParent();
+		USkeletalMeshComponent* SkeletalMeshComponent = GetPublic()->GetAttachParent();
 		AActor* ParentActor = SkeletalMeshComponent ? SkeletalMeshComponent->GetAttachmentRootActor() : nullptr;
 		if (ParentActor && ParentActor->IsValidLowLevel())
 		{
@@ -425,7 +482,7 @@ void UCustomizableObjectInstanceUsage::UpdateDistFromComponentToLevelEditorCamer
 		}
 
 		USkeletalMesh* AttachedSkeletalMesh = GetAttachedSkeletalMesh();
-		const FName& ComponentName = GetComponentName();
+		const FName& ComponentName = GetPublic()->GetComponentName();
 
 		const bool bInstanceGenerated = CustomizableObjectInstance->GetPrivate()->SkeletalMeshStatus != ESkeletalMeshStatus::NotGenerated;
 		USkeletalMesh* GeneratedSkeletalMesh = bInstanceGenerated ? CustomizableObjectInstance->GetComponentMeshSkeletalMesh(ComponentName) :
@@ -439,26 +496,26 @@ void UCustomizableObjectInstanceUsage::UpdateDistFromComponentToLevelEditorCamer
 }
 
 
-void UCustomizableObjectInstanceUsage::EditorUpdateComponent()
+void UCustomizableObjectInstanceUsagePrivate::EditorUpdateComponent()
 {
-	UCustomizableObjectInstance* CustomizableObjectInstance = GetCustomizableObjectInstance();
+	UCustomizableObjectInstance* CustomizableObjectInstance = GetPublic()->GetCustomizableObjectInstance();
 
 	if (CustomizableObjectInstance)
 	{
 		CustomizableObjectInstance->GetPrivate()->SetCOInstanceFlags(UsedByComponent);
 
-		USkeletalMeshComponent* SkeletalMeshComponent = GetAttachParent();
+		USkeletalMeshComponent* SkeletalMeshComponent = GetPublic()->GetAttachParent();
 		AActor* ParentActor = SkeletalMeshComponent ? SkeletalMeshComponent->GetAttachmentRootActor() : nullptr;
 		
 		if (ParentActor)
 		{
 			USkeletalMesh* AttachedSkeletalMesh = GetAttachedSkeletalMesh();
-			const FName& ComponentName = GetComponentName();
+			const FName& ComponentName = GetPublic()->GetComponentName();
 
 			const bool bInstanceGenerated = CustomizableObjectInstance->GetPrivate()->SkeletalMeshStatus != ESkeletalMeshStatus::NotGenerated;
 			USkeletalMesh* GeneratedSkeletalMesh = bInstanceGenerated ?
 				CustomizableObjectInstance->GetComponentMeshSkeletalMesh(ComponentName) :
-				GetSkipSetReferenceSkeletalMesh() ?
+				GetPublic()->GetSkipSetReferenceSkeletalMesh() ?
 					nullptr :
 					CustomizableObjectInstance->GetCustomizableObject()->GetComponentMeshReferenceSkeletalMesh(ComponentName);
 
@@ -472,53 +529,37 @@ void UCustomizableObjectInstanceUsage::EditorUpdateComponent()
 #endif
 
 
-void UCustomizableObjectInstanceUsage::SetVisibilityOfSkeletalMeshSectionWithMaterialName(bool bInVisible, const FString& MaterialName, int32 LOD)
+UCustomizableObjectInstanceUsagePrivate* UCustomizableObjectInstanceUsage::GetPrivate()
 {
-	UCustomizableObjectInstance* CustomizableObjectInstance = GetCustomizableObjectInstance();
-	const FName& ComponentName = GetComponentName();
+	return Private;
+}
 
-	USkeletalMesh* SkeletalMesh = CustomizableObjectInstance->GetComponentMeshSkeletalMesh(ComponentName);
 
-	if (!SkeletalMesh)
-	{
-		return;
-	}
+const UCustomizableObjectInstanceUsagePrivate* UCustomizableObjectInstanceUsage::GetPrivate() const
+{
+	return Private;
+}
 
-	for (int32 i = 0; i < SkeletalMesh->GetMaterials().Num(); ++i)
-	{
-		UMaterialInstanceDynamic* MatInst = Cast<UMaterialInstanceDynamic>(SkeletalMesh->GetMaterials()[i].MaterialInterface);
 
-		if (MatInst && MatInst->Parent->GetName() == MaterialName)
-		{
-			if (FSkeletalMeshRenderData* SkelResource = SkeletalMesh->GetResourceForRendering() )
-			{
-				if ( SkelResource->LODRenderData.IsValidIndex(LOD))
-				{
-					FSkeletalMeshLODRenderData& LodData = SkelResource->LODRenderData[LOD];
+bool UCustomizableObjectInstanceUsage::IsNetMode(ENetMode InNetMode) const
+{
+	return GetPrivate()->IsNetMode(InNetMode);
+}
 
-					for (int32 j = 0; j < LodData.RenderSections.Num(); ++j)
-					{
-						FSkelMeshRenderSection& Section = LodData.RenderSections[j];
 
-						if (Section.MaterialIndex == i)
-						{
-							Section.bDisabled = !bInVisible;
-						}
-					}
-				}
-			}
-		}
-	}
+void UCustomizableObjectInstanceUsage::EditorUpdateComponent()
+{
+	GetPrivate()->EditorUpdateComponent();
 }
 
 
-void UCustomizableObjectInstanceUsage::UpdateDistFromComponentToPlayer(const AActor* ViewCenter, const bool bForceEvenIfNotBegunPlay)
+void UCustomizableObjectInstanceUsagePrivate::UpdateDistFromComponentToPlayer(const AActor* ViewCenter, const bool bForceEvenIfNotBegunPlay)
 {
-	UCustomizableObjectInstance* CustomizableObjectInstance = GetCustomizableObjectInstance();
+	UCustomizableObjectInstance* CustomizableObjectInstance = GetPublic()->GetCustomizableObjectInstance();
 
 	if (CustomizableObjectInstance)
 	{
-		USkeletalMeshComponent* SkeletalMeshComponent = GetAttachParent();
+		USkeletalMeshComponent* SkeletalMeshComponent = GetPublic()->GetAttachParent();
 		AActor* ParentActor = SkeletalMeshComponent ? SkeletalMeshComponent->GetAttachmentRootActor() : nullptr;
 		
 		CustomizableObjectInstance->SetIsPlayerOrNearIt(false);
@@ -560,13 +601,13 @@ void UCustomizableObjectInstanceUsage::UpdateDistFromComponentToPlayer(const AAc
 
 				if (CustomizableObjectInstance->GetPrivate()->MinSquareDistFromComponentToPlayer == SquareDist)
 				{
-					CustomizableObjectInstance->GetPrivate()->NearestToActor = this;
+					CustomizableObjectInstance->GetPrivate()->NearestToActor = GetPublic();
 					CustomizableObjectInstance->GetPrivate()->NearestToViewCenter = ViewCenter;
 				}
 			}
 		}
 
-		const FName& ComponentName = GetComponentName();
+		const FName& ComponentName = GetPublic()->GetComponentName();
 
 		if (ParentActor && GetAttachedSkeletalMesh() == nullptr && CustomizableObjectInstance->GetComponentMeshSkeletalMesh(ComponentName))
 		{
@@ -576,16 +617,16 @@ void UCustomizableObjectInstanceUsage::UpdateDistFromComponentToPlayer(const AAc
 }
 
 
-void UCustomizableObjectInstanceUsage::Tick(float DeltaTime)
+void UCustomizableObjectInstanceUsagePrivate::Tick(float DeltaTime)
 {
 	if (!IsValid(this))
 	{
 		return;
 	}
 
-	UCustomizableObjectInstance* CustomizableObjectInstance = GetCustomizableObjectInstance();
+	UCustomizableObjectInstance* CustomizableObjectInstance = GetPublic()->GetCustomizableObjectInstance();
 
-	if (!GetPendingSetSkeletalMesh() || !CustomizableObjectInstance || !IsValid(CustomizableObjectInstance))
+	if (!GetPendingSetSkeletalMesh() || GetPublic()->GetSkipSetSkeletalMeshOnAttach() || !CustomizableObjectInstance || !IsValid(CustomizableObjectInstance))
 	{
 		return;
 	}
@@ -596,9 +637,9 @@ void UCustomizableObjectInstanceUsage::Tick(float DeltaTime)
 		return;
 	}	
 
-	if (USkeletalMeshComponent* Parent = Cast<USkeletalMeshComponent>(GetAttachParent()))
+	if (USkeletalMeshComponent* Parent = Cast<USkeletalMeshComponent>(GetPublic()->GetAttachParent()))
 	{
-		const FName& ComponentName = GetComponentName();
+		const FName& ComponentName = GetPublic()->GetComponentName();
 
 		USkeletalMesh* SkeletalMesh = nullptr;
 
@@ -611,7 +652,7 @@ void UCustomizableObjectInstanceUsage::Tick(float DeltaTime)
 		else
 		{
 			// If not generated yet, conditionally set the SkeletalMesh of reference
-			if (!bInstanceGenerated && !GetSkipSetReferenceSkeletalMesh() && 
+			if (!bInstanceGenerated && !GetPublic()->GetSkipSetReferenceSkeletalMesh() && 
 				CustomizableObject->bEnableUseRefSkeletalMeshAsPlaceholder)
 			{
 				// Can be nullptr
@@ -628,27 +669,39 @@ void UCustomizableObjectInstanceUsage::Tick(float DeltaTime)
 }
 
 
-TStatId UCustomizableObjectInstanceUsage::GetStatId() const
+TStatId UCustomizableObjectInstanceUsagePrivate::GetStatId() const
 {
 	RETURN_QUICK_DECLARE_CYCLE_STAT(UCustomizableObjectInstanceUsage, STATGROUP_Tickables);
 }
 
 
-bool UCustomizableObjectInstanceUsage::IsTickable() const
+bool UCustomizableObjectInstanceUsagePrivate::IsTickableWhenPaused() const
+{
+	return true;
+}
+
+
+bool UCustomizableObjectInstanceUsagePrivate::IsTickableInEditor() const
+{
+	return true;
+}
+
+
+bool UCustomizableObjectInstanceUsagePrivate::IsTickable() const
 {
 	return !HasAnyFlags(RF_BeginDestroyed);
 }
 
 
-bool UCustomizableObjectInstanceUsage::IsNetMode(ENetMode InNetMode) const
+bool UCustomizableObjectInstanceUsagePrivate::IsNetMode(ENetMode InNetMode) const
 {
-	if (CustomizableSkeletalComponent)
+	if (GetPublic()->CustomizableSkeletalComponent)
 	{
-		return CustomizableSkeletalComponent->IsNetMode(InNetMode);
+		return GetPublic()->CustomizableSkeletalComponent->IsNetMode(InNetMode);
 	}
-	else if(UsedSkeletalMeshComponent.IsValid())
+	else if (GetPublic()->UsedSkeletalMeshComponent.IsValid())
 	{
-		return UsedSkeletalMeshComponent->IsNetMode(InNetMode);
+		return GetPublic()->UsedSkeletalMeshComponent->IsNetMode(InNetMode);
 	}
 
 	return false;
diff --git a/Engine/Plugins/Mutable/Source/CustomizableObject/Private/MuCO/CustomizableObjectModule.cpp b/Engine/Plugins/Mutable/Source/CustomizableObject/Private/MuCO/CustomizableObjectModule.cpp
--- a/Engine/Plugins/Mutable/Source/CustomizableObject/Private/MuCO/CustomizableObjectModule.cpp
+++ b/Engine/Plugins/Mutable/Source/CustomizableObject/Private/MuCO/CustomizableObjectModule.cpp
@@ -1,5 +1,6 @@
 // Copyright Epic Games, Inc. All Rights Reserved.
 
+#include "MuCO/CustomizableObjectInstanceUsagePrivate.h"
 #include "Algo/Find.h"
 #include "Engine/Engine.h"
 #include "Interfaces/IPluginManager.h"
@@ -238,7 +239,7 @@ UCustomizableObjectInstanceUsage* GetPlayerCustomizableObjectInstanceUsage(const
 		for (TObjectIterator<UCustomizableObjectInstanceUsage> CustomizableObjectInstanceUsage; CustomizableObjectInstanceUsage; ++CustomizableObjectInstanceUsage)
 		{
 #if WITH_EDITOR
-			if (IsValid(*CustomizableObjectInstanceUsage) && CustomizableObjectInstanceUsage->IsNetMode(NM_DedicatedServer))
+			if (IsValid(*CustomizableObjectInstanceUsage) && CustomizableObjectInstanceUsage->GetPrivate()->IsNetMode(NM_DedicatedServer))
 			{
 				continue;
 			}
@@ -270,7 +271,7 @@ UCustomizableObjectInstanceUsage* GetPlayerCustomizableObjectInstanceUsage(const
 		for (TObjectIterator<UCustomizableObjectInstanceUsage> CustomizableObjectInstanceUsage; CustomizableObjectInstanceUsage; ++CustomizableObjectInstanceUsage)
 		{
 #if WITH_EDITOR
-			if (IsValid(*CustomizableObjectInstanceUsage) && CustomizableObjectInstanceUsage->IsNetMode(NM_DedicatedServer))
+			if (IsValid(*CustomizableObjectInstanceUsage) && CustomizableObjectInstanceUsage->GetPrivate()->IsNetMode(NM_DedicatedServer))
 			{
 				continue;
 			}
diff --git a/Engine/Plugins/Mutable/Source/CustomizableObject/Private/MuCO/CustomizableObjectSystem.cpp b/Engine/Plugins/Mutable/Source/CustomizableObject/Private/MuCO/CustomizableObjectSystem.cpp
--- a/Engine/Plugins/Mutable/Source/CustomizableObject/Private/MuCO/CustomizableObjectSystem.cpp
+++ b/Engine/Plugins/Mutable/Source/CustomizableObject/Private/MuCO/CustomizableObjectSystem.cpp
@@ -12,6 +12,7 @@
 #include "MuCO/CustomizableInstanceLODManagement.h"
 #include "MuCO/CustomizableObjectInstancePrivate.h"
 #include "MuCO/CustomizableObjectPrivate.h"
+#include "MuCO/CustomizableObjectInstanceUsagePrivate.h"
 #include "MuCO/CustomizableObjectUIData.h"
 #include "MuCO/DefaultImageProvider.h"
 #include "MuCO/CustomizableObjectInstanceUsage.h"
@@ -438,7 +439,7 @@ void UCustomizableObjectSystem::LogShowData(bool bFullInfo, bool ShowMaterialInf
 		const UCustomizableObjectInstanceUsage* CustomizableObjectInstanceUsage = *It;
 
 #if WITH_EDITOR
-		if (IsValid(CustomizableObjectInstanceUsage) && CustomizableObjectInstanceUsage->IsNetMode(NM_DedicatedServer))
+		if (IsValid(CustomizableObjectInstanceUsage) && CustomizableObjectInstanceUsage->GetPrivate()->IsNetMode(NM_DedicatedServer))
 		{
 			continue;
 		}
@@ -862,7 +863,7 @@ void FinishUpdateGlobal(const TSharedRef<FUpdateContextPrivate>& Context)
 			}
 
 #if WITH_EDITOR
-			if (It->IsNetMode(NM_DedicatedServer))
+			if (It->GetPrivate()->IsNetMode(NM_DedicatedServer))
 			{
 				continue;
 			}
@@ -871,7 +872,7 @@ void FinishUpdateGlobal(const TSharedRef<FUpdateContextPrivate>& Context)
 			if (InstanceUsage->GetCustomizableObjectInstance() == Instance &&
 				(!Context->bOptimizedUpdate || Context->AttachedParentUpdated.Find(InstanceUsage)))
 			{
-				InstanceUsage->Callbacks();
+				InstanceUsage->GetPrivate()->Callbacks();
 			}
 		}
 	}
@@ -943,7 +944,7 @@ void UpdateSkeletalMesh(const TSharedRef<FUpdateContextPrivate>& Context)
 		UCustomizableObjectInstanceUsage* CustomizableObjectInstanceUsage = *It;
 
 #if WITH_EDITOR
-		if (IsValid(CustomizableObjectInstanceUsage) && CustomizableObjectInstanceUsage->IsNetMode(NM_DedicatedServer))
+		if (IsValid(CustomizableObjectInstanceUsage) && CustomizableObjectInstanceUsage->GetPrivate()->IsNetMode(NM_DedicatedServer))
 		{
 			continue;
 		}
@@ -959,12 +960,12 @@ void UpdateSkeletalMesh(const TSharedRef<FUpdateContextPrivate>& Context)
 			MUTABLE_CPUPROFILER_SCOPE(UpdateSkeletalMesh_SetSkeletalMesh);
 
 			USkeletalMesh* SkeletalMesh = CustomizableObjectInstance->GetComponentMeshSkeletalMesh(CustomizableObjectInstanceUsage->GetComponentName());
-			CustomizableObjectInstanceUsage->SetSkeletalMesh(SkeletalMesh, &bSkeletalMeshUpdated, &bMaterialsUpdated);
+			CustomizableObjectInstanceUsage->GetPrivate()->SetSkeletalMesh(SkeletalMesh, &bSkeletalMeshUpdated, &bMaterialsUpdated);
 
 			if (CustomizableObjectInstancePrivateData->HasCOInstanceFlags(ReplacePhysicsAssets) &&
 				SkeletalMesh)
 			{
-				CustomizableObjectInstanceUsage->SetPhysicsAsset(SkeletalMesh->GetPhysicsAsset(), &bPhysicsAssetUpdated);	
+				CustomizableObjectInstanceUsage->GetPrivate()->SetPhysicsAsset(SkeletalMesh->GetPhysicsAsset(), &bPhysicsAssetUpdated);	
 			}
 		}
 
@@ -3422,7 +3423,7 @@ namespace impl
 
 		for (TObjectIterator<UCustomizableObjectInstanceUsage> CustomizableObjectInstanceUsage; CustomizableObjectInstanceUsage && !bRequestAllLODs; ++CustomizableObjectInstanceUsage)
 		{
-			if (IsValid(*CustomizableObjectInstanceUsage) && CustomizableObjectInstanceUsage->IsNetMode(NM_DedicatedServer))
+			if (IsValid(*CustomizableObjectInstanceUsage) && CustomizableObjectInstanceUsage->GetPrivate()->IsNetMode(NM_DedicatedServer))
 			{
 				continue;
 			}
diff --git a/Engine/Plugins/Mutable/Source/CustomizableObject/Private/MuCO/CustomizableSkeletalComponent.cpp b/Engine/Plugins/Mutable/Source/CustomizableObject/Private/MuCO/CustomizableSkeletalComponent.cpp
--- a/Engine/Plugins/Mutable/Source/CustomizableObject/Private/MuCO/CustomizableSkeletalComponent.cpp
+++ b/Engine/Plugins/Mutable/Source/CustomizableObject/Private/MuCO/CustomizableSkeletalComponent.cpp
@@ -1,25 +1,28 @@
 // Copyright Epic Games, Inc. All Rights Reserved.
 
 #include "MuCO/CustomizableSkeletalComponent.h"
-#include "MuCO/CustomizableObjectInstanceUsage.h"
 
+#include "MuCO/CustomizableObjectInstanceUsagePrivate.h"
+#include "MuCO/CustomizableSkeletalComponentPrivate.h"
+
+#include "MuCO/CustomizableObjectInstanceUsage.h"
 #include "UObject/UObjectGlobals.h"
 #include "Components/SkeletalMeshComponent.h"
 #include "MuCO/CustomizableObject.h"
 
 
-void UCustomizableSkeletalComponent::CreateCustomizableObjectInstanceUsage()
+void UCustomizableSkeletalComponentPrivate::CreateCustomizableObjectInstanceUsage()
 {
 	if (CustomizableObjectInstanceUsage)
 	{
 		// CustomizableObjectInstanceUsage may already exist if duplicated from an existing Customizable Skeletal Component
-		if (CustomizableObjectInstanceUsage->CustomizableSkeletalComponent != this)
+		if (CustomizableObjectInstanceUsage->GetPrivate()->GetCustomizableSkeletalComponent() != GetPublic())
 		{
 			CustomizableObjectInstanceUsage = nullptr;
 		}
 	}
 
-	AActor* RootActor = GetAttachmentRootActor();
+	AActor* RootActor = GetPublic()->GetAttachmentRootActor();
 	bool bIsDefaultActor = RootActor ?
 		RootActor->HasAnyFlags(RF_ClassDefaultObject) :
 		false;
@@ -27,11 +30,17 @@ void UCustomizableSkeletalComponent::CreateCustomizableObjectInstanceUsage()
 	if (!CustomizableObjectInstanceUsage && !HasAnyFlags(RF_ClassDefaultObject) && !bIsDefaultActor)
 	{
 		CustomizableObjectInstanceUsage = NewObject<UCustomizableObjectInstanceUsage>(this, TEXT("InstanceUsage"), RF_Transient);
-		CustomizableObjectInstanceUsage->CustomizableSkeletalComponent = this;
+		CustomizableObjectInstanceUsage->GetPrivate()->SetCustomizableSkeletalComponent(GetPublic());
 	}
 }
 
 
+UCustomizableSkeletalComponent::UCustomizableSkeletalComponent()
+{
+	Private = CreateDefaultSubobject<UCustomizableSkeletalComponentPrivate>(FName("Private"));
+}
+
+
 #if WITH_EDITOR
 void UCustomizableSkeletalComponent::PostEditChangeProperty(FPropertyChangedEvent& PropertyChangedEvent)
 {
@@ -52,49 +61,49 @@ void UCustomizableSkeletalComponent::PostEditChangeProperty(FPropertyChangedEven
 #endif
 
 
-void UCustomizableSkeletalComponent::Callbacks() const
+void UCustomizableSkeletalComponentPrivate::Callbacks() const
 {
 	if (CustomizableObjectInstanceUsage)
 	{
-		CustomizableObjectInstanceUsage->Callbacks();
+		CustomizableObjectInstanceUsage->GetPrivate()->Callbacks();
 	}
 }
 
 
-USkeletalMesh* UCustomizableSkeletalComponent::GetSkeletalMesh() const
+USkeletalMesh* UCustomizableSkeletalComponentPrivate::GetSkeletalMesh() const
 {
 	if (CustomizableObjectInstanceUsage)
 	{
-		return CustomizableObjectInstanceUsage->GetSkeletalMesh();
+		return CustomizableObjectInstanceUsage->GetPrivate()->GetSkeletalMesh();
 	}
 
 	return nullptr;
 }
 
 
-void UCustomizableSkeletalComponent::SetSkeletalMesh(USkeletalMesh* SkeletalMesh)
+void UCustomizableSkeletalComponentPrivate::SetSkeletalMesh(USkeletalMesh* SkeletalMesh)
 {
 	if (CustomizableObjectInstanceUsage)
 	{
-		CustomizableObjectInstanceUsage->SetSkeletalMesh(SkeletalMesh);
+		CustomizableObjectInstanceUsage->GetPrivate()->SetSkeletalMesh(SkeletalMesh);
 	}
 }
 
 
-void UCustomizableSkeletalComponent::SetPhysicsAsset(UPhysicsAsset* PhysicsAsset)
+void UCustomizableSkeletalComponentPrivate::SetPhysicsAsset(UPhysicsAsset* PhysicsAsset)
 {
 	if (CustomizableObjectInstanceUsage)
 	{
-		CustomizableObjectInstanceUsage->SetPhysicsAsset(PhysicsAsset);
+		CustomizableObjectInstanceUsage->GetPrivate()->SetPhysicsAsset(PhysicsAsset);
 	}
 }
 
 
-USkeletalMesh* UCustomizableSkeletalComponent::GetAttachedSkeletalMesh() const
+USkeletalMesh* UCustomizableSkeletalComponentPrivate::GetAttachedSkeletalMesh() const
 {
 	if (CustomizableObjectInstanceUsage)
 	{
-		return CustomizableObjectInstanceUsage->GetAttachedSkeletalMesh();
+		return CustomizableObjectInstanceUsage->GetPrivate()->GetAttachedSkeletalMesh();
 	}
 
 	return nullptr;
@@ -121,6 +130,42 @@ FName UCustomizableSkeletalComponent::GetComponentName() const
 }
 
 
+UCustomizableObjectInstance* UCustomizableSkeletalComponent::GetCustomizableObjectInstance() const
+{
+	return CustomizableObjectInstance;
+}
+
+
+void UCustomizableSkeletalComponent::SetCustomizableObjectInstance(UCustomizableObjectInstance* Instance)
+{
+	CustomizableObjectInstance = Instance;
+}
+
+
+void UCustomizableSkeletalComponent::SetSkipSetReferenceSkeletalMesh(bool bSkip)
+{
+	bSkipSetReferenceSkeletalMesh = bSkip;
+}
+
+
+bool UCustomizableSkeletalComponent::GetSkipSetReferenceSkeletalMesh() const
+{
+	return bSkipSetReferenceSkeletalMesh;
+}
+
+
+void UCustomizableSkeletalComponent::SetSkipSetSkeletalMeshOnAttach(bool bSkip)
+{
+	GetPrivate()->bSkipSkipSetSkeletalMeshOnAttach = bSkip;
+}
+
+
+bool UCustomizableSkeletalComponent::GetSkipSetSkeletalMeshOnAttach() const
+{
+	return GetPrivate()->bSkipSkipSetSkeletalMeshOnAttach;
+}
+
+
 void UCustomizableSkeletalComponent::UpdateSkeletalMeshAsync(bool bNeverSkipUpdate)
 {
 	if (CustomizableObjectInstanceUsage)
@@ -139,18 +184,49 @@ void UCustomizableSkeletalComponent::UpdateSkeletalMeshAsyncResult(FInstanceUpda
 }
 
 
-#if WITH_EDITOR
+UCustomizableSkeletalComponentPrivate* UCustomizableSkeletalComponent::GetPrivate()
+{
+	check(Private);
+	return Private;
+}
+
 
-void UCustomizableSkeletalComponent::EditorUpdateComponent()
+const UCustomizableSkeletalComponentPrivate* UCustomizableSkeletalComponent::GetPrivate() const
+{
+	check(Private);
+	return Private;
+}
+
+
+#if WITH_EDITOR
+void UCustomizableSkeletalComponentPrivate::EditorUpdateComponent()
 {
 	if (CustomizableObjectInstanceUsage)
 	{
-		CustomizableObjectInstanceUsage->EditorUpdateComponent();
+		CustomizableObjectInstanceUsage->GetPrivate()->EditorUpdateComponent();
 	}
 }
 #endif
 
 
+UCustomizableSkeletalComponent* UCustomizableSkeletalComponentPrivate::GetPublic()
+{
+	UCustomizableSkeletalComponent* Public = StaticCast<UCustomizableSkeletalComponent*>(GetOuter());
+	check(Public);
+
+	return Public;
+}
+
+
+const UCustomizableSkeletalComponent* UCustomizableSkeletalComponentPrivate::GetPublic() const
+{
+	UCustomizableSkeletalComponent* Public = StaticCast<UCustomizableSkeletalComponent*>(GetOuter());
+	check(Public);
+
+	return Public;	
+}
+
+
 void UCustomizableSkeletalComponent::OnAttachmentChanged()
 {
 	Super::OnAttachmentChanged();
@@ -159,7 +235,7 @@ void UCustomizableSkeletalComponent::OnAttachmentChanged()
 
 	if (Parent && CustomizableObjectInstanceUsage)
 	{
-		CustomizableObjectInstanceUsage->SetPendingSetSkeletalMesh(true);
+		CustomizableObjectInstanceUsage->GetPrivate()->SetPendingSetSkeletalMesh(true);
 	}
 	else if(!GetAttachParent())
 	{
@@ -171,12 +247,12 @@ void UCustomizableSkeletalComponent::OnAttachmentChanged()
 void UCustomizableSkeletalComponent::PostInitProperties()
 {
 	Super::PostInitProperties();
-	CreateCustomizableObjectInstanceUsage();
+	GetPrivate()->CreateCustomizableObjectInstanceUsage();
 }
 
 void UCustomizableSkeletalComponent::PostReinitProperties()
 {
 	Super::PostReinitProperties();
-	CreateCustomizableObjectInstanceUsage();
+	GetPrivate()->CreateCustomizableObjectInstanceUsage();
 }
 
diff --git a/Engine/Plugins/Mutable/Source/CustomizableObject/Private/MuCO/LogBenchmarkUtil.cpp b/Engine/Plugins/Mutable/Source/CustomizableObject/Private/MuCO/LogBenchmarkUtil.cpp
--- a/Engine/Plugins/Mutable/Source/CustomizableObject/Private/MuCO/LogBenchmarkUtil.cpp
+++ b/Engine/Plugins/Mutable/Source/CustomizableObject/Private/MuCO/LogBenchmarkUtil.cpp
@@ -201,7 +201,7 @@ void FLogBenchmarkUtil::UpdateStats()
 
 	for (TObjectIterator<UCustomizableObjectInstanceUsage> CustomizableObjectInstanceUsage; CustomizableObjectInstanceUsage; ++CustomizableObjectInstanceUsage)
 	{
-		if (!IsValid(*CustomizableObjectInstanceUsage) || CustomizableObjectInstanceUsage->IsNetMode(NM_DedicatedServer))
+		if (!IsValid(*CustomizableObjectInstanceUsage) || CustomizableObjectInstanceUsage->GetPrivate()->IsNetMode(NM_DedicatedServer))
 		{
 			continue;
 		}
diff --git a/Engine/Plugins/Mutable/Source/CustomizableObject/Public/MuCO/CustomizableInstanceLODManagement.h b/Engine/Plugins/Mutable/Source/CustomizableObject/Public/MuCO/CustomizableInstanceLODManagement.h
--- a/Engine/Plugins/Mutable/Source/CustomizableObject/Public/MuCO/CustomizableInstanceLODManagement.h
+++ b/Engine/Plugins/Mutable/Source/CustomizableObject/Public/MuCO/CustomizableInstanceLODManagement.h
@@ -5,6 +5,7 @@
 
 #include "CustomizableInstanceLODManagement.generated.h"
 
+class UCustomizableObjectInstanceUsage;
 typedef TMap<const class UCustomizableObjectInstance*, class FMutableUpdateCandidate> FMutableInstanceUpdateMap;
 
 // This is an abstract base class, override it to create a new Instance LOD management system and register with UCustomizableObjectSystem::SetInstanceLODManagement
@@ -24,6 +25,10 @@ public:
 	virtual float GetOnlyUpdateCloseCustomizableObjectsDist() const PURE_VIRTUAL(UCustomizableInstanceLODManagementBase::GetOnlyUpdateCloseCustomizableObjectsDist, return 0.f;);
 	virtual bool IsOnlyUpdateCloseCustomizableObjectsEnabled() const PURE_VIRTUAL(UCustomizableInstanceLODManagementBase::IsOnlyUpdateCloseCustomizableObjectsEnabled, return false;);
 	virtual bool IsOnlyGenerateRequestedLODLevelsEnabled() const PURE_VIRTUAL(UCustomizableInstanceLODManagementBase::IsOnlyGenerateRequestedLODLevelsEnabled, return false;);
+
+#if WITH_EDITOR
+	void EditorUpdateComponent(UCustomizableObjectInstanceUsage* InstanceUsage);
+#endif
 };
 
 
diff --git a/Engine/Plugins/Mutable/Source/CustomizableObject/Public/MuCO/CustomizableObjectInstanceUsage.h b/Engine/Plugins/Mutable/Source/CustomizableObject/Public/MuCO/CustomizableObjectInstanceUsage.h
--- a/Engine/Plugins/Mutable/Source/CustomizableObject/Public/MuCO/CustomizableObjectInstanceUsage.h
+++ b/Engine/Plugins/Mutable/Source/CustomizableObject/Public/MuCO/CustomizableObjectInstanceUsage.h
@@ -3,6 +3,7 @@
 #pragma once
 
 #include "MuCO/CustomizableObjectInstance.h"
+#include "MuCO/CustomizableObjectInstanceUsagePrivate.h"
 #include "Tickable.h"
 #include "Engine/EngineBaseTypes.h"
 
@@ -12,7 +13,6 @@ class AActor;
 class FObjectPreSaveContext;
 class UObject;
 class USkeletalMesh;
-class UCustomizableSkeletalComponent;
 class UPhysicsAsset;
 struct FFrame;
 enum class EUpdateResult : uint8;
@@ -24,22 +24,22 @@ DECLARE_DELEGATE(FCustomizableObjectInstanceUsageUpdatedDelegate);
 // streaming, etc. It's a UObject, so it will be much cheaper than a UCustomizableComponent as it won't have to refresh its transforms
 // every time it's moved.
 UCLASS(Blueprintable, BlueprintType, ClassGroup = (CustomizableObject), meta = (BlueprintSpawnableComponent))
-class CUSTOMIZABLEOBJECT_API UCustomizableObjectInstanceUsage : public UObject, public FTickableGameObject
+class CUSTOMIZABLEOBJECT_API UCustomizableObjectInstanceUsage : public UObject
 {
 public:
 	GENERATED_BODY()
 
-	UPROPERTY(Transient)
-	float SkippedLastRenderTime;
-
-	FCustomizableObjectInstanceUsageUpdatedDelegate UpdatedDelegate;
-
+	// Own interface
+	UCustomizableObjectInstanceUsage();
+	
+	UFUNCTION(BlueprintCallable, Category = CustomizableObjectInstanceUsage)
 	void SetCustomizableObjectInstance(UCustomizableObjectInstance* CustomizableObjectInstance);
-	UCustomizableObjectInstance* GetCustomizableObjectInstance() const;
 
-	// TODO: What is the proper way to refer to dynamic components?
-	// For now assume this is an *Object* component index (to index the array of all possible components in the CO) and not an *Instance* component index.
+	UFUNCTION(BlueprintCallable, Category = CustomizableObjectInstanceUsage)
+	UCustomizableObjectInstance* GetCustomizableObjectInstance() const;
+	
 	void SetComponentIndex(int32 ObjectComponentIndex);
+
 	int32 GetComponentIndex() const;
 	
 	UFUNCTION(BlueprintCallable, Category = CustomizableObjectInstanceUsage)
@@ -47,16 +47,15 @@ public:
 
 	UFUNCTION(BlueprintCallable, Category = CustomizableObjectInstanceUsage)
 	FName GetComponentName() const;
-	
+
+	/** Attach this Customizable Object Instance Usage to a Skeletal Mesh Component to be customized. */
+	UFUNCTION(BlueprintCallable, Category = CustomizableObjectInstanceUsage)
 	void AttachTo(USkeletalMeshComponent* SkeletalMeshComponent);
+	
+	/** Get the parent Skeletal Mesh Component this Customizable Object Instance Usage is attached to. */
+	UFUNCTION(BlueprintCallable, Category = CustomizableObjectInstanceUsage)
 	USkeletalMeshComponent* GetAttachParent() const;
-
-	/** Common end point of all updates. Even those which failed. */
-	void Callbacks() const;
 	
-	USkeletalMesh* GetSkeletalMesh() const;
-	USkeletalMesh* GetAttachedSkeletalMesh() const;
-
 	/** Update Skeletal Mesh asynchronously. */
 	UFUNCTION(BlueprintCallable, Category = CustomizableObject)
 	void UpdateSkeletalMeshAsync(bool bNeverSkipUpdate = false);
@@ -64,33 +63,32 @@ public:
 	/** Update Skeletal Mesh asynchronously. Callback will be called once the update finishes, even if it fails. */
 	UFUNCTION(BlueprintCallable, Category = CustomizableObjectInstance)
 	void UpdateSkeletalMeshAsyncResult(FInstanceUpdateDelegate Callback, bool bIgnoreCloseDist = false, bool bForceHighPriority = false);
+	
+	/** Set to true to avoid automatically replacing the Skeletal Mesh of the parent Skeletal Mesh Component by the Reference Skeletal Mesh.
+	 * If SkipSetSkeletalMeshOnAttach is true, it will not replace it. */
+	UFUNCTION(BlueprintCallable, Category = CustomizableObjectInstanceUsage)
+	void SetSkipSetReferenceSkeletalMesh(bool bSkip);
 
-	void SetSkeletalMesh(USkeletalMesh* SkeletalMesh, bool* bOutSkeletalMeshUpdated = nullptr, bool* bOutMaterialsUpdated = nullptr);
-	void SetPhysicsAsset(UPhysicsAsset* PhysicsAsset, bool* bOutPhysicsAssetUpdated = nullptr);
-	void UpdateDistFromComponentToPlayer(const AActor* const Pawn, bool bForceEvenIfNotBegunPlay = false);
+	UFUNCTION(BlueprintCallable, Category = CustomizableObjectInstanceUsage)
+	bool GetSkipSetReferenceSkeletalMesh() const;
 
-	void SetVisibilityOfSkeletalMeshSectionWithMaterialName(bool bVisible, const FString& MaterialName, int32 LOD);
+	/** Set to true to avoid automatically replacing the Skeletal Mesh of the parent Skeletal Mesh Component with any mesh. */
+	UFUNCTION(BlueprintCallable, Category = CustomizableObjectInstanceUsage)
+	void SetSkipSetSkeletalMeshOnAttach(bool bSkip);
 
-	// Set to true to replace the SkeletalMesh of the parent component by the ReferenceSkeletalMesh or the generated SkeletalMesh 
-	void SetPendingSetSkeletalMesh(bool bIsActive);
-	bool GetPendingSetSkeletalMesh() const;
+	UFUNCTION(BlueprintCallable, Category = CustomizableObjectInstanceUsage)
+	bool GetSkipSetSkeletalMeshOnAttach() const;
+	
+	UCustomizableObjectInstanceUsagePrivate* GetPrivate();
 
-	// Set to true to avoid replacing the SkeletalMesh of the parent component by the ReferenceSkeletalMesh if bPendingSetSkeletalMesh is true
-	void SetSkipSetReferenceSkeletalMesh(bool bIsActive);
-	bool GetSkipSetReferenceSkeletalMesh() const;
+	const UCustomizableObjectInstanceUsagePrivate* GetPrivate() const;
 
-	// USceneComponent interface
-	virtual void BeginDestroy() override;
+	FCustomizableObjectInstanceUsageUpdatedDelegate UpdatedDelegate;
 
-	// Returns true if the NetMode of the associated UCustomizableSkeletalComponent (or the associated SkeletalMeshComponent if the former does not exist) is equal to InNetMode
+	// TODO GMT Remove
 	bool IsNetMode(ENetMode InNetMode) const;
-
-#if WITH_EDITOR
-	// Used to generate instances outside the CustomizableObject editor and PIE
-	void UpdateDistFromComponentToLevelEditorCamera(const FVector& CameraPosition);
 	void EditorUpdateComponent();
-#endif
-
+	
 private:
 	// If this CustomizableSkeletalComponent is not null, it means this Usage was created by it, and all persistent properties should be obtained through it
 	UPROPERTY()
@@ -111,24 +109,10 @@ private:
 	/** Only used if the ComponentIndex is INDEX_NONE. */
 	UPROPERTY()
 	FName UsedComponentName;
-
-	// Used to replace the SkeletalMesh of the parent component by the ReferenceSkeletalMesh or the generated SkeletalMesh 
-	bool bUsedPendingSetSkeletalMesh = false;
-
-	// Used to avoid replacing the SkeletalMesh of the parent component by the ReferenceSkeletalMesh if bPendingSetSkeletalMesh is true
-	bool bUsedSkipSetReferenceSkeletalMesh = false;
-
-	// Begin FTickableGameObject
-	virtual void Tick(float DeltaTime) override;
-	virtual ETickableTickType GetTickableTickType() const override;
-	virtual TStatId GetStatId() const override;
-	virtual bool IsTickableWhenPaused() const override { return true; }
-	virtual bool IsTickableInEditor() const override { return true; }
-	virtual bool IsTickable() const override;
-	// End FTickableGameObject
-
-	void SetSkeletalMeshAndOverrideMaterials(USkeletalMeshComponent& Parent, USkeletalMesh* SkeletalMesh, const UCustomizableObjectInstance& CustomizableObjectInstance, bool* bOutSkeletalMeshUpdated, bool* bOutMaterialsUpdated);
-
-	friend UCustomizableSkeletalComponent;
+	
+	UPROPERTY()
+	TObjectPtr<UCustomizableObjectInstanceUsagePrivate> Private;
+	
+	friend UCustomizableObjectInstanceUsagePrivate;
 };
 
diff --git a/Engine/Plugins/Mutable/Source/CustomizableObject/Public/MuCO/CustomizableSkeletalComponent.h b/Engine/Plugins/Mutable/Source/CustomizableObject/Public/MuCO/CustomizableSkeletalComponent.h
--- a/Engine/Plugins/Mutable/Source/CustomizableObject/Public/MuCO/CustomizableSkeletalComponent.h
+++ b/Engine/Plugins/Mutable/Source/CustomizableObject/Public/MuCO/CustomizableSkeletalComponent.h
@@ -8,85 +8,101 @@
 
 #include "CustomizableSkeletalComponent.generated.h"
 
+class UCustomizableSkeletalComponentPrivate;
 
 DECLARE_DELEGATE(FCustomizableSkeletalComponentUpdatedDelegate);
 
+
 UCLASS(Blueprintable, BlueprintType, ClassGroup = (CustomizableObject), meta = (BlueprintSpawnableComponent))
 class CUSTOMIZABLEOBJECT_API UCustomizableSkeletalComponent : public USceneComponent
 {
 public:
 	GENERATED_BODY()
-
-	UPROPERTY(BlueprintReadWrite, EditAnywhere, Category = CustomizableSkeletalMesh)
+	
+	UPROPERTY(BlueprintReadWrite, EditAnywhere, Category = CustomizableSkeletalComponent)
 	TObjectPtr<UCustomizableObjectInstance> CustomizableObjectInstance;
 
 	/** This component index refers to the object list of components */
-	UPROPERTY(BlueprintReadWrite, EditAnywhere, Category = CustomizableSkeletalMesh)
+	UPROPERTY(BlueprintReadWrite, EditAnywhere, Category = CustomizableSkeletalComponent)
 	int32 ComponentIndex;
 
 private:
 	/** Only used if the ComponentIndex is INDEX_NONE.
 	 *  Editing this property will set ComponentIndex to INDEX_NONE. */
-	UPROPERTY(EditAnywhere, Category = CustomizableSkeletalMesh)
+	UPROPERTY(EditAnywhere, Category = CustomizableSkeletalComponent)
 	FName ComponentName;
 
-public:
-	// Used to replace the SkeletalMesh of the parent component by the ReferenceSkeletalMesh or the generated SkeletalMesh 
-	bool bPendingSetSkeletalMesh = false;
-
-	// Used to avoid replacing the SkeletalMesh of the parent component by the ReferenceSkeletalMesh if bPendingSetSkeletalMesh is true
+public: // TODO GMT Private
+	UPROPERTY(EditAnywhere, Category = CustomizableSkeletalComponent)
 	bool bSkipSetReferenceSkeletalMesh = false;
-
+	
+public:
 	FCustomizableSkeletalComponentUpdatedDelegate UpdatedDelegate;
 
+private:
+	UPROPERTY()
+	TObjectPtr<UCustomizableSkeletalComponentPrivate> Private;
+
+public:
+	UCustomizableSkeletalComponent();
+	
 	// UObject interface
+protected:
 #if WITH_EDITOR
 	virtual void PostEditChangeProperty(FPropertyChangedEvent& PropertyChangedEvent) override;
 #endif
-	
+
+	// UActor interface
+	virtual void PostInitProperties() override;
+	virtual void PostReinitProperties() override;
+
+	// USceneComponent interface
+	virtual void OnAttachmentChanged() override;
+
+public:
 	// Own interface
 	
-	/** Common end point of all updates. Even those which failed. */
-	void Callbacks() const;
-	
-	USkeletalMesh* GetSkeletalMesh() const;
-	USkeletalMesh* GetAttachedSkeletalMesh() const;
-
-	UFUNCTION(BlueprintCallable, Category = CustomizableSkeletalMesh)
+	UFUNCTION(BlueprintCallable, Category = CustomizableSkeletalComponent)
 	void SetComponentName(const FName& Name);
 
-	UFUNCTION(BlueprintCallable, Category = CustomizableSkeletalMesh)
+	UFUNCTION(BlueprintCallable, Category = CustomizableSkeletalComponent)
 	FName GetComponentName() const;
 	
+	UFUNCTION(BlueprintCallable, Category = CustomizableSkeletalComponent)
+	UCustomizableObjectInstance* GetCustomizableObjectInstance() const;
+	
+	UFUNCTION(BlueprintCallable, Category = CustomizableSkeletalComponent)
+	void SetCustomizableObjectInstance(UCustomizableObjectInstance* Instance);
+	
+	/** Set to true to avoid automatically replacing the Skeletal Mesh of the parent Skeletal Mesh Component by the Reference Skeletal Mesh.
+	 * If SkipSetSkeletalMeshOnAttach is true, it will not replace it. */
+	UFUNCTION(BlueprintCallable, Category = CustomizableSkeletalComponent)
+	void SetSkipSetReferenceSkeletalMesh(bool bSkip);
+
+	UFUNCTION(BlueprintCallable, Category = CustomizableSkeletalComponent)
+	bool GetSkipSetReferenceSkeletalMesh() const;
+
+	/** Set to true to avoid automatically replacing the Skeletal Mesh of the parent Skeletal Mesh Component with any mesh. */
+	UFUNCTION(BlueprintCallable, Category = CustomizableObjectInstanceUsage)
+	void SetSkipSetSkeletalMeshOnAttach(bool bSkip);
+
+	UFUNCTION(BlueprintCallable, Category = CustomizableObjectInstanceUsage)
+	bool GetSkipSetSkeletalMeshOnAttach() const;
+	
 	/** Update Skeletal Mesh asynchronously. */
-	UFUNCTION(BlueprintCallable, Category = CustomizableObject)
+	UFUNCTION(BlueprintCallable, Category = CustomizableSkeletalComponent)
 	void UpdateSkeletalMeshAsync(bool bNeverSkipUpdate = false);
 
 	/** Update Skeletal Mesh asynchronously. Callback will be called once the update finishes, even if it fails. */
-	UFUNCTION(BlueprintCallable, Category = CustomizableObjectInstance)
+	UFUNCTION(BlueprintCallable, Category = CustomizableSkeletalComponent)
 	void UpdateSkeletalMeshAsyncResult(FInstanceUpdateDelegate Callback, bool bIgnoreCloseDist = false, bool bForceHighPriority = false);
 
-	void SetSkeletalMesh(USkeletalMesh* SkeletalMesh);
-	void SetPhysicsAsset(class UPhysicsAsset* PhysicsAsset);
-	void UpdateDistFromComponentToPlayer(const AActor* const Pawn, bool bForceEvenIfNotBegunPlay = false);
-
-	void SetVisibilityOfSkeletalMeshSectionWithMaterialName(bool bVisible, const FString& MaterialName, int32 LOD);
-
-#if WITH_EDITOR
-	// Used to generate instances outside the CustomizableObject editor and PIE
-	void UpdateDistFromComponentToLevelEditorCamera(const FVector& CameraPosition);
-	void EditorUpdateComponent();
-#endif
+	UCustomizableSkeletalComponentPrivate* GetPrivate();
 
+	const UCustomizableSkeletalComponentPrivate* GetPrivate() const;
+	
 private:
 	UPROPERTY(Transient)
 	TObjectPtr<UCustomizableObjectInstanceUsage> CustomizableObjectInstanceUsage;
-
-	void OnAttachmentChanged() override;
-	void CreateCustomizableObjectInstanceUsage();
-
-protected:
-	virtual void PostInitProperties() override;
-	virtual void PostReinitProperties() override;
 };
 
diff --git a/Engine/Plugins/Mutable/Source/CustomizableObjectEditor/Private/MuCOE/CustomizableObjectEditor.cpp b/Engine/Plugins/Mutable/Source/CustomizableObjectEditor/Private/MuCOE/CustomizableObjectEditor.cpp
--- a/Engine/Plugins/Mutable/Source/CustomizableObjectEditor/Private/MuCOE/CustomizableObjectEditor.cpp
+++ b/Engine/Plugins/Mutable/Source/CustomizableObjectEditor/Private/MuCOE/CustomizableObjectEditor.cpp
@@ -66,6 +66,7 @@
 #include "Widgets/Layout/SScrollBox.h"
 #include "Widgets/Notifications/SNotificationList.h"
 #include "SMutableObjectViewer.h"
+#include "MuCO/CustomizableSkeletalComponentPrivate.h"
 #include "MuCO/CustomizableObjectCustomVersion.h"
 #include "MuCO/CustomizableSkeletalMeshActor.h"
 #include "MuCOE/CustomizableObjectInstanceFactory.h"
@@ -2328,7 +2329,7 @@ void FCustomizableObjectEditor::CreatePreviewActor()
 		PreviewSkeletalMeshComponents.Add(DebugComponent);
 		
 		UCustomizableSkeletalComponent* CustomizableComponent = NewObject<UCustomizableSkeletalComponent>(DebugComponent);
-		CustomizableComponent->bSkipSetReferenceSkeletalMesh = true;
+		CustomizableComponent->SetSkipSetReferenceSkeletalMesh(true);
 		CustomizableComponent->CustomizableObjectInstance = PreviewInstance;
 		CustomizableComponent->SetComponentName(CustomizableObject->GetComponentName(ObjectComponentIndex));
 		CustomizableComponent->AttachToComponent(DebugComponent, FAttachmentTransformRules::KeepRelativeTransform);
diff --git a/Engine/Plugins/Mutable/Source/CustomizableObjectEditor/Private/MuCOE/CustomizableObjectInstanceEditor.cpp b/Engine/Plugins/Mutable/Source/CustomizableObjectEditor/Private/MuCOE/CustomizableObjectInstanceEditor.cpp
--- a/Engine/Plugins/Mutable/Source/CustomizableObjectEditor/Private/MuCOE/CustomizableObjectInstanceEditor.cpp
+++ b/Engine/Plugins/Mutable/Source/CustomizableObjectEditor/Private/MuCOE/CustomizableObjectInstanceEditor.cpp
@@ -987,7 +987,7 @@ void FCustomizableObjectInstanceEditor::CreatePreviewActor()
 		PreviewSkeletalMeshComponents.Add(DebugComponent);
 		
 		UCustomizableSkeletalComponent* CustomizableComponent = NewObject<UCustomizableSkeletalComponent>(DebugComponent);
-		CustomizableComponent->bSkipSetReferenceSkeletalMesh = true;
+		CustomizableComponent->SetSkipSetReferenceSkeletalMesh(true);
 		CustomizableComponent->CustomizableObjectInstance = CustomizableObjectInstance;
 		CustomizableComponent->SetComponentName(Object->GetComponentName(ObjectComponentIndex));
 		CustomizableComponent->AttachToComponent(DebugComponent, FAttachmentTransformRules::KeepRelativeTransform);
diff --git a/Engine/Plugins/Mutable/Source/CustomizableObjectEditor/Private/MuCOE/CustomizableObjectInstanceFactory.cpp b/Engine/Plugins/Mutable/Source/CustomizableObjectEditor/Private/MuCOE/CustomizableObjectInstanceFactory.cpp
--- a/Engine/Plugins/Mutable/Source/CustomizableObjectEditor/Private/MuCOE/CustomizableObjectInstanceFactory.cpp
+++ b/Engine/Plugins/Mutable/Source/CustomizableObjectEditor/Private/MuCOE/CustomizableObjectInstanceFactory.cpp
@@ -5,6 +5,7 @@
 #include "AssetRegistry/AssetData.h"
 #include "AssetRegistry/AssetRegistryModule.h"
 #include "Components/SkeletalMeshComponent.h"
+#include "MuCO/CustomizableSkeletalComponentPrivate.h"
 #include "Engine/SkeletalMesh.h"
 #include "Modules/ModuleManager.h"
 #include "MuCO/CustomizableObjectInstance.h"
@@ -78,7 +79,7 @@ void UCustomizableObjectInstanceFactory::PostSpawnActor(UObject* Asset, AActor*
 				CustomSkeletalComp->UnregisterComponent();
 				CustomSkeletalComp->CustomizableObjectInstance = Instance;
 				CustomSkeletalComp->SetComponentName(ComponentName);
-				CustomSkeletalComp->SetSkeletalMesh(SkeletalMesh);
+				CustomSkeletalComp->GetPrivate()->SetSkeletalMesh(SkeletalMesh);
 				CustomSkeletalComp->UpdateSkeletalMeshAsync();
 				CustomSkeletalComp->RegisterComponent();
 			}
diff --git a/Engine/Plugins/Mutable/Source/CustomizableObjectEditor/Private/MuCOE/Thumbnails/CustomizableObjectInstanceThumbnailRenderer.cpp b/Engine/Plugins/Mutable/Source/CustomizableObjectEditor/Private/MuCOE/Thumbnails/CustomizableObjectInstanceThumbnailRenderer.cpp
--- a/Engine/Plugins/Mutable/Source/CustomizableObjectEditor/Private/MuCOE/Thumbnails/CustomizableObjectInstanceThumbnailRenderer.cpp
+++ b/Engine/Plugins/Mutable/Source/CustomizableObjectEditor/Private/MuCOE/Thumbnails/CustomizableObjectInstanceThumbnailRenderer.cpp
@@ -10,6 +10,7 @@
 #include "SceneView.h"
 #include "TextureResource.h"
 #include "ThumbnailHelpers.h"
+#include "MuCO/CustomizableSkeletalComponentPrivate.h"
 #include "UObject/ConstructorHelpers.h"
 
 class FRenderTarget;
@@ -82,7 +83,7 @@ void UCustomizableObjectInstanceThumbnailRenderer::Draw(UObject* Object, int32 X
 			PreviewCustomizableSkeletalComponent->CustomizableObjectInstance = CustomizableObjectInstance;
 
 			PreviewCustomizableSkeletalComponent->UpdateSkeletalMeshAsync();
-			USkeletalMesh* SkeletalMesh = PreviewCustomizableSkeletalComponent->GetSkeletalMesh();
+			USkeletalMesh* SkeletalMesh = PreviewCustomizableSkeletalComponent->GetPrivate()->GetSkeletalMesh();
 
 			if (SkeletalMesh)
 			{
diff --git a/Engine/Plugins/Mutable/Source/CustomizableObjectEditor/Private/MuCOE/Thumbnails/CustomizableObjectThumbnailRenderer.cpp b/Engine/Plugins/Mutable/Source/CustomizableObjectEditor/Private/MuCOE/Thumbnails/CustomizableObjectThumbnailRenderer.cpp
--- a/Engine/Plugins/Mutable/Source/CustomizableObjectEditor/Private/MuCOE/Thumbnails/CustomizableObjectThumbnailRenderer.cpp
+++ b/Engine/Plugins/Mutable/Source/CustomizableObjectEditor/Private/MuCOE/Thumbnails/CustomizableObjectThumbnailRenderer.cpp
@@ -10,6 +10,7 @@
 #include "SceneView.h"
 #include "TextureResource.h"
 #include "ThumbnailHelpers.h"
+#include "MuCO/CustomizableSkeletalComponentPrivate.h"
 #include "UObject/ConstructorHelpers.h"
 
 class FRenderTarget;
@@ -70,7 +71,7 @@ void UCustomizableObjectThumbnailRenderer::Draw(UObject* Object, int32 X, int32
 				PreviewCustomizableSkeletalComponent->CustomizableObjectInstance = CustomizableObjectInstance;
 
 				PreviewCustomizableSkeletalComponent->UpdateSkeletalMeshAsync();
-				USkeletalMesh* SkeletalMesh = PreviewCustomizableSkeletalComponent->GetSkeletalMesh();
+				USkeletalMesh* SkeletalMesh = PreviewCustomizableSkeletalComponent->GetPrivate()->GetSkeletalMesh();
 
 				if (SkeletalMesh && SkeletalMesh->GetLODNum())
 				{
