[Mutable] Removed temporal public functions in UCustomizableSkeletalComponent and UCustomizableInstanceUsage.

Those were kept to avoid breaking internal CI. Not that all usages are removed, we can safely remove them.

#rnx

diff --git a/Engine/Plugins/Mutable/Source/CustomizableObject/Private/MuCO/CustomizableObjectInstanceUsage.cpp b/Engine/Plugins/Mutable/Source/CustomizableObject/Private/MuCO/CustomizableObjectInstanceUsage.cpp
--- a/Engine/Plugins/Mutable/Source/CustomizableObject/Private/MuCO/CustomizableObjectInstanceUsage.cpp
+++ b/Engine/Plugins/Mutable/Source/CustomizableObject/Private/MuCO/CustomizableObjectInstanceUsage.cpp
@@ -541,18 +541,6 @@ const UCustomizableObjectInstanceUsagePrivate* UCustomizableObjectInstanceUsage:
 }
 
 
-bool UCustomizableObjectInstanceUsage::IsNetMode(ENetMode InNetMode) const
-{
-	return GetPrivate()->IsNetMode(InNetMode);
-}
-
-
-void UCustomizableObjectInstanceUsage::EditorUpdateComponent()
-{
-	GetPrivate()->EditorUpdateComponent();
-}
-
-
 void UCustomizableObjectInstanceUsagePrivate::UpdateDistFromComponentToPlayer(const AActor* ViewCenter, const bool bForceEvenIfNotBegunPlay)
 {
 	UCustomizableObjectInstance* CustomizableObjectInstance = GetPublic()->GetCustomizableObjectInstance();
diff --git a/Engine/Plugins/Mutable/Source/CustomizableObject/Public/MuCO/CustomizableObjectInstanceUsage.h b/Engine/Plugins/Mutable/Source/CustomizableObject/Public/MuCO/CustomizableObjectInstanceUsage.h
--- a/Engine/Plugins/Mutable/Source/CustomizableObject/Public/MuCO/CustomizableObjectInstanceUsage.h
+++ b/Engine/Plugins/Mutable/Source/CustomizableObject/Public/MuCO/CustomizableObjectInstanceUsage.h
@@ -85,10 +85,6 @@ public:
 
 	FCustomizableObjectInstanceUsageUpdatedDelegate UpdatedDelegate;
 
-	// TODO GMT Remove
-	bool IsNetMode(ENetMode InNetMode) const;
-	void EditorUpdateComponent();
-	
 private:
 	// If this CustomizableSkeletalComponent is not null, it means this Usage was created by it, and all persistent properties should be obtained through it
 	UPROPERTY()
diff --git a/Engine/Plugins/Mutable/Source/CustomizableObject/Public/MuCO/CustomizableSkeletalComponent.h b/Engine/Plugins/Mutable/Source/CustomizableObject/Public/MuCO/CustomizableSkeletalComponent.h
--- a/Engine/Plugins/Mutable/Source/CustomizableObject/Public/MuCO/CustomizableSkeletalComponent.h
+++ b/Engine/Plugins/Mutable/Source/CustomizableObject/Public/MuCO/CustomizableSkeletalComponent.h
@@ -32,7 +32,6 @@ private:
 	UPROPERTY(EditAnywhere, Category = CustomizableSkeletalComponent)
 	FName ComponentName;
 
-public: // TODO GMT Private
 	UPROPERTY(EditAnywhere, Category = CustomizableSkeletalComponent)
 	bool bSkipSetReferenceSkeletalMesh = false;
 	
