[Mutable] Fix possible case where texture project could generate 1 pixel texture if ReferenceTextureSize is set to 0.

- This can happen if the target texture size is generated from a CustomizableObjectNodeModifierExtendMeshSection. In this case the value is foced to 0.

#rnx

diff --git a/Engine/Plugins/Mutable/Source/CustomizableObjectEditor/Private/MuCOE/GenerateMutableSource/GenerateMutableSourceImage.cpp b/Engine/Plugins/Mutable/Source/CustomizableObjectEditor/Private/MuCOE/GenerateMutableSource/GenerateMutableSourceImage.cpp
--- a/Engine/Plugins/Mutable/Source/CustomizableObjectEditor/Private/MuCOE/GenerateMutableSource/GenerateMutableSourceImage.cpp
+++ b/Engine/Plugins/Mutable/Source/CustomizableObjectEditor/Private/MuCOE/GenerateMutableSource/GenerateMutableSourceImage.cpp
@@ -631,7 +631,7 @@ mu::NodeImagePtr GenerateMutableSourceImage(const UEdGraphPin* Pin, FMutableGrap
 		const int32 MaxProjectedTextureSizeInGame = ReferenceTextureSize > 0 && ReferenceTextureSize < MaxReferenceTextureSizeInGame ? ReferenceTextureSize : MaxReferenceTextureSizeInGame;
 
 		const int32 ProjectorNodeTextureSize = FMath::Max(TextureSize.X, TextureSize.Y);
-		if (ProjectorNodeTextureSize > 0 && ProjectorNodeTextureSize > MaxProjectedTextureSizeInGame)
+		if (ProjectorNodeTextureSize > 0 && MaxProjectedTextureSizeInGame > 0 && ProjectorNodeTextureSize > MaxProjectedTextureSizeInGame)
 		{
 			const int32 NumMips = FMath::CeilLogTwo(ProjectorNodeTextureSize) + 1;
 			const int32 MaxNumMips = FMath::CeilLogTwo(MaxProjectedTextureSizeInGame) + 1;
