[Mutable] Exposed EditorGameplayTags to Blueprints.

#rnx

diff --git a/Engine/Plugins/Mutable/Source/CustomizableObject/Internal/MuCO/CustomizableObjectPrivate.h b/Engine/Plugins/Mutable/Source/CustomizableObject/Internal/MuCO/CustomizableObjectPrivate.h
--- a/Engine/Plugins/Mutable/Source/CustomizableObject/Internal/MuCO/CustomizableObjectPrivate.h
+++ b/Engine/Plugins/Mutable/Source/CustomizableObject/Internal/MuCO/CustomizableObjectPrivate.h
@@ -1480,6 +1480,8 @@ public:
 		RemovedGroupRomsBySource,
 
 		ReGroupRomsBySource,
+
+		UIMetadataGameplayTags,
 		
 		// -----<new versions can be added above this line>--------
 		LastCustomizableObjectVersion
diff --git a/Engine/Plugins/Mutable/Source/CustomizableObject/Private/MuCO/CustomizableObjectUIData.cpp b/Engine/Plugins/Mutable/Source/CustomizableObject/Private/MuCO/CustomizableObjectUIData.cpp
--- a/Engine/Plugins/Mutable/Source/CustomizableObject/Private/MuCO/CustomizableObjectUIData.cpp
+++ b/Engine/Plugins/Mutable/Source/CustomizableObject/Private/MuCO/CustomizableObjectUIData.cpp
@@ -80,11 +80,10 @@ FArchive& operator<<(FArchive& Ar, FMutableParamUIMetadata& Struct)
 	Ar << Struct.MinimumValue;
 	Ar << Struct.MaximumValue;
 
-#if WITH_EDITORONLY_DATA
 	FString ExportString;
 	if (Ar.IsSaving())
 	{
-		ExportString = Struct.EditorGameplayTags.ToString();
+		ExportString = Struct.GameplayTags.ToString();
 	}
 
 	Ar << ExportString;
@@ -93,10 +92,9 @@ FArchive& operator<<(FArchive& Ar, FMutableParamUIMetadata& Struct)
 	{
 		if (!ExportString.IsEmpty())
 		{
-			Struct.EditorGameplayTags.FromExportString(ExportString);
+			Struct.GameplayTags.FromExportString(ExportString);
 		}
 	}
-#endif
 
 	return Ar;
 }
diff --git a/Engine/Plugins/Mutable/Source/CustomizableObject/Public/MuCO/CustomizableObjectUIData.h b/Engine/Plugins/Mutable/Source/CustomizableObject/Public/MuCO/CustomizableObjectUIData.h
--- a/Engine/Plugins/Mutable/Source/CustomizableObject/Public/MuCO/CustomizableObjectUIData.h
+++ b/Engine/Plugins/Mutable/Source/CustomizableObject/Public/MuCO/CustomizableObjectUIData.h
@@ -5,10 +5,7 @@
 #include "StructUtils/InstancedStruct.h"
 #include "Engine/Texture2D.h"
 #include "MuCO/CustomizableObjectParameterTypeDefinitions.h"
-
-#if WITH_EDITORONLY_DATA
 #include "GameplayTagContainer.h"
-#endif
 
 #include "CustomizableObjectUIData.generated.h"
 
@@ -65,12 +62,10 @@ struct FMutableParamUIMetadata : public FMutableUIMetadata
 	UPROPERTY(BlueprintReadWrite, EditAnywhere, Category = UI)
 	float MaximumValue = 1.0f;
 
-#if WITH_EDITORONLY_DATA
 	/** Gameplay tags to take into consideration when filtering parameters.
 	  * Only applies to the options shown in drop downs. Child object parameters and data table rows. */
-	UPROPERTY(EditAnywhere, Category = UI)
-	FGameplayTagContainer EditorGameplayTags;
-#endif
+	UPROPERTY(BlueprintReadWrite, EditAnywhere, Category = UI)
+	FGameplayTagContainer GameplayTags;
 
 	friend FArchive& operator<<(FArchive& Ar, FMutableParamUIMetadata& Struct);
 };
diff --git a/Engine/Plugins/Mutable/Source/CustomizableObjectEditor/Private/MuCOE/CustomizableInstanceDetails.cpp b/Engine/Plugins/Mutable/Source/CustomizableObjectEditor/Private/MuCOE/CustomizableInstanceDetails.cpp
--- a/Engine/Plugins/Mutable/Source/CustomizableObjectEditor/Private/MuCOE/CustomizableInstanceDetails.cpp
+++ b/Engine/Plugins/Mutable/Source/CustomizableObjectEditor/Private/MuCOE/CustomizableInstanceDetails.cpp
@@ -858,11 +858,11 @@ bool FCustomizableInstanceDetails::IsIntParameterFilteredOut(const UCustomizable
 	{
 	case EGameplayContainerMatchType::Any:
 	{
-		return !Metadata.EditorGameplayTags.HasAny(Filter);
+		return !Metadata.GameplayTags.HasAny(Filter);
 	}
 	case EGameplayContainerMatchType::All:
 	{
-		return !Metadata.EditorGameplayTags.HasAll(Filter);
+		return !Metadata.GameplayTags.HasAll(Filter);
 	}
 	default:
 		unimplemented();
diff --git a/Engine/Plugins/Mutable/Source/CustomizableObjectEditor/Private/MuCOE/GenerateMutableSource/GenerateMutableSourceTable.cpp b/Engine/Plugins/Mutable/Source/CustomizableObjectEditor/Private/MuCOE/GenerateMutableSource/GenerateMutableSourceTable.cpp
--- a/Engine/Plugins/Mutable/Source/CustomizableObjectEditor/Private/MuCOE/GenerateMutableSource/GenerateMutableSourceTable.cpp
+++ b/Engine/Plugins/Mutable/Source/CustomizableObjectEditor/Private/MuCOE/GenerateMutableSource/GenerateMutableSourceTable.cpp
@@ -1084,7 +1084,7 @@ void GenerateTableParameterUIData(const UDataTable* DataTable, const UCustomizab
 									if (uint8* TagCellData = GetCellData(RowName, *DataTable, *TagColumnProperty))
 									{
 										FGameplayTagContainer* TagContainer = reinterpret_cast<FGameplayTagContainer*>(TagCellData);
-										IntegerMetadata.ParamUIMetadata.EditorGameplayTags.AppendTags(*TagContainer);
+										IntegerMetadata.ParamUIMetadata.GameplayTags.AppendTags(*TagContainer);
 									}
 								}
 							}
