[mutable] Improve matching of texture parameters in edit and extend modifiers.
- Use the material parameter name instead of just the index.
- Add warning message on missing required texture connections in the mesh sextion extend modifier.
- Added second error context to error messages so that they can point at two nodes if necessary.
- Removed unnecessary state related to tags from the code generator.
#rnx

diff --git a/Engine/Plugins/Mutable/Source/CustomizableObjectEditor/Private/MuCOE/CustomizableObjectCompileRunnable.cpp b/Engine/Plugins/Mutable/Source/CustomizableObjectEditor/Private/MuCOE/CustomizableObjectCompileRunnable.cpp
--- a/Engine/Plugins/Mutable/Source/CustomizableObjectEditor/Private/MuCOE/CustomizableObjectCompileRunnable.cpp
+++ b/Engine/Plugins/Mutable/Source/CustomizableObjectEditor/Private/MuCOE/CustomizableObjectCompileRunnable.cpp
@@ -214,11 +214,15 @@ uint32 FCustomizableObjectCompileRunnable::Run()
 				TSharedPtr<FErrorAttachedData> ErrorAttachedData = MakeShared<FErrorAttachedData>();
 				ErrorAttachedData->UnassignedUVs.Reset();
 				ErrorAttachedData->UnassignedUVs.Append(MessageAttachedData.m_unassignedUVs, MessageAttachedData.m_unassignedUVsSize);
-				ArrayErrors.Add(FError(Severity, FText::AsCultureInvariant(Message), ErrorAttachedData, pLog->GetMessageContext(i), SpamBin));
+				const UObject* Context = static_cast<const UObject*>(pLog->GetMessageContext(i));
+				ArrayErrors.Add(FError(Severity, FText::AsCultureInvariant(Message), ErrorAttachedData, Context, SpamBin));
 			}
 			else
 			{
-				ArrayErrors.Add(FError(Severity, FText::AsCultureInvariant(Message), pLog->GetMessageContext(i), SpamBin));
+				// TODO: Review, and probably propagate the UObject type into the runtime.
+				const UObject* Context = static_cast<const UObject*>(pLog->GetMessageContext(i));
+				const UObject* Context2 = static_cast<const UObject*>(pLog->GetMessageContext2(i));
+				ArrayErrors.Add(FError(Severity, FText::AsCultureInvariant(Message), Context, Context2, SpamBin));
 			}
 		}
 	}
diff --git a/Engine/Plugins/Mutable/Source/CustomizableObjectEditor/Private/MuCOE/CustomizableObjectCompileRunnable.h b/Engine/Plugins/Mutable/Source/CustomizableObjectEditor/Private/MuCOE/CustomizableObjectCompileRunnable.h
--- a/Engine/Plugins/Mutable/Source/CustomizableObjectEditor/Private/MuCOE/CustomizableObjectCompileRunnable.h
+++ b/Engine/Plugins/Mutable/Source/CustomizableObjectEditor/Private/MuCOE/CustomizableObjectCompileRunnable.h
@@ -28,14 +28,16 @@ public:
 
 	struct FError
 	{
-		EMessageSeverity::Type Severity;
+		EMessageSeverity::Type Severity = EMessageSeverity::Error;
 		ELoggerSpamBin SpamBin = ELoggerSpamBin::ShowAll;
 		FText Message;
 		TSharedPtr<FErrorAttachedData> AttachedData;
-		const void* Context;
+		TObjectPtr<const UObject> Context = nullptr;
+		TObjectPtr<const UObject> Context2 = nullptr;
 
-		FError(const EMessageSeverity::Type InSeverity, const FText& InMessage, const void* InContext, const ELoggerSpamBin InSpamBin = ELoggerSpamBin::ShowAll) : Severity(InSeverity), SpamBin(InSpamBin), Message(InMessage), Context(InContext) {}
-		FError(const EMessageSeverity::Type InSeverity, const FText& InMessage, const TSharedPtr<FErrorAttachedData>& InAttachedData, const void* InContext, const ELoggerSpamBin InSpamBin = ELoggerSpamBin::ShowAll)
+		FError(const EMessageSeverity::Type InSeverity, const FText& InMessage, const UObject* InContext, const UObject* InContext2=nullptr, const ELoggerSpamBin InSpamBin = ELoggerSpamBin::ShowAll )
+			: Severity(InSeverity), SpamBin(InSpamBin), Message(InMessage), Context(InContext), Context2(InContext2) {}
+		FError(const EMessageSeverity::Type InSeverity, const FText& InMessage, const TSharedPtr<FErrorAttachedData>& InAttachedData, const UObject* InContext, const ELoggerSpamBin InSpamBin = ELoggerSpamBin::ShowAll)
 			: Severity(InSeverity), SpamBin(InSpamBin), Message(InMessage), AttachedData(InAttachedData), Context(InContext) {}
 	};
 
diff --git a/Engine/Plugins/Mutable/Source/CustomizableObjectEditor/Private/MuCOE/CustomizableObjectCompiler.cpp b/Engine/Plugins/Mutable/Source/CustomizableObjectEditor/Private/MuCOE/CustomizableObjectCompiler.cpp
--- a/Engine/Plugins/Mutable/Source/CustomizableObjectEditor/Private/MuCOE/CustomizableObjectCompiler.cpp
+++ b/Engine/Plugins/Mutable/Source/CustomizableObjectEditor/Private/MuCOE/CustomizableObjectCompiler.cpp
@@ -1527,11 +1527,19 @@ void FCustomizableObjectCompiler::UpdateCompilerLogData()
 
 	for (const FCustomizableObjectCompileRunnable::FError& CompileError : ArrayCompileErrors)
 	{
-		const UObject* Object = static_cast<const UObject*>(CompileError.Context); // Context are always UObjects
+		TArray<const UObject*> ObjectArray;
+		if (CompileError.Context)
+		{
+			ObjectArray.Add(CompileError.Context);
+		}
+		if (CompileError.Context2)
+		{
+			ObjectArray.Add(CompileError.Context2);
+		}
 
-		if (const UCustomizableObjectNode* Node = Cast<UCustomizableObjectNode>(Object))
+		if (CompileError.Context && CompileError.AttachedData)
 		{
-			if (CompileError.AttachedData)
+			if (const UCustomizableObjectNode* Node = Cast<UCustomizableObjectNode>(CompileError.Context))
 			{
 				UCustomizableObjectNode::FAttachedErrorDataView ErrorDataView;
 				ErrorDataView.UnassignedUVs = { CompileError.AttachedData->UnassignedUVs.GetData(),
@@ -1542,7 +1550,7 @@ void FCustomizableObjectCompiler::UpdateCompilerLogData()
 		}
 
 		FText FullMsg = FText::Format(LOCTEXT("MutableMessage", "{0} : {1}"), ObjectName, CompileError.Message);
-		CompilerLog(FullMsg, Object, CompileError.Severity, true, CompileError.SpamBin);
+		CompilerLog(FullMsg, ObjectArray, CompileError.Severity, true, CompileError.SpamBin);
 	}
 }
 
diff --git a/Engine/Plugins/Mutable/Source/CustomizableObjectEditor/Private/MuCOE/GenerateMutableSource/GenerateMutableSourceModifier.cpp b/Engine/Plugins/Mutable/Source/CustomizableObjectEditor/Private/MuCOE/GenerateMutableSource/GenerateMutableSourceModifier.cpp
--- a/Engine/Plugins/Mutable/Source/CustomizableObjectEditor/Private/MuCOE/GenerateMutableSource/GenerateMutableSourceModifier.cpp
+++ b/Engine/Plugins/Mutable/Source/CustomizableObjectEditor/Private/MuCOE/GenerateMutableSource/GenerateMutableSourceModifier.cpp
@@ -225,6 +225,7 @@ mu::Ptr<mu::NodeModifier> GenerateMutableSourceModifier(const UEdGraphPin * Pin,
 		mu::Ptr<mu::NodeModifierMeshClipWithUVMask> ClipNode = new mu::NodeModifierMeshClipWithUVMask();
 		Result = ClipNode;
 
+		ClipNode->SetMessageContext(Node);
 		ClipNode->FaceCullStrategy = TypedNodeClipUVMask->FaceCullStrategy;
 
 		if (const UEdGraphPin* ConnectedPin = FollowInputPin(*TypedNodeClipUVMask->ClipMaskPin()))
@@ -256,6 +257,7 @@ mu::Ptr<mu::NodeModifier> GenerateMutableSourceModifier(const UEdGraphPin * Pin,
 		GenerationContext.MeshGenerationFlags.Push(ModifiersMeshFlags);
 
 		mu::Ptr<mu::NodeModifierSurfaceEdit> SurfNode = new mu::NodeModifierSurfaceEdit();
+		SurfNode->SetMessageContext(Node);
 		Result = SurfNode;
 
 		// TODO: This was used in the non-modifier version for group projectors. It may affect the "drop projection from LOD" feature.
@@ -330,6 +332,7 @@ mu::Ptr<mu::NodeModifier> GenerateMutableSourceModifier(const UEdGraphPin * Pin,
 			for (int32 ImageIndex = 0; ImageIndex < NumImages; ++ImageIndex)
 			{
 				mu::NodeImagePtr ImageNode;
+				FString MaterialParameterName;
 
 				if (!ImageNode) // If
 				{
@@ -375,14 +378,15 @@ mu::Ptr<mu::NodeModifier> GenerateMutableSourceModifier(const UEdGraphPin * Pin,
 							const int32 ReferenceTextureSize = 0; // TODO GetBaseTextureSize(GenerationContext, TypedNodeExt, ImageIndex);
 
 							ImageNode = GenerateMutableSourceImage(ConnectedPin, GenerationContext, ReferenceTextureSize);
+							MaterialParameterName = TypedNodeExt->GetParameterName(EMaterialParameterType::Texture, ImageIndex).ToString();
+
 						}
 					}
 				}
 
-				if (ImageNode)
-				{
-					SurfNode->LODs[LODIndex].Textures[ImageIndex].Extend = ImageNode;
-				}
+
+				SurfNode->LODs[LODIndex].Textures[ImageIndex].Extend = ImageNode;
+				SurfNode->LODs[LODIndex].Textures[ImageIndex].MaterialParameterName = MaterialParameterName;
 			}
 		}
 
@@ -405,6 +409,7 @@ mu::Ptr<mu::NodeModifier> GenerateMutableSourceModifier(const UEdGraphPin * Pin,
 		mu::Ptr<mu::NodeModifierSurfaceEdit> SurfNode = new mu::NodeModifierSurfaceEdit();
 		Result = SurfNode;
 
+		SurfNode->SetMessageContext(Node);
 		SurfNode->MultipleTagsPolicy = TypedNodeRem->MultipleTagPolicy;
 		SurfNode->RequiredTags = TypedNodeRem->RequiredTags;
 
@@ -442,6 +447,7 @@ mu::Ptr<mu::NodeModifier> GenerateMutableSourceModifier(const UEdGraphPin * Pin,
 		mu::Ptr<mu::NodeModifierMeshClipWithUVMask> ClipNode = new mu::NodeModifierMeshClipWithUVMask();
 		Result = ClipNode;
 
+		ClipNode->SetMessageContext(Node);
 		ClipNode->FaceCullStrategy = TypedNodeRemBlocks->FaceCullStrategy;
 
 		ClipNode->MultipleTagsPolicy = TypedNodeRemBlocks->MultipleTagPolicy;
@@ -465,6 +471,7 @@ mu::Ptr<mu::NodeModifier> GenerateMutableSourceModifier(const UEdGraphPin * Pin,
 		mu::Ptr<mu::NodeModifierSurfaceEdit> SurfNode = new mu::NodeModifierSurfaceEdit();
 		Result = SurfNode;
 
+		SurfNode->SetMessageContext(Node);
 		SurfNode->MultipleTagsPolicy = TypedNodeEdit->MultipleTagPolicy;
 		SurfNode->RequiredTags = TypedNodeEdit->RequiredTags;
 
@@ -493,6 +500,8 @@ mu::Ptr<mu::NodeModifier> GenerateMutableSourceModifier(const UEdGraphPin * Pin,
 
 					mu::NodeModifierSurfaceEdit::FTexture& ImagePatch = SurfNode->LODs[LODIndex].Textures[ImageIndex];
 
+					ImagePatch.MaterialParameterName = TypedNodeEdit->GetParameterName(EMaterialParameterType::Texture, ImageIndex).ToString();
+
 					// \todo: expose these two options?
 					ImagePatch.PatchBlendType = mu::EBlendType::BT_BLEND;
 					ImagePatch.bPatchApplyToAlpha = true;
@@ -540,6 +549,8 @@ mu::Ptr<mu::NodeModifier> GenerateMutableSourceModifier(const UEdGraphPin * Pin,
 		mu::Ptr<mu::NodeModifierSurfaceEdit> SurfNode = new mu::NodeModifierSurfaceEdit();
 		Result = SurfNode;
 
+		SurfNode->SetMessageContext(Node);
+
 		// This modifier needs to be applied right after the mesh constant is generated
 		SurfNode->bApplyBeforeNormalOperations = true;
 
diff --git a/Engine/Plugins/Mutable/Source/MutableTools/Internal/MuT/ErrorLogPrivate.h b/Engine/Plugins/Mutable/Source/MutableTools/Internal/MuT/ErrorLogPrivate.h
--- a/Engine/Plugins/Mutable/Source/MutableTools/Internal/MuT/ErrorLogPrivate.h
+++ b/Engine/Plugins/Mutable/Source/MutableTools/Internal/MuT/ErrorLogPrivate.h
@@ -15,23 +15,25 @@ namespace mu
 
         struct FErrorData
         {
-            TArray< float > m_unassignedUVs;
+            TArray< float > UnassignedUVs;
         };
 
 		struct FMessage
 		{
-			ErrorLogMessageType m_type = ELMT_NONE;
-			ErrorLogMessageSpamBin m_spam = ELMSB_ALL;
-			FString m_text;
-            TSharedPtr<FErrorData> m_data;
-			const void* m_context = nullptr;
+			ErrorLogMessageType Type = ELMT_NONE;
+			ErrorLogMessageSpamBin Spam = ELMSB_ALL;
+			FString Text;
+            TSharedPtr<FErrorData> Data;
+			const void* Context = nullptr;
+			const void* Context2 = nullptr;
 		};
 
-		TArray<FMessage> m_messages;
+		TArray<FMessage> Messages;
 
 
 		//!
 		void Add(const FString& Message, ErrorLogMessageType Type, const void* Context, ErrorLogMessageSpamBin SpamBin = ELMSB_ALL);
+		void Add(const FString& Message, ErrorLogMessageType Type, const void* Context, const void* Context2, ErrorLogMessageSpamBin SpamBin = ELMSB_ALL);
 
         //!
         void Add(const FString& Message, const ErrorLogMessageAttachedDataView& Data, ErrorLogMessageType Type, const void* Context, ErrorLogMessageSpamBin SpamBin = ELMSB_ALL);
diff --git a/Engine/Plugins/Mutable/Source/MutableTools/Private/MuT/CodeGenerator.cpp b/Engine/Plugins/Mutable/Source/MutableTools/Private/MuT/CodeGenerator.cpp
--- a/Engine/Plugins/Mutable/Source/MutableTools/Private/MuT/CodeGenerator.cpp
+++ b/Engine/Plugins/Mutable/Source/MutableTools/Private/MuT/CodeGenerator.cpp
@@ -894,15 +894,24 @@ namespace mu
 		// If this is true, we will reuse the surface properties from a higher LOD, se we can skip the generation of material properties and images.
 		const bool bShareSurface = node.SharedSurfaceId != INDEX_NONE && !bIsBaseForSharedSurface;
 
+		// Gather all modifiers that apply to this surface
+		TArray<FirstPassGenerator::FModifier> Modifiers;
+		constexpr bool bModifiersForBeforeOperations = false;
+
+		// Store the data necessary to apply modifiers for the pre-normal operations stage.
+		// TODO: Should we merge with currently active tags from the InOptions?
+		GetModifiersFor(node.Tags, bModifiersForBeforeOperations, Modifiers);
+
+		// This pass on the modifiers is only to detect errors that cannot be detected at the point they are applied.
+		CheckModifiersForSurface( node, Modifiers);
+
+		// Generate the mesh
         if (node.Mesh)
         {
             MUTABLE_CPUPROFILER_SCOPE(SurfaceMesh);
 
             Ptr<ASTOp> lastMeshOp;
 
-			// Store the data necessary to apply modifiers for the pre-normal operations stage.
-			ActiveTags.Add(node.Tags);
-
             // Generate the mesh
 			FMeshGenerationOptions MeshOptions;
 			MeshOptions.bLayouts = true;
@@ -932,13 +941,8 @@ namespace mu
 
             GenerateMesh(MeshOptions, MeshResults, node.Mesh);
 
-			ActiveTags.Pop();
-
 			// Apply the modifier for the post-normal operations stage.
-			bool bModifiersForBeforeOperations = false;
-			FMeshGenerationOptions ModifierOptions(MeshOptions);
-			ModifierOptions.ActiveTags = node.Tags;
-			lastMeshOp = ApplyMeshModifiers(MeshOptions, MeshResults, SharedMeshResults, bModifiersForBeforeOperations, surfaceNode->GetMessageContext(), nullptr);
+			lastMeshOp = ApplyMeshModifiers(Modifiers, MeshOptions, MeshResults, SharedMeshResults, surfaceNode->GetMessageContext(), nullptr);
 
             // Layouts
             for ( int32 LayoutIndex=0; LayoutIndex < MeshResults.GeneratedLayouts.Num(); ++LayoutIndex)
@@ -1125,8 +1129,9 @@ namespace mu
 
 				if (bFound)
 				{
+					const NodeSurfaceNew::FImageData& ImageData = node.Images[ImageIndex];
 
-					const int32 LayoutIndex = node.Images[ImageIndex].LayoutIndex;
+					const int32 LayoutIndex = ImageData.LayoutIndex;
 
 					// If the layout index has been set to negative, it means we should ignore the layout for this image.
 					CompilerOptions::TextureLayoutStrategy ImageLayoutStrategy = (LayoutIndex < 0)
@@ -1136,8 +1141,6 @@ namespace mu
 
 					if (ImageLayoutStrategy == CompilerOptions::TextureLayoutStrategy::None)
 					{
-						//check(desc.m_format != EImageFormat::IF_NONE);
-
 						// Generate the image
 						FImageGenerationOptions ImageOptions;
 						ImageOptions.State = Options.State;
@@ -1158,8 +1161,7 @@ namespace mu
 						RectInCells.min = { 0,0 };
 						RectInCells.size = { FakeLayoutSize ,FakeLayoutSize };
 						
-						bool bModifiersForBeforeOperations = false;
-						imageAd = ApplyImageBlockModifiers(ImageOptions, imageAd, ImageIndex, GridSize, LayoutBlockDesc, RectInCells, bModifiersForBeforeOperations, surfaceNode->GetMessageContext());
+						imageAd = ApplyImageBlockModifiers( Modifiers, ImageOptions, imageAd, ImageData, GridSize, LayoutBlockDesc, RectInCells, surfaceNode->GetMessageContext());
 						
 						check(imageAd);
 
@@ -1293,8 +1295,7 @@ namespace mu
 								// Even if we force the size afterwards, we need some size hint in some cases, like image projections.
 								ImageOptions.RectSize = UE::Math::TIntVector2<int32>(BlockDesc.m_size);
 
-								bool bModifiersForBeforeOperations = false;
-								blockAd = ApplyImageBlockModifiers(ImageOptions, blockAd, ImageIndex, GridSize, LayoutBlockDesc, RectInCells, bModifiersForBeforeOperations, surfaceNode->GetMessageContext());
+								blockAd = ApplyImageBlockModifiers( Modifiers, ImageOptions, blockAd, ImageData, GridSize, LayoutBlockDesc, RectInCells, surfaceNode->GetMessageContext());
 
 								// Enforce block size and optimizations
 								blockAd = GenerateImageSize(blockAd, FIntVector2(BlockDesc.m_size));
@@ -1323,8 +1324,9 @@ namespace mu
 							FMeshGenerationOptions ModifierOptions;
 							ModifierOptions.State = Options.State;
 							ModifierOptions.ActiveTags = node.Tags;
-							bool bModifiersForBeforeOperations = false;
-							imageAd = ApplyImageExtendModifiers(ModifierOptions, MeshResults, imageAd, ImageLayoutStrategy, LayoutIndex, ImageIndex, GridSize, LayoutBlockDesc, bModifiersForBeforeOperations, surfaceNode->GetMessageContext());
+							imageAd = ApplyImageExtendModifiers( Modifiers, ModifierOptions, MeshResults, imageAd, ImageLayoutStrategy, 
+								LayoutIndex, ImageData, GridSize, LayoutBlockDesc, 
+								surfaceNode->GetMessageContext());
 
 							// Complete the base op
 							BlankImageOp->op.args.ImageBlankLayout.blockSize[0] = uint16(LayoutBlockDesc.BlockPixelsX);
@@ -1865,7 +1867,7 @@ namespace mu
 	//---------------------------------------------------------------------------------------------
 	void CodeGenerator::GetModifiersFor(
 		const TArray<FString>& SurfaceTags,
-		int32 LOD, bool bModifiersForBeforeOperations,
+		bool bModifiersForBeforeOperations,
 		TArray<FirstPassGenerator::FModifier>& OutModifiers)
 	{
         MUTABLE_CPUPROFILER_SCOPE(GetModifiersFor);
@@ -1877,13 +1879,6 @@ namespace mu
 
 		for (const FirstPassGenerator::FModifier& m: FirstPass.Modifiers)
 		{
-			// Correct LOD?
-			// TODO: LOD restrictions inside the modifiers data.
-			//if (m.LOD != LOD)
-			//{
-			//	continue;
-			//}
-
 			// Correct stage?
 			if (m.Node->bApplyBeforeNormalOperations != bModifiersForBeforeOperations)
 			{
@@ -1941,24 +1936,19 @@ namespace mu
 	}
 
 	//---------------------------------------------------------------------------------------------
-	Ptr<ASTOp> CodeGenerator::ApplyMeshModifiers(const FMeshGenerationOptions& Options, 
+	Ptr<ASTOp> CodeGenerator::ApplyMeshModifiers(
+		const TArray<FirstPassGenerator::FModifier>& Modifiers,
+		const FMeshGenerationOptions& Options, 
 		FMeshGenerationResult& BaseMeshResult,
 		const FMeshGenerationResult* SharedMeshResults,
-		bool bModifiersForBeforeOperations,
 		const void* ErrorContext,
 		const NodeMeshConstant* OriginalMeshNode )
 	{
 		Ptr<ASTOp> LastMeshOp = BaseMeshResult.MeshOp;
 
-		// Apply mesh modifiers
-		TArray<FirstPassGenerator::FModifier> Modifiers;
-
-		int32 CurrentLOD = CurrentParents.Last().Lod;
-		GetModifiersFor(Options.ActiveTags, CurrentLOD, bModifiersForBeforeOperations, Modifiers);
-
 		Ptr<ASTOp> PreModifiersMesh = LastMeshOp;
 
-		ActiveTags.Push({});
+		int32 CurrentLOD = CurrentParents.Last().Lod;
 
 		// Process mesh extend modifiers (from edit modifiers)
 		int32 EditIndex = 0;
@@ -1982,8 +1972,6 @@ namespace mu
 					Ptr<NodeMesh> pAdd = Edit->LODs[CurrentLOD].MeshAdd;
 
 					// Store the data necessary to apply modifiers for the pre-normal operations stage.
-					ActiveTags.Push(Edit->EnableTags);
-
 					FMeshGenerationOptions MergedMeshOptions(Options);
 					MergedMeshOptions.ActiveTags = Edit->EnableTags; // TODO: Append to current?
 
@@ -1996,13 +1984,16 @@ namespace mu
 					FMeshGenerationResult AddResults;
 					GenerateMesh(MergedMeshOptions, AddResults, pAdd);
 
-					ActiveTags.Pop();
-
 					// Apply the modifier for the post-normal operations stage to the added mesh
 					FMeshGenerationOptions ModifierOptions(Options);
 					ModifierOptions.ActiveTags = Edit->EnableTags;
+
+					TArray<FirstPassGenerator::FModifier> ChildModifiers;
+					constexpr bool bModifiersForBeforeOperations = false;
+					GetModifiersFor(ModifierOptions.ActiveTags, bModifiersForBeforeOperations, ChildModifiers);
+
 					ModifiersToIgnore.Push(m);
-					Ptr<ASTOp> AddedMeshOp = ApplyMeshModifiers(ModifierOptions, AddResults, SharedMeshResults, bModifiersForBeforeOperations, ErrorContext, nullptr);
+					Ptr<ASTOp> AddedMeshOp = ApplyMeshModifiers(ChildModifiers, ModifierOptions, AddResults, SharedMeshResults, ErrorContext, nullptr);
 					ModifiersToIgnore.Pop();
 
 					FMeshGenerationResult::FExtraLayouts data;
@@ -2419,28 +2410,23 @@ namespace mu
 				LastMeshOp = Op;
 			}
 		}
-			
-		ActiveTags.Pop();
 
 		return LastMeshOp;
 	}
 
 
-	Ptr<ASTOp> CodeGenerator::ApplyImageBlockModifiers(const FImageGenerationOptions& Options, Ptr<ASTOp> BaseImageOp, int32 ImageIndex, 
+	Ptr<ASTOp> CodeGenerator::ApplyImageBlockModifiers(
+		const TArray<FirstPassGenerator::FModifier>& Modifiers,
+		const FImageGenerationOptions& Options, Ptr<ASTOp> BaseImageOp, 
+		const NodeSurfaceNew::FImageData& ImageData,
 		FIntPoint GridSize,
 		const FLayoutBlockDesc& LayoutBlockDesc,
 		box< FIntVector2 > RectInCells,
-		bool bModifiersForBeforeOperations, const void* ErrorContext)
+		const void* ErrorContext)
 	{
 		Ptr<ASTOp> LastImageOp = BaseImageOp;
 
-		// Apply mesh modifiers
-		TArray<FirstPassGenerator::FModifier> Modifiers;
-
 		int32 CurrentLOD = CurrentParents.Last().Lod;
-		GetModifiersFor(Options.ActiveTags, CurrentLOD, bModifiersForBeforeOperations, Modifiers);
-
-		ActiveTags.Push({});
 
 		// Process patch image modifiers (from edit modifiers)
 		for (const FirstPassGenerator::FModifier& m : Modifiers)
@@ -2451,27 +2437,40 @@ namespace mu
 
 				bool bAffectsCurrentLOD = Edit->LODs.IsValidIndex(CurrentLOD);
 
-				if (!bAffectsCurrentLOD || ImageIndex >= Edit->LODs[CurrentLOD].Textures.Num())
+				if (!bAffectsCurrentLOD)
+				{
+					continue;
+				}
+
+				const NodeModifierSurfaceEdit::FTexture* MatchingEdit = Edit->LODs[CurrentLOD].Textures.FindByPredicate(
+					[&](const NodeModifierSurfaceEdit::FTexture& Candidate)
+					{
+						return (Candidate.MaterialParameterName == ImageData.MaterialParameterName);
+					});
+
+				if ( !MatchingEdit)
 				{
 					continue;
 				}
 
-				const NodeModifierSurfaceEdit::FTexture& Patch = Edit->LODs[CurrentLOD].Textures[ImageIndex];
-				if (Patch.PatchImage.get())
+				if (MatchingEdit->PatchImage.get())
 				{
 					// Does the current block need to be patched? Find out by building a mask.
-					Ptr<Image> PatchMask = GenerateImageBlockPatchMask(Patch, GridSize, LayoutBlockDesc.BlockPixelsX, LayoutBlockDesc.BlockPixelsY, RectInCells);
+					Ptr<Image> PatchMask = GenerateImageBlockPatchMask(*MatchingEdit, GridSize, LayoutBlockDesc.BlockPixelsX, LayoutBlockDesc.BlockPixelsY, RectInCells);
 
 					if (PatchMask)
 					{
-						LastImageOp = GenerateImageBlockPatch(LastImageOp, Patch, PatchMask, m.FinalCondition, Options);
+						LastImageOp = GenerateImageBlockPatch(LastImageOp, *MatchingEdit, PatchMask, m.FinalCondition, Options);
 					}
 				}
 			}
 
-		}
+			else
+			{
+				// This modifier doesn't affect the per-block image operations.
+			}
 
-		ActiveTags.Pop();
+		}
 
 		return LastImageOp;
 	}
@@ -2494,24 +2493,21 @@ namespace mu
 
 
 	Ptr<ASTOp> CodeGenerator::ApplyImageExtendModifiers(
+		const TArray<FirstPassGenerator::FModifier>& Modifiers,
 		const FGenericGenerationOptions& Options, 
 		const FMeshGenerationResult& BaseMeshResults,
 		Ptr<ASTOp> BaseImageOp, 
 		CompilerOptions::TextureLayoutStrategy ImageLayoutStrategy,
-		int32 LayoutIndex, int32 ImageIndex,
+		int32 LayoutIndex, 
+		const NodeSurfaceNew::FImageData& ImageData,
 		FIntPoint GridSize, 
 		CodeGenerator::FLayoutBlockDesc& InOutLayoutBlockDesc,
-		bool bModifiersForBeforeOperations, const void* ErrorContext)
+		const void* ModifiedNodeErrorContext)
 	{
 		Ptr<ASTOp> LastImageOp = BaseImageOp;
 
-		// Apply mesh modifiers
-		TArray<FirstPassGenerator::FModifier> Modifiers;
-
 		int32 CurrentLOD = CurrentParents.Last().Lod;
-		GetModifiersFor(Options.ActiveTags, CurrentLOD, bModifiersForBeforeOperations, Modifiers);
 
-		ActiveTags.Push({});
 
 		// Process mesh extend modifiers (from edit modifiers)
 		int32 EditIndex = 0;
@@ -2530,14 +2526,21 @@ namespace mu
 					continue;
 				}
 
-				if (ImageIndex >= Edit->LODs[CurrentLOD].Textures.Num())
-				{
-					continue;
-				}
+				const NodeModifierSurfaceEdit::FTexture* MatchingEdit = Edit->LODs[CurrentLOD].Textures.FindByPredicate(
+					[&](const NodeModifierSurfaceEdit::FTexture& Candidate)
+					{
+						return (Candidate.MaterialParameterName == ImageData.MaterialParameterName);
+					});
 
-				Ptr<NodeImage> pExtend = Edit->LODs[CurrentLOD].Textures[ImageIndex].Extend;
-				if (!pExtend)
+				if (!MatchingEdit || (MatchingEdit && !MatchingEdit->Extend) )
 				{
+					if (Edit->LODs[CurrentLOD].MeshAdd)
+					{
+						// When extending a mesh section it is mandatory to provide textures for all section textures handled by Mutable.
+						FString Msg = FString::Printf(TEXT("Required texture [%s] is missing when trying to extend a mesh section."), *ImageData.MaterialParameterName);
+						ErrorLog->GetPrivate()->Add(Msg, ELMT_WARNING, Edit->GetMessageContext(), ModifiedNodeErrorContext);
+					}
+
 					continue;
 				}
 
@@ -2545,8 +2548,7 @@ namespace mu
 
 				if (LayoutIndex >= ExtraLayouts.Num() || !ExtraLayouts[LayoutIndex].Layout)
 				{
-					FString Msg = FString::Printf(TEXT("Trying to extend a layout that doesn't exist in object [%s]."), *CurrentParents.Last().ObjectNode->Name );
-					ErrorLog->GetPrivate()->Add(Msg, ELMT_ERROR, ErrorContext);
+					ErrorLog->GetPrivate()->Add(TEXT("Trying to extend a layout that doesn't exist."), ELMT_WARNING, Edit->GetMessageContext(), ModifiedNodeErrorContext);
 				}
 				else
 				{
@@ -2565,7 +2567,7 @@ namespace mu
 						ImageOptions.LayoutToApply = pExtendLayout;
 						ImageOptions.LayoutBlockId = pExtendLayout->Blocks[b].Id;
 						FImageGenerationResult ExtendResult;
-						GenerateImage(ImageOptions, ExtendResult, pExtend);
+						GenerateImage(ImageOptions, ExtendResult, MatchingEdit->Extend);
 						Ptr<ASTOp> fragmentAd = ExtendResult.op;
 
 						// Block in layout grid units
@@ -2624,12 +2626,59 @@ namespace mu
 			}
 		}
 
-		ActiveTags.Pop();
-
 		return LastImageOp;
 	}
 
 
+	void CodeGenerator::CheckModifiersForSurface(const NodeSurfaceNew& Node, const TArray<FirstPassGenerator::FModifier>& Modifiers )
+	{
+		int32 CurrentLOD = CurrentParents.Last().Lod;
+
+		for (const FirstPassGenerator::FModifier& m : Modifiers)
+		{
+			// A mistake in the surface edit modifier usually results in no change visible. Try to detect it.
+			if (m.Node->GetType() == NodeModifierSurfaceEdit::GetStaticType())
+			{
+				const NodeModifierSurfaceEdit* Edit = static_cast<const NodeModifierSurfaceEdit*>(m.Node);
+
+				bool bAffectsCurrentLOD = Edit->LODs.IsValidIndex(CurrentLOD);
+
+				if (!bAffectsCurrentLOD)
+				{
+					continue;
+				}
+
+				if (Node.Images.IsEmpty())
+				{
+					continue;
+				}
+
+				bool bAtLeastSomeTexture = false;
+
+				for (NodeSurfaceNew::FImageData Data : Node.Images)
+				{
+					const NodeModifierSurfaceEdit::FTexture* MatchingEdit = Edit->LODs[CurrentLOD].Textures.FindByPredicate(
+						[&](const NodeModifierSurfaceEdit::FTexture& Candidate)
+						{
+							return (Candidate.MaterialParameterName == Data.MaterialParameterName);
+						});
+
+					if (MatchingEdit)
+					{
+						bAtLeastSomeTexture = true;
+						break;
+					}
+				}
+
+				if (!bAtLeastSomeTexture)
+				{
+					ErrorLog->GetPrivate()->Add(TEXT("A mesh section modifier applies to a section but no texture matches."), ELMT_WARNING, Edit->GetMessageContext(), Node.GetMessageContext());
+				}
+			}
+		}
+	}
+
+
 	Ptr<ASTOp> CodeGenerator::GenerateDefaultTableValue(ETableColumnType NodeType)
 	{
 		switch (NodeType)
diff --git a/Engine/Plugins/Mutable/Source/MutableTools/Private/MuT/CodeGenerator.h b/Engine/Plugins/Mutable/Source/MutableTools/Private/MuT/CodeGenerator.h
--- a/Engine/Plugins/Mutable/Source/MutableTools/Private/MuT/CodeGenerator.h
+++ b/Engine/Plugins/Mutable/Source/MutableTools/Private/MuT/CodeGenerator.h
@@ -252,9 +252,6 @@ namespace mu
         //! List of already used vertex ID groups that must be unique.
         TSet<uint32> UniqueVertexIDGroups;
 
-		// (top-down) Tags that are active when generating nodes.
-		TArray< TArray<FString> > ActiveTags;
-
         struct FParentKey
         {
             const NodeObjectNew* ObjectNode = nullptr;
@@ -812,32 +809,43 @@ namespace mu
 		void UpdateLayoutBlockDesc(FLayoutBlockDesc& Out, FImageDesc BlockDesc, FIntVector2 LayoutCellSize);
 
 		// Get the modifiers that have to be applied to elements with a specific tag.
-		void GetModifiersFor(const TArray<FString>& SurfaceTags, int32 LOD, bool bModifiersForBeforeOperations, TArray<FirstPassGenerator::FModifier>& OutModifiers);
+		void GetModifiersFor(const TArray<FString>& SurfaceTags, bool bModifiersForBeforeOperations, TArray<FirstPassGenerator::FModifier>& OutModifiers);
 
 		// Used to avoid recursion when generating modifiers.
 		TArray<FirstPassGenerator::FModifier> ModifiersToIgnore;
 
 		// Apply the required mesh modifiers to the given operation.
-		Ptr<ASTOp> ApplyMeshModifiers(const FMeshGenerationOptions&, FMeshGenerationResult& BaseResults, 
+		Ptr<ASTOp> ApplyMeshModifiers(
+			const TArray<FirstPassGenerator::FModifier>&,
+			const FMeshGenerationOptions&, 
+			FMeshGenerationResult& BaseResults, 
 			const FMeshGenerationResult* SharedSurfaceResults, 
-			bool bModifiersForBeforeOperations, const void* ErrorContext,
+			const void* ErrorContext,
 			const NodeMeshConstant* OriginalMeshNode);
 
-		Ptr<ASTOp> ApplyImageBlockModifiers(const FImageGenerationOptions&, Ptr<ASTOp> BaseImageOp, int32 ImageIndex, 
+		Ptr<ASTOp> ApplyImageBlockModifiers(
+			const TArray<FirstPassGenerator::FModifier>&, 
+			const FImageGenerationOptions&, 
+			Ptr<ASTOp> BaseImageOp,
+			const NodeSurfaceNew::FImageData& ImageData,
 			FIntPoint GridSize,
 			const FLayoutBlockDesc& LayoutBlockDesc,
 			box< FIntVector2 > RectInCells,
-			bool bModifiersForBeforeOperations, const void* ErrorContext);
+			const void* ErrorContext);
 
 		Ptr<ASTOp> ApplyImageExtendModifiers(
+			const TArray<FirstPassGenerator::FModifier>&,
 			const FGenericGenerationOptions& Options, 
 			const FMeshGenerationResult& BaseMeshResults,
 			Ptr<ASTOp> ImageAd, 
 			CompilerOptions::TextureLayoutStrategy ImageLayoutStrategy, 
-			int32 LayoutIndex, int32 ImageIndex,
+			int32 LayoutIndex, 
+			const NodeSurfaceNew::FImageData& ImageData,
 			FIntPoint GridSize,
 			CodeGenerator::FLayoutBlockDesc& InOutLayoutBlockDesc,
-			bool bModifiersForBeforeOperations, const void* ErrorContext);
+			const void* ModifiedNodeErrorContext);
+
+		void CheckModifiersForSurface(const NodeSurfaceNew&, const TArray<FirstPassGenerator::FModifier>&);
 
     };
 
diff --git a/Engine/Plugins/Mutable/Source/MutableTools/Private/MuT/CodeGenerator_Color.cpp b/Engine/Plugins/Mutable/Source/MutableTools/Private/MuT/CodeGenerator_Color.cpp
--- a/Engine/Plugins/Mutable/Source/MutableTools/Private/MuT/CodeGenerator_Color.cpp
+++ b/Engine/Plugins/Mutable/Source/MutableTools/Private/MuT/CodeGenerator_Color.cpp
@@ -303,10 +303,7 @@ namespace mu
 		// Source image
 		FImageGenerationOptions ImageOptions;
 		ImageOptions.State = Options.State;
-		if (!ActiveTags.IsEmpty())
-		{
-			ImageOptions.ActiveTags = ActiveTags.Last();
-		}
+		ImageOptions.ActiveTags = Options.ActiveTags;
 
 		Ptr<ASTOp> base;
 		if (node.Image)
diff --git a/Engine/Plugins/Mutable/Source/MutableTools/Private/MuT/CodeGenerator_Mesh.cpp b/Engine/Plugins/Mutable/Source/MutableTools/Private/MuT/CodeGenerator_Mesh.cpp
--- a/Engine/Plugins/Mutable/Source/MutableTools/Private/MuT/CodeGenerator_Mesh.cpp
+++ b/Engine/Plugins/Mutable/Source/MutableTools/Private/MuT/CodeGenerator_Mesh.cpp
@@ -1322,8 +1322,11 @@ namespace mu
 		OutResult.MeshOp = LastMeshOp;
 
 		// Apply the modifier for the pre-normal operations stage.
-		bool bModifiersForBeforeOperations = true;
-		OutResult.MeshOp = ApplyMeshModifiers(InOptions, OutResult, nullptr, bModifiersForBeforeOperations, InNode->GetMessageContext(), InNode);
+		TArray<FirstPassGenerator::FModifier> Modifiers;
+		constexpr bool bModifiersForBeforeOperations = true;
+		GetModifiersFor(InOptions.ActiveTags, bModifiersForBeforeOperations, Modifiers);
+
+		OutResult.MeshOp = ApplyMeshModifiers(Modifiers, InOptions, OutResult, nullptr, InNode->GetMessageContext(), InNode);
     }
 
 
diff --git a/Engine/Plugins/Mutable/Source/MutableTools/Private/MuT/ErrorLog.cpp b/Engine/Plugins/Mutable/Source/MutableTools/Private/MuT/ErrorLog.cpp
--- a/Engine/Plugins/Mutable/Source/MutableTools/Private/MuT/ErrorLog.cpp
+++ b/Engine/Plugins/Mutable/Source/MutableTools/Private/MuT/ErrorLog.cpp
@@ -25,18 +25,16 @@
 #include "Templates/Tuple.h"
 #include "Trace/Detail/Channel.h"
 
+
 namespace mu
 {
 
-
-	//---------------------------------------------------------------------------------------------
 	ErrorLog::ErrorLog()
 	{
 		m_pD = new Private();
 	}
 
 
-	//---------------------------------------------------------------------------------------------
 	ErrorLog::~ErrorLog()
 	{
         check( m_pD );
@@ -45,29 +43,26 @@ namespace mu
 	}
 
 
-	//---------------------------------------------------------------------------------------------
 	ErrorLog::Private* ErrorLog::GetPrivate() const
 	{
 		return m_pD;
 	}
 
 
-	//---------------------------------------------------------------------------------------------
-    int ErrorLog::GetMessageCount( ErrorLogMessageType ) const
+    int32 ErrorLog::GetMessageCount( ErrorLogMessageType ) const
 	{
         // \todo: by type argument
-		return (int)m_pD->m_messages.Num();
+		return (int32)m_pD->Messages.Num();
 	}
 
 
-	//---------------------------------------------------------------------------------------------
-	const FString& ErrorLog::GetMessageText( int index ) const
+	const FString& ErrorLog::GetMessageText( int32 index ) const
 	{
 		const char* strResult = "";
 
-		if ( index >=0 && index<(int)m_pD->m_messages.Num() )
+		if (m_pD->Messages.IsValidIndex(index))
 		{
-			return m_pD->m_messages[index].m_text;
+			return m_pD->Messages[index].Text;
 		}
 
 		static FString Empty;
@@ -75,122 +70,146 @@ namespace mu
 	}
 
 
-	//---------------------------------------------------------------------------------------------
-	const void* ErrorLog::GetMessageContext( int index ) const
+	const void* ErrorLog::GetMessageContext(int32 index) const
 	{
 		const void* result = 0;
 
-		if ( index >=0 && index<(int)m_pD->m_messages.Num() )
+		if (m_pD->Messages.IsValidIndex(index))
 		{
-			result = m_pD->m_messages[index].m_context;
+			result = m_pD->Messages[index].Context;
 		}
 
 		return result;
 	}
 
 
-	//---------------------------------------------------------------------------------------------
-	ErrorLogMessageType ErrorLog::GetMessageType( int index ) const
+	const void* ErrorLog::GetMessageContext2(int32 index) const
+	{
+		const void* result = 0;
+
+		if (m_pD->Messages.IsValidIndex(index))
+		{
+			result = m_pD->Messages[index].Context2;
+		}
+
+		return result;
+	}
+
+
+	ErrorLogMessageType ErrorLog::GetMessageType( int32 index ) const
 	{
 		ErrorLogMessageType result = ELMT_NONE;
 
-		if ( index >=0 && index<(int)m_pD->m_messages.Num() )
+		if (m_pD->Messages.IsValidIndex(index))
 		{
-			result = m_pD->m_messages[index].m_type;
+			result = m_pD->Messages[index].Type;
 		}
 
 		return result;
 	}
 
 
-	//---------------------------------------------------------------------------------------------
-	ErrorLogMessageSpamBin ErrorLog::GetMessageSpamBin(int index) const
+	ErrorLogMessageSpamBin ErrorLog::GetMessageSpamBin(int32 index) const
 	{
 		ErrorLogMessageSpamBin result = ELMSB_ALL;
 
-		if (index >= 0 && index < (int)m_pD->m_messages.Num())
+		if (m_pD->Messages.IsValidIndex(index))
 		{
-			result = m_pD->m_messages[index].m_spam;
+			result = m_pD->Messages[index].Spam;
 		}
 
 		return result;
 	}
 
-	//---------------------------------------------------------------------------------------------
-    ErrorLogMessageAttachedDataView ErrorLog::GetMessageAttachedData( int index ) const
+
+	ErrorLogMessageAttachedDataView ErrorLog::GetMessageAttachedData( int32 index ) const
 	{
         ErrorLogMessageAttachedDataView result;
 
-		if ( index >=0 && index<(int)m_pD->m_messages.Num() )
+		if (m_pD->Messages.IsValidIndex(index))
 		{
-			const Private::FMessage& message = m_pD->m_messages[index];
+			const Private::FMessage& message = m_pD->Messages[index];
             
-            if ( message.m_data ) 
+            if ( message.Data ) 
             {
-                result.m_unassignedUVs = message.m_data->m_unassignedUVs.GetData();
-			    result.m_unassignedUVsSize = message.m_data->m_unassignedUVs.Num();
+                result.m_unassignedUVs = message.Data->UnassignedUVs.GetData();
+			    result.m_unassignedUVsSize = message.Data->UnassignedUVs.Num();
             }
 		}
 
 		return result;
 	}
 
-	//---------------------------------------------------------------------------------------------
+
 	void ErrorLog::Private::Add(const FString& InMessage,
-								ErrorLogMessageType InType,
-								const void* InContext,
-								ErrorLogMessageSpamBin InSpamBin)
+		ErrorLogMessageType InType,
+		const void* InContext,
+		ErrorLogMessageSpamBin InSpamBin)
 	{
-		m_messages.Add(FMessage());
-		m_messages.Last().m_type = InType;
-		m_messages.Last().m_spam = InSpamBin;
-		m_messages.Last().m_text = InMessage;
-		m_messages.Last().m_context = InContext;
+		FMessage& Msg = Messages.Emplace_GetRef();
+		Msg.Type = InType;
+		Msg.Spam = InSpamBin;
+		Msg.Text = InMessage;
+		Msg.Context = InContext;
 	}
 
-	//---------------------------------------------------------------------------------------------
+
+	void ErrorLog::Private::Add(const FString& InMessage,
+		ErrorLogMessageType InType,
+		const void* InContext,
+		const void* InContext2,
+		ErrorLogMessageSpamBin InSpamBin)
+	{
+		FMessage& Msg = Messages.Emplace_GetRef();
+		Msg.Type = InType;
+		Msg.Spam = InSpamBin;
+		Msg.Text = InMessage;
+		Msg.Context = InContext;
+		Msg.Context2 = InContext2;
+	}
+
+
 	void ErrorLog::Private::Add(const FString& InMessage,
                                 const ErrorLogMessageAttachedDataView& InDataView,
-                                ErrorLogMessageType InType, const void* InContext,
+                                ErrorLogMessageType InType, 
+								const void* InContext,
 								ErrorLogMessageSpamBin InSpamBin)
 	{
-		m_messages.Add(FMessage() );
-		m_messages.Last().m_type = InType;
-		m_messages.Last().m_spam = InSpamBin;
-		m_messages.Last().m_text = InMessage;
-		m_messages.Last().m_context = InContext;
-		m_messages.Last().m_data = MakeShared<FErrorData>();
+		FMessage& Msg = Messages.Emplace_GetRef();
+		Msg.Type = InType;
+		Msg.Spam = InSpamBin;
+		Msg.Text = InMessage;
+		Msg.Context = InContext;
+		Msg.Data = MakeShared<FErrorData>();
 
         if ( InDataView.m_unassignedUVs && InDataView.m_unassignedUVsSize > 0 )
         {
 			// \TODO: Review
-			m_messages.Last().m_data->m_unassignedUVs.Append(InDataView.m_unassignedUVs, InDataView.m_unassignedUVsSize);
+			Msg.Data->UnassignedUVs.Append(InDataView.m_unassignedUVs, InDataView.m_unassignedUVsSize);
         }
 	}
 	
 	
-	//---------------------------------------------------------------------------------------------
 	void ErrorLog::Log() const
 	{
 		UE_LOG(LogMutableCore, Log, TEXT(" Error Log :\n"));
 
-		for ( const Private::FMessage& msg : m_pD->m_messages )
+		for ( const Private::FMessage& msg : m_pD->Messages )
 		{
-			switch ( msg.m_type )
+			switch ( msg.Type )
 			{
-			case ELMT_ERROR: 	UE_LOG(LogMutableCore, Log, TEXT("  ERR  %s\n"), *msg.m_text); break;
-			case ELMT_WARNING: 	UE_LOG(LogMutableCore, Log, TEXT("  WRN  %s\n"), *msg.m_text); break;
-			case ELMT_INFO: 	UE_LOG(LogMutableCore, Log, TEXT("  INF  %s\n"), *msg.m_text); break;
-			default: 			UE_LOG(LogMutableCore, Log, TEXT("  NON  %s\n"), *msg.m_text); break;
+			case ELMT_ERROR: 	UE_LOG(LogMutableCore, Log, TEXT("  ERR  %s\n"), *msg.Text); break;
+			case ELMT_WARNING: 	UE_LOG(LogMutableCore, Log, TEXT("  WRN  %s\n"), *msg.Text); break;
+			case ELMT_INFO: 	UE_LOG(LogMutableCore, Log, TEXT("  INF  %s\n"), *msg.Text); break;
+			default: 			UE_LOG(LogMutableCore, Log, TEXT("  NON  %s\n"), *msg.Text); break;
 			}
 		}
 	}
 
 
-	//---------------------------------------------------------------------------------------------
 	void ErrorLog::Merge( const ErrorLog* pOther )
 	{
-		m_pD->m_messages.Append(pOther->GetPrivate()->m_messages);
+		m_pD->Messages.Append(pOther->GetPrivate()->Messages);
 	}
 
 
@@ -331,11 +350,10 @@ namespace mu
 
     // clang-format on
 
-    //---------------------------------------------------------------------------------------------
     const TCHAR* GetOpName( OP_TYPE type )
     {
-        //check( type>=0 && type<(int)OP_TYPE::COUNT );
-        return s_opNames[(int)type];
+        //check( type>=0 && type<(int32)OP_TYPE::COUNT );
+        return s_opNames[(int32)type];
     }
 
 }
diff --git a/Engine/Plugins/Mutable/Source/MutableTools/Public/MuT/ErrorLog.h b/Engine/Plugins/Mutable/Source/MutableTools/Public/MuT/ErrorLog.h
--- a/Engine/Plugins/Mutable/Source/MutableTools/Public/MuT/ErrorLog.h
+++ b/Engine/Plugins/Mutable/Source/MutableTools/Public/MuT/ErrorLog.h
@@ -58,25 +58,26 @@ namespace mu
 
 		//! Get the number of messages.
 		//! If a message type is provided, return the number of that type of message only.
-		int GetMessageCount( ErrorLogMessageType type = ELMT_NONE ) const;
+		int32 GetMessageCount( ErrorLogMessageType type = ELMT_NONE ) const;
 
 		//! Get the text of a message.
 		//! \param index index of the message from 0 to GetMessageCount(ELMT_NONE)-1
-		const FString& GetMessageText( int index ) const;
+		const FString& GetMessageText( int32 index ) const;
 
 		//! Get the opaque context of a message.
 		//! \param index index of the message from 0 to GetMessageCount(ELMT_NONE)-1
-		const void* GetMessageContext( int index ) const;
+		const void* GetMessageContext(int32 index) const;
+		const void* GetMessageContext2(int32 index) const;
 
 		//!
-		ErrorLogMessageType GetMessageType( int index ) const;
+		ErrorLogMessageType GetMessageType( int32 index ) const;
 
 		//!
-		ErrorLogMessageSpamBin GetMessageSpamBin(int index) const;
+		ErrorLogMessageSpamBin GetMessageSpamBin(int32 index) const;
 
 		//! Get the attached data of a message.
 		//! \param index index of message data from 0 to GetMessageCount(ELMT_NONE)-1
-		ErrorLogMessageAttachedDataView GetMessageAttachedData( int index ) const;
+		ErrorLogMessageAttachedDataView GetMessageAttachedData( int32 index ) const;
 
 		//!
 		void Log() const;
diff --git a/Engine/Plugins/Mutable/Source/MutableTools/Public/MuT/NodeModifierSurfaceEdit.h b/Engine/Plugins/Mutable/Source/MutableTools/Public/MuT/NodeModifierSurfaceEdit.h
--- a/Engine/Plugins/Mutable/Source/MutableTools/Public/MuT/NodeModifierSurfaceEdit.h
+++ b/Engine/Plugins/Mutable/Source/MutableTools/Public/MuT/NodeModifierSurfaceEdit.h
@@ -25,6 +25,11 @@ namespace mu
 		/** Data for every modified texture. */
 		struct FTexture
 		{
+			/** Name used to match the image with the original one being modified. 
+			* This should match the MaterialParameterName in a NodeSurfaceNew::FImageData.
+			*/
+			FString MaterialParameterName;
+
 			/** Image to add if extgending. */
 			Ptr<NodeImage> Extend;
 
