[Mutable] Fix case where a layout with Overlay packing strategy could become resizable.

- The case could happen when generating the layout for a passthrough texture. It would become an issue when extending the mesh section. This can be a valid scenario if all layouts use the overlay strategy.

- Now in case the layout is detected to not be important the packing startegy is kept to the default generated layout.

#rnx

diff --git a/Engine/Plugins/Mutable/Source/CustomizableObjectEditor/Private/MuCOE/GenerateMutableSource/GenerateMutableSourceMesh.cpp b/Engine/Plugins/Mutable/Source/CustomizableObjectEditor/Private/MuCOE/GenerateMutableSource/GenerateMutableSourceMesh.cpp
--- a/Engine/Plugins/Mutable/Source/CustomizableObjectEditor/Private/MuCOE/GenerateMutableSource/GenerateMutableSourceMesh.cpp
+++ b/Engine/Plugins/Mutable/Source/CustomizableObjectEditor/Private/MuCOE/GenerateMutableSource/GenerateMutableSourceMesh.cpp
@@ -3534,6 +3534,13 @@ mu::NodeMeshPtr GenerateMutableSourceMesh(const UEdGraphPin* Pin,
 						LayoutFlags.TexturePinModes[LayoutIndex] != EPinMode::Mutable)
 					{
 						MeshNode->SetLayout(LayoutIndex, CreateDefaultLayout());
+					
+						// Keep packing strategy if possible, Overlay can be valid with EPinMode == Passthrough	
+						if (Layouts[LayoutIndex])
+						{
+							MeshNode->GetLayout(LayoutIndex)->Strategy = ConvertLayoutStrategy(Layouts[LayoutIndex]->PackingStrategy);
+						}
+
 						// Ignore layout
 						continue;
 					}
@@ -4296,12 +4303,19 @@ mu::NodeMeshPtr GenerateMutableSourceMesh(const UEdGraphPin* Pin,
 								// Generating node Layouts
 								const FLayoutGenerationFlags& LayoutFlags = GenerationContext.LayoutGenerationFlags.Last();
 								
-								for (int32 i = 0; i < Layouts.Num(); ++i)
+								for (int32 LayoutIndex = 0; LayoutIndex < Layouts.Num(); ++LayoutIndex)
 								{
-									if (!LayoutFlags.TexturePinModes.IsValidIndex(i) ||
-										LayoutFlags.TexturePinModes[i] != EPinMode::Mutable)
+									if (!LayoutFlags.TexturePinModes.IsValidIndex(LayoutIndex) ||
+										LayoutFlags.TexturePinModes[LayoutIndex] != EPinMode::Mutable)
 									{
-										MeshTableNode->SetLayout(i, CreateDefaultLayout());
+										MeshTableNode->SetLayout(LayoutIndex, CreateDefaultLayout());
+
+										// Keep packing strategy if possible, Overlay can be valid with EPinMode == Passthrough	
+										if (Layouts[LayoutIndex])
+										{
+											MeshTableNode->GetLayout(LayoutIndex)->Strategy = ConvertLayoutStrategy(Layouts[LayoutIndex]->PackingStrategy);
+										}
+
 										// Ignore layouts
 										continue;
 									}
@@ -4309,14 +4323,14 @@ mu::NodeMeshPtr GenerateMutableSourceMesh(const UEdGraphPin* Pin,
 									bool bWasEmpty = false;
 									// In tables, mimic the legacy behaviour and ignore all layout warnings beyond LOD 0.
 									bool bIgnoreLayoutWarnings = true;
-									mu::Ptr<mu::NodeLayout> LayoutNode = CreateMutableLayoutNode(GenerationContext, Layouts[i], bIgnoreLayoutWarnings, bWasEmpty);
+									mu::Ptr<mu::NodeLayout> LayoutNode = CreateMutableLayoutNode(GenerationContext, Layouts[LayoutIndex], bIgnoreLayoutWarnings, bWasEmpty);
 									if (bWasEmpty)
 									{
 										FString msg = "Mesh Column [" + MutableColumnName + "] Layout doesn't has any block. A grid sized block will be used instead.";
 										GenerationContext.Log(FText::FromString(msg), Node, EMessageSeverity::Warning);
 									}
 
-									MeshTableNode->SetLayout(i, LayoutNode);
+									MeshTableNode->SetLayout(LayoutIndex, LayoutNode);
 								}
 							}
 
