[mutable] Remove false warning with remove mesh nodes. Ensure modifier layout previews are detected correctly with internal tags.
#rnx

diff --git a/Engine/Plugins/Mutable/Source/CustomizableObjectEditor/Private/MuCOE/Nodes/CustomizableObjectNodeModifierBase.cpp b/Engine/Plugins/Mutable/Source/CustomizableObjectEditor/Private/MuCOE/Nodes/CustomizableObjectNodeModifierBase.cpp
--- a/Engine/Plugins/Mutable/Source/CustomizableObjectEditor/Private/MuCOE/Nodes/CustomizableObjectNodeModifierBase.cpp
+++ b/Engine/Plugins/Mutable/Source/CustomizableObjectEditor/Private/MuCOE/Nodes/CustomizableObjectNodeModifierBase.cpp
@@ -33,13 +33,15 @@ bool UCustomizableObjectNodeModifierBase::IsApplicableTo(UCustomizableObjectNode
 	const TArray<FString>* EnabledTags = Candidate->GetEnableTags();
 	if (EnabledTags)
 	{
+		FString InternalTag = Candidate->GetInternalTag();
+
 		switch ( MultipleTagPolicy )
 		{
 		case EMutableMultipleTagPolicy::OnlyOneRequired:
 		{
 			for (const FString& RequiredTag : RequiredTags)
 			{
-				if (EnabledTags->Contains(RequiredTag))
+				if (InternalTag==RequiredTag || EnabledTags->Contains(RequiredTag))
 				{
 					return true;
 				}
@@ -51,7 +53,7 @@ bool UCustomizableObjectNodeModifierBase::IsApplicableTo(UCustomizableObjectNode
 		{
 			for (const FString& RequiredTag : RequiredTags)
 			{
-				if (!EnabledTags->Contains(RequiredTag))
+				if (InternalTag != RequiredTag && !EnabledTags->Contains(RequiredTag))
 				{
 					return false;
 				}
diff --git a/Engine/Plugins/Mutable/Source/MutableTools/Private/MuT/CodeGenerator.cpp b/Engine/Plugins/Mutable/Source/MutableTools/Private/MuT/CodeGenerator.cpp
--- a/Engine/Plugins/Mutable/Source/MutableTools/Private/MuT/CodeGenerator.cpp
+++ b/Engine/Plugins/Mutable/Source/MutableTools/Private/MuT/CodeGenerator.cpp
@@ -2634,21 +2634,20 @@ namespace mu
 	{
 		int32 CurrentLOD = CurrentParents.Last().Lod;
 
-		for (const FirstPassGenerator::FModifier& m : Modifiers)
+		for (const FirstPassGenerator::FModifier& Mod : Modifiers)
 		{
 			// A mistake in the surface edit modifier usually results in no change visible. Try to detect it.
-			if (m.Node->GetType() == NodeModifierSurfaceEdit::GetStaticType())
+			if (Mod.Node->GetType() == NodeModifierSurfaceEdit::GetStaticType())
 			{
-				const NodeModifierSurfaceEdit* Edit = static_cast<const NodeModifierSurfaceEdit*>(m.Node);
+				const NodeModifierSurfaceEdit* Edit = static_cast<const NodeModifierSurfaceEdit*>(Mod.Node);
 
 				bool bAffectsCurrentLOD = Edit->LODs.IsValidIndex(CurrentLOD);
-
 				if (!bAffectsCurrentLOD)
 				{
 					continue;
 				}
 
-				if (Node.Images.IsEmpty())
+				if (Node.Images.IsEmpty() || Edit->LODs[CurrentLOD].Textures.IsEmpty())
 				{
 					continue;
 				}
