[mutable] Fix possible crash when compiling degenerated COs.
#rnx

diff --git a/Engine/Plugins/Mutable/Source/MutableTools/Private/MuT/CodeGenerator_Mesh.cpp b/Engine/Plugins/Mutable/Source/MutableTools/Private/MuT/CodeGenerator_Mesh.cpp
--- a/Engine/Plugins/Mutable/Source/MutableTools/Private/MuT/CodeGenerator_Mesh.cpp
+++ b/Engine/Plugins/Mutable/Source/MutableTools/Private/MuT/CodeGenerator_Mesh.cpp
@@ -1368,7 +1368,10 @@ namespace mu
             Ptr<ASTOpConstantResource> cop = new ASTOpConstantResource();
             cop->Type = OP_TYPE::ME_CONSTANT;
             cop->SetValue( FormatMesh, CompilerOptions->OptimisationOptions.DiskCacheContext );
-			cop->SourceDataDescriptor = baseResult.BaseMeshOp->GetSourceDataDescriptor();
+			if (baseResult.BaseMeshOp)
+			{
+				cop->SourceDataDescriptor = baseResult.BaseMeshOp->GetSourceDataDescriptor();
+			}
             op->Format = cop;
 
             OutResult.MeshOp = op;
