Mutable: Add a modifier to transform vertices contained within another mesh.

diff --git a/Engine/Plugins/Mutable/Source/CustomizableObject/Internal/MuCO/CustomizableObjectInstanceDescriptor.h b/Engine/Plugins/Mutable/Source/CustomizableObject/Internal/MuCO/CustomizableObjectInstanceDescriptor.h
--- a/Engine/Plugins/Mutable/Source/CustomizableObject/Internal/MuCO/CustomizableObjectInstanceDescriptor.h
+++ b/Engine/Plugins/Mutable/Source/CustomizableObject/Internal/MuCO/CustomizableObjectInstanceDescriptor.h
@@ -98,6 +98,8 @@ struct CUSTOMIZABLEOBJECT_API FCustomizableObjectInstanceDescriptor
 
 	const TArray<FCustomizableObjectProjectorParameterValue>& GetProjectorParameters() const;
 	
+	const TArray<FCustomizableObjectTransformParameterValue>& GetTransformParameters() const;
+	
 	/** Return true if there are any parameters. */
 	bool HasAnyParameters() const;
 
@@ -128,6 +130,12 @@ struct CUSTOMIZABLEOBJECT_API FCustomizableObjectInstanceDescriptor
 	/** Sets the color value "ColorValue" of a color parameter with index "ColorParamIndex". */
 	void SetColorParameterSelectedOption(const FString& ColorParamName, const FLinearColor& ColorValue);
 
+	/** Gets the value of a transform parameter with name "TransformParamName". */
+	FTransform GetTransformParameterSelectedOption(const FString& TransformParamName) const;
+
+	/** Sets the transform value "TransformValue" of a transform parameter with name "TransformParamName". */
+	void SetTransformParameterSelectedOption(const FString& TransformParamName, const FTransform& TransformValue);
+
 	/** Gets the value of the bool parameter with name "BoolParamName". */
 	bool GetBoolParameterSelectedOption(const FString& BoolParamName) const;
 
@@ -319,6 +327,9 @@ private:
 	UPROPERTY()
 	TArray<FCustomizableObjectProjectorParameterValue> ProjectorParameters;
 
+	UPROPERTY()
+	TArray<FCustomizableObjectTransformParameterValue> TransformParameters;
+
 	/** Mutable parameters optimization state. Transient UProperty to make it transactable. */
 	UPROPERTY(Transient)
 	int32 State = 0;
diff --git a/Engine/Plugins/Mutable/Source/CustomizableObject/Internal/MuCO/CustomizableObjectPrivate.h b/Engine/Plugins/Mutable/Source/CustomizableObject/Internal/MuCO/CustomizableObjectPrivate.h
--- a/Engine/Plugins/Mutable/Source/CustomizableObject/Internal/MuCO/CustomizableObjectPrivate.h
+++ b/Engine/Plugins/Mutable/Source/CustomizableObject/Internal/MuCO/CustomizableObjectPrivate.h
@@ -1482,6 +1482,8 @@ public:
 		ReGroupRomsBySource,
 
 		UIMetadataGameplayTags,
+
+		TransformInMeshModifier,
 		
 		// -----<new versions can be added above this line>--------
 		LastCustomizableObjectVersion
diff --git a/Engine/Plugins/Mutable/Source/CustomizableObject/Private/MuCO/CustomizableObject.cpp b/Engine/Plugins/Mutable/Source/CustomizableObject/Private/MuCO/CustomizableObject.cpp
--- a/Engine/Plugins/Mutable/Source/CustomizableObject/Private/MuCO/CustomizableObject.cpp
+++ b/Engine/Plugins/Mutable/Source/CustomizableObject/Private/MuCO/CustomizableObject.cpp
@@ -1600,6 +1600,12 @@ void UCustomizableObjectPrivate::UpdateParameterPropertiesFromModel(const TShare
 				break;
 			}
 
+			case mu::PARAMETER_TYPE::T_MATRIX:
+			{
+				Data.Type = EMutableParameterType::Transform;
+				break;
+			}
+				
 			case mu::PARAMETER_TYPE::T_IMAGE:
 			{
 				Data.Type = EMutableParameterType::Texture;
@@ -1908,6 +1914,28 @@ FLinearColor UCustomizableObject::GetColorParameterDefaultValue(const FString& I
 }
 
 
+FTransform UCustomizableObject::GetTransformParameterDefaultValue(const FString& InParameterName) const
+{
+	const int32 ParameterIndex = FindParameter(InParameterName);
+	if (ParameterIndex == INDEX_NONE)
+	{
+		UE_LOG(LogMutable, Error, TEXT("Tried to access the default value of the nonexistent color parameter [%s] in the CustomizableObject [%s]."), *InParameterName, *GetName());
+		return FCustomizableObjectTransformParameterValue::DEFAULT_PARAMETER_VALUE;
+	}
+
+	const TSharedPtr<const mu::Model>& Model = GetPrivate()->GetModel();
+	if (!Model)
+	{
+		checkNoEntry();
+		return FCustomizableObjectTransformParameterValue::DEFAULT_PARAMETER_VALUE;
+	}
+	
+	const FMatrix44f Matrix = Model->GetMatrixDefaultValue(ParameterIndex);
+
+	return FTransform(FMatrix(Matrix));
+}
+
+
 void UCustomizableObject::GetProjectorParameterDefaultValue(const FString& InParameterName, FVector3f& OutPos,
 	FVector3f& OutDirection, FVector3f& OutUp, FVector3f& OutScale, float& OutAngle,
 	ECustomizableObjectProjectorType& OutType) const
diff --git a/Engine/Plugins/Mutable/Source/CustomizableObject/Private/MuCO/CustomizableObjectInstance.cpp b/Engine/Plugins/Mutable/Source/CustomizableObject/Private/MuCO/CustomizableObjectInstance.cpp
--- a/Engine/Plugins/Mutable/Source/CustomizableObject/Private/MuCO/CustomizableObjectInstance.cpp
+++ b/Engine/Plugins/Mutable/Source/CustomizableObject/Private/MuCO/CustomizableObjectInstance.cpp
@@ -2565,11 +2565,21 @@ void UCustomizableObjectInstance::SetVectorParameterSelectedOption(const FString
 	Descriptor.SetVectorParameterSelectedOption(VectorParamName, VectorValue);
 }
 
+FTransform UCustomizableObjectInstance::GetTransformParameterSelectedOption(const FString& TransformParamName) const
+{
+	return Descriptor.GetTransformParameterSelectedOption(TransformParamName);
+}
+
+void UCustomizableObjectInstance::SetTransformParameterSelectedOption(const FString& TransformParamName, const FTransform& TransformValue)
+{
+	Descriptor.SetTransformParameterSelectedOption(TransformParamName, TransformValue);
+}
+
 
 void UCustomizableObjectInstance::SetProjectorValue(const FString& ProjectorParamName,
-	const FVector& Pos, const FVector& Direction, const FVector& Up, const FVector& Scale,
-	const float Angle,
-	const int32 RangeIndex)
+                                                    const FVector& Pos, const FVector& Direction, const FVector& Up, const FVector& Scale,
+                                                    const float Angle,
+                                                    const int32 RangeIndex)
 {
 	Descriptor.SetProjectorValue(ProjectorParamName, Pos, Direction, Up, Scale, Angle, RangeIndex);
 }
@@ -2746,6 +2756,7 @@ bool UCustomizableInstancePrivate::LoadParametersFromProfile(int32 ProfileIndex)
 	GetPublic()->Descriptor.TextureParameters = Profile.TextureParameters;
 	GetPublic()->Descriptor.ProjectorParameters = Profile.ProjectorParameters;
 	GetPublic()->Descriptor.VectorParameters = Profile.VectorParameters;
+	GetPublic()->Descriptor.TransformParameters = Profile.TransformParameters;
 #endif
 	return true;
 
@@ -2775,6 +2786,7 @@ bool UCustomizableInstancePrivate::SaveParametersToProfile(int32 ProfileIndex)
 	Profile.TextureParameters = GetPublic()->Descriptor.TextureParameters;
 	Profile.ProjectorParameters = GetPublic()->Descriptor.ProjectorParameters;
 	Profile.VectorParameters = GetPublic()->Descriptor.VectorParameters;
+	Profile.TransformParameters = GetPublic()->Descriptor.TransformParameters;
 #endif
 	return true;
 }
@@ -2803,7 +2815,7 @@ bool UCustomizableInstancePrivate::MigrateProfileParametersToCurrentInstance(int
 	TempProfile.ProjectorParameters = GetPublic()->Descriptor.ProjectorParameters;
 	TempProfile.TextureParameters = GetPublic()->Descriptor.TextureParameters;
 	TempProfile.VectorParameters = GetPublic()->Descriptor.VectorParameters;
-	
+	TempProfile.TransformParameters = GetPublic()->Descriptor.TransformParameters;
 
 	// Populate TempProfile with the parameters found in the profile.
 	// Any profile parameter missing will be discarded.
diff --git a/Engine/Plugins/Mutable/Source/CustomizableObject/Private/MuCO/CustomizableObjectInstanceDescriptor.cpp b/Engine/Plugins/Mutable/Source/CustomizableObject/Private/MuCO/CustomizableObjectInstanceDescriptor.cpp
--- a/Engine/Plugins/Mutable/Source/CustomizableObject/Private/MuCO/CustomizableObjectInstanceDescriptor.cpp
+++ b/Engine/Plugins/Mutable/Source/CustomizableObject/Private/MuCO/CustomizableObjectInstanceDescriptor.cpp
@@ -190,6 +190,22 @@ void FCustomizableObjectInstanceDescriptor::SaveDescriptor(FArchive& Ar, bool bU
 			break;
 		}
 
+		case EMutableParameterType::Transform:
+		{
+			FTransform Value(FTransform::Identity);
+			for (const FCustomizableObjectTransformParameterValue& P : TransformParameters)
+			{
+				if (P.ParameterName == Name)
+				{
+					Value = P.ParameterValue;
+					break;
+				}
+			}
+			Ar << Value;
+
+			break;
+		}
+
 		case EMutableParameterType::Texture:
 		{
 			FName Value;
@@ -387,6 +403,22 @@ void FCustomizableObjectInstanceDescriptor::LoadDescriptor(FArchive& Ar)
 			break;
 		}
 
+		case EMutableParameterType::Transform:
+		{
+			FTransform Value(FTransform::Identity);
+			Ar << Value;
+			for (FCustomizableObjectTransformParameterValue& P : TransformParameters)
+			{
+				if (P.ParameterName == Name)
+				{
+					P.ParameterValue = Value;
+					break;
+				}
+			}
+
+			break;
+		}
+			
 		case EMutableParameterType::Texture:
 		{
 			FName Value;
@@ -590,6 +622,21 @@ mu::ParametersPtr FCustomizableObjectInstanceDescriptor::GetParameters() const
 			break;
 		}
 
+		case mu::PARAMETER_TYPE::T_MATRIX:
+			{
+				for (const FCustomizableObjectTransformParameterValue& TransformParameter : TransformParameters)
+				{
+					if (TransformParameter.ParameterName == Name || (Uid.IsValid() && TransformParameter.Id == Uid))
+					{
+						MutableParameters->SetMatrixValue(ParamIndex, FMatrix44f(TransformParameter.ParameterValue.ToMatrixWithScale()));
+
+						break;
+					}
+				}
+
+				break;
+			}
+			
 		case mu::PARAMETER_TYPE::T_PROJECTOR:
 		{
 			for (const auto& ProjectorParameter : ProjectorParameters)
@@ -724,6 +771,7 @@ void FCustomizableObjectInstanceDescriptor::ReloadParameters()
 	TArray<FCustomizableObjectTextureParameterValue> OldTextureParameters = TextureParameters;
 	TArray<FCustomizableObjectVectorParameterValue> OldVectorParameters = VectorParameters;
 	TArray<FCustomizableObjectProjectorParameterValue> OldProjectorParameters = ProjectorParameters;
+	TArray<FCustomizableObjectTransformParameterValue> OldTransformParameters = TransformParameters;
 
 	BoolParameters.Reset();
 	IntParameters.Reset();
@@ -731,6 +779,7 @@ void FCustomizableObjectInstanceDescriptor::ReloadParameters()
 	TextureParameters.Reset();
 	VectorParameters.Reset();
 	ProjectorParameters.Reset();
+	TransformParameters.Reset();
 	
 	if (!CustomizableObject->GetPrivate()->GetModel())
 	{
@@ -953,6 +1002,32 @@ void FCustomizableObjectInstanceDescriptor::ReloadParameters()
 			break;
 		}
 
+		case mu::PARAMETER_TYPE::T_MATRIX:
+			{
+				FCustomizableObjectTransformParameterValue Param;
+				Param.ParameterName = Name;
+				Param.Id = Uid;
+
+				auto FindByNameAndUid = [&](const FCustomizableObjectTransformParameterValue& P)
+				{
+					return P.ParameterName == Name || (Uid.IsValid() && P.Id == Uid);
+				};
+
+				if (FCustomizableObjectTransformParameterValue* Result = OldTransformParameters.FindByPredicate(FindByNameAndUid))
+				{	
+					Param.ParameterValue = Result->ParameterValue;
+				}
+				else // Not found in Instance Parameters. Use Mutable Parameters.
+				{
+					FMatrix44f Matrix;
+					MutableParameters->GetMatrixValue(ParamIndex, Matrix);
+					Param.ParameterValue = FTransform(FMatrix(Matrix));
+				}
+
+				TransformParameters.Add(Param);
+				break;
+			}
+			
 		case mu::PARAMETER_TYPE::T_PROJECTOR:
 		{
 			FCustomizableObjectProjectorParameterValue Param;
@@ -1160,6 +1235,11 @@ const TArray<FCustomizableObjectProjectorParameterValue>& FCustomizableObjectIns
 	return ProjectorParameters;	
 }
 
+const TArray<FCustomizableObjectTransformParameterValue>& FCustomizableObjectInstanceDescriptor::GetTransformParameters() const
+{
+	return TransformParameters;
+}
+
 
 bool FCustomizableObjectInstanceDescriptor::HasAnyParameters() const
 {
@@ -1168,6 +1248,7 @@ bool FCustomizableObjectInstanceDescriptor::HasAnyParameters() const
 		!FloatParameters.IsEmpty() || 
 		!TextureParameters.IsEmpty() ||
 		!ProjectorParameters.IsEmpty() ||
+		!TransformParameters.IsEmpty() ||
 		!VectorParameters.IsEmpty();
 }
 
@@ -1518,6 +1599,49 @@ void FCustomizableObjectInstanceDescriptor::SetColorParameterSelectedOption(cons
 	SetVectorParameterSelectedOption(ColorParamName, ColorValue);
 }
 
+FTransform FCustomizableObjectInstanceDescriptor::GetTransformParameterSelectedOption(const FString& TransformParamName) const
+{
+	if (!CustomizableObject)
+	{
+		CustomizableObjectNullErrorMessage();
+		return FCustomizableObjectTransformParameterValue::DEFAULT_PARAMETER_VALUE;
+	}
+
+	const int32 ParameterIndexInObject = CustomizableObject->FindParameter(TransformParamName);
+	const int32 TransformParamIndex = FindTypedParameterIndex(TransformParamName, EMutableParameterType::Transform);
+
+	if (TransformParamIndex == INDEX_NONE)
+	{
+		LogParameterNotFoundWarning(TransformParamName, ParameterIndexInObject, TransformParamIndex, CustomizableObject, __FUNCTION__);
+		return FCustomizableObjectTransformParameterValue::DEFAULT_PARAMETER_VALUE;
+	}
+
+	return TransformParameters.IsValidIndex(TransformParamIndex) ? TransformParameters[TransformParamIndex].ParameterValue : FTransform::Identity;
+}
+
+void FCustomizableObjectInstanceDescriptor::SetTransformParameterSelectedOption(const FString& TransformParamName, const FTransform& TransformValue)
+{
+	if (!CustomizableObject)
+	{
+		CustomizableObjectNullErrorMessage();
+		return;
+	}
+
+	RETURN_ON_UNCOMPILED_CO(CustomizableObject, TEXT("Error: Cannot set Transform parameter "));
+
+	const int32 ParameterIndexInObject = CustomizableObject->FindParameter(TransformParamName);
+	const int32 ParameterIndexInInstance = FindTypedParameterIndex(TransformParamName, EMutableParameterType::Transform);
+
+	if (ParameterIndexInObject < 0 || ParameterIndexInInstance < 0)
+	{
+		// Early out since we could not find the parameter to set.
+		LogParameterNotFoundWarning(TransformParamName, ParameterIndexInObject, ParameterIndexInInstance, CustomizableObject, __FUNCTION__);
+		return;
+	}
+
+	TransformParameters[ParameterIndexInInstance].ParameterValue = TransformValue;
+}
+
 
 bool FCustomizableObjectInstanceDescriptor::GetBoolParameterSelectedOption(const FString& BoolParamName) const
 {
@@ -2393,6 +2517,13 @@ void FCustomizableObjectInstanceDescriptor::SetDefaultValue(int32 ParamIndex)
 		break;
 	}
 
+	case EMutableParameterType::Transform:
+	{
+		const FTransform DefaultValue = CustomizableObject->GetTransformParameterDefaultValue(ParamName);
+		TransformParameters[TypedIndex].ParameterValue = DefaultValue;
+		break;
+	}
+		
 	case EMutableParameterType::Projector:
 	{
 		const FCustomizableObjectProjector DefaultValue = CustomizableObject->GetProjectorParameterDefaultValue(ParamName);
diff --git a/Engine/Plugins/Mutable/Source/CustomizableObject/Private/MuCO/DescriptorHash.cpp b/Engine/Plugins/Mutable/Source/CustomizableObject/Private/MuCO/DescriptorHash.cpp
--- a/Engine/Plugins/Mutable/Source/CustomizableObject/Private/MuCO/DescriptorHash.cpp
+++ b/Engine/Plugins/Mutable/Source/CustomizableObject/Private/MuCO/DescriptorHash.cpp
@@ -44,6 +44,11 @@ FDescriptorHash::FDescriptorHash(const FCustomizableObjectInstanceDescriptor& De
 		Hash = HashCombine(Hash, GetTypeHash(Value));
 	}
 	
+	for (const FCustomizableObjectTransformParameterValue& Value : Descriptor.TransformParameters)
+	{
+		Hash = HashCombine(Hash, GetTypeHash(Value));
+	}
+
 	Hash = HashCombine(Hash, GetTypeHash(Descriptor.State));
 	Hash = HashCombine(Hash, GetTypeHash(Descriptor.GetBuildParameterRelevancy()));
 	
diff --git a/Engine/Plugins/Mutable/Source/CustomizableObject/Public/MuCO/CustomizableObject.h b/Engine/Plugins/Mutable/Source/CustomizableObject/Public/MuCO/CustomizableObject.h
--- a/Engine/Plugins/Mutable/Source/CustomizableObject/Public/MuCO/CustomizableObject.h
+++ b/Engine/Plugins/Mutable/Source/CustomizableObject/Public/MuCO/CustomizableObject.h
@@ -96,6 +96,10 @@ struct FProfileParameterDat
 	////
 	UPROPERTY()
 	TArray<FCustomizableObjectProjectorParameterValue> ProjectorParameters;
+	
+	////
+	UPROPERTY()
+	TArray<FCustomizableObjectTransformParameterValue> TransformParameters;
 };
 
 // A USTRUCT version of FMeshToMeshVertData in SkeletalMeshTypes.h
@@ -437,6 +441,13 @@ public:
 	UFUNCTION(BlueprintCallable, Category = CustomizableObject)
 	FLinearColor GetColorParameterDefaultValue(UPARAM(DisplayName = "Parameter Name") const FString& InParameterName) const;
 	
+	/** Get the default value of a parameter of type Transform.
+	  * @param InParameterName The name of the Transform parameter to get the default value of.
+	  * @return The default value of the provided parameter name. */
+	UFUNCTION(BlueprintCallable, Category = CustomizableObject)
+	FTransform GetTransformParameterDefaultValue(UPARAM(DisplayName = "Parameter Name") const FString& InParameterName) const;
+
+	
 	/** Get the default value of a parameter of type Projector.
 	  * @param InParameterName The name of the Projector parameter to get the default value of.
 	  * @param OutPos The default position of the Projector.
diff --git a/Engine/Plugins/Mutable/Source/CustomizableObject/Public/MuCO/CustomizableObjectInstance.h b/Engine/Plugins/Mutable/Source/CustomizableObject/Public/MuCO/CustomizableObjectInstance.h
--- a/Engine/Plugins/Mutable/Source/CustomizableObject/Public/MuCO/CustomizableObjectInstance.h
+++ b/Engine/Plugins/Mutable/Source/CustomizableObject/Public/MuCO/CustomizableObjectInstance.h
@@ -159,7 +159,7 @@ enum class EUpdateRequired : uint8
 
 
 /** Instance Update Result. */
-UENUM()
+UENUM(BlueprintType)
 enum class EUpdateResult : uint8
 {
 	Success, // Update finished without issues.
@@ -184,12 +184,12 @@ enum class ESkeletalMeshStatus : uint8
 
 /** Instance Update Context.
  * Used to avoid changing the delegate signature in the future.  */
-USTRUCT()
+USTRUCT(BlueprintType)
 struct FUpdateContext
 {	
 	GENERATED_BODY()
 
-	UPROPERTY()
+	UPROPERTY(BlueprintReadOnly, Category=UpdateResult)
 	EUpdateResult UpdateResult = EUpdateResult::Success;
 };
 
@@ -414,6 +414,15 @@ public:
 	UFUNCTION(BlueprintCallable, Category = CustomizableObjectInstance)
 	void SetVectorParameterSelectedOption(const FString& VectorParamName, const FLinearColor& VectorValue);
 
+	/** Gets the value of a transform parameter with name "TransformParamName". */
+	UFUNCTION(BlueprintCallable, Category = CustomizableObjectInstance)
+	FTransform GetTransformParameterSelectedOption(const FString& TransformParamName) const;
+
+	/** Sets the transform value "TransformValue" of a transform parameter with name "TransformParamName". */
+	UFUNCTION(BlueprintCallable, Category = CustomizableObjectInstance)
+	void SetTransformParameterSelectedOption(const FString& TransformParamName, const FTransform& TransformValue);
+	
+	
 	/** Sets the projector values of a projector parameter with index "ProjectorParamIndex". */
 	UFUNCTION(BlueprintCallable, Category = CustomizableObjectInstance)
 	void SetProjectorValue(const FString& ProjectorParamName,
diff --git a/Engine/Plugins/Mutable/Source/CustomizableObject/Public/MuCO/CustomizableObjectParameterTypeDefinitions.h b/Engine/Plugins/Mutable/Source/CustomizableObject/Public/MuCO/CustomizableObjectParameterTypeDefinitions.h
--- a/Engine/Plugins/Mutable/Source/CustomizableObject/Public/MuCO/CustomizableObjectParameterTypeDefinitions.h
+++ b/Engine/Plugins/Mutable/Source/CustomizableObject/Public/MuCO/CustomizableObjectParameterTypeDefinitions.h
@@ -17,7 +17,8 @@ enum class EMutableParameterType : uint8
 	Float		UMETA(DisplayName = "Float"),
 	Color		UMETA(DisplayName = "Color"),
 	Projector	UMETA(DisplayName = "Projector"),
-	Texture		UMETA(DisplayName = "Texture")
+	Texture		UMETA(DisplayName = "Texture"),
+	Transform	UMETA(DisplayName = "Transform"),
 };
 
 
@@ -216,6 +217,34 @@ inline uint32 GetTypeHash(const FCustomizableObjectVectorParameterValue& Key)
 }
 
 
+USTRUCT(BlueprintType)
+struct FCustomizableObjectTransformParameterValue
+{
+	GENERATED_USTRUCT_BODY()
+
+	inline static const FTransform DEFAULT_PARAMETER_VALUE = FTransform::Identity;
+	
+	UPROPERTY(Category = CustomizableObjectTransformParameterValue, VisibleAnywhere)
+	FString ParameterName;
+
+	UPROPERTY(Category = CustomizableObjectTransformParameterValue, VisibleAnywhere)
+	FTransform ParameterValue = FTransform::Identity;
+
+	UPROPERTY(Category = CustomizableObjectTransformParameterValue, VisibleAnywhere)
+	FGuid Id;
+};
+
+
+inline uint32 GetTypeHash(const FCustomizableObjectTransformParameterValue& Key)
+{
+	uint32 Hash = GetTypeHash(Key.ParameterName);
+	Hash = HashCombine(Hash, GetTypeHash(Key.ParameterValue));
+
+	return Hash;
+}
+
+
+
 UENUM()
 enum class ECustomizableObjectProjectorType : uint8
 {
diff --git a/Engine/Plugins/Mutable/Source/CustomizableObjectEditor/Private/MuCOE/CustomizableInstanceDetails.cpp b/Engine/Plugins/Mutable/Source/CustomizableObjectEditor/Private/MuCOE/CustomizableInstanceDetails.cpp
--- a/Engine/Plugins/Mutable/Source/CustomizableObjectEditor/Private/MuCOE/CustomizableInstanceDetails.cpp
+++ b/Engine/Plugins/Mutable/Source/CustomizableObjectEditor/Private/MuCOE/CustomizableInstanceDetails.cpp
@@ -14,6 +14,7 @@
 #include "ScopedTransaction.h"
 #include "Serialization/BufferArchive.h"
 #include "Toolkits/ToolkitManager.h"
+#include "Math/Transform.h"
 
 #include "MuCO/CustomizableObjectInstancePrivate.h"
 #include "MuCO/CustomizableObject.h"
@@ -33,7 +34,10 @@
 #include "Widgets/Input/SSpinBox.h"
 #include "Widgets/Input/SSlider.h"
 #include "Widgets/Input/STextComboBox.h"
+#include "Widgets/Input/SRotatorInputBox.h"
+#include "Widgets/Input/SVectorInputBox.h"
 #include "Widgets/Layout/SExpandableArea.h"
+#include "Widgets/Layout/SGridPanel.h"
 #include "Widgets/Layout/SUniformGridPanel.h"
 
 class UObject;
@@ -946,6 +950,10 @@ TSharedRef<SWidget> FCustomizableInstanceDetails::GenerateParameterWidget(const
 		{
 			return GenerateTextureWidget(CustomizableObject, ParamName);
 		}
+		case EMutableParameterType::Transform:
+		{
+			return GenerateTransformWidget(ParamName);
+		}
 		case EMutableParameterType::Projector:
 		{
 			bool bIsParamMultidimensional = CustomInstance->GetCustomizableObject()->IsParameterMultidimensional(ParamIndexInObject);
@@ -1419,6 +1427,203 @@ void FCustomizableInstanceDetails::OnSetColorFromColorPicker(FLinearColor NewCol
 	UpdateInstance();
 }
 
+// TRANSFORM PARAMETERS -----------------------------------------------------------------------------------------------------------------
+
+TSharedRef<SWidget> FCustomizableInstanceDetails::GenerateTransformWidget(const FString& ParamName)
+{
+	auto OnLocationChanged = [this, ParamName](double Value, EAxis::Type Axis)
+	{
+		FTransform Transform = GetTransformParameterValue(ParamName);
+		FVector Location = Transform.GetLocation();
+		if (!FMath::IsNearlyEqualByULP(Value, Location.GetComponentForAxis(Axis)))
+		{
+			Location.SetComponentForAxis(Axis, Value);
+			Transform.SetLocation(Location);
+			OnTransformParameterChanged(Transform, ParamName);
+		}
+	};
+	auto OnLocationCommitted = [this, ParamName](double Value, ETextCommit::Type Type, EAxis::Type Axis)
+	{
+		FTransform Transform = GetTransformParameterValue(ParamName);
+		FVector Location = Transform.GetLocation();
+		if (!FMath::IsNearlyEqualByULP(Value, Location.GetComponentForAxis(Axis)))
+		{
+			Location.SetComponentForAxis(Axis, Value);
+			Transform.SetLocation(Location);
+			OnTransformParameterCommitted(Transform, Type, ParamName);
+		}
+	};
+
+	auto OnRotationChanged = [this, ParamName](double Value, EAxis::Type Axis)
+	{
+		FTransform Transform = GetTransformParameterValue(ParamName);
+		FRotator Rotation = Transform.Rotator();
+		Rotation.SetComponentForAxis(Axis, Value);
+		
+		if (!Transform.Rotator().Equals(Rotation))
+		{
+			Transform.SetRotation(Rotation.Quaternion());
+			OnTransformParameterChanged(Transform, ParamName);
+		}
+	};
+	auto OnRotationCommitted = [this, ParamName](double Value, ETextCommit::Type Type, EAxis::Type Axis)
+	{
+		FTransform Transform = GetTransformParameterValue(ParamName);
+		FRotator Rotation = Transform.Rotator();
+		Rotation.SetComponentForAxis(Axis, Value);
+		
+		if (!Transform.Rotator().Equals(Rotation))
+		{
+			Transform.SetRotation(Rotation.Quaternion());
+			OnTransformParameterCommitted(Transform, Type, ParamName);
+		}
+	};
+	
+	auto OnScaleChanged = [this, ParamName](double Value, EAxis::Type Axis)
+	{
+		FTransform Transform = GetTransformParameterValue(ParamName);
+		FVector Scale = Transform.GetScale3D();
+		if (!FMath::IsNearlyEqualByULP(Value, Scale.GetComponentForAxis(Axis)))
+		{
+			Scale.SetComponentForAxis(Axis, Value);
+			Transform.SetScale3D(Scale);
+			OnTransformParameterChanged(Transform, ParamName);
+		}
+	};
+	auto OnScaleCommitted = [this, ParamName](double Value, ETextCommit::Type Type, EAxis::Type Axis)
+	{
+		FTransform Transform = GetTransformParameterValue(ParamName);
+		FVector Scale = Transform.GetScale3D();
+		if (!FMath::IsNearlyEqualByULP(Value, Scale.GetComponentForAxis(Axis)))
+		{
+			Scale.SetComponentForAxis(Axis, Value);
+			Transform.SetScale3D(Scale);
+			OnTransformParameterCommitted(Transform, Type, ParamName);
+		}
+	};
+	
+	
+	return SNew(SGridPanel)
+		.FillColumn(1, 1)
+		+ SGridPanel::Slot(0, 0)
+		.HAlign(HAlign_Right)
+		.VAlign(VAlign_Center)
+		[
+			SNew(STextBlock)
+			.Text(LOCTEXT("Transform_Location", "Location"))
+			.Font(IDetailLayoutBuilder::GetDetailFont())
+		]
+		+ SGridPanel::Slot(1, 0)
+		.Padding(4.0f, 0.0f, 0.0f, 0.0f)
+		.HAlign(HAlign_Fill)
+		[
+			SNew(SNumericVectorInputBox<double>)
+			.bColorAxisLabels(true)
+			.AllowSpin(true)
+			.X_Lambda([this, ParamName]() { return GetTransformParameterValue(ParamName).GetLocation().X; })
+			.Y_Lambda([this, ParamName]() { return GetTransformParameterValue(ParamName).GetLocation().Y; })
+			.Z_Lambda([this, ParamName]() { return GetTransformParameterValue(ParamName).GetLocation().Z; })
+			.OnXChanged_Lambda([OnLocationChanged](float Value) { OnLocationChanged(Value, EAxis::X); })
+			.OnYChanged_Lambda([OnLocationChanged](float Value) { OnLocationChanged(Value, EAxis::Y); })
+			.OnZChanged_Lambda([OnLocationChanged](float Value) { OnLocationChanged(Value, EAxis::Z); })
+			.OnXCommitted_Lambda([OnLocationCommitted](float Value, ETextCommit::Type Type) { OnLocationCommitted(Value, Type, EAxis::X); })
+			.OnYCommitted_Lambda([OnLocationCommitted](float Value, ETextCommit::Type Type) { OnLocationCommitted(Value, Type, EAxis::Y); })
+			.OnZCommitted_Lambda([OnLocationCommitted](float Value, ETextCommit::Type Type) { OnLocationCommitted(Value, Type, EAxis::Z); })
+			.Font(IDetailLayoutBuilder::GetDetailFont())
+		]
+		
+		+ SGridPanel::Slot(0, 1)
+		.HAlign(EHorizontalAlignment::HAlign_Right)
+		.VAlign(VAlign_Center)
+		[
+			SNew(STextBlock)
+			.Text(LOCTEXT("Transform_Rotation", "Rotation"))
+			.Font(IDetailLayoutBuilder::GetDetailFont())
+		]
+		+ SGridPanel::Slot(1, 1)
+		.HAlign(HAlign_Fill)
+		.Padding(4.0f, 0.0f, 0.0f, 0.0f)
+		[
+			SNew(SNumericRotatorInputBox<double>)
+			.bColorAxisLabels(true)
+			.AllowSpin(true)
+			.Roll_Lambda([this, ParamName]() { return GetTransformParameterValue(ParamName).Rotator().Roll; })
+			.Pitch_Lambda([this, ParamName]() { return GetTransformParameterValue(ParamName).Rotator().Pitch; })
+			.Yaw_Lambda([this, ParamName]() { return GetTransformParameterValue(ParamName).Rotator().Yaw; })
+			.OnRollChanged_Lambda([OnRotationChanged](float Value) { OnRotationChanged(Value, EAxis::X); })
+			.OnPitchChanged_Lambda([OnRotationChanged](float Value) { OnRotationChanged(Value, EAxis::Y); })
+			.OnYawChanged_Lambda([OnRotationChanged](float Value) { OnRotationChanged(Value, EAxis::Z); })
+			.OnRollCommitted_Lambda([OnRotationCommitted](float Value, ETextCommit::Type Type) { OnRotationCommitted(Value, Type, EAxis::X); })
+			.OnPitchCommitted_Lambda([OnRotationCommitted](float Value, ETextCommit::Type Type) { OnRotationCommitted(Value, Type, EAxis::Y); })
+			.OnYawCommitted_Lambda([OnRotationCommitted](float Value, ETextCommit::Type Type) { OnRotationCommitted(Value, Type, EAxis::Z); })
+			.Font(IDetailLayoutBuilder::GetDetailFont())
+		]
+
+		+ SGridPanel::Slot(0, 2)
+		.HAlign(EHorizontalAlignment::HAlign_Right)
+		.VAlign(VAlign_Center)
+		[
+			SNew(STextBlock)
+			.Text(LOCTEXT("Transform_Scale", "Scale"))
+			.Font(IDetailLayoutBuilder::GetDetailFont())
+		]
+		+ SGridPanel::Slot(1, 2)
+		.HAlign(HAlign_Fill)
+		.Padding(4.0f, 0.0f, 0.0f, 0.0f)
+		[
+			SNew(SNumericVectorInputBox<double>)
+			.bColorAxisLabels(true)
+			.AllowSpin(false)
+			.X_Lambda([this, ParamName]() { return GetTransformParameterValue(ParamName).GetScale3D().X; })
+			.Y_Lambda([this, ParamName]() { return GetTransformParameterValue(ParamName).GetScale3D().Y; })
+			.Z_Lambda([this, ParamName]() { return GetTransformParameterValue(ParamName).GetScale3D().Z; })
+
+			.OnXChanged_Lambda([OnScaleChanged](float Value) { OnScaleChanged(Value, EAxis::X); })
+			.OnYChanged_Lambda([OnScaleChanged](float Value) { OnScaleChanged(Value, EAxis::Y); })
+			.OnZChanged_Lambda([OnScaleChanged](float Value) { OnScaleChanged(Value, EAxis::Z); })
+			.OnXCommitted_Lambda([OnScaleCommitted](float Value, ETextCommit::Type Type) { OnScaleCommitted(Value, Type, EAxis::X); })
+			.OnYCommitted_Lambda([OnScaleCommitted](float Value, ETextCommit::Type Type) { OnScaleCommitted(Value, Type, EAxis::Y); })
+			.OnZCommitted_Lambda([OnScaleCommitted](float Value, ETextCommit::Type Type) { OnScaleCommitted(Value, Type, EAxis::Z); })
+			.Font(IDetailLayoutBuilder::GetDetailFont())
+		];
+}
+
+FTransform FCustomizableInstanceDetails::GetTransformParameterValue(const FString ParamName) const
+{
+	return CustomInstance->GetTransformParameterSelectedOption(ParamName);
+}
+
+void FCustomizableInstanceDetails::OnTransformParameterChanged(FTransform NewValue, const FString ParamName)
+{
+	FTransform OldValue = CustomInstance->GetTransformParameterSelectedOption(ParamName);
+
+	if (!OldValue.Equals(NewValue))
+	{
+		//No transaction is needed here as this is called when the transaction has already started
+		CustomInstance->SetTransformParameterSelectedOption(ParamName, NewValue);
+		UpdateInstance();
+	}
+}
+
+void FCustomizableInstanceDetails::OnTransformParameterCommitted(FTransform NewTransform, ETextCommit::Type Type, const FString ParamName)
+{
+	if (Type == ETextCommit::OnEnter)
+	{
+		if (!Transaction)
+		{
+			BeginTransaction(FText::Format(LOCTEXT("OnTransformParameterCommited", "Set Transform Parameter: {0}"), FText::FromString(ParamName)));
+		}
+
+		CustomInstance->SetTransformParameterSelectedOption(ParamName, NewTransform);
+		UpdateInstance();
+		EndTransaction();
+
+		// Non-continuous change: collect garbage.
+		GEngine->ForceGarbageCollection();
+	}
+
+}
+
 
 // PROJECTOR PARAMETERS -----------------------------------------------------------------------------------------------------------------
 
diff --git a/Engine/Plugins/Mutable/Source/CustomizableObjectEditor/Private/MuCOE/CustomizableInstanceDetails.h b/Engine/Plugins/Mutable/Source/CustomizableObjectEditor/Private/MuCOE/CustomizableInstanceDetails.h
--- a/Engine/Plugins/Mutable/Source/CustomizableObjectEditor/Private/MuCOE/CustomizableInstanceDetails.h
+++ b/Engine/Plugins/Mutable/Source/CustomizableObjectEditor/Private/MuCOE/CustomizableInstanceDetails.h
@@ -115,6 +115,13 @@ private:
 	FReply OnColorBlockMouseButtonDown(const FGeometry& MyGeometry, const FPointerEvent& MouseEvent, const FString ParamName);
 	void OnSetColorFromColorPicker(FLinearColor NewColor, const FString PickerParamName);
 
+	// Transform Parameters Functions
+	TSharedRef<SWidget> GenerateTransformWidget(const FString& ParamName);
+	FTransform GetTransformParameterValue(const FString ParamName) const;
+	void OnTransformParameterChanged(FTransform NewTransform, const FString ParamName);
+	void OnTransformParameterCommitted(FTransform NewTransform, ETextCommit::Type Type, const FString ParamName);
+
+
 	// Bool Parameters Functions
 	TSharedRef<SWidget> GenerateBoolWidget(const FString& ParamName);
 	ECheckBoxState GetBoolParameterValue(const FString ParamName) const;
diff --git a/Engine/Plugins/Mutable/Source/CustomizableObjectEditor/Private/MuCOE/CustomizableObjectEditor.cpp b/Engine/Plugins/Mutable/Source/CustomizableObjectEditor/Private/MuCOE/CustomizableObjectEditor.cpp
--- a/Engine/Plugins/Mutable/Source/CustomizableObjectEditor/Private/MuCOE/CustomizableObjectEditor.cpp
+++ b/Engine/Plugins/Mutable/Source/CustomizableObjectEditor/Private/MuCOE/CustomizableObjectEditor.cpp
@@ -45,6 +45,7 @@
 #include "MuCOE/Nodes/CustomizableObjectNodeMaterialVariation.h"
 #include "MuCOE/Nodes/CustomizableObjectNodeModifierClipMorph.h"
 #include "MuCOE/Nodes/CustomizableObjectNodeModifierClipWithMesh.h"
+#include "MuCOE/Nodes/CustomizableObjectNodeModifierTransformInMesh.h"
 #include "MuCOE/Nodes/CustomizableObjectNodeObject.h"
 #include "MuCOE/Nodes/CustomizableObjectNodeObjectGroup.h"
 #include "MuCOE/Nodes/CustomizableObjectNodeProjectorConstant.h"
@@ -1030,14 +1031,14 @@ void FCustomizableObjectEditor::HideGizmoClipMorph()
 }
 
 
-void FCustomizableObjectEditor::ShowGizmoClipMesh(UCustomizableObjectNodeModifierClipWithMesh& Node)
+void FCustomizableObjectEditor::ShowGizmoClipMesh(UCustomizableObjectNode& Node, FTransform* Transform, const UEdGraphPin& MeshPin)
 {
 	UObject* ClipMesh = nullptr;
 	int32 LODIndex = 0;
 	int32 SectionIndex = 0;
 	int32 MaterialSlotIndex = 0;
 
-	if (const UEdGraphPin* ConnectedPin = FollowInputPin(*Node.ClipMeshPin()))
+	if (const UEdGraphPin* ConnectedPin = FollowInputPin(MeshPin))
 	{
 		if (const UEdGraphNode* ConnectedNode = ConnectedPin->GetOwningNode())
 		{
@@ -1080,7 +1081,7 @@ void FCustomizableObjectEditor::ShowGizmoClipMesh(UCustomizableObjectNodeModifie
 
 		SelectSingleNode(Node);
 
-		Viewport->ShowGizmoClipMesh(Node, *ClipMesh, LODIndex, SectionIndex, MaterialSlotIndex);
+		Viewport->ShowGizmoClipMesh(Node, Transform, *ClipMesh, LODIndex, SectionIndex, MaterialSlotIndex);
 	}
 }
 
@@ -1100,7 +1101,8 @@ void FCustomizableObjectEditor::HideGizmoClipMesh()
 	for (FGraphPanelSelectionSet::TConstIterator NodeIt( SelectedNodes ); NodeIt; ++NodeIt)
 	{
 		const UObject* Node = *NodeIt;
-		if (Node->IsA<UCustomizableObjectNodeModifierClipWithMesh>())
+		if (Node->IsA<UCustomizableObjectNodeModifierClipWithMesh>() ||
+			Node->IsA<UCustomizableObjectNodeModifierTransformInMesh>())
 		{
 			GraphEditor->ClearSelectionSet();
 			break;
@@ -1711,7 +1713,17 @@ void FCustomizableObjectEditor::OnSelectedGraphNodesChanged(const FGraphPanelSel
 		}
 		else if (UCustomizableObjectNodeModifierClipWithMesh* NodeModifierClipWithMesh = Cast<UCustomizableObjectNodeModifierClipWithMesh>(Objects[0]))
 		{
-			ShowGizmoClipMesh(*NodeModifierClipWithMesh);
+			if (const UEdGraphPin* Pin = NodeModifierClipWithMesh->ClipMeshPin())
+			{
+				ShowGizmoClipMesh(*NodeModifierClipWithMesh, &NodeModifierClipWithMesh->Transform, *Pin);
+			}
+		}
+		else if (UCustomizableObjectNodeModifierTransformInMesh* NodeModifierTransformInMesh = Cast<UCustomizableObjectNodeModifierTransformInMesh>(Objects[0]))
+		{
+			if (const UEdGraphPin* Pin = NodeModifierTransformInMesh->BoundingMeshPin())
+			{
+				ShowGizmoClipMesh(*NodeModifierTransformInMesh, &NodeModifierTransformInMesh->BoundingMeshTransform, *Pin);
+			}
 		}
 		else if (UCustomizableObjectNodeProjectorConstant* NodeProjectorConstant = Cast<UCustomizableObjectNodeProjectorConstant>(Objects[0]))
 		{
diff --git a/Engine/Plugins/Mutable/Source/CustomizableObjectEditor/Private/MuCOE/CustomizableObjectEditor.h b/Engine/Plugins/Mutable/Source/CustomizableObjectEditor/Private/MuCOE/CustomizableObjectEditor.h
--- a/Engine/Plugins/Mutable/Source/CustomizableObjectEditor/Private/MuCOE/CustomizableObjectEditor.h
+++ b/Engine/Plugins/Mutable/Source/CustomizableObjectEditor/Private/MuCOE/CustomizableObjectEditor.h
@@ -208,7 +208,7 @@ public:
 	virtual void HideGizmoProjectorParameter() override;
 	virtual void ShowGizmoClipMorph(UCustomizableObjectNodeModifierClipMorph& Node) override;
 	virtual void HideGizmoClipMorph() override;
-	virtual void ShowGizmoClipMesh(UCustomizableObjectNodeModifierClipWithMesh& Node) override;
+	virtual void ShowGizmoClipMesh(UCustomizableObjectNode& Node, FTransform* Transform, const UEdGraphPin& MeshPin) override;
 	virtual void HideGizmoClipMesh() override;
 	virtual void ShowGizmoLight(ULightComponent& SelectedLight) override;
 	virtual void HideGizmoLight() override;
diff --git a/Engine/Plugins/Mutable/Source/CustomizableObjectEditor/Private/MuCOE/CustomizableObjectEditorModule.cpp b/Engine/Plugins/Mutable/Source/CustomizableObjectEditor/Private/MuCOE/CustomizableObjectEditorModule.cpp
--- a/Engine/Plugins/Mutable/Source/CustomizableObjectEditor/Private/MuCOE/CustomizableObjectEditorModule.cpp
+++ b/Engine/Plugins/Mutable/Source/CustomizableObjectEditor/Private/MuCOE/CustomizableObjectEditorModule.cpp
@@ -51,6 +51,7 @@
 #include "MuCOE/Nodes/CustomizableObjectNodeModifierClipWithMeshDetails.h"
 #include "MuCOE/Nodes/CustomizableObjectNodeModifierClipWithUVMask.h"
 #include "MuCOE/Nodes/CustomizableObjectNodeModifierClipDeform.h"
+#include "MuCOE/Nodes/CustomizableObjectNodeModifierTransformInMesh.h"
 #include "MuCOE/Nodes/CustomizableObjectNodeMeshMorph.h"
 #include "MuCOE/Nodes/CustomizableObjectNodeMeshMorphDetails.h"
 #include "MuCOE/Nodes/CustomizableObjectNodeMeshReshapeCommon.h"
@@ -83,6 +84,7 @@
 #include "MuCO/CustomizableObjectSystemPrivate.h"
 #include "MuCOE/CustomizableObjectGraph.h"
 #include "Nodes/CustomizableObjectNodeComponentMeshDetails.h"
+#include "Nodes/CustomizableObjectNodeModifierTransformInMeshDetails.h"
 
 class AActor;
 class FString;
@@ -257,6 +259,7 @@ void FCustomizableObjectEditorModule::StartupModule()
 	RegisterCustomDetails(PropertyModule, UCustomizableObjectNodeModifierClipWithMesh::StaticClass(), FOnGetDetailCustomizationInstance::CreateStatic(&FCustomizableObjectNodeModifierClipWithMeshDetails::MakeInstance));
 	RegisterCustomDetails(PropertyModule, UCustomizableObjectNodeModifierClipWithUVMask::StaticClass(), FOnGetDetailCustomizationInstance::CreateStatic(&FCustomizableObjectNodeModifierBaseDetails::MakeInstance));
 	RegisterCustomDetails(PropertyModule, UCustomizableObjectNodeModifierClipDeform::StaticClass(), FOnGetDetailCustomizationInstance::CreateStatic(&FCustomizableObjectNodeModifierBaseDetails::MakeInstance));
+	RegisterCustomDetails(PropertyModule, UCustomizableObjectNodeModifierTransformInMesh::StaticClass(), FOnGetDetailCustomizationInstance::CreateStatic(&FCustomizableObjectNodeModifierTransformInMeshDetails::MakeInstance));
 	RegisterCustomDetails(PropertyModule, UCustomizableObjectNodeObject::StaticClass(), FOnGetDetailCustomizationInstance::CreateStatic(&FCustomizableObjectNodeObjectDetails::MakeInstance));
 	RegisterCustomDetails(PropertyModule, UCustomizableObjectNodeObjectGroup::StaticClass(), FOnGetDetailCustomizationInstance::CreateStatic(&FCustomizableObjectNodeObjectGroupDetails::MakeInstance));
 	RegisterCustomDetails(PropertyModule, UCustomizableObjectNodeProjectorParameter::StaticClass(), FOnGetDetailCustomizationInstance::CreateStatic(&FCustomizableObjectNodeProjectorParameterDetails::MakeInstance));
diff --git a/Engine/Plugins/Mutable/Source/CustomizableObjectEditor/Private/MuCOE/CustomizableObjectEditorViewportClient.cpp b/Engine/Plugins/Mutable/Source/CustomizableObjectEditor/Private/MuCOE/CustomizableObjectEditorViewportClient.cpp
--- a/Engine/Plugins/Mutable/Source/CustomizableObjectEditor/Private/MuCOE/CustomizableObjectEditorViewportClient.cpp
+++ b/Engine/Plugins/Mutable/Source/CustomizableObjectEditor/Private/MuCOE/CustomizableObjectEditorViewportClient.cpp
@@ -599,19 +599,20 @@ void FCustomizableObjectEditorViewportClient::HideGizmoClipMorph()
 }
 
 
-void FCustomizableObjectEditorViewportClient::ShowGizmoClipMesh(UCustomizableObjectNodeModifierClipWithMesh& InClipMeshNode, UObject& ClipMesh, int32 LODIndex, int32 SectionIndex, int32 MaterialSlotIndex)
+void FCustomizableObjectEditorViewportClient::ShowGizmoClipMesh(UCustomizableObjectNode& InClipMeshNode, FTransform* InClipMeshTransform, UObject& ClipMesh, int32 LODIndex, int32 SectionIndex, int32 MaterialSlotIndex)
 {
 	HideGizmoClipMesh();
 
 	SetWidgetType(EWidgetType::ClipMesh);
 
 	ClipMeshNode = &InClipMeshNode;
+	ClipMeshTransform = InClipMeshTransform;
 
 	if (UStaticMesh* StaticMesh = Cast<UStaticMesh>(&ClipMesh))
 	{
 		ClipMeshStaticMeshComp->SetStaticMesh(StaticMesh);
 		ClipMeshStaticMeshComp->SetVisibility(true);
-		ClipMeshStaticMeshComp->SetWorldTransform(InClipMeshNode.Transform);
+		ClipMeshStaticMeshComp->SetWorldTransform(*ClipMeshTransform);
 		ClipMeshStaticMeshComp->EmptyOverrideMaterials();
 		ClipMeshStaticMeshComp->SetMaterial(MaterialSlotIndex, ClipMeshMaterial);
 		ClipMeshStaticMeshComp->SetSectionPreview(SectionIndex);
@@ -620,7 +621,7 @@ void FCustomizableObjectEditorViewportClient::ShowGizmoClipMesh(UCustomizableObj
 	{
 		ClipMeshSkeletalMeshComp->SetSkeletalMesh(SkeletalMesh);
 		ClipMeshSkeletalMeshComp->SetVisibility(true);
-		ClipMeshSkeletalMeshComp->SetWorldTransform(InClipMeshNode.Transform);
+		ClipMeshSkeletalMeshComp->SetWorldTransform(*ClipMeshTransform);
 		ClipMeshSkeletalMeshComp->SetForcedLOD(LODIndex);
 		ClipMeshSkeletalMeshComp->EmptyOverrideMaterials();
 		ClipMeshSkeletalMeshComp->SetMaterial(MaterialSlotIndex, ClipMeshMaterial);
@@ -1082,21 +1083,21 @@ bool FCustomizableObjectEditorViewportClient::InputWidgetDelta(FViewport* InView
 		{
 			if (WidgetMode == UE::Widget::WM_Translate)
 			{
-				ClipMeshNode->Transform.AddToTranslation(Drag);
+				ClipMeshTransform->AddToTranslation(Drag);
 			}
 			else if (WidgetMode == UE::Widget::WM_Rotate)
 			{
-				ClipMeshNode->Transform.ConcatenateRotation(Rot.Quaternion());
+				ClipMeshTransform->ConcatenateRotation(Rot.Quaternion());
 			}
 			if (WidgetMode == UE::Widget::WM_Scale)
 			{
-				ClipMeshNode->Transform.SetScale3D(ClipMeshNode->Transform.GetScale3D() + Scale);
+				ClipMeshTransform->SetScale3D(ClipMeshTransform->GetScale3D() + Scale);
 			}
 
 			ClipMeshStaticMeshComp->Modify();
 			ClipMeshSkeletalMeshComp->Modify();
-			ClipMeshStaticMeshComp->SetWorldTransform(ClipMeshNode->Transform);
-			ClipMeshSkeletalMeshComp->SetWorldTransform(ClipMeshNode->Transform);
+			ClipMeshStaticMeshComp->SetWorldTransform(*ClipMeshTransform);
+			ClipMeshSkeletalMeshComp->SetWorldTransform(*ClipMeshTransform);
 
 			return true;
 		}
@@ -1253,7 +1254,7 @@ FVector FCustomizableObjectEditorViewportClient::GetWidgetLocation() const
 		return ClipMorphOrigin + ClipMorphOffset;
 
 	case EWidgetType::ClipMesh:
-		return ClipMeshNode->Transform.GetTranslation();
+		return ClipMeshTransform->GetTranslation();
 
 	case EWidgetType::Light:
 		return SelectedLightComponent->GetComponentLocation();
@@ -1299,7 +1300,7 @@ FMatrix FCustomizableObjectEditorViewportClient::GetWidgetCoordSystem() const
 		
 	case EWidgetType::ClipMesh:
 		{
-			return ClipMeshNode->Transform.ToMatrixNoScale().RemoveTranslation();			
+			return ClipMeshTransform->ToMatrixNoScale().RemoveTranslation();			
 		}
 		
 	case EWidgetType::Light:
diff --git a/Engine/Plugins/Mutable/Source/CustomizableObjectEditor/Private/MuCOE/EdGraphSchema_CustomizableObject.cpp b/Engine/Plugins/Mutable/Source/CustomizableObjectEditor/Private/MuCOE/EdGraphSchema_CustomizableObject.cpp
--- a/Engine/Plugins/Mutable/Source/CustomizableObjectEditor/Private/MuCOE/EdGraphSchema_CustomizableObject.cpp
+++ b/Engine/Plugins/Mutable/Source/CustomizableObjectEditor/Private/MuCOE/EdGraphSchema_CustomizableObject.cpp
@@ -44,6 +44,7 @@
 #include "MuCOE/Nodes/CustomizableObjectNodeModifierClipMorph.h"
 #include "MuCOE/Nodes/CustomizableObjectNodeModifierClipWithMesh.h"
 #include "MuCOE/Nodes/CustomizableObjectNodeModifierClipWithUVMask.h"
+#include "MuCOE/Nodes/CustomizableObjectNodeModifierTransformInMesh.h"
 #include "MuCOE/Nodes/CustomizableObjectNodeMeshGeometryOperation.h"
 #include "MuCOE/Nodes/CustomizableObjectNodeMeshMorph.h"
 #include "MuCOE/Nodes/CustomizableObjectNodeMeshMorphStackApplication.h"
@@ -90,6 +91,8 @@
 #include "MuCO/CustomizableObjectCustomVersion.h"
 #include "MuCOE/GraphTraversal.h"
 #include "MuCOE/Nodes/CustomizableObjectNodeComponentMeshAddTo.h"
+#include "MuCOE/Nodes/CustomizableObjectNodeTransformConstant.h"
+#include "MuCOE/Nodes/CustomizableObjectNodeTransformParameter.h"
 
 class IToolkit;
 
@@ -243,6 +246,7 @@ const FName UEdGraphSchema_CustomizableObject::PC_Stack("stack");
 const FName UEdGraphSchema_CustomizableObject::PC_MaterialAsset("materialAsset");
 const FName UEdGraphSchema_CustomizableObject::PC_Wildcard("wildcard");
 const FName UEdGraphSchema_CustomizableObject::PC_PoseAsset("poseAsset");
+const FName UEdGraphSchema_CustomizableObject::PC_Transform("transform");
 
 
 // Node categories
@@ -523,6 +527,17 @@ void UEdGraphSchema_CustomizableObject::GetGraphContextActions(FGraphContextMenu
 		AddNewNodeActionFiltered(Node, ContextMenuBuilder, TEXT("Float"), LOCTEXT("Float_Division", "Float Division /"), GeneralGrouping, Filter);
 	}
 
+	{
+		UCustomizableObjectNode* TransformTemplateNodes[] =
+		{
+			ContextMenuBuilder.CreateTemplateNode<UCustomizableObjectNodeTransformConstant>(),
+			ContextMenuBuilder.CreateTemplateNode<UCustomizableObjectNodeTransformParameter>(),
+		};
+
+		AddNewNodeCategoryActionsFiltered(TransformTemplateNodes, ContextMenuBuilder, TEXT("Transform"), GeneralGrouping, Filter);
+	}
+	
+	
 	{
 		UCustomizableObjectNode* ProjectorTemplateNodes[] =
 		{
@@ -540,6 +555,7 @@ void UEdGraphSchema_CustomizableObject::GetGraphContextActions(FGraphContextMenu
 			ContextMenuBuilder.CreateTemplateNode<UCustomizableObjectNodeComponentPassthroughMesh>(),
 			ContextMenuBuilder.CreateTemplateNode<UCustomizableObjectNodeMeshReshape>(),
 			ContextMenuBuilder.CreateTemplateNode<UCustomizableObjectNodeModifierClipDeform>(),
+			ContextMenuBuilder.CreateTemplateNode<UCustomizableObjectNodeModifierTransformInMesh>(),
 			ContextMenuBuilder.CreateTemplateNode<UCustomizableObjectNodeTextureParameter>(),
 		};
 
@@ -548,7 +564,7 @@ void UEdGraphSchema_CustomizableObject::GetGraphContextActions(FGraphContextMenu
 
 	{
 		// External Pin Nodes
-		TArray<FName> PinTypes({ PC_Material, PC_Modifier, PC_Mesh, PC_Image, PC_Projector, PC_GroupProjector, PC_Color, PC_Float, PC_Bool, PC_Enum, PC_Stack, PC_PassThroughImage, PC_MaterialAsset, PC_PoseAsset, PC_Component });
+		TArray<FName> PinTypes({ PC_Material, PC_Modifier, PC_Mesh, PC_Image, PC_Projector, PC_GroupProjector, PC_Color, PC_Float, PC_Bool, PC_Enum, PC_Transform, PC_Stack, PC_PassThroughImage, PC_MaterialAsset, PC_PoseAsset, PC_Component });
 
 		// Add pin types from extensions
 		for (const FRegisteredCustomizableObjectPinType& PinType : ICustomizableObjectModule::Get().GetExtendedPinTypes())
@@ -763,6 +779,10 @@ FLinearColor UEdGraphSchema_CustomizableObject::GetPinTypeColor(const FName& Typ
 	{
 		return FLinearColor(0.700000f, 0.000000f, 0.000000f, 1.000000f); // Dark Red
 	}
+	else if (TypeString == PC_Transform)
+	{
+		return FLinearColor(0.5f, 1.0f, 1.0f, 1.0f); // Cyan
+	}
 
 	for (const FRegisteredCustomizableObjectPinType& PinType : ICustomizableObjectModule::Get().GetExtendedPinTypes())
 	{
@@ -1414,6 +1434,10 @@ FText UEdGraphSchema_CustomizableObject::GetPinCategoryName(const FName& PinCate
 	{
 		return LOCTEXT("Pose_Pin_Category", "PoseAsset");
 	}
+	else if (PinCategory == UEdGraphSchema_CustomizableObject::PC_Transform)
+	{
+		return LOCTEXT("Transform_Pin_Category", "Transform");
+	}
 	else
 	{
 		for (const FRegisteredCustomizableObjectPinType& PinType : ICustomizableObjectModule::Get().GetExtendedPinTypes())
diff --git a/Engine/Plugins/Mutable/Source/CustomizableObjectEditor/Private/MuCOE/GenerateMutableSource/GenerateMutableSourceModifier.cpp b/Engine/Plugins/Mutable/Source/CustomizableObjectEditor/Private/MuCOE/GenerateMutableSource/GenerateMutableSourceModifier.cpp
--- a/Engine/Plugins/Mutable/Source/CustomizableObjectEditor/Private/MuCOE/GenerateMutableSource/GenerateMutableSourceModifier.cpp
+++ b/Engine/Plugins/Mutable/Source/CustomizableObjectEditor/Private/MuCOE/GenerateMutableSource/GenerateMutableSourceModifier.cpp
@@ -10,6 +10,7 @@
 #include "MuCOE/GenerateMutableSource/GenerateMutableSourceSurface.h"
 #include "MuCOE/GenerateMutableSource/GenerateMutableSourceLayout.h"
 #include "MuCOE/GenerateMutableSource/GenerateMutableSourceGroupProjector.h"
+#include "MuCOE/GenerateMutableSource/GenerateMutableSourceTransform.h"
 #include "MuCOE/GraphTraversal.h"
 #include "MuCOE/MutableUtils.h"
 #include "MuCOE/Nodes/CustomizableObjectNodeModifierClipDeform.h"
@@ -24,11 +25,13 @@
 #include "MuCOE/Nodes/CustomizableObjectNodeStaticMesh.h"
 #include "MuCOE/Nodes/CustomizableObjectNodeFloatParameter.h"
 #include "MuCOE/Nodes/CustomizableObjectNodeFloatConstant.h"
+#include "MuCOE/Nodes/CustomizableObjectNodeModifierTransformInMesh.h"
 #include "MuR/Mesh.h"
 #include "MuT/NodeMeshTransform.h"
 #include "MuT/NodeModifierMeshClipDeform.h"
 #include "MuT/NodeModifierMeshClipMorphPlane.h"
 #include "MuT/NodeModifierMeshClipWithUVMask.h"
+#include "MuT/NodeModifierMeshTransformInMesh.h"
 #include "MuT/NodeModifierSurfaceEdit.h"
 #include "MuT/NodeMeshConstant.h"
 #include "MuT/NodeMeshFragment.h"
@@ -604,6 +607,77 @@ mu::Ptr<mu::NodeModifier> GenerateMutableSourceModifier(const UEdGraphPin * Pin,
 		GenerationContext.MeshGenerationFlags.Pop();
 	}
 
+	else if (const UCustomizableObjectNodeModifierTransformInMesh* TypedNodeTransformMesh = Cast<UCustomizableObjectNodeModifierTransformInMesh>(Node))
+	{
+		const EMutableMeshConversionFlags ModifiersMeshFlags =
+			EMutableMeshConversionFlags::IgnoreSkinning |
+			EMutableMeshConversionFlags::IgnorePhysics;
+		GenerationContext.MeshGenerationFlags.Push(ModifiersMeshFlags);
+
+		// MeshTransformInMesh can be connected to multiple objects, so the compiled NodeModifierMeshTransformInMesh
+		// needs to be different for each object. If it were added to the Generated cache, all the objects would get the same.
+		bDoNotAddToGeneratedCache = true;
+
+		mu::Ptr<mu::NodeModifierMeshTransformInMesh> TransformNode = new mu::NodeModifierMeshTransformInMesh();
+		TransformNode->SetMessageContext(Node);
+		Result = TransformNode;
+
+		if (const UEdGraphPin* ConnectedPin = FollowInputPin(*TypedNodeTransformMesh->TransformPin()))
+		{
+			TransformNode->MatrixNode = GenerateMutableSourceTransform(ConnectedPin, GenerationContext);
+		}
+
+		// If no bounding mesh is provided, we transform the entire mesh.
+		if (const UEdGraphPin* ConnectedPin = FollowInputPin(*TypedNodeTransformMesh->BoundingMeshPin()))
+		{
+			FMutableGraphMeshGenerationData DummyMeshData;
+
+			mu::NodeMeshPtr BoundingMesh = GenerateMutableSourceMesh(ConnectedPin, GenerationContext, DummyMeshData, false, true);
+
+			const FPinDataValue* PinData = GenerationContext.PinData.Find(ConnectedPin);
+			for (const FMeshData& MeshData : PinData->MeshesData)
+			{
+				bool bClosed = true;
+				if (const USkeletalMesh* SkeletalMesh = Cast<USkeletalMesh>(MeshData.Mesh))
+				{
+					bClosed = IsMeshClosed(SkeletalMesh, MeshData.LOD, MeshData.MaterialIndex);
+				}
+				else if (const UStaticMesh* StaticMesh = Cast<UStaticMesh>(MeshData.Mesh))
+				{
+					bClosed = IsMeshClosed(StaticMesh, MeshData.LOD, MeshData.MaterialIndex);
+				}
+				else
+				{
+					// TODO: We support the bounding mesh not being constant. This message is not precise enough. It should say that it hasn't been 
+					// possible to check if the mesh is closed or not.
+					GenerationContext.Log(LOCTEXT("UnimplementedNode", "Node type not implemented yet."), MeshData.Node);
+				}
+
+				if (!bClosed)
+				{
+					FText ErrorMsg = FText::Format(LOCTEXT("Clipping mesh", "The bounding [{0}] not watertight (i.e. it does not fully enclose a volume)."), FText::FromName(MeshData.Mesh->GetFName()));
+					GenerationContext.Log(ErrorMsg, MeshData.Node, EMessageSeverity::Warning);
+				}
+			}
+
+			if (FMatrix Matrix = TypedNodeTransformMesh->BoundingMeshTransform.ToMatrixWithScale(); Matrix != FMatrix::Identity)
+			{
+				mu::NodeMeshTransformPtr TransformMesh = new mu::NodeMeshTransform();
+				TransformMesh->SetSource(BoundingMesh.get());
+
+				TransformMesh->SetTransform(FMatrix44f(Matrix));
+				BoundingMesh = TransformMesh;
+			}
+
+			TransformNode->BoundingMesh = BoundingMesh;
+		}
+
+		TransformNode->MultipleTagsPolicy = TypedNodeTransformMesh->MultipleTagPolicy;
+		TransformNode->RequiredTags = TypedNodeTransformMesh->RequiredTags;
+
+		GenerationContext.MeshGenerationFlags.Pop();
+	}
+
 	else
 	{
 		GenerationContext.Log(LOCTEXT("UnimplementedNode", "Node type not implemented yet."), Node);
diff --git a/Engine/Plugins/Mutable/Source/CustomizableObjectEditor/Private/MuCOE/GenerateMutableSource/GenerateMutableSourceTransform.cpp b/Engine/Plugins/Mutable/Source/CustomizableObjectEditor/Private/MuCOE/GenerateMutableSource/GenerateMutableSourceTransform.cpp
new file mode 100644
--- /dev/null
+++ b/Engine/Plugins/Mutable/Source/CustomizableObjectEditor/Private/MuCOE/GenerateMutableSource/GenerateMutableSourceTransform.cpp
@@ -0,0 +1,74 @@
+﻿// Copyright Epic Games, Inc. All Rights Reserved.
+
+#include "GenerateMutableSourceTransform.h"
+
+#include "MuCOE/EdGraphSchema_CustomizableObject.h"
+#include "MuCOE/GenerateMutableSource/GenerateMutableSource.h"
+#include "MuCOE/Nodes/CustomizableObjectNodeTransformConstant.h"
+#include "MuCOE/Nodes/CustomizableObjectNodeTransformParameter.h"
+#include "MuT/NodeMatrixConstant.h"
+#include "MuT/NodeMatrixParameter.h"
+
+#define LOCTEXT_NAMESPACE "CustomizableObjectEditor"
+
+
+mu::Ptr<mu::NodeMatrix> GenerateMutableSourceTransform(const UEdGraphPin* Pin, FMutableGraphGenerationContext& GenerationContext)
+{
+	check(Pin);
+
+	RETURN_ON_CYCLE(*Pin, GenerationContext)
+
+	CheckNumOutputs(*Pin, GenerationContext);
+
+	UCustomizableObjectNode* Node = CastChecked<UCustomizableObjectNode>(Pin->GetOwningNode());
+	const FGeneratedKey Key(reinterpret_cast<void*>(&GenerateMutableSourceTransform), *Pin, *Node, GenerationContext);
+	if (const FGeneratedData* Generated = GenerationContext.Generated.Find(Key))
+	{
+		return static_cast<mu::NodeMatrix*>(Generated->Node.get());
+	}
+
+	if (Node->IsNodeOutDatedAndNeedsRefresh())
+	{
+		Node->SetRefreshNodeWarning();
+	}
+
+	mu::Ptr<mu::NodeMatrix> Result;
+	if (const UCustomizableObjectNodeTransformConstant* TypedNodeTransformConstant = Cast<UCustomizableObjectNodeTransformConstant>(Node))
+	{
+		mu::Ptr<mu::NodeMatrixConstant> MatrixNode = new mu::NodeMatrixConstant{};
+		Result = MatrixNode;
+		MatrixNode->Value = FMatrix44f(TypedNodeTransformConstant->Value.ToMatrixWithScale());
+	}
+	else if (const UCustomizableObjectNodeTransformParameter* TypedNodeTransformParameter = Cast<UCustomizableObjectNodeTransformParameter>(Node))
+	{
+		mu::Ptr<mu::NodeMatrixParameter> MatrixNode = new mu::NodeMatrixParameter();
+		Result = MatrixNode;
+
+		GenerationContext.AddParameterNameUnique(Node, TypedNodeTransformParameter->ParameterName);
+
+		MatrixNode->Name = TypedNodeTransformParameter->ParameterName;
+		MatrixNode->Uid = GenerationContext.GetNodeIdUnique(Node).ToString();
+		MatrixNode->DefaultValue = FMatrix44f(TypedNodeTransformParameter->DefaultValue.ToMatrixWithScale());
+
+		GenerationContext.ParameterUIDataMap.Add(TypedNodeTransformParameter->ParameterName, FMutableParameterData(
+			TypedNodeTransformParameter->ParamUIMetadata,
+			EMutableParameterType::Transform));
+	}
+	else
+	{
+		GenerationContext.Log(LOCTEXT("UnimplementedNode", "Node type not implemented yet."), Node);
+	}
+	
+	FGeneratedData CacheData = FGeneratedData(Node, Result);
+	GenerationContext.Generated.Add(Key, CacheData);
+	GenerationContext.GeneratedNodes.Add(Node);
+	
+	if (Result)
+	{
+		Result->SetMessageContext(Node);
+	}
+
+	return Result;
+}
+
+#undef LOCTEXT_NAMESPACE
diff --git a/Engine/Plugins/Mutable/Source/CustomizableObjectEditor/Private/MuCOE/GenerateMutableSource/GenerateMutableSourceTransform.h b/Engine/Plugins/Mutable/Source/CustomizableObjectEditor/Private/MuCOE/GenerateMutableSource/GenerateMutableSourceTransform.h
new file mode 100644
--- /dev/null
+++ b/Engine/Plugins/Mutable/Source/CustomizableObjectEditor/Private/MuCOE/GenerateMutableSource/GenerateMutableSourceTransform.h
@@ -0,0 +1,11 @@
+﻿// Copyright Epic Games, Inc. All Rights Reserved.
+
+#pragma once
+
+#include "MuT/NodeMatrix.h"
+
+class UEdGraphPin;
+struct FMutableGraphGenerationContext;
+
+/** Convert a CustomizableObject Source Graph into a mutable source graph. */
+mu::Ptr<mu::NodeMatrix> GenerateMutableSourceTransform(const UEdGraphPin* Pin, FMutableGraphGenerationContext& GenerationContext);
diff --git a/Engine/Plugins/Mutable/Source/CustomizableObjectEditor/Private/MuCOE/Nodes/CustomizableObjectNodeModifierTransformInMesh.cpp b/Engine/Plugins/Mutable/Source/CustomizableObjectEditor/Private/MuCOE/Nodes/CustomizableObjectNodeModifierTransformInMesh.cpp
new file mode 100644
--- /dev/null
+++ b/Engine/Plugins/Mutable/Source/CustomizableObjectEditor/Private/MuCOE/Nodes/CustomizableObjectNodeModifierTransformInMesh.cpp
@@ -0,0 +1,52 @@
+﻿// Copyright Epic Games, Inc. All Rights Reserved.
+
+#include "MuCOE/Nodes/CustomizableObjectNodeModifierTransformInMesh.h"
+
+#include "MuCOE/EdGraphSchema_CustomizableObject.h"
+
+#define LOCTEXT_NAMESPACE "CustomizableObjectEditor"
+
+static const TCHAR* OutputPinName = TEXT("Modifier");
+static const TCHAR* BoundingMeshPinName = TEXT("Bounding Mesh");
+static const TCHAR* TransformPinName = TEXT("Transform");
+
+
+FText UCustomizableObjectNodeModifierTransformInMesh::GetNodeTitle(ENodeTitleType::Type TitleType) const
+{
+	return LOCTEXT("Transform_Mesh_In_Mesh", "Transform Mesh In Mesh");
+}
+
+FText UCustomizableObjectNodeModifierTransformInMesh::GetTooltipText() const
+{
+	return LOCTEXT("Transform_Mesh_In_Mesh_Tooltip", "Applies a transform to the vertices of a mesh that is contained within the given bounding mesh");
+}
+
+void UCustomizableObjectNodeModifierTransformInMesh::AllocateDefaultPins(UCustomizableObjectNodeRemapPins* RemapPins)
+{
+	const UEdGraphSchema_CustomizableObject* Schema = GetDefault<UEdGraphSchema_CustomizableObject>();
+
+	UEdGraphPin* ClipMeshPin = CustomCreatePin(EGPD_Input, Schema->PC_Mesh, FName(BoundingMeshPinName));
+	ClipMeshPin->bDefaultValueIsIgnored = true;
+
+	UEdGraphPin* TransformPin = CustomCreatePin(EGPD_Input, Schema->PC_Transform, FName(TransformPinName));
+	TransformPin->bDefaultValueIsIgnored = true;
+	
+	(void)CustomCreatePin(EGPD_Output, Schema->PC_Modifier, FName(OutputPinName));
+}
+
+UEdGraphPin* UCustomizableObjectNodeModifierTransformInMesh::OutputPin() const
+{
+	return FindPin(OutputPinName);
+}
+
+UEdGraphPin* UCustomizableObjectNodeModifierTransformInMesh::BoundingMeshPin() const
+{
+	return FindPin(BoundingMeshPinName);
+}
+
+UEdGraphPin* UCustomizableObjectNodeModifierTransformInMesh::TransformPin() const
+{
+	return FindPin(TransformPinName);
+}
+
+#undef LOCTEXT_NAMESPACE
diff --git a/Engine/Plugins/Mutable/Source/CustomizableObjectEditor/Private/MuCOE/Nodes/CustomizableObjectNodeModifierTransformInMeshDetails.cpp b/Engine/Plugins/Mutable/Source/CustomizableObjectEditor/Private/MuCOE/Nodes/CustomizableObjectNodeModifierTransformInMeshDetails.cpp
new file mode 100644
--- /dev/null
+++ b/Engine/Plugins/Mutable/Source/CustomizableObjectEditor/Private/MuCOE/Nodes/CustomizableObjectNodeModifierTransformInMeshDetails.cpp
@@ -0,0 +1,45 @@
+﻿// Copyright Epic Games, Inc. All Rights Reserved.
+
+#include "MuCOE/Nodes/CustomizableObjectNodeModifierTransformInMeshDetails.h"
+
+#include "MuCOE/Nodes/CustomizableObjectNodeModifierTransformInMesh.h"
+
+#include "DetailCategoryBuilder.h"
+#include "DetailLayoutBuilder.h"
+#include "DetailWidgetRow.h"
+
+#define LOCTEXT_NAMESPACE "CustomizableObjectNodeModifierTransformInMeshDetails"
+
+
+TSharedRef<IDetailCustomization> FCustomizableObjectNodeModifierTransformInMeshDetails::MakeInstance()
+{
+	return MakeShareable(new FCustomizableObjectNodeModifierTransformInMeshDetails);
+}
+
+
+void FCustomizableObjectNodeModifierTransformInMeshDetails::CustomizeDetails(IDetailLayoutBuilder& DetailBuilder)
+{
+	FCustomizableObjectNodeModifierBaseDetails::CustomizeDetails(DetailBuilder);
+
+	Node = nullptr;
+	DetailBuilderPtr = &DetailBuilder;
+
+	const IDetailsView* DetailsView = DetailBuilder.GetDetailsView();
+	if (DetailsView->GetSelectedObjects().Num())
+	{
+		Node = Cast<UCustomizableObjectNodeModifierTransformInMesh>(DetailsView->GetSelectedObjects()[0].Get());
+	}
+
+	IDetailCategoryBuilder& CustomizableObjectToClipCategory = DetailBuilder.EditCategory("BoundingMesh");
+
+	if(!Node)
+	{
+		CustomizableObjectToClipCategory.AddCustomRow(LOCTEXT("Node", "Node"))
+			[
+				SNew(STextBlock)
+				.Text(LOCTEXT("Node not found", "Node not found"))
+			];
+	}
+}
+
+#undef LOCTEXT_NAMESPACE
diff --git a/Engine/Plugins/Mutable/Source/CustomizableObjectEditor/Private/MuCOE/Nodes/CustomizableObjectNodeModifierTransformInMeshDetails.h b/Engine/Plugins/Mutable/Source/CustomizableObjectEditor/Private/MuCOE/Nodes/CustomizableObjectNodeModifierTransformInMeshDetails.h
new file mode 100644
--- /dev/null
+++ b/Engine/Plugins/Mutable/Source/CustomizableObjectEditor/Private/MuCOE/Nodes/CustomizableObjectNodeModifierTransformInMeshDetails.h
@@ -0,0 +1,29 @@
+﻿// Copyright Epic Games, Inc. All Rights Reserved.
+
+#pragma once
+
+#include "CustomizableObjectNodeModifierBaseDetails.h"
+#include "IDetailCustomization.h"
+
+class IDetailLayoutBuilder;
+class UCustomizableObjectNodeModifierTransformInMesh;
+
+
+class FCustomizableObjectNodeModifierTransformInMeshDetails : public FCustomizableObjectNodeModifierBaseDetails
+{
+public:
+	/** Makes a new instance of this detail layout class for a specific detail view requesting it */
+	static TSharedRef<IDetailCustomization> MakeInstance();
+
+	/** ILayoutDetails interface */
+	virtual void CustomizeDetails( IDetailLayoutBuilder& DetailBuilder ) override;
+
+private:
+
+	/** The node for which details are being customized */
+	UCustomizableObjectNodeModifierTransformInMesh* Node = nullptr;
+
+	/** Pointer to the builder passed by CustomizeDetails method */
+	IDetailLayoutBuilder* DetailBuilderPtr = nullptr;
+
+};
diff --git a/Engine/Plugins/Mutable/Source/CustomizableObjectEditor/Private/MuCOE/Nodes/CustomizableObjectNodeTransformConstant.cpp b/Engine/Plugins/Mutable/Source/CustomizableObjectEditor/Private/MuCOE/Nodes/CustomizableObjectNodeTransformConstant.cpp
new file mode 100644
--- /dev/null
+++ b/Engine/Plugins/Mutable/Source/CustomizableObjectEditor/Private/MuCOE/Nodes/CustomizableObjectNodeTransformConstant.cpp
@@ -0,0 +1,233 @@
+﻿// Copyright Epic Games, Inc. All Rights Reserved.
+
+#include "MuCOE/Nodes/CustomizableObjectNodeTransformConstant.h"
+
+#include "Components/HorizontalBox.h"
+#include "MuCOE/CustomizableObjectEditorStyle.h"
+#include "MuCOE/CustomizableObjectGraph.h"
+#include "MuCOE/EdGraphSchema_CustomizableObject.h"
+#include "Widgets/Images/SImage.h"
+#include "Widgets/Input/SCheckBox.h"
+#include "Widgets/Input/SRotatorInputBox.h"
+#include "Widgets/Input/SVectorInputBox.h"
+#include "Widgets/Layout/SGridPanel.h"
+
+
+#define LOCTEXT_NAMESPACE "CustomizableObjectEditor"
+
+
+void SGraphNodeTransformConstant::Construct(const FArguments& InArgs, UEdGraphNode* InGraphNode)
+{
+	NodeTransformConstant = Cast<UCustomizableObjectNodeTransformConstant>(InGraphNode);
+
+	SCustomizableObjectNode::Construct({}, InGraphNode);
+}
+
+void SGraphNodeTransformConstant::SetDefaultTitleAreaWidget(TSharedRef<SOverlay> DefaultTitleAreaWidget)
+{
+	// Collapsing arrow of the title area
+	DefaultTitleAreaWidget->AddSlot()
+	.HAlign(HAlign_Right)
+	.VAlign(VAlign_Center)
+	.Padding(FMargin(5))
+	[
+		SNew(SCheckBox)
+		.OnCheckStateChanged(this, &SGraphNodeTransformConstant::OnExpressionPreviewChanged)
+		.IsChecked(IsExpressionPreviewChecked())
+		.Cursor(EMouseCursor::Default)
+		.Style(FAppStyle::Get(), "Graph.Node.AdvancedView")
+		[
+			SNew(SHorizontalBox)
+			+ SHorizontalBox::Slot()
+			.VAlign(VAlign_Center)
+			.HAlign(HAlign_Center)
+			[
+				SNew(SImage)
+				.Image(GetExpressionPreviewArrow())
+			]
+		]
+	];
+}
+
+void SGraphNodeTransformConstant::CreateBelowPinControls(TSharedPtr<SVerticalBox> MainBox)
+{
+	constexpr float TransformBoxMinWidth = 250.0f;
+	
+	MainBox->AddSlot()
+	.AutoHeight()
+	.Padding(10.0f, 0.0f, 10.0f, 10.0f)
+	[
+		SNew(SHorizontalBox)
+		.Visibility(NodeTransformConstant->bCollapsed ? EVisibility::Collapsed : EVisibility::Visible)
+
+		+ SHorizontalBox::Slot()
+		.MinWidth(TransformBoxMinWidth)
+		[
+			SNew(SGridPanel)
+			.FillColumn(1, 1)
+			+ SGridPanel::Slot(0, 0)
+			.HAlign(HAlign_Right)
+			.VAlign(VAlign_Center)
+			[
+				SNew(STextBlock)
+				.Text(LOCTEXT("Transform_Location", "Location"))
+			]
+			+ SGridPanel::Slot(1, 0)
+			.Padding(4.0f, 0.0f, 0.0f, 0.0f)
+			.HAlign(HAlign_Fill)
+			[
+				SNew(SNumericVectorInputBox<double>)
+				.bColorAxisLabels(true)
+				.AllowSpin(true)
+				.X_Lambda([Node=NodeTransformConstant]() { return Node->Value.GetLocation().X; })
+				.Y_Lambda([Node=NodeTransformConstant]() { return Node->Value.GetLocation().Y; })
+				.Z_Lambda([Node=NodeTransformConstant]() { return Node->Value.GetLocation().Z; })
+				.OnXChanged(this, &SGraphNodeTransformConstant::OnLocationChanged, ETextCommit::Default, EAxis::X)
+				.OnYChanged(this, &SGraphNodeTransformConstant::OnLocationChanged, ETextCommit::Default, EAxis::Y)
+				.OnZChanged(this, &SGraphNodeTransformConstant::OnLocationChanged, ETextCommit::Default, EAxis::Z)
+				.OnXCommitted(this, &SGraphNodeTransformConstant::OnLocationChanged, EAxis::X)
+				.OnYCommitted(this, &SGraphNodeTransformConstant::OnLocationChanged, EAxis::Y)
+				.OnZCommitted(this, &SGraphNodeTransformConstant::OnLocationChanged, EAxis::Z)
+			]
+			
+			+ SGridPanel::Slot(0, 1)
+			.HAlign(EHorizontalAlignment::HAlign_Right)
+			.VAlign(VAlign_Center)
+			[
+				SNew(STextBlock)
+				.Text(LOCTEXT("Transform_Rotation", "Rotation"))
+			]
+			+ SGridPanel::Slot(1, 1)
+			.HAlign(HAlign_Fill)
+			.Padding(4.0f, 0.0f, 0.0f, 0.0f)
+			[
+				SNew(SNumericRotatorInputBox<double>)
+				.bColorAxisLabels(true)
+				.AllowSpin(true)
+				.Roll_Lambda([Node=NodeTransformConstant]() { return Node->Value.Rotator().Roll; })
+				.Pitch_Lambda([Node=NodeTransformConstant]() { return Node->Value.Rotator().Pitch; })
+				.Yaw_Lambda([Node=NodeTransformConstant]() { return Node->Value.Rotator().Yaw; })
+				.OnRollChanged(this, &SGraphNodeTransformConstant::OnRotationChanged, ETextCommit::Default, EAxis::X)
+				.OnPitchChanged(this, &SGraphNodeTransformConstant::OnRotationChanged, ETextCommit::Default, EAxis::Y)
+				.OnYawChanged(this, &SGraphNodeTransformConstant::OnRotationChanged, ETextCommit::Default, EAxis::Z)
+				.OnRollCommitted(this, &SGraphNodeTransformConstant::OnRotationChanged, EAxis::X)
+				.OnPitchCommitted(this, &SGraphNodeTransformConstant::OnRotationChanged, EAxis::Y)
+				.OnYawCommitted(this, &SGraphNodeTransformConstant::OnRotationChanged, EAxis::Z)
+			]
+
+			+ SGridPanel::Slot(0, 2)
+			.HAlign(EHorizontalAlignment::HAlign_Right)
+			.VAlign(VAlign_Center)
+			[
+				SNew(STextBlock)
+				.Text(LOCTEXT("Transform_Scale", "Scale"))
+			]
+			+ SGridPanel::Slot(1, 2)
+			.HAlign(HAlign_Fill)
+			.Padding(4.0f, 0.0f, 0.0f, 0.0f)
+			[
+				SNew(SNumericVectorInputBox<double>)
+				.bColorAxisLabels(true)
+				.AllowSpin(true)
+				.X_Lambda([Node=NodeTransformConstant]() { return Node->Value.GetScale3D().X; })
+				.Y_Lambda([Node=NodeTransformConstant]() { return Node->Value.GetScale3D().Y; })
+				.Z_Lambda([Node=NodeTransformConstant]() { return Node->Value.GetScale3D().Z; })
+				.OnXChanged(this, &SGraphNodeTransformConstant::OnScaleChanged, ETextCommit::Default, EAxis::X)
+				.OnYChanged(this, &SGraphNodeTransformConstant::OnScaleChanged, ETextCommit::Default, EAxis::Y)
+				.OnZChanged(this, &SGraphNodeTransformConstant::OnScaleChanged, ETextCommit::Default, EAxis::Z)
+				.OnXCommitted(this, &SGraphNodeTransformConstant::OnScaleChanged, EAxis::X)
+				.OnYCommitted(this, &SGraphNodeTransformConstant::OnScaleChanged, EAxis::Y)
+				.OnZCommitted(this, &SGraphNodeTransformConstant::OnScaleChanged, EAxis::Z)
+			]
+		]
+	];
+}
+
+void SGraphNodeTransformConstant::OnExpressionPreviewChanged(const ECheckBoxState NewCheckedState)
+{
+	NodeTransformConstant->bCollapsed = (NewCheckedState != ECheckBoxState::Checked);
+	UpdateGraphNode();
+}
+
+ECheckBoxState SGraphNodeTransformConstant::IsExpressionPreviewChecked() const
+{
+	return NodeTransformConstant->bCollapsed ? ECheckBoxState::Unchecked : ECheckBoxState::Checked;
+}
+
+const FSlateBrush* SGraphNodeTransformConstant::GetExpressionPreviewArrow() const
+{
+	return FCustomizableObjectEditorStyle::Get().GetBrush(NodeTransformConstant->bCollapsed ? TEXT("Nodes.ArrowDown") : TEXT("Nodes.ArrowUp"));
+}
+
+void SGraphNodeTransformConstant::OnLocationChanged(double Value, ETextCommit::Type, EAxis::Type Axis)
+{
+	FVector Location = NodeTransformConstant->Value.GetLocation();
+	if (!FMath::IsNearlyEqualByULP(Value, Location.GetComponentForAxis(Axis)))
+	{
+		Location.SetComponentForAxis(Axis, Value);
+		NodeTransformConstant->Value.SetLocation(Location);
+
+		(void)NodeTransformConstant->MarkPackageDirty();
+	}
+}
+
+void SGraphNodeTransformConstant::OnRotationChanged(double Value, ETextCommit::Type, EAxis::Type Axis)
+{
+	FRotator Rotation = NodeTransformConstant->Value.Rotator();
+	Rotation.SetComponentForAxis(Axis, Value);
+	
+	if (!NodeTransformConstant->Value.Rotator().Equals(Rotation))
+	{
+		NodeTransformConstant->Value.SetRotation(Rotation.Quaternion());
+
+		(void)NodeTransformConstant->MarkPackageDirty();
+	}
+}
+
+void SGraphNodeTransformConstant::OnScaleChanged(double Value, ETextCommit::Type, EAxis::Type Axis)
+{
+	FVector Scale = NodeTransformConstant->Value.GetScale3D();
+	if (!FMath::IsNearlyEqualByULP(Value, Scale.GetComponentForAxis(Axis)))
+	{
+		Scale.SetComponentForAxis(Axis, Value);
+		NodeTransformConstant->Value.SetScale3D(Scale);
+
+		(void)NodeTransformConstant->MarkPackageDirty();
+	}
+}
+
+
+// ============================================
+// === UCustomizableObjectNodeTransformConstant
+
+
+FText UCustomizableObjectNodeTransformConstant::GetNodeTitle(ENodeTitleType::Type TitleType) const
+{
+	return LOCTEXT("Transform_Constant", "Transform Constant");
+}
+
+FLinearColor UCustomizableObjectNodeTransformConstant::GetNodeTitleColor() const
+{
+	const UEdGraphSchema_CustomizableObject* Schema = GetDefault<UEdGraphSchema_CustomizableObject>();
+	return Schema->GetPinTypeColor(Schema->PC_Transform);
+}
+
+FText UCustomizableObjectNodeTransformConstant::GetTooltipText() const
+{
+	return LOCTEXT("Transform_Constant_Tooltip", "Define a constant transform value.");
+}
+
+void UCustomizableObjectNodeTransformConstant::AllocateDefaultPins(UCustomizableObjectNodeRemapPins* RemapPins)
+{
+	const UEdGraphSchema_CustomizableObject* Schema = GetDefault<UEdGraphSchema_CustomizableObject>();
+
+	UEdGraphPin* ValuePin = CustomCreatePin(EGPD_Output, Schema->PC_Transform, FName("Value"));
+	ValuePin->bDefaultValueIsIgnored = true;
+}
+
+TSharedPtr<SGraphNode> UCustomizableObjectNodeTransformConstant::CreateVisualWidget()
+{
+	return SNew(SGraphNodeTransformConstant, this);
+}
+
+#undef LOCTEXT_NAMESPACE
diff --git a/Engine/Plugins/Mutable/Source/CustomizableObjectEditor/Private/MuCOE/Nodes/CustomizableObjectNodeTransformParameter.cpp b/Engine/Plugins/Mutable/Source/CustomizableObjectEditor/Private/MuCOE/Nodes/CustomizableObjectNodeTransformParameter.cpp
new file mode 100644
--- /dev/null
+++ b/Engine/Plugins/Mutable/Source/CustomizableObjectEditor/Private/MuCOE/Nodes/CustomizableObjectNodeTransformParameter.cpp
@@ -0,0 +1,54 @@
+﻿// Copyright Epic Games, Inc. All Rights Reserved.
+
+#include "MuCOE/Nodes/CustomizableObjectNodeTransformParameter.h"
+
+#include "MuCOE/EdGraphSchema_CustomizableObject.h"
+
+
+#define LOCTEXT_NAMESPACE "CustomizableObjectEditor"
+
+
+FText UCustomizableObjectNodeTransformParameter::GetNodeTitle(ENodeTitleType::Type TitleType) const
+{
+	if (TitleType == ENodeTitleType::ListView || ParameterName.IsEmpty())
+	{
+		return LOCTEXT("Transform_Parameter", "Transform Parameter");
+	}
+	else if (TitleType == ENodeTitleType::EditableTitle)
+	{
+		return FText::Format(LOCTEXT("Transform_Parameter_EditableTitle", "{0}"), FText::FromString(ParameterName));
+	}
+	else
+	{
+		return FText::Format(LOCTEXT("Transform_Parameter_Title", "{0}\nTransform Parameter"), FText::FromString(ParameterName));
+	}
+}
+
+FLinearColor UCustomizableObjectNodeTransformParameter::GetNodeTitleColor() const
+{
+	const UEdGraphSchema_CustomizableObject* Schema = GetDefault<UEdGraphSchema_CustomizableObject>();
+	return Schema->GetPinTypeColor(Schema->PC_Transform);
+}
+
+FText UCustomizableObjectNodeTransformParameter::GetTooltipText() const
+{
+	return LOCTEXT("Transform_Parameter_Tooltip", "Expose a runtime modifiable transform parameter from the Customizable Object.");
+}
+
+void UCustomizableObjectNodeTransformParameter::OnRenameNode(const FString& NewName)
+{
+	if (!NewName.IsEmpty())
+	{
+		ParameterName = NewName;
+	}
+}
+
+void UCustomizableObjectNodeTransformParameter::AllocateDefaultPins(UCustomizableObjectNodeRemapPins* RemapPins)
+{
+	const UEdGraphSchema_CustomizableObject* Schema = GetDefault<UEdGraphSchema_CustomizableObject>();
+
+	UEdGraphPin* ValuePin = CustomCreatePin(EGPD_Output, Schema->PC_Transform, FName("Value"));
+	ValuePin->bDefaultValueIsIgnored = true;
+}
+
+#undef LOCTEXT_NAMESPACE
diff --git a/Engine/Plugins/Mutable/Source/CustomizableObjectEditor/Private/MuCOE/SCustomizableObjectEditorViewport.cpp b/Engine/Plugins/Mutable/Source/CustomizableObjectEditor/Private/MuCOE/SCustomizableObjectEditorViewport.cpp
--- a/Engine/Plugins/Mutable/Source/CustomizableObjectEditor/Private/MuCOE/SCustomizableObjectEditorViewport.cpp
+++ b/Engine/Plugins/Mutable/Source/CustomizableObjectEditor/Private/MuCOE/SCustomizableObjectEditorViewport.cpp
@@ -361,14 +361,14 @@ void SCustomizableObjectEditorViewportTabBody::HideGizmoClipMorph() const
 }
 
 
-void SCustomizableObjectEditorViewportTabBody::ShowGizmoClipMesh(UCustomizableObjectNodeModifierClipWithMesh& ClipMeshNode, UObject& ClipMesh, int32 LODIndex, int32 SectionIndex, int32 MaterialSlotIndex) const
+void SCustomizableObjectEditorViewportTabBody::ShowGizmoClipMesh(UCustomizableObjectNode& Node, FTransform* Transform, UObject& ClipMesh, int32 LODIndex, int32 SectionIndex, int32 MaterialSlotIndex) const
 {
 	if (const TSharedPtr<SWidget> Toolbar = ViewportToolbarTransformWidget.Pin())
 	{
 		Toolbar->SetVisibility(EVisibility::Visible);		
 	}
 	
-	LevelViewportClient->ShowGizmoClipMesh(ClipMeshNode, ClipMesh, LODIndex, SectionIndex, MaterialSlotIndex);
+	LevelViewportClient->ShowGizmoClipMesh(Node, Transform, ClipMesh, LODIndex, SectionIndex, MaterialSlotIndex);
 }
 
 
diff --git a/Engine/Plugins/Mutable/Source/CustomizableObjectEditor/Private/MuCOE/SCustomizableObjectEditorViewport.h b/Engine/Plugins/Mutable/Source/CustomizableObjectEditor/Private/MuCOE/SCustomizableObjectEditorViewport.h
--- a/Engine/Plugins/Mutable/Source/CustomizableObjectEditor/Private/MuCOE/SCustomizableObjectEditorViewport.h
+++ b/Engine/Plugins/Mutable/Source/CustomizableObjectEditor/Private/MuCOE/SCustomizableObjectEditorViewport.h
@@ -22,6 +22,7 @@ class FUICommandList;
 class STextComboBox;
 class SWidget;
 class UCustomizableObject;
+class UCustomizableObjectNode;
 class UCustomizableObjectNodeProjectorConstant;
 class UCustomizableObjectNodeProjectorParameter;
 class UDebugSkelMeshComponent;
@@ -117,7 +118,7 @@ public:
 	void HideGizmoClipMorph() const;
 
 	/** Do not call directly. Use ICustomizableObjectEditor functions instead. */
-	void ShowGizmoClipMesh(UCustomizableObjectNodeModifierClipWithMesh& ClipMeshNode, UObject& ClipMesh, int32 LODIndex, int32 SectionIndex, int32 MaterialSlotIndex) const;
+	void ShowGizmoClipMesh(UCustomizableObjectNode& Node, FTransform* Transform, UObject& ClipMesh, int32 LODIndex, int32 SectionIndex, int32 MaterialSlotIndex) const;
 
 	/** Do not call directly. Use ICustomizableObjectEditor functions instead. */
 	void HideGizmoClipMesh() const;
diff --git a/Engine/Plugins/Mutable/Source/CustomizableObjectEditor/Private/MuCOE/SMutableCodeViewer.cpp b/Engine/Plugins/Mutable/Source/CustomizableObjectEditor/Private/MuCOE/SMutableCodeViewer.cpp
--- a/Engine/Plugins/Mutable/Source/CustomizableObjectEditor/Private/MuCOE/SMutableCodeViewer.cpp
+++ b/Engine/Plugins/Mutable/Source/CustomizableObjectEditor/Private/MuCOE/SMutableCodeViewer.cpp
@@ -2386,6 +2386,10 @@ bool SMutableCodeViewer::IsConstantResourceUsedByOperation(const int32 IndexOnCo
 				{
 					bResourceLocated = IndexOnConstantsArray == InProgram.GetOpArgs<mu::OP::MeshTransformArgs>(OperationAddress).matrix;
 				}
+				else if (OperationType == mu::OP_TYPE::ME_TRANSFORMWITHMESH)
+				{
+					bResourceLocated = IndexOnConstantsArray == InProgram.GetOpArgs<mu::OP::MeshTransformWithinMeshArgs>(OperationAddress).matrix;
+				}
 				break;
 			}
 
diff --git a/Engine/Plugins/Mutable/Source/CustomizableObjectEditor/Private/MuCOE/SMutableGraphViewer.cpp b/Engine/Plugins/Mutable/Source/CustomizableObjectEditor/Private/MuCOE/SMutableGraphViewer.cpp
--- a/Engine/Plugins/Mutable/Source/CustomizableObjectEditor/Private/MuCOE/SMutableGraphViewer.cpp
+++ b/Engine/Plugins/Mutable/Source/CustomizableObjectEditor/Private/MuCOE/SMutableGraphViewer.cpp
@@ -80,6 +80,7 @@
 #include "MuT/NodeImageInvertPrivate.h"
 #include "MuT/NodeImageSwizzlePrivate.h"
 #include "MuT/NodeImageMultiLayerPrivate.h"
+#include "MuT/NodeModifierMeshTransformInMesh.h"
 #include "Widgets/MutableExpanderArrow.h"
 
 class FExtender;
@@ -498,6 +499,13 @@ void SMutableGraphViewer::GetChildrenForInfo(TSharedPtr<FMutableGraphTreeElement
 		AddChildFunc(ModifierMeshClipWithUVMaskVar->ClipLayout.get(), FString::Printf(TEXT("CLIP LAYOUT")));
 	}
 
+	else if (ParentNode->GetType() == mu::NodeModifierMeshTransformInMesh::GetStaticType())
+	{
+		mu::NodeModifierMeshTransformInMesh* ModifierMeshTransformInMeshVar = StaticCast<mu::NodeModifierMeshTransformInMesh*>(ParentNode);
+		AddChildFunc(ModifierMeshTransformInMeshVar->BoundingMesh.get(), FString::Printf(TEXT("BOUNDING MESH")));
+		AddChildFunc(ModifierMeshTransformInMeshVar->MatrixNode.get(), FString::Printf(TEXT("MESH TRANSFORM")));
+	}
+
 	else if (ParentNode->GetType() == mu::NodeImageSwitch::GetStaticType())
 	{
 		mu::NodeImageSwitch* ImageSwitchVar = StaticCast<mu::NodeImageSwitch*>(ParentNode);
diff --git a/Engine/Plugins/Mutable/Source/CustomizableObjectEditor/Public/MuCOE/CustomizableObjectEditorViewportClient.h b/Engine/Plugins/Mutable/Source/CustomizableObjectEditor/Public/MuCOE/CustomizableObjectEditorViewportClient.h
--- a/Engine/Plugins/Mutable/Source/CustomizableObjectEditor/Public/MuCOE/CustomizableObjectEditorViewportClient.h
+++ b/Engine/Plugins/Mutable/Source/CustomizableObjectEditor/Public/MuCOE/CustomizableObjectEditorViewportClient.h
@@ -18,8 +18,8 @@ class FViewport;
 class ICustomizableObjectInstanceEditor;
 class UAnimationAsset;
 class UCustomizableObject;
+class UCustomizableObjectNode;
 class UCustomizableObjectNodeModifierClipMorph;
-class UCustomizableObjectNodeModifierClipWithMesh;
 class UCustomizableObjectNodeProjectorConstant;
 class UCustomizableObjectNodeProjectorParameter;
 class UDebugSkelMeshComponent;
@@ -188,7 +188,7 @@ public:
 	void HideGizmoClipMorph();
 	
 	/** Do not call directly. Use ICustomizableObjectEditor functions instead. */
-	void ShowGizmoClipMesh(UCustomizableObjectNodeModifierClipWithMesh& ClipMeshNode, UObject& ClipMesh, int32 LODIndex, int32 SectionIndex, int32 MaterialSlotIndex);
+	void ShowGizmoClipMesh(UCustomizableObjectNode& ClipMeshNode, FTransform* ClipMeshTransform, UObject& ClipMesh, int32 LODIndex, int32 SectionIndex, int32 MaterialSlotIndex);
 
 	/** Do not call directly. Use ICustomizableObjectEditor functions instead. */
 	void HideGizmoClipMesh();
@@ -355,7 +355,8 @@ private:
 	FVector ClipMorphYAxis;
 	float MorphLength;
 	FBoxSphereBounds MorphBounds;
-	UCustomizableObjectNodeModifierClipWithMesh* ClipMeshNode;
+	TObjectPtr<UObject> ClipMeshNode;
+	FTransform* ClipMeshTransform = nullptr;
 	TObjectPtr<UMaterial> ClipMeshMaterial;
 	UStaticMeshComponent* ClipMeshStaticMeshComp;
 	USkeletalMeshComponent* ClipMeshSkeletalMeshComp;
diff --git a/Engine/Plugins/Mutable/Source/CustomizableObjectEditor/Public/MuCOE/EdGraphSchema_CustomizableObject.h b/Engine/Plugins/Mutable/Source/CustomizableObjectEditor/Public/MuCOE/EdGraphSchema_CustomizableObject.h
--- a/Engine/Plugins/Mutable/Source/CustomizableObjectEditor/Public/MuCOE/EdGraphSchema_CustomizableObject.h
+++ b/Engine/Plugins/Mutable/Source/CustomizableObjectEditor/Public/MuCOE/EdGraphSchema_CustomizableObject.h
@@ -107,6 +107,7 @@ public:
 	static const FName PC_MaterialAsset;
 	static const FName PC_Wildcard;
 	static const FName PC_PoseAsset;
+	static const FName PC_Transform;
 
 	// Node categories
 	static const FText NC_Experimental;
diff --git a/Engine/Plugins/Mutable/Source/CustomizableObjectEditor/Public/MuCOE/ICustomizableObjectInstanceEditor.h b/Engine/Plugins/Mutable/Source/CustomizableObjectEditor/Public/MuCOE/ICustomizableObjectInstanceEditor.h
--- a/Engine/Plugins/Mutable/Source/CustomizableObjectEditor/Public/MuCOE/ICustomizableObjectInstanceEditor.h
+++ b/Engine/Plugins/Mutable/Source/CustomizableObjectEditor/Public/MuCOE/ICustomizableObjectInstanceEditor.h
@@ -12,15 +12,14 @@ class UCustomizableObject;
 class UCustomizableObjectInstance;
 class UProjectorParameter;
 class UCustomSettings;
-class UCustomizableObjectNodeModifierClipWithMesh;
+class UCustomizableObjectNode;
 class UCustomizableObjectNodeModifierClipMorph;
-class UCustomizableObjectNodeModifierClipWithMesh;
+class UEdGraphPin;
 class ULightComponent;
 class UPoseAsset;
 class UCustomizableObjectEditorProperties;
 class SCustomizableObjectEditorAdvancedPreviewSettings;
 
-
 /**
  * Public interface to Customizable Object Instance Editor
  */
@@ -76,9 +75,8 @@ public:
 	 * Synchronizes all editor widgets so that the selection and the widgets is always consistent. */
 	virtual void HideGizmoClipMorph() {}
 	
-	/** Show the clipping mesh gizmo from the NodeMeshClipWithMesh.
-	 * Synchronizes all editor widgets so that the selection and the widgets is always consistent. */
-	virtual void ShowGizmoClipMesh(UCustomizableObjectNodeModifierClipWithMesh& ClipMeshNode) {}
+	/** Show the clipping mesh gizmo from a node that uses a transformable bounding mesh. */
+	virtual void ShowGizmoClipMesh(UCustomizableObjectNode& Node, FTransform* Transform, const UEdGraphPin& MeshPin) {}
 
 	/** Hide the clipping mesh gizmo from the NodeMeshClipWithMesh.
 	 * Synchronizes all editor widgets so that the selection and the widgets is always consistent. */
diff --git a/Engine/Plugins/Mutable/Source/CustomizableObjectEditor/Public/MuCOE/Nodes/CustomizableObjectNodeModifierTransformInMesh.h b/Engine/Plugins/Mutable/Source/CustomizableObjectEditor/Public/MuCOE/Nodes/CustomizableObjectNodeModifierTransformInMesh.h
new file mode 100644
--- /dev/null
+++ b/Engine/Plugins/Mutable/Source/CustomizableObjectEditor/Public/MuCOE/Nodes/CustomizableObjectNodeModifierTransformInMesh.h
@@ -0,0 +1,42 @@
+﻿// Copyright Epic Games, Inc. All Rights Reserved.
+
+#pragma once
+
+#include "MuCOE/Nodes/CustomizableObjectNodeModifierBase.h"
+#include "MuT/NodeModifier.h"
+
+#include "CustomizableObjectNodeModifierTransformInMesh.generated.h"
+
+namespace ENodeTitleType { enum Type : int; }
+
+class UCustomizableObject;
+class UCustomizableObjectNodeRemapPins;
+class UEdGraphPin;
+
+UCLASS()
+class CUSTOMIZABLEOBJECTEDITOR_API UCustomizableObjectNodeModifierTransformInMesh : public UCustomizableObjectNodeModifierBase
+{
+	GENERATED_BODY()
+
+public:
+	/** Transform to apply to the bounding mesh before selecting for vertices to transform. */
+	UPROPERTY(EditAnywhere, Category = BoundingMesh)
+	FTransform BoundingMeshTransform = FTransform::Identity;
+
+public:
+	// EdGraphNode interface
+	virtual FText GetNodeTitle(ENodeTitleType::Type TitleType) const override;
+	virtual FText GetTooltipText() const override;
+
+	// UCustomizableObjectNode interface
+	virtual void AllocateDefaultPins(UCustomizableObjectNodeRemapPins* RemapPins) override;
+	virtual bool IsExperimental() const override { return true; }
+
+	// UCustomizableObjectNodeModifierBase interface
+	UEdGraphPin* OutputPin() const override;
+
+	// Own interface
+	UEdGraphPin* BoundingMeshPin() const;
+	UEdGraphPin* TransformPin() const;
+};
+
diff --git a/Engine/Plugins/Mutable/Source/CustomizableObjectEditor/Public/MuCOE/Nodes/CustomizableObjectNodeTransformConstant.h b/Engine/Plugins/Mutable/Source/CustomizableObjectEditor/Public/MuCOE/Nodes/CustomizableObjectNodeTransformConstant.h
new file mode 100644
--- /dev/null
+++ b/Engine/Plugins/Mutable/Source/CustomizableObjectEditor/Public/MuCOE/Nodes/CustomizableObjectNodeTransformConstant.h
@@ -0,0 +1,80 @@
+﻿// Copyright Epic Games, Inc. All Rights Reserved.
+
+#pragma once
+
+#include "MuCOE/Nodes/CustomizableObjectNode.h"
+#include "MuCOE/Nodes/SCustomizableObjectNode.h"
+
+#include "CustomizableObjectNodeTransformConstant.generated.h"
+
+
+class UCustomizableObjectNodeTransformConstant;
+
+
+// Class create color inputs in the NodeColorConstant
+class SGraphNodeTransformConstant : public SCustomizableObjectNode
+{
+public:
+
+	SLATE_BEGIN_ARGS(SGraphNodeTransformConstant) {}
+	SLATE_END_ARGS();
+
+	SGraphNodeTransformConstant() = default;
+
+	// Builds the SGraphNodeFloatConstant when needed
+	void Construct(const FArguments& InArgs, UEdGraphNode* InGraphNode);
+
+	// Overriden functions to build the SGraphNode widgets
+	virtual void SetDefaultTitleAreaWidget(TSharedRef<SOverlay> DefaultTitleAreaWidget) override;
+	virtual void CreateBelowPinControls(TSharedPtr<SVerticalBox> MainBox) override;
+	virtual bool ShouldAllowCulling() const override { return false; }
+
+	// Callbacks for the widget
+	void OnExpressionPreviewChanged(const ECheckBoxState NewCheckedState);
+	ECheckBoxState IsExpressionPreviewChecked() const;
+	const FSlateBrush* GetExpressionPreviewArrow() const;
+
+private:
+
+	// Callback for when the location part of the transform changes.
+	void OnLocationChanged(double Value, ETextCommit::Type, EAxis::Type Axis);
+
+	// Callback for when the rotation part of the transform changes.
+	void OnRotationChanged(double Value, ETextCommit::Type, EAxis::Type Axis);
+
+	// Callback for when the scale part of the transform changes.
+	void OnScaleChanged(double Value, ETextCommit::Type, EAxis::Type Axis);
+
+private:
+	// Pointer to the transform constant node that owns this SGraphNode
+	UCustomizableObjectNodeTransformConstant* NodeTransformConstant = nullptr;
+
+};
+
+
+UCLASS()
+class CUSTOMIZABLEOBJECTEDITOR_API UCustomizableObjectNodeTransformConstant : public UCustomizableObjectNode
+{
+	GENERATED_BODY()
+
+public:	
+	/**  */
+	UPROPERTY(EditAnywhere, Category=CustomizableObject, meta = (DontUpdateWhileEditing))
+	FTransform Value = FTransform::Identity;
+
+	// Begin EdGraphNode interface
+	FText GetNodeTitle(ENodeTitleType::Type TitleType) const override;
+	FLinearColor GetNodeTitleColor() const override;
+	FText GetTooltipText() const override;
+	// End EdGraphNode interface
+
+	// UCustomizableObjectNode interface
+	void AllocateDefaultPins(UCustomizableObjectNodeRemapPins* RemapPins) override;
+	bool IsAffectedByLOD() const override { return false; }
+
+	// Creates the SGraph Node widget for the Color Editor
+	TSharedPtr<SGraphNode> CreateVisualWidget() override;
+
+	// Determines if the Node is collapsed or not
+	bool bCollapsed = true;
+};
diff --git a/Engine/Plugins/Mutable/Source/CustomizableObjectEditor/Public/MuCOE/Nodes/CustomizableObjectNodeTransformParameter.h b/Engine/Plugins/Mutable/Source/CustomizableObjectEditor/Public/MuCOE/Nodes/CustomizableObjectNodeTransformParameter.h
new file mode 100644
--- /dev/null
+++ b/Engine/Plugins/Mutable/Source/CustomizableObjectEditor/Public/MuCOE/Nodes/CustomizableObjectNodeTransformParameter.h
@@ -0,0 +1,38 @@
+﻿// Copyright Epic Games, Inc. All Rights Reserved.
+
+#pragma once
+
+#include "MuCOE/Nodes/CustomizableObjectNode.h"
+
+#include "CustomizableObjectNodeTransformParameter.generated.h"
+
+namespace ENodeTitleType { enum Type : int; }
+
+
+UCLASS()
+class CUSTOMIZABLEOBJECTEDITOR_API UCustomizableObjectNodeTransformParameter : public UCustomizableObjectNode
+{
+	GENERATED_BODY()
+
+public:
+	UPROPERTY(EditAnywhere, Category=CustomizableObject, meta = (DontUpdateWhileEditing))
+	FTransform DefaultValue = FTransform::Identity;
+
+	UPROPERTY(EditAnywhere, Category=CustomizableObject)
+	FString ParameterName = "Default Name";
+
+	UPROPERTY(EditAnywhere, Category = UI, meta = (DisplayName = "Parameter UI Metadata"))
+	FMutableParamUIMetadata ParamUIMetadata;
+
+	// Begin EdGraphNode interface
+	FText GetNodeTitle(ENodeTitleType::Type TitleType) const override;
+	FLinearColor GetNodeTitleColor() const override;
+	FText GetTooltipText() const override;
+	void OnRenameNode(const FString& NewName) override;
+	bool GetCanRenameNode() const override { return true; }
+
+	// UCustomizableObjectNode interface
+	void AllocateDefaultPins(UCustomizableObjectNodeRemapPins* RemapPins) override;
+	bool IsAffectedByLOD() const override { return false; }
+	
+};
diff --git a/Engine/Plugins/Mutable/Source/MutableRuntime/Internal/MuR/Operations.h b/Engine/Plugins/Mutable/Source/MutableRuntime/Internal/MuR/Operations.h
--- a/Engine/Plugins/Mutable/Source/MutableRuntime/Internal/MuR/Operations.h
+++ b/Engine/Plugins/Mutable/Source/MutableRuntime/Internal/MuR/Operations.h
@@ -40,6 +40,7 @@ namespace mu
         PR_CONSTANT,
         ST_CONSTANT,
 		ED_CONSTANT,
+    	MA_CONSTANT,
 
         //! User parameter
         BO_PARAMETER,
@@ -49,6 +50,7 @@ namespace mu
         PR_PARAMETER,
         IM_PARAMETER,
         ST_PARAMETER,
+    	MA_PARAMETER,
 
 		//! A referenced, but opaque engine resource
 		IM_REFERENCE,
@@ -293,6 +295,9 @@ namespace mu
 		//! Add a set of tags to a mesh
 		ME_ADDTAGS,
 
+    	//! Transform with a 4x4 matrix the geometry channels of a mesh that are bounded by another mesh
+    	ME_TRANSFORMWITHMESH,
+
         //-----------------------------------------------------------------------------------------
         // Instance operations
         //-----------------------------------------------------------------------------------------
@@ -411,6 +416,11 @@ namespace mu
             float value[4];
         };
 
+    	struct MatrixConstantArgs
+    	{
+    		ADDRESS value;
+    	};
+
         struct ResourceConstantArgs
         {
             ADDRESS value;
@@ -987,6 +997,13 @@ namespace mu
 			EFaceCullStrategy FaceCullStrategy;
 		};
 
+    	struct MeshTransformWithinMeshArgs
+    	{
+    		ADDRESS sourceMesh;
+    		ADDRESS boundingMesh;
+    		ADDRESS matrix;
+    	};
+
 		struct MeshOptimizeSkinningArgs
 		{
 			ADDRESS source;
@@ -1099,6 +1116,7 @@ namespace mu
             IntConstantArgs IntConstant;
             ScalarConstantArgs ScalarConstant;
             ColourConstantArgs ColourConstant;
+        	MatrixConstantArgs TransformConstant;
             //ParameterArgs Parameter;
 
             //-------------------------------------------------------------------------------------
@@ -1178,9 +1196,9 @@ namespace mu
 		DT_PROJECTOR,
 		DT_STRING,
 		DT_EXTENSION_DATA,
+		DT_MATRIX,
 
 		// Supporting data types : Never returned as an actual data type for any operation.
-		DT_MATRIX,
 		DT_SHAPE,
 		DT_CURVE,
 		DT_SKELETON,
diff --git a/Engine/Plugins/Mutable/Source/MutableRuntime/Internal/MuR/ParametersPrivate.h b/Engine/Plugins/Mutable/Source/MutableRuntime/Internal/MuR/ParametersPrivate.h
--- a/Engine/Plugins/Mutable/Source/MutableRuntime/Internal/MuR/ParametersPrivate.h
+++ b/Engine/Plugins/Mutable/Source/MutableRuntime/Internal/MuR/ParametersPrivate.h
@@ -255,9 +255,10 @@ namespace Private
 	using ParamProjectorType = Private::TIndirectObject<FProjector>;
 	using ParamImageType = FName;
 	using ParamStringType = Private::TIndirectObject<FString>;
+	using ParamMatrixType = Private::TIndirectObject<FMatrix44f>;
 	
 	using PARAMETER_VALUE = TVariant<
-            ParamBoolType, ParamIntType, ParamFloatType, ParamColorType, ParamProjectorType, ParamImageType, ParamStringType>;
+            ParamBoolType, ParamIntType, ParamFloatType, ParamColorType, ParamProjectorType, ParamImageType, ParamStringType, ParamMatrixType>;
 
     // static_assert to track PARAMETER_VALUE size changes. It is ok to change if needed.
     static_assert(sizeof(PARAMETER_VALUE) == 8*4, "PARAMETER_VALUE size has changed.");
diff --git a/Engine/Plugins/Mutable/Source/MutableRuntime/Internal/MuR/SystemPrivate.h b/Engine/Plugins/Mutable/Source/MutableRuntime/Internal/MuR/SystemPrivate.h
--- a/Engine/Plugins/Mutable/Source/MutableRuntime/Internal/MuR/SystemPrivate.h
+++ b/Engine/Plugins/Mutable/Source/MutableRuntime/Internal/MuR/SystemPrivate.h
@@ -511,6 +511,7 @@ namespace mu
 		TMemoryTrackedArray<FProjector> ProjectorResults;
 		TMemoryTrackedArray<Ptr<const String>> StringResults;
 		TMemoryTrackedArray<Ptr<const ExtensionData>> ExtensionDataResults;
+		TMemoryTrackedArray<FMatrix44f> MatrixResults;
 
 		/** */
 		inline const ExecutionIndex& GetRangeIndex(uint32_t i)
@@ -567,6 +568,7 @@ namespace mu
 				InstanceResults.Add(nullptr);
 				ProjectorResults.Add(FProjector());
 				StringResults.Add(nullptr);
+				MatrixResults.Add(FMatrix44f());
 				ExtensionDataResults.Add(nullptr);
 			}
 		}
@@ -723,6 +725,16 @@ namespace mu
 			return ColorResults[Data->DataTypeIndex];
 		}
 
+		FMatrix44f GetMatrix(FCacheAddress at)
+		{
+			if (!at.At) return FMatrix44f::Identity;
+			const FOpExecutionData* Data = OpExecutionData.get_ptr(at);
+			if (!Data) return FMatrix44f::Identity;
+
+			check(Data->DataType == DATATYPE::DT_MATRIX);
+			return MatrixResults[Data->DataTypeIndex];
+		}
+
 		FProjector GetProjector(FCacheAddress at)
 		{
 			if (!at.At) return FProjector();
@@ -897,6 +909,26 @@ namespace mu
 			check(Data->DataTypeIndex != 0);
 		}
 
+		void SetMatrix(FCacheAddress at, const FMatrix44f& v)
+		{
+			check(at.At < OpExecutionData.size_code());
+			FOpExecutionData* Data = OpExecutionData.get_ptr(at);
+			check(Data->DataType == DATATYPE::DT_MATRIX || Data->DataType == DATATYPE::DT_NONE);
+			Data->DataType = DATATYPE::DT_MATRIX;
+			Data->IsValueValid = true;
+
+			if (!Data->DataTypeIndex)
+			{
+				Data->DataTypeIndex = MatrixResults.Num();
+				MatrixResults.Add(v);
+			}
+			else
+			{
+				MatrixResults[Data->DataTypeIndex] = v;
+			}
+			check(Data->DataTypeIndex != 0);
+		}
+		
 		void SetProjector(FCacheAddress at, const FProjector& v)
 		{
 			check(at.At < OpExecutionData.size_code());
diff --git a/Engine/Plugins/Mutable/Source/MutableRuntime/Private/MuR/CodeRunner.cpp b/Engine/Plugins/Mutable/Source/MutableRuntime/Private/MuR/CodeRunner.cpp
--- a/Engine/Plugins/Mutable/Source/MutableRuntime/Private/MuR/CodeRunner.cpp
+++ b/Engine/Plugins/Mutable/Source/MutableRuntime/Private/MuR/CodeRunner.cpp
@@ -2,6 +2,7 @@
 
 #include "MuR/CodeRunner.h"
 
+#include "OpMeshTransformWithMesh.h"
 #include "GenericPlatform/GenericPlatformMath.h"
 #include "HAL/UnrealMemory.h"
 #include "Logging/LogCategory.h"
@@ -3080,6 +3081,67 @@ namespace mu
 			break;
 		}
 
+        case OP_TYPE::ME_TRANSFORMWITHMESH:
+		{
+			OP::MeshTransformWithinMeshArgs args = Program.GetOpArgs<OP::MeshTransformWithinMeshArgs>(item.At);
+			switch (item.Stage)
+			{
+			case 0:
+			{
+				if (args.sourceMesh)
+				{
+					AddOp(FScheduledOp(item.At, item, 1),
+						FScheduledOp(args.sourceMesh, item),
+						FScheduledOp(args.boundingMesh, item),
+						FScheduledOp(args.matrix, item));
+				}
+				else
+				{
+					StoreMesh(item, nullptr);
+				}
+				break;
+			}
+			case 1:
+			{
+				MUTABLE_CPUPROFILER_SCOPE(ME_TRANSFORMWITHMESH_1)
+            
+				Ptr<const Mesh> SourceMesh = LoadMesh(FCacheAddress(args.sourceMesh,item));
+				Ptr<const Mesh> BoundingMesh = LoadMesh(FCacheAddress(args.boundingMesh, item));
+				const FMatrix44f& Transform = LoadMatrix(FCacheAddress(args.matrix, item));
+
+				if (SourceMesh)
+				{
+					Ptr<Mesh> Result = CreateMesh(SourceMesh->GetDataSize());
+
+					bool bOutSuccess = false;
+					MeshTransformWithMesh(Result.get(), SourceMesh.get(), BoundingMesh.get(), Transform, bOutSuccess);
+					Release(BoundingMesh);
+
+					if (!bOutSuccess)
+					{
+						Release(Result);
+						StoreMesh(item, SourceMesh);
+					}
+					else
+					{
+						Release(SourceMesh);
+						StoreMesh(item, Result);
+					}
+				}
+				else
+				{
+					Release(BoundingMesh);
+					StoreMesh(item, SourceMesh);
+				}
+				break;
+			}
+
+			default:
+				check(false);
+			}
+			break;
+		}
+
         default:
             if (type!=OP_TYPE::NONE)
             {
@@ -5857,6 +5919,36 @@ namespace mu
         }
     }
 
+    //---------------------------------------------------------------------------------------------
+    void CodeRunner::RunCode_Matrix(const FScheduledOp& item, const Parameters* pParams, const Model* pModel )
+    {
+		MUTABLE_CPUPROFILER_SCOPE(RunCode_Transform);
+
+		const FProgram& Program = m_pModel->GetPrivate()->m_program;
+
+		OP_TYPE type = Program.GetOpType(item.At);
+
+		switch ( type )
+		{
+		case OP_TYPE::MA_CONSTANT:
+			{
+				OP::MatrixConstantArgs args = Program.GetOpArgs<OP::MatrixConstantArgs>(item.At);
+				StoreMatrix( item, Program.m_constantMatrices[args.value] );
+				break;
+			}
+
+		case OP_TYPE::MA_PARAMETER:
+			{
+				OP::ParameterArgs args = Program.GetOpArgs<OP::ParameterArgs>(item.At);
+				Ptr<RangeIndex> index = BuildCurrentOpRangeIndex( item, pParams, pModel, args.variable );
+				FMatrix44f Value;
+				pParams->GetMatrixValue( args.variable, Value, index );
+				StoreMatrix( item, Value );
+				break;
+			}
+		}
+    }
+
 
     //---------------------------------------------------------------------------------------------
     void CodeRunner::RunCode_Layout(const FScheduledOp& item, const Model* pModel )
@@ -6132,6 +6224,10 @@ namespace mu
 				RunCode_Colour(item, pParams, pModel);
 				break;
 
+			case DT_MATRIX:
+				RunCode_Matrix(item, pParams, pModel);
+				break;
+				
 			default:
 				check(false);
 				break;
diff --git a/Engine/Plugins/Mutable/Source/MutableRuntime/Private/MuR/CodeRunner.h b/Engine/Plugins/Mutable/Source/MutableRuntime/Private/MuR/CodeRunner.h
--- a/Engine/Plugins/Mutable/Source/MutableRuntime/Private/MuR/CodeRunner.h
+++ b/Engine/Plugins/Mutable/Source/MutableRuntime/Private/MuR/CodeRunner.h
@@ -136,6 +136,7 @@ namespace  mu
         void RunCode_String(const FScheduledOp&, const Parameters*, const Model* );
         void RunCode_Colour(const FScheduledOp&, const Parameters*, const Model* );
         void RunCode_Projector(const FScheduledOp&, const Parameters*, const Model* );
+        void RunCode_Matrix(const FScheduledOp&, const Parameters*, const Model* );
 
 		void RunCodeImageDesc(const FScheduledOp&, const Parameters*, const Model*, uint32 LodMask);
 
@@ -468,6 +469,11 @@ namespace  mu
 			return m_pSystem->WorkingMemoryManager.CurrentInstanceCache->GetColour(From);
 		}
 
+    	inline FMatrix44f LoadMatrix(const FCacheAddress& From)
+		{
+			return m_pSystem->WorkingMemoryManager.CurrentInstanceCache->GetMatrix(From);
+		}
+
 		inline Ptr<const String> LoadString(const FCacheAddress& From)
 		{
 			return m_pSystem->WorkingMemoryManager.CurrentInstanceCache->GetString(From);
@@ -533,6 +539,11 @@ namespace  mu
 			m_pSystem->WorkingMemoryManager.CurrentInstanceCache->SetColour(To, Value);
 		}
 
+    	inline void StoreMatrix(const FCacheAddress& To, const FMatrix44f& Value)
+		{
+			m_pSystem->WorkingMemoryManager.CurrentInstanceCache->SetMatrix(To, Value);
+		}
+    	
 		inline void StoreProjector(const FCacheAddress& To, const FProjector& Value)
 		{
 			m_pSystem->WorkingMemoryManager.CurrentInstanceCache->SetProjector(To, Value);
diff --git a/Engine/Plugins/Mutable/Source/MutableRuntime/Private/MuR/Model.cpp b/Engine/Plugins/Mutable/Source/MutableRuntime/Private/MuR/Model.cpp
--- a/Engine/Plugins/Mutable/Source/MutableRuntime/Private/MuR/Model.cpp
+++ b/Engine/Plugins/Mutable/Source/MutableRuntime/Private/MuR/Model.cpp
@@ -382,9 +382,24 @@ namespace mu
 		}
     }
 
+	FMatrix44f Model::GetMatrixDefaultValue(int32 Index) const
+	{
+    	check(m_pD->m_program.m_parameters.IsValidIndex(Index));
+    	check(m_pD->m_program.m_parameters[Index].m_type == PARAMETER_TYPE::T_MATRIX);
+
+    	// Early out in case of invalid parameters
+    	if (!m_pD->m_program.m_parameters.IsValidIndex(Index) ||
+			m_pD->m_program.m_parameters[Index].m_type != PARAMETER_TYPE::T_MATRIX)
+    	{
+    		return FMatrix44f::Identity;
+    	}
+
+    	return m_pD->m_program.m_parameters[Index].m_defaultValue.Get<ParamMatrixType>();
+	}
+
 
 	void Model::GetProjectorDefaultValue(int32 Index, PROJECTOR_TYPE* OutProjectionType, FVector3f* OutPos,
-		FVector3f* OutDir, FVector3f* OutUp, FVector3f* OutScale, float* OutProjectionAngle) const
+	                                     FVector3f* OutDir, FVector3f* OutUp, FVector3f* OutScale, float* OutProjectionAngle) const
 	{
     	check(m_pD->m_program.m_parameters.IsValidIndex(Index));
 		check(m_pD->m_program.m_parameters[Index].m_type == PARAMETER_TYPE::T_PROJECTOR);
diff --git a/Engine/Plugins/Mutable/Source/MutableRuntime/Private/MuR/OpMeshClipWithMesh.cpp b/Engine/Plugins/Mutable/Source/MutableRuntime/Private/MuR/OpMeshClipWithMesh.cpp
--- a/Engine/Plugins/Mutable/Source/MutableRuntime/Private/MuR/OpMeshClipWithMesh.cpp
+++ b/Engine/Plugins/Mutable/Source/MutableRuntime/Private/MuR/OpMeshClipWithMesh.cpp
@@ -320,8 +320,187 @@
         return NumIntersections;
     }
 
+	void MeshUVMaskClassifyVertices(TBitArray<>& VertexClipped, const Mesh* Base, const Image* Mask, const uint8 LayoutIndex)
+	{
+		using namespace UE::Geometry;
+
+		MUTABLE_CPUPROFILER_SCOPE(MeshUVMaskClassifyVertices);
+
+		uint32 VertexCount = Base->GetVertexCount();
+
+		// Stores whether each vertex in the original mesh in the clip mesh volume
+		VertexClipped.SetNum(VertexCount, false);
+
+		// Now go through all vertices in the mesh and record whether they are inside or outside of the ClipMesh
+		const FMeshBufferSet& MBSPriv = Base->GetVertexBuffers();
+		for (int32 b = 0; b < MBSPriv.Buffers.Num(); ++b)
+		{
+			const TArray<mu::FMeshBufferChannel>& Channels = MBSPriv.Buffers[b].Channels;
+
+			for (int32 c = 0; c < Channels.Num(); ++c)
+			{
+				EMeshBufferSemantic Sem = Channels[c].Semantic;
+				if (Sem != MBS_TEXCOORDS)
+				{
+					continue;
+				}
+
+				int32 SemIndex = Channels[c].SemanticIndex;
+				if (SemIndex != LayoutIndex)
+				{
+					continue;
+				}
+
+				UntypedMeshBufferIteratorConst It(Base->GetVertexBuffers(), Sem, SemIndex);
+				for (uint32 V = 0; V < VertexCount; ++V)
+				{
+					// \TODO: This could be optimized.
+					FVector2f UV = It.GetAsVec2f();
+
+					// \TODO: This could also be optimized
+					FVector4f Color = Mask->Sample(UV);
+
+					VertexClipped[V] = Color[0] >= 0.5f;
+
+					++It;
+				}
+
+				break;
+			}
+		}
+	}
+
+
+	void MeshLayoutMaskClassifyVertices(TBitArray<>& VertexClipped, const Mesh* Base, const Layout* Mask, const uint8 LayoutIndex)
+	{
+		using namespace UE::Geometry;
+
+		MUTABLE_CPUPROFILER_SCOPE(MeshLayoutMaskClassifyVertices);
+
+		uint32 VertexCount = Base->GetVertexCount();
+
+		// Stores whether each vertex in the original mesh in the clip mesh volume
+		VertexClipped.SetNum(VertexCount, false);
+
+		// Now go through all vertices in the mesh and record whether they are inside or outside of the ClipMesh
+		const FMeshBufferSet& MBSPriv = Base->GetVertexBuffers();
+		for (int32 b = 0; b < MBSPriv.Buffers.Num(); ++b)
+		{
+			const TArray<mu::FMeshBufferChannel>& Channels = MBSPriv.Buffers[b].Channels;
+
+			for (int32 c = 0; c < Channels.Num(); ++c)
+			{
+				EMeshBufferSemantic Sem = Channels[c].Semantic;
+				if (Sem != MBS_TEXCOORDS)
+				{
+					continue;
+				}
+
+				int32 SemIndex = Channels[c].SemanticIndex;
+				if (SemIndex != LayoutIndex)
+				{
+					continue;
+				}
+
+				UntypedMeshBufferIteratorConst It(Base->GetVertexBuffers(), Sem, SemIndex);
+				for (uint32 V = 0; V < VertexCount; ++V)
+				{
+					// \TODO: This could be optimized.
+					FVector2f UV = It.GetAsVec2f();
+
+					FVector2f Cell = UV * FVector2f(Mask->Size.X, Mask->Size.Y);
+
+					// \TODO: This could also be optimized
+					for (const FLayoutBlock& Block : Mask->Blocks)
+					{
+						if (
+							float(Block.Min.X) <= Cell.X && float(Block.Min.Y) <= Cell.Y
+							&&
+							float(Block.Min.X + Block.Size.X) >= Cell.X && float(Block.Min.Y + Block.Size.Y) >= Cell.Y
+							)
+						{
+							VertexClipped[V] = true; 
+							break;
+						}
+					}
+
+					++It;
+				}
+
+				break;
+			}
+		}
+	}
+
+
+	/** Make a mask with the indices of the vertices with 0 in the IncludedVertices array. */
+	void CreateMask(Mesh* Result, const Mesh* Base, const TArray<uint8>& IncludedVertices)
+	{
+		int32 MaskVertexCount = 0;
+		for (uint8 b : IncludedVertices)
+		{
+			if (!b)
+			{
+				++MaskVertexCount;
+			}
+		}
+
+		Result->GetVertexBuffers().SetElementCount(MaskVertexCount);
+		Result->GetVertexBuffers().SetBufferCount(1);
+		Result->MeshIDPrefix = Base->MeshIDPrefix;
+
+		EMeshBufferSemantic Semantic = MBS_VERTEXINDEX;
+		int32 SemanticIndex = 0;
+		int32 Components = 1;
+		int32 Offsets = 0;
+
+		bool bMakeRelativeIds = !Base->AreVertexIdsExplicit();
+		if (bMakeRelativeIds)
+		{
+			EMeshBufferFormat Format = MBF_UINT32;
+			Result->GetVertexBuffers().SetBuffer( 0, sizeof(uint32), 1, &Semantic, &SemanticIndex, &Format, &Components, &Offsets );
+
+			// Fill the buffer
+			uint32* Data = reinterpret_cast<uint32*>(Result->GetVertexBuffers().GetBufferData(0));
+			MeshVertexIdIteratorConst ItBase(Base);
+			for (int32 v = 0; v < IncludedVertices.Num(); ++v)
+			{
+				if (!IncludedVertices[v])
+				{
+					uint64 ID = ItBase.Get();
+					*Data = uint32(ID & 0xffffffff);
+					++Data;
+				}
+				++ItBase;
+			}
+		}
+		else
+		{
+			EMeshBufferFormat Format = MBF_UINT64;
+			Result->GetVertexBuffers().SetBuffer( 0, sizeof(uint64), 1, &Semantic, &SemanticIndex, &Format, &Components, &Offsets );
+
+			// Fill the buffer
+			uint64* Data = reinterpret_cast<uint64*>( Result->GetVertexBuffers().GetBufferData(0) );
+			MeshVertexIdIteratorConst ItBase(Base);
+			for (int32 v = 0; v < IncludedVertices.Num(); ++v)
+			{
+				if (!IncludedVertices[v])
+				{
+					*Data = ItBase.Get();
+					++Data;
+				}
+				++ItBase;
+			}
+		}
+	}
+
+}} // namespace mu::anonymous
+
+
+namespace mu
+{
     /** Core Geometry version */
-    void MeshClipMeshClassifyVertices(TArray<uint8>& VertexInClipMesh, const Mesh* pBase, const Mesh* pClipMesh)
+    void MeshClipMeshClassifyVertices(TBitArray<>& VertexInClipMesh, const Mesh* pBase, const Mesh* pClipMesh)
     {
 		using namespace UE::Geometry;
 
@@ -337,7 +516,7 @@
         int32 OrigVertCount = pBase->GetVertexBuffers().GetElementCount();
 
         // Stores whether each vertex in the original mesh in the clip mesh volume
-        VertexInClipMesh.AddZeroed(OrigVertCount);
+        VertexInClipMesh.SetNum(OrigVertCount, false);
 
         if (VCount == 0)
         {
@@ -506,189 +685,6 @@
         }
     }	
 
-
-	void MeshUVMaskClassifyVertices(TArray<uint8>& VertexClipped, const Mesh* Base, const Image* Mask, const uint8 LayoutIndex)
-	{
-		using namespace UE::Geometry;
-
-		MUTABLE_CPUPROFILER_SCOPE(MeshUVMaskClassifyVertices);
-
-		uint32 VertexCount = Base->GetVertexCount();
-
-		// Stores whether each vertex in the original mesh in the clip mesh volume
-		VertexClipped.SetNumUninitialized(VertexCount,EAllowShrinking::No);
-		FMemory::Memzero(VertexClipped.GetData(),VertexClipped.GetAllocatedSize());
-
-		// Now go through all vertices in the mesh and record whether they are inside or outside of the ClipMesh
-		const FMeshBufferSet& MBSPriv = Base->GetVertexBuffers();
-		for (int32 b = 0; b < MBSPriv.Buffers.Num(); ++b)
-		{
-			const TArray<mu::FMeshBufferChannel>& Channels = MBSPriv.Buffers[b].Channels;
-
-			for (int32 c = 0; c < Channels.Num(); ++c)
-			{
-				EMeshBufferSemantic Sem = Channels[c].Semantic;
-				if (Sem != MBS_TEXCOORDS)
-				{
-					continue;
-				}
-
-				int32 SemIndex = Channels[c].SemanticIndex;
-				if (SemIndex != LayoutIndex)
-				{
-					continue;
-				}
-
-				UntypedMeshBufferIteratorConst It(Base->GetVertexBuffers(), Sem, SemIndex);
-				for (uint32 V = 0; V < VertexCount; ++V)
-				{
-					// \TODO: This could be optimized.
-					FVector2f UV = It.GetAsVec2f();
-
-					// \TODO: This could also be optimized
-					FVector4f Color = Mask->Sample(UV);
-
-					VertexClipped[V] = Color[0] >= 0.5f;
-
-					++It;
-				}
-
-				break;
-			}
-		}
-	}
-
-
-	void MeshLayoutMaskClassifyVertices(TArray<uint8>& VertexClipped, const Mesh* Base, const Layout* Mask, const uint8 LayoutIndex)
-	{
-		using namespace UE::Geometry;
-
-		MUTABLE_CPUPROFILER_SCOPE(MeshUVMaskClassifyVertices);
-
-		uint32 VertexCount = Base->GetVertexCount();
-
-		// Stores whether each vertex in the original mesh in the clip mesh volume
-		VertexClipped.SetNumUninitialized(VertexCount, EAllowShrinking::No);
-		FMemory::Memzero(VertexClipped.GetData(), VertexClipped.Num());
-
-		// Now go through all vertices in the mesh and record whether they are inside or outside of the ClipMesh
-		const FMeshBufferSet& MBSPriv = Base->GetVertexBuffers();
-		for (int32 b = 0; b < MBSPriv.Buffers.Num(); ++b)
-		{
-			const TArray<mu::FMeshBufferChannel>& Channels = MBSPriv.Buffers[b].Channels;
-
-			for (int32 c = 0; c < Channels.Num(); ++c)
-			{
-				EMeshBufferSemantic Sem = Channels[c].Semantic;
-				if (Sem != MBS_TEXCOORDS)
-				{
-					continue;
-				}
-
-				int32 SemIndex = Channels[c].SemanticIndex;
-				if (SemIndex != LayoutIndex)
-				{
-					continue;
-				}
-
-				UntypedMeshBufferIteratorConst It(Base->GetVertexBuffers(), Sem, SemIndex);
-				for (uint32 V = 0; V < VertexCount; ++V)
-				{
-					// \TODO: This could be optimized.
-					FVector2f UV = It.GetAsVec2f();
-
-					FVector2f Cell = UV * FVector2f(Mask->Size.X, Mask->Size.Y);
-
-					// \TODO: This could also be optimized
-					for (const FLayoutBlock& Block : Mask->Blocks)
-					{
-						if (
-							float(Block.Min.X) <= Cell.X && float(Block.Min.Y) <= Cell.Y
-							&&
-							float(Block.Min.X + Block.Size.X) >= Cell.X && float(Block.Min.Y + Block.Size.Y) >= Cell.Y
-							)
-						{
-							VertexClipped[V] = true; 
-							break;
-						}
-					}
-
-					++It;
-				}
-
-				break;
-			}
-		}
-	}
-
-
-	/** Make a mask with the indices of the vertices with 0 in the IncludedVertices array. */
-	void CreateMask(Mesh* Result, const Mesh* Base, const TArray<uint8>& IncludedVertices)
-	{
-		int32 MaskVertexCount = 0;
-		for (uint8 b : IncludedVertices)
-		{
-			if (!b)
-			{
-				++MaskVertexCount;
-			}
-		}
-
-		Result->GetVertexBuffers().SetElementCount(MaskVertexCount);
-		Result->GetVertexBuffers().SetBufferCount(1);
-		Result->MeshIDPrefix = Base->MeshIDPrefix;
-
-		EMeshBufferSemantic Semantic = MBS_VERTEXINDEX;
-		int32 SemanticIndex = 0;
-		int32 Components = 1;
-		int32 Offsets = 0;
-
-		bool bMakeRelativeIds = !Base->AreVertexIdsExplicit();
-		if (bMakeRelativeIds)
-		{
-			EMeshBufferFormat Format = MBF_UINT32;
-			Result->GetVertexBuffers().SetBuffer( 0, sizeof(uint32), 1, &Semantic, &SemanticIndex, &Format, &Components, &Offsets );
-
-			// Fill the buffer
-			uint32* Data = reinterpret_cast<uint32*>(Result->GetVertexBuffers().GetBufferData(0));
-			MeshVertexIdIteratorConst ItBase(Base);
-			for (int32 v = 0; v < IncludedVertices.Num(); ++v)
-			{
-				if (!IncludedVertices[v])
-				{
-					uint64 ID = ItBase.Get();
-					*Data = uint32(ID & 0xffffffff);
-					++Data;
-				}
-				++ItBase;
-			}
-		}
-		else
-		{
-			EMeshBufferFormat Format = MBF_UINT64;
-			Result->GetVertexBuffers().SetBuffer( 0, sizeof(uint64), 1, &Semantic, &SemanticIndex, &Format, &Components, &Offsets );
-
-			// Fill the buffer
-			uint64* Data = reinterpret_cast<uint64*>( Result->GetVertexBuffers().GetBufferData(0) );
-			MeshVertexIdIteratorConst ItBase(Base);
-			for (int32 v = 0; v < IncludedVertices.Num(); ++v)
-			{
-				if (!IncludedVertices[v])
-				{
-					*Data = ItBase.Get();
-					++Data;
-				}
-				++ItBase;
-			}
-		}
-	}
-
-}} // namespace mu::anonymous
-
-
-namespace mu
-{
-
     void MeshClipWithMesh(Mesh* Result, const Mesh* pBase, const Mesh* pClipMesh, bool& bOutSuccess)
     {
         MUTABLE_CPUPROFILER_SCOPE(MeshClipWithMesh);
@@ -703,7 +699,7 @@
 
 		Result->CopyFrom(*pBase);
 
-        TArray<uint8> VertexInClipMesh;  // Stores whether each vertex in the original mesh in in the clip mesh volume
+        TBitArray<> VertexInClipMesh;  // Stores whether each vertex in the original mesh is in the clip mesh volume
         MeshClipMeshClassifyVertices(VertexInClipMesh, pBase, pClipMesh);
 
         // Now remove all the faces from the result mesh that have all the vertices outside the clip volume
@@ -757,7 +753,7 @@
 			return;
         }
 
-        TArray<uint8> VertexInClipMesh;  // Stores whether each vertex in the original mesh in in the clip mesh volume
+        TBitArray<> VertexInClipMesh;  // Stores whether each vertex in the original mesh in in the clip mesh volume
         MeshClipMeshClassifyVertices( VertexInClipMesh, pBase, pClipMesh );
 
         // We only remove vertices if all their faces are clipped
@@ -800,8 +796,8 @@
 
 		bOutSuccess = true;
 
-		// Stores whether each vertex in the original mesh in in the clip mesh volume
-		TArray<uint8> VertexClipped; 
+		// Stores whether each vertex in the original mesh is in the clip mesh volume
+		TBitArray<> VertexClipped; 
 
 		MeshUVMaskClassifyVertices(VertexClipped, BaseForUVs, Mask, LayoutIndex);
 
@@ -845,8 +841,8 @@
 
 		bOutSuccess = true;
 
-		// Stores whether each vertex in the original mesh in in the clip mesh volume
-		TArray<uint8> VertexClipped;
+		// Stores whether each vertex in the original mesh is in the clip mesh volume
+		TBitArray<> VertexClipped;
 
 		MeshLayoutMaskClassifyVertices(VertexClipped, BaseForUVs, Mask, LayoutIndex);
 
diff --git a/Engine/Plugins/Mutable/Source/MutableRuntime/Private/MuR/OpMeshClipWithMesh.h b/Engine/Plugins/Mutable/Source/MutableRuntime/Private/MuR/OpMeshClipWithMesh.h
--- a/Engine/Plugins/Mutable/Source/MutableRuntime/Private/MuR/OpMeshClipWithMesh.h
+++ b/Engine/Plugins/Mutable/Source/MutableRuntime/Private/MuR/OpMeshClipWithMesh.h
@@ -10,6 +10,9 @@ namespace mu
 	class Image;
 	class Layout;
 
+    /** Generate a classification list for which vertex of pBase is fully contained in pClipMesh */
+	extern void MeshClipMeshClassifyVertices(TBitArray<>& VertexInClipMesh, const Mesh* pBase, const Mesh* pClipMesh);
+
     /**  */
 	extern void MeshClipWithMesh(Mesh* Result, const Mesh* pBase, const Mesh* pClipMesh, bool& bOutSuccess);
 
diff --git a/Engine/Plugins/Mutable/Source/MutableRuntime/Private/MuR/OpMeshTransformWithMesh.cpp b/Engine/Plugins/Mutable/Source/MutableRuntime/Private/MuR/OpMeshTransformWithMesh.cpp
new file mode 100644
--- /dev/null
+++ b/Engine/Plugins/Mutable/Source/MutableRuntime/Private/MuR/OpMeshTransformWithMesh.cpp
@@ -0,0 +1,102 @@
+﻿// Copyright Epic Games, Inc. All Rights Reserved.
+
+#include "MuR/OpMeshTransformWithMesh.h"
+
+#include "OpMeshClipWithMesh.h"
+#include "MuR/Mesh.h"
+#include "MuR/MeshPrivate.h"
+#include "MuR/MutableTrace.h"
+
+
+namespace mu
+{
+void MeshTransformWithMesh(Mesh* Result, const Mesh* SourceMesh, const Mesh* BoundingMesh, const FMatrix44f& Transform, bool& bOutSuccess)
+{
+	MUTABLE_CPUPROFILER_SCOPE(MeshClipWithMesh);
+	bOutSuccess = true;
+
+	const uint32 VCount = SourceMesh->GetVertexBuffers().GetElementCount();
+	if (!VCount)
+	{
+		bOutSuccess = false;
+		return; // OutSuccess false indicates the SourceMesh can be reused in this case. 
+	}
+
+	Result->CopyFrom(*SourceMesh);
+
+	// Classify which vertices in the SourceMesh are completely bounded by the BoundingMesh geometry.
+	// If no BoundingMesh is provided, this defaults to act exactly like mu::MeshTransform
+	TBitArray<> VertexInBoundaryMesh;  
+	if (BoundingMesh)
+	{
+		MeshClipMeshClassifyVertices(VertexInBoundaryMesh, SourceMesh, BoundingMesh);
+	}
+
+	const FMatrix44f TransformInvT = Transform.Inverse().GetTransposed();
+
+    const FMeshBufferSet& MBSPriv = Result->GetVertexBuffers();
+    for ( int32 b=0; b<MBSPriv.Buffers.Num(); ++b )
+    {
+        for ( int32 c=0; c<MBSPriv.Buffers[b].Channels.Num(); ++c )
+        {
+            EMeshBufferSemantic sem = MBSPriv.Buffers[b].Channels[c].Semantic;
+            int semIndex = MBSPriv.Buffers[b].Channels[c].SemanticIndex;
+
+            UntypedMeshBufferIterator it( Result->GetVertexBuffers(), sem, semIndex );
+
+            switch ( sem )
+            {
+            case MBS_POSITION:
+                for ( uint32_t v=0; v<VCount; ++v )
+                {
+                	if (!BoundingMesh || VertexInBoundaryMesh[v])
+                	{
+                		FVector4f value( 0.0f, 0.0f, 0.0f, 1.0f );
+                		for( int i=0; i<it.GetComponents(); ++i )
+                		{
+                			ConvertData( i, &value[0], MBF_FLOAT32, it.ptr(), it.GetFormat() );
+                		}
+
+                		value = Transform.TransformFVector4( value );
+
+                		for( int i=0; i<it.GetComponents(); ++i )
+                		{
+                			ConvertData( i, it.ptr(), it.GetFormat(), &value[0], MBF_FLOAT32 );
+                		}
+                	}
+                	++it;
+                }
+                break;
+
+            case MBS_NORMAL:
+            case MBS_TANGENT:
+            case MBS_BINORMAL:
+                for ( uint32_t v=0; v<VCount; ++v )
+                {
+                	if (!BoundingMesh || VertexInBoundaryMesh[v])
+                	{
+                		FVector4f value( 0.0f, 0.0f, 0.0f, 1.0f );
+	                    for( int i=0; i<it.GetComponents(); ++i )
+	                    {
+	                        ConvertData( i, &value[0], MBF_FLOAT32, it.ptr(), it.GetFormat() );
+	                    }
+
+	                    value = TransformInvT.TransformFVector4(value);
+
+	                    for( int i=0; i<it.GetComponents(); ++i )
+	                    {
+	                        ConvertData( i, it.ptr(), it.GetFormat(), &value[0], MBF_FLOAT32 );
+	                    }
+                    }
+					++it;
+                }
+                break;
+
+            default:
+                break;
+            }
+        }
+    }
+}
+
+}
diff --git a/Engine/Plugins/Mutable/Source/MutableRuntime/Private/MuR/OpMeshTransformWithMesh.h b/Engine/Plugins/Mutable/Source/MutableRuntime/Private/MuR/OpMeshTransformWithMesh.h
new file mode 100644
--- /dev/null
+++ b/Engine/Plugins/Mutable/Source/MutableRuntime/Private/MuR/OpMeshTransformWithMesh.h
@@ -0,0 +1,14 @@
+﻿// Copyright Epic Games, Inc. All Rights Reserved.
+
+#pragma once
+
+#include "Math/MathFwd.h"
+
+namespace mu
+{
+class Mesh;
+
+/**  */
+extern void MeshTransformWithMesh(Mesh* Result, const Mesh* SourceMesh, const Mesh* BoundingMesh, const FMatrix44f& Transform, bool& bOutSuccess);
+
+}
diff --git a/Engine/Plugins/Mutable/Source/MutableRuntime/Private/MuR/Operations.cpp b/Engine/Plugins/Mutable/Source/MutableRuntime/Private/MuR/Operations.cpp
--- a/Engine/Plugins/Mutable/Source/MutableRuntime/Private/MuR/Operations.cpp
+++ b/Engine/Plugins/Mutable/Source/MutableRuntime/Private/MuR/Operations.cpp
@@ -31,6 +31,7 @@ namespace mu
 		{ DT_PROJECTOR,		false,	{ 0, 0, 0, 0, 0, 0, 0, 0, 0, 0, 0, 0, 0, 0, 0, 0, 0, 0, 0, 0, 0, 0, 0, 0, 0, 0, 0, 0, 0 }		},	// PR_CONSTANT
 		{ DT_STRING,		false,	{ 0, 0, 0, 0, 0, 0, 0, 0, 0, 0, 0, 0, 0, 0, 0, 0, 0, 0, 0, 0, 0, 0, 0, 0, 0, 0, 0, 0, 0 }		},	// ST_CONSTANT
 		{ DT_EXTENSION_DATA,true,	{ 0, 0, 0, 0, 0, 0, 0, 0, 0, 0, 0, 0, 0, 0, 0, 0, 0, 0, 0, 0, 0, 0, 0, 0, 0, 0, 0, 0, 0 }		},	// ED_CONSTANT
+		{ DT_MATRIX,		false,	{ 0, 0, 0, 0, 0, 0, 0, 0, 0, 0, 0, 0, 0, 0, 0, 0, 0, 0, 0, 0, 0, 0, 0, 0, 0, 0, 0, 0, 0 }		},	// MA_CONSTANT
 
 		{ DT_BOOL,			false,	{ 0, 0, 0, 0, 0, 0, 0, 0, 0, 0, 0, 0, 0, 0, 0, 0, 0, 0, 0, 0, 0, 0, 0, 0, 0, 0, 0, 0, 0 }		},	// BO_PARAMETER
 		{ DT_INT,			false,	{ 0, 0, 0, 0, 0, 0, 0, 0, 0, 0, 0, 0, 0, 0, 0, 0, 0, 0, 0, 0, 0, 0, 0, 0, 0, 0, 0, 0, 0 }		},	// NU_PARAMETER
@@ -39,6 +40,7 @@ namespace mu
 		{ DT_PROJECTOR,		false,	{ 0, 0, 0, 0, 0, 0, 0, 0, 0, 0, 0, 0, 0, 0, 0, 0, 0, 0, 0, 0, 0, 0, 0, 0, 0, 0, 0, 0, 0 }		},	// PR_PARAMETER
 		{ DT_IMAGE,			false,	{ 0, 0, 0, 0, 0, 0, 0, 0, 0, 0, 0, 0, 0, 0, 0, 0, 0, 0, 0, 0, 0, 0, 0, 0, 0, 0, 0, 0, 0 }		},	// IM_PARAMETER
 		{ DT_STRING,		false,	{ 0, 0, 0, 0, 0, 0, 0, 0, 0, 0, 0, 0, 0, 0, 0, 0, 0, 0, 0, 0, 0, 0, 0, 0, 0, 0, 0, 0, 0 }		},	// ST_PARAMETER
+		{ DT_MATRIX,		false,	{ 0, 0, 0, 0, 0, 0, 0, 0, 0, 0, 0, 0, 0, 0, 0, 0, 0, 0, 0, 0, 0, 0, 0, 0, 0, 0, 0, 0, 0 }		},	// MA_PARAMETER
 
 		{ DT_IMAGE,			true,	{ 0, 0, 0, 0, 0, 0, 0, 0, 0, 0, 0, 0, 0, 0, 0, 0, 0, 0, 0, 0, 0, 0, 0, 0, 0, 0, 0, 0, 0 }		},	// IM_REFERENCE
 		{ DT_MESH,			true,	{ 0, 0, 0, 0, 0, 0, 0, 0, 0, 0, 0, 0, 0, 0, 0, 0, 0, 0, 0, 0, 0, 0, 0, 0, 0, 0, 0, 0, 0 }		},	// ME_REFERENCE
@@ -127,6 +129,7 @@ namespace mu
 		{ DT_MESH,			true,	{ 0, 0, 0, 0, 0, 0, 0, 0, 0, 0, 0, 0, 0, 0, 0, 0, 0, 0, 0, 0, 0, 0, 0, 0, 0, 0, 0, 0, 0 }		},	// ME_MORPHRESHAPE
 		{ DT_MESH,			true,	{ 0, 0, 0, 0, 0, 0, 0, 0, 0, 0, 0, 0, 0, 0, 0, 0, 0, 0, 0, 0, 0, 0, 0, 0, 0, 0, 0, 0, 0 }		},	// ME_OPTIMIZESKINNING
 		{ DT_MESH,			true,	{ 0, 0, 0, 0, 0, 0, 0, 0, 0, 0, 0, 0, 0, 0, 0, 0, 0, 0, 0, 0, 0, 0, 0, 0, 0, 0, 0, 0, 0 }		},	// ME_ADDTAGS
+		{ DT_MESH,			true,	{ 0, 0, 0, 0, 0, 0, 0, 0, 0, 0, 0, 0, 0, 0, 0, 0, 0, 0, 0, 0, 0, 0, 0, 0, 0, 0, 0, 0, 0 }		},	// ME_TRANSFORMWITHMESH
 
 		{ DT_INSTANCE,		false,	{ 0, 0, 0, 0, 0, 0, 0, 0, 0, 0, 0, 0, 0, 0, 0, 0, 0, 0, 0, 0, 0, 0, 0, 0, 0, 0, 0, 0, 0 }		},	// IN_ADDMESH
 		{ DT_INSTANCE,		false,	{ 0, 0, 0, 0, 0, 0, 0, 0, 0, 0, 0, 0, 0, 0, 0, 0, 0, 0, 0, 0, 0, 0, 0, 0, 0, 0, 0, 0, 0 }		},	// IN_ADDIMAGE
@@ -376,12 +379,14 @@ namespace mu
         case OP_TYPE::LA_CONSTANT:
         case OP_TYPE::PR_CONSTANT:
 		case OP_TYPE::ED_CONSTANT:
+		case OP_TYPE::MA_CONSTANT:
         case OP_TYPE::BO_PARAMETER:
         case OP_TYPE::NU_PARAMETER:
         case OP_TYPE::SC_PARAMETER:
         case OP_TYPE::CO_PARAMETER:
         case OP_TYPE::PR_PARAMETER:
         case OP_TYPE::IM_PARAMETER:
+        case OP_TYPE::MA_PARAMETER:
 		case OP_TYPE::IM_REFERENCE:
 		case OP_TYPE::ME_REFERENCE:
 			break;
@@ -1010,6 +1015,15 @@ namespace mu
 			break;
 		}
 
+        case OP_TYPE::ME_TRANSFORMWITHMESH :
+        {
+        	OP::MeshTransformWithinMeshArgs args = program.GetOpArgs<OP::MeshTransformWithinMeshArgs>(at);
+        	f(args.sourceMesh);
+        	f(args.boundingMesh);
+        	f(args.matrix);
+        	break;
+        }
+        	
         //-------------------------------------------------------------------------------------
         case OP_TYPE::IN_ADDMESH:
         case OP_TYPE::IN_ADDIMAGE:
diff --git a/Engine/Plugins/Mutable/Source/MutableRuntime/Private/MuR/Parameters.cpp b/Engine/Plugins/Mutable/Source/MutableRuntime/Private/MuR/Parameters.cpp
--- a/Engine/Plugins/Mutable/Source/MutableRuntime/Private/MuR/Parameters.cpp
+++ b/Engine/Plugins/Mutable/Source/MutableRuntime/Private/MuR/Parameters.cpp
@@ -953,6 +953,87 @@ namespace mu
         }
     }
 
+    void Parameters::GetMatrixValue(int index, FMatrix44f& OutValue, const Ptr<const RangeIndex>& pos) const
+    {
+    	check( index >= 0 && index < (int)m_pD->m_values.Num() );
+    	check( GetType( index ) == PARAMETER_TYPE::T_MATRIX );
+
+    	// Early out in case of invalid parameters
+    	if ( index < 0 || index >= (int)m_pD->m_values.Num() ||
+			 GetType( index ) != PARAMETER_TYPE::T_MATRIX )
+    	{
+    		OutValue = FMatrix44f::Identity;
+    		return;
+    	}
+
+    	// Single value case
+    	if ( !pos )
+    	{
+    		// Return the single value
+    		OutValue = m_pD->m_values[index].Get<ParamMatrixType>();
+    		return;
+    	}
+
+    	// Multivalue case
+    	check( pos->m_pD->m_parameter == index );
+
+    	if ( index < int( m_pD->m_multiValues.Num() ) )
+    	{
+    		const TMap< TArray<int32_t>, PARAMETER_VALUE >& m = m_pD->m_multiValues[index];
+    		const PARAMETER_VALUE* it = m.Find( pos->m_pD->m_values );
+    		if ( it )
+    		{
+    			OutValue = it->Get<ParamMatrixType>();
+    			return;
+    		}
+    	}
+
+    	// Multivalue parameter, but no multivalue set. Return single value.
+    	OutValue = m_pD->m_values[index].Get<ParamMatrixType>();
+    }
+
+    void Parameters::SetMatrixValue(int index, const FMatrix44f& Value, const Ptr<const RangeIndex>& pos)
+    {
+    	LLM_SCOPE_BYNAME(TEXT("MutableRuntime"));
+
+    	check( index >= 0 && index < (int)m_pD->m_values.Num() );
+    	check( GetType( index ) == PARAMETER_TYPE::T_MATRIX );
+
+    	// Early out in case of invalid parameters
+    	if ( index < 0 || index >= (int)m_pD->m_values.Num() ||
+			 GetType( index ) != PARAMETER_TYPE::T_MATRIX )
+    	{
+    		return;
+    	}
+
+    	// Single value case
+    	if ( !pos )
+    	{
+    		// Clear multivalue, if set.
+    		if ( index < int( m_pD->m_multiValues.Num() ) )
+    		{
+    			m_pD->m_multiValues[index].Empty();
+    		}
+
+    		m_pD->m_values[index].Set<ParamMatrixType>(Value);
+    	}
+
+    	// Multivalue case
+    	else
+    	{
+    		check( pos->m_pD->m_parameter == index );
+
+    		if ( index >= int( m_pD->m_multiValues.Num() ) )
+    		{
+    			m_pD->m_multiValues.SetNum( index + 1 );
+    		}
+
+    		TMap< TArray<int32_t>, PARAMETER_VALUE >& m = m_pD->m_multiValues[index];
+    		PARAMETER_VALUE& it = m.FindOrAdd( pos->m_pD->m_values );
+    		it.Set<ParamMatrixType>(Value);
+    	}
+    }
+
 
     //---------------------------------------------------------------------------------------------
     FProjector Parameters::Private::GetProjectorValue( int index, const Ptr<const RangeIndex>& pos ) const
diff --git a/Engine/Plugins/Mutable/Source/MutableRuntime/Private/MuR/System.cpp b/Engine/Plugins/Mutable/Source/MutableRuntime/Private/MuR/System.cpp
--- a/Engine/Plugins/Mutable/Source/MutableRuntime/Private/MuR/System.cpp
+++ b/Engine/Plugins/Mutable/Source/MutableRuntime/Private/MuR/System.cpp
@@ -733,6 +733,7 @@ namespace mu
             case OP_TYPE::CO_PARAMETER:
             case OP_TYPE::PR_PARAMETER:
             case OP_TYPE::IM_PARAMETER:
+            case OP_TYPE::MA_PARAMETER:
             {
 				OP::ParameterArgs args = program.GetOpArgs<OP::ParameterArgs>(at);
 				OP::ADDRESS ParamIndex = args.variable;
diff --git a/Engine/Plugins/Mutable/Source/MutableRuntime/Public/MuR/Model.h b/Engine/Plugins/Mutable/Source/MutableRuntime/Public/MuR/Model.h
--- a/Engine/Plugins/Mutable/Source/MutableRuntime/Public/MuR/Model.h
+++ b/Engine/Plugins/Mutable/Source/MutableRuntime/Public/MuR/Model.h
@@ -133,6 +133,11 @@ namespace mu
         //! \param R,G,B Pointers to values where every resulting colour channel will be stored
     	void GetColourDefaultValue(int32 Index, float* R, float* G, float* B, float* A) const;
 
+		//! Return the default value of a colour parameter.
+		//! \pre The parameter specified by index is a T_FLOAT.
+        //! \param Index Index of the parameter from 0 to GetCount()-1
+    	FMatrix44f GetMatrixDefaultValue(int32 Index) const;
+
     	//! Return the default value of a projector parameter, as a 4x4 matrix. The matrix is supposed to be
 		//! a linear transform in column-major.
 		//! \pre The parameter specified by index is a T_PROJECTOR.
diff --git a/Engine/Plugins/Mutable/Source/MutableRuntime/Public/MuR/Parameters.h b/Engine/Plugins/Mutable/Source/MutableRuntime/Public/MuR/Parameters.h
--- a/Engine/Plugins/Mutable/Source/MutableRuntime/Public/MuR/Parameters.h
+++ b/Engine/Plugins/Mutable/Source/MutableRuntime/Public/MuR/Parameters.h
@@ -62,6 +62,9 @@ namespace mu
         //! A text string.
         T_STRING,
 
+    	//! A 4x4 matrix.
+    	T_MATRIX,
+
         //! Utility enumeration value, not really a parameter type.
         T_COUNT
 
@@ -355,6 +358,18 @@ namespace mu
         //! \param pos Only for multidimensional parameters: relevant position to set in the ranges
         void SetStringValue( int index, const FString& Value, const Ptr<const RangeIndex>& pos = nullptr );
 
+		//! Return the value of a matrix parameter.
+		//! \pre The parameter specified by index is a T_MATRIX.
+		//! \param index Index of the parameter from 0 to GetCount()-1
+		//! \param pos Only for multidimensional parameters: relevant position to get in the ranges
+		void GetMatrixValue( int index, FMatrix44f& OutValue, const Ptr<const RangeIndex>& pos = nullptr) const;
+
+		//! If the parameter is of the matrix type, set its value.
+		//! \param index Index of the parameter from 0 to GetCount()-1
+		//! \param value new value of the parameter
+		//! \param pos Only for multidimensional parameters: relevant position to set in the ranges
+		void SetMatrixValue( int index, const FMatrix44f& Value, const Ptr<const RangeIndex>& pos = nullptr );
+
         //! Utility method to compare the values of a specific parameter with the values of another
         //! Parameters object. It returns false if type or values are different.
         bool HasSameValue( int thisParamIndex, const ParametersPtrConst& other, int otherParamIndex ) const;
diff --git a/Engine/Plugins/Mutable/Source/MutableTools/Private/MuT/ASTOpConstantMatrix.cpp b/Engine/Plugins/Mutable/Source/MutableTools/Private/MuT/ASTOpConstantMatrix.cpp
new file mode 100644
--- /dev/null
+++ b/Engine/Plugins/Mutable/Source/MutableTools/Private/MuT/ASTOpConstantMatrix.cpp
@@ -0,0 +1,46 @@
+﻿// Copyright Epic Games, Inc. All Rights Reserved.
+
+#include "MuT/ASTOpConstantMatrix.h"
+
+mu::ASTOpConstantMatrix::ASTOpConstantMatrix(FMatrix44f InitValue) : value(InitValue)
+{
+}
+
+uint64 mu::ASTOpConstantMatrix::Hash() const
+{
+	uint64 res = std::hash<uint64>()(uint64(OP_TYPE::MA_CONSTANT));
+	hash_combine(res, value.ComputeHash());
+	return res;
+}
+
+bool mu::ASTOpConstantMatrix::IsEqual(const ASTOp& otherUntyped) const
+{
+	if (otherUntyped.GetOpType() == GetOpType())
+	{
+		const ASTOpConstantMatrix* other = static_cast<const ASTOpConstantMatrix*>(&otherUntyped);
+		return value == other->value;
+	}
+	
+	return false;
+}
+
+mu::Ptr<mu::ASTOp> mu::ASTOpConstantMatrix::Clone(MapChildFuncRef mapChild) const
+{
+	Ptr<ASTOpConstantMatrix> n = new ASTOpConstantMatrix();
+	n->value = value;
+	return n;
+}
+
+void mu::ASTOpConstantMatrix::Link(FProgram& program, FLinkerOptions* Options)
+{
+	if (!linkedAddress)
+	{
+		OP::MatrixConstantArgs args;
+		args.value = program.AddConstant(value);
+
+		linkedAddress = static_cast<OP::ADDRESS>(program.m_opAddress.Num());
+		program.m_opAddress.Add(static_cast<uint32_t>(program.m_byteCode.Num()));
+		AppendCode(program.m_byteCode, OP_TYPE::MA_CONSTANT);
+		AppendCode(program.m_byteCode, args);
+	}
+}
diff --git a/Engine/Plugins/Mutable/Source/MutableTools/Private/MuT/ASTOpConstantMatrix.h b/Engine/Plugins/Mutable/Source/MutableTools/Private/MuT/ASTOpConstantMatrix.h
new file mode 100644
--- /dev/null
+++ b/Engine/Plugins/Mutable/Source/MutableTools/Private/MuT/ASTOpConstantMatrix.h
@@ -0,0 +1,38 @@
+﻿// Copyright Epic Games, Inc. All Rights Reserved.
+
+#pragma once
+
+#include "HAL/Platform.h"
+#include "MuR/Operations.h"
+#include "MuR/Ptr.h"
+#include "MuT/AST.h"
+
+
+namespace mu
+{
+struct FProgram;
+
+
+//---------------------------------------------------------------------------------------------
+//!
+//---------------------------------------------------------------------------------------------
+class ASTOpConstantMatrix final : public ASTOp
+{
+public:
+
+	//!
+	FMatrix44f value;
+
+public:
+
+	ASTOpConstantMatrix(FMatrix44f value = FMatrix44f::Identity);
+
+	OP_TYPE GetOpType() const override { return OP_TYPE::MA_CONSTANT; }
+	uint64 Hash() const override;
+	void ForEachChild(const TFunctionRef<void(ASTChild&)>) override {}
+	bool IsEqual(const ASTOp& otherUntyped) const override;
+	Ptr<ASTOp> Clone(MapChildFuncRef mapChild) const override;
+	void Link(FProgram& program, FLinkerOptions* Options) override;
+};
+
+}
diff --git a/Engine/Plugins/Mutable/Source/MutableTools/Private/MuT/ASTOpMeshTransformWithBoundingMesh.cpp b/Engine/Plugins/Mutable/Source/MutableTools/Private/MuT/ASTOpMeshTransformWithBoundingMesh.cpp
new file mode 100644
--- /dev/null
+++ b/Engine/Plugins/Mutable/Source/MutableTools/Private/MuT/ASTOpMeshTransformWithBoundingMesh.cpp
@@ -0,0 +1,79 @@
+﻿// Copyright Epic Games, Inc. All Rights Reserved.
+
+#include "MuT/ASTOpMeshTransformWithBoundingMesh.h"
+
+#include "ASTOpImageTransform.h"
+
+
+mu::ASTOpMeshTransformWithBoundingMesh::ASTOpMeshTransformWithBoundingMesh()
+	: source(this)
+	, boundingMesh(this)
+	, matrix(this)
+{
+}
+
+mu::ASTOpMeshTransformWithBoundingMesh::~ASTOpMeshTransformWithBoundingMesh()
+{
+	// Explicit call needed to avoid recursive destruction
+	ASTOp::RemoveChildren();
+}
+
+uint64 mu::ASTOpMeshTransformWithBoundingMesh::Hash() const
+{
+	uint64 res = std::hash<OP_TYPE>()(GetOpType());
+	hash_combine(res, source.child());
+	hash_combine(res, boundingMesh.child());
+	hash_combine(res, matrix.child());
+
+	return res;
+}
+
+bool mu::ASTOpMeshTransformWithBoundingMesh::IsEqual(const ASTOp& otherUntyped) const
+{
+	if (GetOpType() == otherUntyped.GetOpType())
+	{
+		const ASTOpMeshTransformWithBoundingMesh& other = static_cast<const ASTOpMeshTransformWithBoundingMesh&>(otherUntyped);
+		return source == other.source && boundingMesh == other.boundingMesh && matrix == other.matrix; 
+	}
+	return false;
+}
+
+mu::Ptr<mu::ASTOp> mu::ASTOpMeshTransformWithBoundingMesh::Clone(MapChildFuncRef mapChild) const
+{
+	Ptr<ASTOpMeshTransformWithBoundingMesh> n = new ASTOpMeshTransformWithBoundingMesh();
+	n->source = mapChild(source.child());
+	n->boundingMesh = mapChild(boundingMesh.child());
+	n->matrix = mapChild(matrix.child());
+	return n;
+}
+
+void mu::ASTOpMeshTransformWithBoundingMesh::ForEachChild(const TFunctionRef<void(ASTChild&)> f)
+{
+	f(source);
+	f(boundingMesh);
+	f(matrix);
+}
+
+void mu::ASTOpMeshTransformWithBoundingMesh::Link(FProgram& program, FLinkerOptions* Options)
+{
+	if (!linkedAddress)
+	{
+		OP::MeshTransformWithinMeshArgs args = {};
+		if (source)
+		{
+			args.sourceMesh = source->linkedAddress;
+		}
+		if (boundingMesh)
+		{
+			args.boundingMesh = boundingMesh->linkedAddress;
+		}
+		if (matrix)
+		{
+			args.matrix = matrix->linkedAddress;
+		}
+		linkedAddress = static_cast<OP::ADDRESS>(program.m_opAddress.Num());
+		program.m_opAddress.Add(program.m_byteCode.Num());
+		AppendCode(program.m_byteCode, OP_TYPE::ME_TRANSFORMWITHMESH);
+		AppendCode(program.m_byteCode, args);
+	}
+}
diff --git a/Engine/Plugins/Mutable/Source/MutableTools/Private/MuT/ASTOpMeshTransformWithBoundingMesh.h b/Engine/Plugins/Mutable/Source/MutableTools/Private/MuT/ASTOpMeshTransformWithBoundingMesh.h
new file mode 100644
--- /dev/null
+++ b/Engine/Plugins/Mutable/Source/MutableTools/Private/MuT/ASTOpMeshTransformWithBoundingMesh.h
@@ -0,0 +1,42 @@
+﻿// Copyright Epic Games, Inc. All Rights Reserved.
+
+#pragma once
+
+#include "HAL/Platform.h"
+#include "MuR/Operations.h"
+#include "MuR/Ptr.h"
+#include "MuT/AST.h"
+
+
+namespace mu
+{
+struct FProgram;
+
+
+//---------------------------------------------------------------------------------------------
+//!
+//---------------------------------------------------------------------------------------------
+class ASTOpMeshTransformWithBoundingMesh final : public ASTOp
+{
+public:
+
+	ASTChild source;
+	ASTChild boundingMesh;
+
+	ASTChild matrix;
+public:
+
+	ASTOpMeshTransformWithBoundingMesh();
+	ASTOpMeshTransformWithBoundingMesh(const ASTOpMeshTransformWithBoundingMesh&) = delete;
+	~ASTOpMeshTransformWithBoundingMesh() override;
+
+	OP_TYPE GetOpType() const override { return OP_TYPE::ME_TRANSFORMWITHMESH; }
+	uint64 Hash() const override;
+	bool IsEqual(const ASTOp& otherUntyped) const override;
+	Ptr<ASTOp> Clone(MapChildFuncRef mapChild) const override;
+	void ForEachChild(const TFunctionRef<void(ASTChild&)>) override;
+	void Link(FProgram& program, FLinkerOptions* Options) override;
+};
+
+}
+
diff --git a/Engine/Plugins/Mutable/Source/MutableTools/Private/MuT/CodeGenerator.cpp b/Engine/Plugins/Mutable/Source/MutableTools/Private/MuT/CodeGenerator.cpp
--- a/Engine/Plugins/Mutable/Source/MutableTools/Private/MuT/CodeGenerator.cpp
+++ b/Engine/Plugins/Mutable/Source/MutableTools/Private/MuT/CodeGenerator.cpp
@@ -2,6 +2,7 @@
 
 #include "MuT/CodeGenerator.h"
 
+#include "ASTOpMeshTransformWithBoundingMesh.h"
 #include "Containers/Array.h"
 #include "Logging/LogCategory.h"
 #include "Logging/LogMacros.h"
@@ -58,6 +59,7 @@
 #include "MuT/NodeImageMipmap.h"
 #include "MuT/NodeImageMipmapPrivate.h"
 #include "MuT/NodeImageSwizzlePrivate.h"
+#include "MuT/NodeMatrixConstant.h"
 #include "MuT/NodeMesh.h"
 #include "MuT/NodeMeshClipMorphPlane.h"
 #include "MuT/NodeMeshClipWithMesh.h"
@@ -73,6 +75,7 @@
 #include "MuT/NodeModifierMeshClipMorphPlane.h"
 #include "MuT/NodeModifierMeshClipWithMesh.h"
 #include "MuT/NodeModifierMeshClipWithUVMask.h"
+#include "MuT/NodeModifierMeshTransformInMesh.h"
 #include "MuT/NodeModifierSurfaceEdit.h"
 #include "MuT/NodeObject.h"
 #include "MuT/NodeObjectGroupPrivate.h"
@@ -2410,6 +2413,64 @@ namespace mu
 				LastMeshOp = Op;
 			}
 		}
+		
+		// Process transform mesh within mesh modifiers.
+		for (const FirstPassGenerator::FModifier& m : Modifiers)
+		{
+			if (m.Node->GetType()== NodeModifierMeshTransformInMesh::GetStaticType())
+			{
+				const NodeModifierMeshTransformInMesh* TypedTransformNode = static_cast<const NodeModifierMeshTransformInMesh*>(m.Node);
+
+				// If a matrix node is not connected, the op won't do anything, so let's not create it at all.
+				if (TypedTransformNode->MatrixNode)
+				{
+					Ptr<ASTOpMeshTransformWithBoundingMesh> transformOp = new ASTOpMeshTransformWithBoundingMesh();
+					transformOp->source = LastMeshOp;
+
+					// Transform matrix.
+					if (TypedTransformNode->MatrixNode)
+					{
+						FMatrixGenerationResult ChildResult;
+						GenerateMatrix(ChildResult, Options, TypedTransformNode->MatrixNode);
+						transformOp->matrix = ChildResult.op;
+					}
+					
+					if (TypedTransformNode->BoundingMesh)
+					{
+						// Parameters
+						FMeshGenerationOptions MeshOptions;
+						MeshOptions.bLayouts = false;
+						MeshOptions.State = Options.State;
+
+						FMeshGenerationResult BoundingMeshResult;
+						GenerateMesh(MeshOptions, BoundingMeshResult, TypedTransformNode->BoundingMesh);
+						transformOp->boundingMesh = BoundingMeshResult.MeshOp;
+
+						if (!transformOp->boundingMesh)
+						{
+							ErrorLog->GetPrivate()->Add("Bounding mesh has not been generated", ELMT_ERROR, ErrorContext);
+							continue;
+						}
+					}
+
+					// Condition to apply the transform op
+					if (m.FinalCondition)
+					{
+						Ptr<ASTOpConditional> conditionalAd = new ASTOpConditional();
+						conditionalAd->type = OP_TYPE::ME_CONDITIONAL;
+						conditionalAd->no = LastMeshOp;
+						conditionalAd->yes = transformOp;
+						conditionalAd->condition = m.FinalCondition;
+						LastMeshOp = conditionalAd;
+					}
+					else
+					{
+						LastMeshOp = transformOp;
+					}
+				}
+			}
+		}
+		
 
 		return LastMeshOp;
 	}
diff --git a/Engine/Plugins/Mutable/Source/MutableTools/Private/MuT/CodeGenerator.h b/Engine/Plugins/Mutable/Source/MutableTools/Private/MuT/CodeGenerator.h
--- a/Engine/Plugins/Mutable/Source/MutableTools/Private/MuT/CodeGenerator.h
+++ b/Engine/Plugins/Mutable/Source/MutableTools/Private/MuT/CodeGenerator.h
@@ -101,6 +101,9 @@ namespace mu
 	class NodeScalarVariation;
 	class NodeStringConstant;
 	class NodeStringParameter;
+	class NodeMatrix;
+	class NodeMatrixConstant;
+	class NodeMatrixParameter;
 	struct FObjectState;
 	struct FProgram;
 
@@ -763,6 +766,20 @@ namespace mu
 		void GenerateString_Constant(FStringGenerationResult&, const FGenericGenerationOptions& Options, const Ptr<const NodeStringConstant>&);
 		void GenerateString_Parameter(FStringGenerationResult&, const FGenericGenerationOptions& Options, const Ptr<const NodeStringParameter>&);
 
+    	//-----------------------------------------------------------------------------------------
+    	// Transforms
+    	struct FMatrixGenerationResult
+    	{
+    		Ptr<ASTOp> op;
+    	};
+
+    	typedef TMap<FGeneratedCacheKey, FMatrixGenerationResult> FGeneratedMatrixMap;
+    	FGeneratedMatrixMap GeneratedMatrices;
+
+    	void GenerateMatrix(FMatrixGenerationResult&, const FGenericGenerationOptions&, const Ptr<const NodeMatrix>&);
+    	void GenerateMatrix_Constant(FMatrixGenerationResult&, const FGenericGenerationOptions&, const Ptr<const NodeMatrixConstant>&);
+    	void GenerateMatrix_Parameter(FMatrixGenerationResult&, const FGenericGenerationOptions&, const Ptr<const NodeMatrixParameter>&);
+
         //-----------------------------------------------------------------------------------------
         // Ranges
         struct FRangeGenerationResult
diff --git a/Engine/Plugins/Mutable/Source/MutableTools/Private/MuT/CodeGenerator_Matrix.cpp b/Engine/Plugins/Mutable/Source/MutableTools/Private/MuT/CodeGenerator_Matrix.cpp
new file mode 100644
--- /dev/null
+++ b/Engine/Plugins/Mutable/Source/MutableTools/Private/MuT/CodeGenerator_Matrix.cpp
@@ -0,0 +1,86 @@
+// Copyright Epic Games, Inc. All Rights Reserved.
+
+#include "MuT/CodeGenerator.h"
+
+#include "MuT/ASTOpConstantMatrix.h"
+#include "MuT/NodeMatrixConstant.h"
+#include "MuT/NodeMatrixParameter.h"
+
+namespace mu
+{
+
+	void CodeGenerator::GenerateMatrix(FMatrixGenerationResult& Result, const FGenericGenerationOptions& Options, const Ptr<const NodeMatrix>& Untyped)
+	{
+		if (!Untyped)
+		{
+			Result = FMatrixGenerationResult();
+			return;
+		}
+
+		// See if it was already generated
+		FGeneratedCacheKey Key;
+		Key.Node = Untyped;
+		Key.Options = Options;
+		FGeneratedMatrixMap::ValueType* it = GeneratedMatrices.Find(Key);
+		if (it)
+		{
+			Result = *it;
+			return;
+		}
+
+		if (Untyped->GetType() == NodeMatrixConstant::GetStaticType())
+		{
+			const NodeMatrixConstant* Constant = static_cast<const NodeMatrixConstant*>(Untyped.get());
+			GenerateMatrix_Constant(Result, Options, Constant);
+		}
+		else if (Untyped->GetType() == NodeMatrixParameter::GetStaticType())
+		{
+			const NodeMatrixParameter* Parameter = static_cast<const NodeMatrixParameter*>(Untyped.get());
+			GenerateMatrix_Parameter(Result, Options, Parameter);
+		}  
+	}
+
+	//-------------------------------------------------------------------------------------------------
+	void CodeGenerator::GenerateMatrix_Constant(FMatrixGenerationResult& result, const FGenericGenerationOptions& Options, const Ptr<const NodeMatrixConstant>& Typed)
+	{
+		Ptr<ASTOpConstantMatrix> op = new ASTOpConstantMatrix();
+		op->value = Typed->Value;
+		result.op = op;
+	}
+
+	void CodeGenerator::GenerateMatrix_Parameter(FMatrixGenerationResult& result, const FGenericGenerationOptions& Options, const Ptr<const NodeMatrixParameter>& Typed)
+	{
+		Ptr<ASTOpParameter> op;
+		Ptr<ASTOpParameter>* it = FirstPass.ParameterNodes.Find(Typed.get());
+		if (!it)
+		{
+			FParameterDesc param;
+			param.m_name = Typed->Name;
+			const TCHAR* CStr = ToCStr(Typed->Uid);
+			param.m_uid.ImportTextItem(CStr, 0, nullptr, nullptr);
+			param.m_type = PARAMETER_TYPE::T_MATRIX;
+			param.m_defaultValue.Set<ParamMatrixType>(Typed->DefaultValue);
+
+			op = new ASTOpParameter();
+			op->type = OP_TYPE::MA_PARAMETER;
+			op->parameter = param;
+
+			// Generate the code for the ranges
+			for (int32 a = 0; a < Typed->Ranges.Num(); ++a)
+			{
+				FRangeGenerationResult rangeResult;
+				GenerateRange(rangeResult, Options, Typed->Ranges[a]);
+				op->ranges.Emplace(op.get(), rangeResult.sizeOp, rangeResult.rangeName, rangeResult.rangeUID);
+			}
+
+			FirstPass.ParameterNodes.Add(Typed.get(), op);
+		}
+		else
+		{
+			op = *it;
+		}
+
+		result.op = op;
+	}
+
+}
diff --git a/Engine/Plugins/Mutable/Source/MutableTools/Private/MuT/CodeOptimiser.cpp b/Engine/Plugins/Mutable/Source/MutableTools/Private/MuT/CodeOptimiser.cpp
--- a/Engine/Plugins/Mutable/Source/MutableTools/Private/MuT/CodeOptimiser.cpp
+++ b/Engine/Plugins/Mutable/Source/MutableTools/Private/MuT/CodeOptimiser.cpp
@@ -647,6 +647,7 @@ namespace mu
 					case OP_TYPE::CO_PARAMETER:
 					case OP_TYPE::PR_PARAMETER:
 					case OP_TYPE::IM_PARAMETER:
+					case OP_TYPE::MA_PARAMETER:
 						bIsConstantSubgraph = false;
 						break;
 					default:
diff --git a/Engine/Plugins/Mutable/Source/MutableTools/Private/MuT/ErrorLog.cpp b/Engine/Plugins/Mutable/Source/MutableTools/Private/MuT/ErrorLog.cpp
--- a/Engine/Plugins/Mutable/Source/MutableTools/Private/MuT/ErrorLog.cpp
+++ b/Engine/Plugins/Mutable/Source/MutableTools/Private/MuT/ErrorLog.cpp
@@ -233,6 +233,7 @@ namespace mu
 		TEXT("PR_CONSTANT      "),
 		TEXT("ST_CONSTANT      "),
 		TEXT("ED_CONSTANT      "),
+		TEXT("MA_CONSTANT      "),
 
 		TEXT("BO_PARAMETER     "),
 		TEXT("NU_PARAMETER     "),
@@ -241,6 +242,7 @@ namespace mu
 		TEXT("PR_PARAMETER     "),
 		TEXT("IM_PARAMETER     "),
 		TEXT("ST_PARAMETER     "),
+		TEXT("MA_PARAMETER     "),
 		
 		TEXT("IM_REFERENCE     "),
 		TEXT("ME_REFERENCE     "),
@@ -329,6 +331,7 @@ namespace mu
 		TEXT("ME_MORPHRESHAPE  "),
 		TEXT("ME_OPTIMIZESKIN  "),
 		TEXT("ME_ADDTAGS       "),
+		TEXT("ME_TRANSFORMWITHMESH"),
 
 		TEXT("IN_ADDMESH       "),
 		TEXT("IN_ADDIMAGE      "),
diff --git a/Engine/Plugins/Mutable/Source/MutableTools/Private/MuT/Node.cpp b/Engine/Plugins/Mutable/Source/MutableTools/Private/MuT/Node.cpp
--- a/Engine/Plugins/Mutable/Source/MutableTools/Private/MuT/Node.cpp
+++ b/Engine/Plugins/Mutable/Source/MutableTools/Private/MuT/Node.cpp
@@ -36,7 +36,11 @@
 #include "MuT/NodeModifierMeshClipMorphPlane.h"
 #include "MuT/NodeModifierMeshClipWithMesh.h"
 #include "MuT/NodeModifierMeshClipWithUVMask.h"
+#include "MuT/NodeModifierMeshTransformInMesh.h"
 #include "MuT/NodeModifierSurfaceEdit.h"
+#include "MuT/NodeMatrix.h"
+#include "MuT/NodeMatrixConstant.h"
+#include "MuT/NodeMatrixParameter.h"
 
 namespace mu
 {
@@ -81,8 +85,12 @@ namespace mu
 	FNodeType NodeModifierMeshClipMorphPlane::StaticType = FNodeType(Node::EType::ModifierMeshClipMorphPlane, NodeModifier::GetStaticType());
 	FNodeType NodeModifierMeshClipWithMesh::StaticType = FNodeType(Node::EType::ModifierMeshClipWithMesh, NodeModifier::GetStaticType());
 	FNodeType NodeModifierMeshClipWithUVMask::StaticType = FNodeType(Node::EType::ModifierMeshClipWithUVMask, NodeModifier::GetStaticType());
+	FNodeType NodeModifierMeshTransformInMesh::StaticType = FNodeType(Node::EType::ModifierTransformInMesh, NodeModifier::GetStaticType());
 	FNodeType NodeModifierSurfaceEdit::StaticType = FNodeType(Node::EType::ModifierSurfaceEdit, NodeSurface::GetStaticType());
 
+	FNodeType NodeMatrix::StaticType = FNodeType(Node::EType::Matrix, Node::GetStaticType());
+	FNodeType NodeMatrixConstant::StaticType = FNodeType(Node::EType::MatrixConstant, Node::GetStaticType());
+	FNodeType NodeMatrixParameter::StaticType = FNodeType(Node::EType::MatrixParameter, Node::GetStaticType());
 
 	FNodeType::FNodeType()
 	{
diff --git a/Engine/Plugins/Mutable/Source/MutableTools/Public/MuT/Node.h b/Engine/Plugins/Mutable/Source/MutableTools/Public/MuT/Node.h
--- a/Engine/Plugins/Mutable/Source/MutableTools/Public/MuT/Node.h
+++ b/Engine/Plugins/Mutable/Source/MutableTools/Public/MuT/Node.h
@@ -141,12 +141,17 @@ namespace mu
 			ModifierMeshClipDeform,
 			ModifierMeshClipWithUVMask,
 			ModifierSurfaceEdit,
+			ModifierTransformInMesh,
 
 			ExtensionData,
 			ExtensionDataConstant,
 			ExtensionDataSwitch,
 			ExtensionDataVariation,
 
+			Matrix,
+			MatrixConstant,
+			MatrixParameter,
+
 			Count
 		};
 
diff --git a/Engine/Plugins/Mutable/Source/MutableTools/Public/MuT/NodeMatrix.h b/Engine/Plugins/Mutable/Source/MutableTools/Public/MuT/NodeMatrix.h
new file mode 100644
--- /dev/null
+++ b/Engine/Plugins/Mutable/Source/MutableTools/Public/MuT/NodeMatrix.h
@@ -0,0 +1,30 @@
+// Copyright Epic Games, Inc. All Rights Reserved.
+
+#pragma once
+
+#include "MuT/Node.h"
+
+namespace mu
+{
+	/** Base class of any node that outputs a colour.
+	*/
+	class MUTABLETOOLS_API NodeMatrix : public Node
+	{
+	public:
+
+		// Node Interface
+		virtual const FNodeType* GetType() const override { return &StaticType; }
+		static const FNodeType* GetStaticType() { return &StaticType; }
+
+	protected:
+
+		//! Forbidden. Manage with the Ptr<> template.
+		virtual ~NodeMatrix() override {}
+
+	private:
+
+		static FNodeType StaticType;
+
+	};
+
+}
diff --git a/Engine/Plugins/Mutable/Source/MutableTools/Public/MuT/NodeMatrixConstant.h b/Engine/Plugins/Mutable/Source/MutableTools/Public/MuT/NodeMatrixConstant.h
new file mode 100644
--- /dev/null
+++ b/Engine/Plugins/Mutable/Source/MutableTools/Public/MuT/NodeMatrixConstant.h
@@ -0,0 +1,30 @@
+﻿// Copyright Epic Games, Inc. All Rights Reserved.
+
+#pragma once
+
+#include "MuT/NodeMatrix.h"
+#include "Math/Matrix.h"
+
+namespace mu
+{
+
+	class MUTABLETOOLS_API NodeMatrixConstant : public NodeMatrix
+	{
+	public:
+		FMatrix44f Value;
+
+		// Node interface
+		const FNodeType* GetType() const override { return &StaticType; }
+		static const FNodeType* GetStaticType() { return &StaticType; }
+
+	protected:
+
+		//! Forbidden. Manage with the Ptr<> template.
+		virtual ~NodeMatrixConstant() override = default;
+
+	private:
+
+		static FNodeType StaticType;
+	};
+
+}
diff --git a/Engine/Plugins/Mutable/Source/MutableTools/Public/MuT/NodeMatrixParameter.h b/Engine/Plugins/Mutable/Source/MutableTools/Public/MuT/NodeMatrixParameter.h
new file mode 100644
--- /dev/null
+++ b/Engine/Plugins/Mutable/Source/MutableTools/Public/MuT/NodeMatrixParameter.h
@@ -0,0 +1,36 @@
+﻿// Copyright Epic Games, Inc. All Rights Reserved.
+
+#pragma once
+
+#include "MuT/NodeMatrix.h"
+
+#include "MuT/NodeRange.h"
+#include "Math/Matrix.h"
+
+
+namespace mu
+{
+
+	class MUTABLETOOLS_API NodeMatrixParameter : public NodeMatrix
+	{
+	public:
+		FMatrix44f DefaultValue;
+		FString Name;
+		FString Uid;
+
+		TArray<Ptr<NodeRange>> Ranges;
+		
+		// Node interface
+		const FNodeType* GetType() const override { return &StaticType; }
+		static const FNodeType* GetStaticType() { return &StaticType; }
+
+	protected:
+		//! Forbidden. Manage with the Ptr<> template.
+		virtual ~NodeMatrixParameter() override = default;
+
+	private:
+
+		static FNodeType StaticType;
+	};
+
+}
diff --git a/Engine/Plugins/Mutable/Source/MutableTools/Public/MuT/NodeModifierMeshTransformInMesh.h b/Engine/Plugins/Mutable/Source/MutableTools/Public/MuT/NodeModifierMeshTransformInMesh.h
new file mode 100644
--- /dev/null
+++ b/Engine/Plugins/Mutable/Source/MutableTools/Public/MuT/NodeModifierMeshTransformInMesh.h
@@ -0,0 +1,42 @@
+// Copyright Epic Games, Inc. All Rights Reserved.
+
+#pragma once
+
+#include "MuR/Ptr.h"
+#include "MuR/Types.h"
+#include "MuT/Node.h"
+#include "MuT/NodeModifier.h"
+#include "MuT/NodeMesh.h"
+#include "MuT/NodeMatrix.h"
+
+namespace mu
+{
+
+/** */
+class MUTABLETOOLS_API NodeModifierMeshTransformInMesh : public NodeModifier
+{
+public:
+
+	/** */
+	Ptr<NodeMesh> BoundingMesh;
+
+	/** */
+	Ptr<NodeMatrix> MatrixNode;
+public:
+
+	// Node interface
+	virtual const FNodeType* GetType() const override { return &StaticType; }
+	static const FNodeType* GetStaticType() { return &StaticType; }
+
+protected:
+
+	//! Forbidden. Manage with the Ptr<> template.
+	virtual ~NodeModifierMeshTransformInMesh() override {}
+
+private:
+
+	static FNodeType StaticType;
+
+};
+
+}
