Mutable: Fix static analysis warnings

#rnx

diff --git a/Engine/Plugins/Mutable/Source/CustomizableObjectEditor/Private/MuCOE/Nodes/CustomizableObjectNodeModifierTransformInMesh.cpp b/Engine/Plugins/Mutable/Source/CustomizableObjectEditor/Private/MuCOE/Nodes/CustomizableObjectNodeModifierTransformInMesh.cpp
--- a/Engine/Plugins/Mutable/Source/CustomizableObjectEditor/Private/MuCOE/Nodes/CustomizableObjectNodeModifierTransformInMesh.cpp
+++ b/Engine/Plugins/Mutable/Source/CustomizableObjectEditor/Private/MuCOE/Nodes/CustomizableObjectNodeModifierTransformInMesh.cpp
@@ -6,9 +6,12 @@
 
 #define LOCTEXT_NAMESPACE "CustomizableObjectEditor"
 
-static const TCHAR* OutputPinName = TEXT("Modifier");
-static const TCHAR* BoundingMeshPinName = TEXT("Bounding Mesh");
-static const TCHAR* TransformPinName = TEXT("Transform");
+namespace
+{
+const TCHAR* OutputPinName = TEXT("Modifier");
+const TCHAR* BoundingMeshPinName = TEXT("Bounding Mesh");
+const TCHAR* TransformPinName = TEXT("Transform");
+}
 
 
 FText UCustomizableObjectNodeModifierTransformInMesh::GetNodeTitle(ENodeTitleType::Type TitleType) const
