Mutable: Temporarily silence a warning that is causing cook issues.

#rnx

diff --git a/Engine/Plugins/Mutable/Source/MutableTools/Private/MuT/CodeGenerator.cpp b/Engine/Plugins/Mutable/Source/MutableTools/Private/MuT/CodeGenerator.cpp
--- a/Engine/Plugins/Mutable/Source/MutableTools/Private/MuT/CodeGenerator.cpp
+++ b/Engine/Plugins/Mutable/Source/MutableTools/Private/MuT/CodeGenerator.cpp
@@ -2595,12 +2595,14 @@ namespace mu
 
 				if (!MatchingEdit || (MatchingEdit && !MatchingEdit->Extend) )
 				{
+#if 0
 					if (Edit->LODs[CurrentLOD].MeshAdd)
 					{
 						// When extending a mesh section it is mandatory to provide textures for all section textures handled by Mutable.
 						FString Msg = FString::Printf(TEXT("Required texture [%s] is missing when trying to extend a mesh section."), *ImageData.MaterialParameterName);
 						ErrorLog->GetPrivate()->Add(Msg, ELMT_WARNING, Edit->GetMessageContext(), ModifiedNodeErrorContext);
 					}
+#endif
 
 					continue;
 				}
