Mutable: Second stab at a static analysis fix.

#rnx

diff --git a/Engine/Plugins/Mutable/Source/CustomizableObjectEditor/Private/MuCOE/Nodes/CustomizableObjectNodeModifierTransformInMesh.cpp b/Engine/Plugins/Mutable/Source/CustomizableObjectEditor/Private/MuCOE/Nodes/CustomizableObjectNodeModifierTransformInMesh.cpp
--- a/Engine/Plugins/Mutable/Source/CustomizableObjectEditor/Private/MuCOE/Nodes/CustomizableObjectNodeModifierTransformInMesh.cpp
+++ b/Engine/Plugins/Mutable/Source/CustomizableObjectEditor/Private/MuCOE/Nodes/CustomizableObjectNodeModifierTransformInMesh.cpp
@@ -6,12 +6,10 @@
 
 #define LOCTEXT_NAMESPACE "CustomizableObjectEditor"
 
-namespace
-{
-const TCHAR* OutputPinName = TEXT("Modifier");
-const TCHAR* BoundingMeshPinName = TEXT("Bounding Mesh");
-const TCHAR* TransformPinName = TEXT("Transform");
-}
+
+const TCHAR* UCustomizableObjectNodeModifierTransformInMesh::OutputPinName = TEXT("Modifier");
+const TCHAR* UCustomizableObjectNodeModifierTransformInMesh::BoundingMeshPinName = TEXT("Bounding Mesh");
+const TCHAR* UCustomizableObjectNodeModifierTransformInMesh::TransformPinName = TEXT("Transform");
 
 
 FText UCustomizableObjectNodeModifierTransformInMesh::GetNodeTitle(ENodeTitleType::Type TitleType) const
diff --git a/Engine/Plugins/Mutable/Source/CustomizableObjectEditor/Public/MuCOE/Nodes/CustomizableObjectNodeModifierTransformInMesh.h b/Engine/Plugins/Mutable/Source/CustomizableObjectEditor/Public/MuCOE/Nodes/CustomizableObjectNodeModifierTransformInMesh.h
--- a/Engine/Plugins/Mutable/Source/CustomizableObjectEditor/Public/MuCOE/Nodes/CustomizableObjectNodeModifierTransformInMesh.h
+++ b/Engine/Plugins/Mutable/Source/CustomizableObjectEditor/Public/MuCOE/Nodes/CustomizableObjectNodeModifierTransformInMesh.h
@@ -38,5 +38,11 @@ public:
 	// Own interface
 	UEdGraphPin* BoundingMeshPin() const;
 	UEdGraphPin* TransformPin() const;
+	
+private:
+	static const TCHAR* OutputPinName;
+	static const TCHAR* BoundingMeshPinName;
+	static const TCHAR* TransformPinName;
+	
 };
 
