[mutable] Fix compilation loop whit malformed input (swizzle of nothing).
Fix some typos in messages and comments.
#rnx

diff --git a/Engine/Plugins/Mutable/Source/CustomizableObjectEditor/Private/MuCOE/GenerateMutableSource/GenerateMutableSourceSurface.cpp b/Engine/Plugins/Mutable/Source/CustomizableObjectEditor/Private/MuCOE/GenerateMutableSource/GenerateMutableSourceSurface.cpp
--- a/Engine/Plugins/Mutable/Source/CustomizableObjectEditor/Private/MuCOE/GenerateMutableSource/GenerateMutableSourceSurface.cpp
+++ b/Engine/Plugins/Mutable/Source/CustomizableObjectEditor/Private/MuCOE/GenerateMutableSource/GenerateMutableSourceSurface.cpp
@@ -416,7 +416,7 @@ mu::Ptr<mu::NodeSurface> GenerateMutableSourceSurface(const UEdGraphPin * Pin, F
 					{
 						bTableMaterialPinLinked = false;
 
-						GenerationContext.Log(LOCTEXT("DifferentParentMaterial", "The Deafult Material Instance of the Data Table must have the same Parent Material."), TypedNodeMat->GetMaterialNode());
+						GenerationContext.Log(LOCTEXT("DifferentParentMaterial", "The Default Material Instance of the Data Table must have the same Parent Material."), TypedNodeMat->GetMaterialNode());
 					}
 				}
 				else
diff --git a/Engine/Plugins/Mutable/Source/MutableRuntime/Private/MuR/System.cpp b/Engine/Plugins/Mutable/Source/MutableRuntime/Private/MuR/System.cpp
--- a/Engine/Plugins/Mutable/Source/MutableRuntime/Private/MuR/System.cpp
+++ b/Engine/Plugins/Mutable/Source/MutableRuntime/Private/MuR/System.cpp
@@ -1736,7 +1736,7 @@ namespace mu
 		const TArray<FParameterDesc>& ParamDescs = Params->GetPrivate()->m_pModel->GetPrivate()->m_program.m_parameters;
 
 		// First make a mask with a bit for each relevant parameter. It will be on for parameters included in the blob.
-		// A parameter will be excluded from the blob if it has the deafult value, and no multivalues.
+		// A parameter will be excluded from the blob if it has the default value, and no multivalues.
 		TBitArray IncludedParameters(0, RelevantParams->Num());
 		if (RelevantParams->Num())
 		{
diff --git a/Engine/Plugins/Mutable/Source/MutableTools/Private/MuT/CodeGenerator.cpp b/Engine/Plugins/Mutable/Source/MutableTools/Private/MuT/CodeGenerator.cpp
--- a/Engine/Plugins/Mutable/Source/MutableTools/Private/MuT/CodeGenerator.cpp
+++ b/Engine/Plugins/Mutable/Source/MutableTools/Private/MuT/CodeGenerator.cpp
@@ -1123,6 +1123,11 @@ namespace mu
 								bFound = true;
 							}
 						}
+						else
+						{
+							// break loop if swizzle has no sources.
+							bFound = true;
+						}
 					}
 					else
 					{
@@ -2595,14 +2600,12 @@ namespace mu
 
 				if (!MatchingEdit || (MatchingEdit && !MatchingEdit->Extend) )
 				{
-#if 0
 					if (Edit->LODs[CurrentLOD].MeshAdd)
 					{
 						// When extending a mesh section it is mandatory to provide textures for all section textures handled by Mutable.
 						FString Msg = FString::Printf(TEXT("Required texture [%s] is missing when trying to extend a mesh section."), *ImageData.MaterialParameterName);
 						ErrorLog->GetPrivate()->Add(Msg, ELMT_WARNING, Edit->GetMessageContext(), ModifiedNodeErrorContext);
 					}
-#endif
 
 					continue;
 				}
diff --git a/Engine/Plugins/Mutable/Source/MutableTools/Private/MuT/CodeGenerator_Image.cpp b/Engine/Plugins/Mutable/Source/MutableTools/Private/MuT/CodeGenerator_Image.cpp
--- a/Engine/Plugins/Mutable/Source/MutableTools/Private/MuT/CodeGenerator_Image.cpp
+++ b/Engine/Plugins/Mutable/Source/MutableTools/Private/MuT/CodeGenerator_Image.cpp
@@ -764,6 +764,12 @@ namespace mu
 
 		//MUTABLE_CPUPROFILER_SCOPE(NodeImageSwizzle);
 
+		if (node.m_format == EImageFormat::IF_NONE)
+		{
+			Result.op = GenerateMissingImageCode(TEXT("Make Texture channel."), EImageFormat::IF_L_UBYTE, InNode->GetMessageContext(), Options);
+			return;
+		}
+
         // This node always produces a swizzle operation and sometimes it may produce a pixelformat
 		// operation to compress the Result
         Ptr<ASTOpImageSwizzle> SwizzleOp = new ASTOpImageSwizzle();
@@ -776,7 +782,7 @@ namespace mu
         case EImageFormat::IF_BC1:
         case EImageFormat::IF_ASTC_4x4_RGB_LDR:
             CompressedFormat = node.m_format;
-            SwizzleOp->Format = node.m_sources[3] ? EImageFormat::IF_RGBA_UBYTE : EImageFormat::IF_RGB_UBYTE;
+            SwizzleOp->Format = (node.m_sources.IsValidIndex(3) && node.m_sources[3]) ? EImageFormat::IF_RGBA_UBYTE : EImageFormat::IF_RGB_UBYTE;
 			break;
 
 		case EImageFormat::IF_BC2:
