[mutable] Reduce log level to info until data is fixed.
#rnx

diff --git a/Engine/Plugins/Mutable/Source/MutableTools/Private/MuT/CodeGenerator.cpp b/Engine/Plugins/Mutable/Source/MutableTools/Private/MuT/CodeGenerator.cpp
--- a/Engine/Plugins/Mutable/Source/MutableTools/Private/MuT/CodeGenerator.cpp
+++ b/Engine/Plugins/Mutable/Source/MutableTools/Private/MuT/CodeGenerator.cpp
@@ -2604,7 +2604,7 @@ namespace mu
 					{
 						// When extending a mesh section it is mandatory to provide textures for all section textures handled by Mutable.
 						FString Msg = FString::Printf(TEXT("Required texture [%s] is missing when trying to extend a mesh section."), *ImageData.MaterialParameterName);
-						ErrorLog->GetPrivate()->Add(Msg, ELMT_WARNING, Edit->GetMessageContext(), ModifiedNodeErrorContext);
+						ErrorLog->GetPrivate()->Add(Msg, ELMT_INFO, Edit->GetMessageContext(), ModifiedNodeErrorContext);
 					}
 
 					continue;
