[Mutable] Allow to extend from mesh section without a mesh connected.

- Surface metadata will be the first found if the main surface does not have a mesh.

#rnx

diff --git a/Engine/Plugins/Mutable/Source/CustomizableObject/Internal/MuCO/CustomizableObjectPrivate.h b/Engine/Plugins/Mutable/Source/CustomizableObject/Internal/MuCO/CustomizableObjectPrivate.h
--- a/Engine/Plugins/Mutable/Source/CustomizableObject/Internal/MuCO/CustomizableObjectPrivate.h
+++ b/Engine/Plugins/Mutable/Source/CustomizableObject/Internal/MuCO/CustomizableObjectPrivate.h
@@ -562,10 +562,14 @@ struct FMutableMeshMetadata
 	UPROPERTY()
 	uint32 ClothingMetadataId = 0;
 
+	UPROPERTY()
+	uint32 SurfaceMetadataId = 0;
+
 	friend FArchive& operator<<(FArchive& Ar, FMutableMeshMetadata& Elem)
 	{
 		Ar << Elem.MorphMetadataId;
 		Ar << Elem.ClothingMetadataId;
+		Ar << Elem.SurfaceMetadataId;
 
 		return Ar;
 	}
@@ -578,22 +582,20 @@ struct FMutableSurfaceMetadata
 	GENERATED_USTRUCT_BODY()
 	
 	UPROPERTY()
-	int32 MaterialSlotIndex = INDEX_NONE;
+	FName MaterialSlotName = FName{};
 	
 	UPROPERTY()
-	bool bCastShadow = false;
-	uint8 UnusedPadding[3] = {0};
+	bool bCastShadow = true;
 
 	friend FArchive& operator<<(FArchive& Ar, FMutableSurfaceMetadata& Elem)
 	{
-		Ar << Elem.MaterialSlotIndex;
+		Ar << Elem.MaterialSlotName;
 		Ar << Elem.bCastShadow;
 
 		return Ar;
 	}
 };
 
-static_assert(sizeof(FMutableSurfaceMetadata) == (sizeof(int32) + sizeof(bool) + sizeof(uint8)*3));
 
 USTRUCT()
 struct FClothingStreamable
@@ -1485,6 +1487,8 @@ public:
 
 		TransformInMeshModifier,
 		
+		SurfaceMetadataSlotNameIndexToName,
+
 		// -----<new versions can be added above this line>--------
 		LastCustomizableObjectVersion
 	};
diff --git a/Engine/Plugins/Mutable/Source/CustomizableObject/Private/MuCO/CustomizableObjectInstance.cpp b/Engine/Plugins/Mutable/Source/CustomizableObject/Private/MuCO/CustomizableObjectInstance.cpp
--- a/Engine/Plugins/Mutable/Source/CustomizableObject/Private/MuCO/CustomizableObjectInstance.cpp
+++ b/Engine/Plugins/Mutable/Source/CustomizableObject/Private/MuCO/CustomizableObjectInstance.cpp
@@ -4997,7 +4997,34 @@ bool UCustomizableInstancePrivate::BuildOrCopyRenderData(const TSharedRef<FUpdat
 			}
 			
 			uint32 SurfaceMetadataId = OperationData->MutableInstance->GetSurfaceCustomId(InstanceComponentIndex, LODIndex, InstanceSurfaceIndex);
-			MeshSurfacesMetadata[MeshSectionIndex] = ModelResources.SurfaceMetadata.Find(SurfaceMetadataId);
+			
+			uint32 UsedSurfaceMetadataId = 0;
+			if (SurfaceMetadataId != 0)
+			{
+				UsedSurfaceMetadataId = SurfaceMetadataId;
+			}
+			else
+			{
+				// In case the surface does not have metadata, check if any submesh has surface metadata.
+				for (const mu::FSurfaceSubMesh& SubMesh : LOD.Mesh->Surfaces[MeshSectionIndex].SubMeshes)	
+				{
+					const FMutableMeshMetadata* FoundMeshMetadata = ModelResources.MeshMetadata.Find(SubMesh.ExternalId);
+
+					if (!FoundMeshMetadata)
+					{
+						continue;
+					}
+
+					UsedSurfaceMetadataId = FoundMeshMetadata->SurfaceMetadataId; 
+
+					if (UsedSurfaceMetadataId != 0)
+					{
+						break;
+					}
+				}
+			}	
+			
+			MeshSurfacesMetadata[MeshSectionIndex] = ModelResources.SurfaceMetadata.Find(UsedSurfaceMetadataId);
 		}
 
 		// Set RenderSections
@@ -6218,16 +6245,44 @@ void UCustomizableInstancePrivate::BuildMaterials(const TSharedRef<FUpdateContex
 				FSkeletalMaterial& MaterialSlot = Materials.AddDefaulted_GetRef();
 				MaterialSlot.MaterialInterface = MaterialTemplate;
 
-				const FMutableSurfaceMetadata* FoundSurfaceMetadata = ModelResources.SurfaceMetadata.Find(Surface.SurfaceMetadataId);
+				uint32 UsedSurfaceMetadataId = Surface.SurfaceMetadataId;
 				
-				if (FoundSurfaceMetadata)
+				// If the surface metadata is invalid, check if any of the mesh fragments has metadata. 
+				// For now use the fisrt found, an aggregate may be needed. 
+				if (Surface.SurfaceMetadataId == 0 && LOD.Mesh)
 				{
-					const int32 MaterialSlotNameIndex = FoundSurfaceMetadata->MaterialSlotIndex;
-					if (ModelResources.MaterialSlotNames.IsValidIndex(MaterialSlotNameIndex))
+					int32 MeshSurfaceIndex = LOD.Mesh->Surfaces.IndexOfByPredicate([SurfaceId = Surface.SurfaceId](const mu::FMeshSurface& Surface)
 					{
-						MaterialSlot.MaterialSlotName = ModelResources.MaterialSlotNames[MaterialSlotNameIndex];
+						return SurfaceId == Surface.Id;
+					});
+
+					if (MeshSurfaceIndex != INDEX_NONE)
+					{
+						for (const mu::FSurfaceSubMesh& SubMesh : LOD.Mesh->Surfaces[SurfaceIndex].SubMeshes)	
+						{
+							const FMutableMeshMetadata* FoundMeshMetadata = ModelResources.MeshMetadata.Find(SubMesh.ExternalId);
+
+							if (!FoundMeshMetadata)
+							{
+								continue;
+							}
+
+							UsedSurfaceMetadataId = FoundMeshMetadata->SurfaceMetadataId; 
+
+							if (UsedSurfaceMetadataId != 0)
+							{
+								break;
+							}
+						}
 					}
 				}
+
+				const FMutableSurfaceMetadata* FoundSurfaceMetadata = ModelResources.SurfaceMetadata.Find(UsedSurfaceMetadataId);
+				
+				if (FoundSurfaceMetadata)
+				{
+					MaterialSlot.MaterialSlotName = FoundSurfaceMetadata->MaterialSlotName;
+				}
 				if (RefSkeletalMeshData)
 				{
 					SetMeshUVChannelDensity(MaterialSlot.UVChannelData, RefSkeletalMeshData->Settings.DefaultUVChannelDensity);
diff --git a/Engine/Plugins/Mutable/Source/CustomizableObjectEditor/Private/MuCOE/GenerateMutableSource/GenerateMutableSource.cpp b/Engine/Plugins/Mutable/Source/CustomizableObjectEditor/Private/MuCOE/GenerateMutableSource/GenerateMutableSource.cpp
--- a/Engine/Plugins/Mutable/Source/CustomizableObjectEditor/Private/MuCOE/GenerateMutableSource/GenerateMutableSource.cpp
+++ b/Engine/Plugins/Mutable/Source/CustomizableObjectEditor/Private/MuCOE/GenerateMutableSource/GenerateMutableSource.cpp
@@ -841,6 +841,56 @@ mu::NodeMeshApplyPosePtr CreateNodeMeshApplyPose(FMutableGraphGenerationContext&
 	return NodeMeshApplyPose;
 }
 
+uint32 AddUniqueSurfaceMetadata(const FSkeletalMaterial* Material, const FSkelMeshSection* MeshSection, TMap<uint32, FMutableSurfaceMetadata>& InOutHashSurfaceMetadataSet)
+{
+	auto HashSurfaceMetadataFunc = [](const FMutableSurfaceMetadata& Data) -> uint32
+	{
+		// Create a string representation of the metadata.
+		FString SlotNameString = Data.MaterialSlotName.ToString();
+		SlotNameString.ToLowerInline();
+
+		TArrayView<const uint8> NameDataView(
+				reinterpret_cast<const uint8*>(*SlotNameString), SlotNameString.Len()*sizeof(FString::ElementType));
+
+		TArray<uint8, TInlineAllocator<256>> DataToHash;
+
+		DataToHash.Add(static_cast<uint8>(Data.bCastShadow));
+		DataToHash.Append(NameDataView);
+
+		return CityHash32(reinterpret_cast<const char*>(DataToHash.GetData()), DataToHash.Num());
+	};
+
+	auto CompareSurfaceMetadataFunc = [](const FMutableSurfaceMetadata& A, const FMutableSurfaceMetadata& B)
+	{
+		return A.bCastShadow == B.bCastShadow && A.MaterialSlotName == B.MaterialSlotName;
+	};
+
+	FMutableSurfaceMetadata SurfaceMetadata;
+
+	if (Material)
+	{
+		SurfaceMetadata.MaterialSlotName = Material->MaterialSlotName;
+	}
+
+	if (MeshSection)
+	{
+		SurfaceMetadata.bCastShadow = MeshSection->bCastShadow;
+	}
+
+	uint32 SurfaceMetadataUniqueHash = Private::GenerateUniquePersistentHash(
+			SurfaceMetadata, InOutHashSurfaceMetadataSet, HashSurfaceMetadataFunc, CompareSurfaceMetadataFunc);
+
+	if (SurfaceMetadataUniqueHash != 0)
+	{
+		InOutHashSurfaceMetadataSet.FindOrAdd(SurfaceMetadataUniqueHash, SurfaceMetadata);
+	}
+	else
+	{
+		UE_LOG(LogMutable, Error, TEXT("Maximum number of surfaces reached."));
+	}
+
+	return SurfaceMetadataUniqueHash;
+}
 
 // Convert a CustomizableObject Source Graph into a mutable source graph  
 mu::Ptr<mu::NodeObject> GenerateMutableSource(const UEdGraphPin * Pin, FMutableGraphGenerationContext & GenerationContext)
diff --git a/Engine/Plugins/Mutable/Source/CustomizableObjectEditor/Private/MuCOE/GenerateMutableSource/GenerateMutableSource.h b/Engine/Plugins/Mutable/Source/CustomizableObjectEditor/Private/MuCOE/GenerateMutableSource/GenerateMutableSource.h
--- a/Engine/Plugins/Mutable/Source/CustomizableObjectEditor/Private/MuCOE/GenerateMutableSource/GenerateMutableSource.h
+++ b/Engine/Plugins/Mutable/Source/CustomizableObjectEditor/Private/MuCOE/GenerateMutableSource/GenerateMutableSource.h
@@ -1069,3 +1069,10 @@ mu::Ptr<mu::Mesh> GenerateMeshConstant(USkeletalMesh*, FMutableGraphGenerationCo
 
 /** Generates a mutable image descriptor from an unreal engine texture */
 mu::FImageDesc GenerateImageDescriptor(UTexture* Texture);
+
+/** 
+ * Add SurfaceMetadata gathered form Material and MeshSection to HashSurfaceMetadataSet.
+ * 
+ * return the unique id for the SurfaceMetadata in HashSurfaceMetadataSet.
+ **/
+uint32 AddUniqueSurfaceMetadata(const FSkeletalMaterial* Material, const FSkelMeshSection* MeshSection, TMap<uint32, FMutableSurfaceMetadata>& InOutHashSurfaceMetadataSet);
diff --git a/Engine/Plugins/Mutable/Source/CustomizableObjectEditor/Private/MuCOE/GenerateMutableSource/GenerateMutableSourceColor.cpp b/Engine/Plugins/Mutable/Source/CustomizableObjectEditor/Private/MuCOE/GenerateMutableSource/GenerateMutableSourceColor.cpp
--- a/Engine/Plugins/Mutable/Source/CustomizableObjectEditor/Private/MuCOE/GenerateMutableSource/GenerateMutableSourceColor.cpp
+++ b/Engine/Plugins/Mutable/Source/CustomizableObjectEditor/Private/MuCOE/GenerateMutableSource/GenerateMutableSourceColor.cpp
@@ -305,7 +305,7 @@ mu::Ptr<mu::NodeColour> GenerateMutableSourceColor(const UEdGraphPin* Pin, FMuta
 					{
 						int32 Dummy = -1; // TODO MTBL-1512
 						bool Dummy2 = false;
-						bSuccess = GenerateTableColumn(TypedNodeTable, Pin, Table, ColumnName, Property, Dummy, Dummy, GenerationContext.CurrentLOD, Dummy, Dummy2, GenerationContext);
+						bSuccess = GenerateTableColumn(TypedNodeTable, Pin, Table, ColumnName, Property, Dummy, Dummy, GenerationContext.CurrentLOD, Dummy, 0, Dummy2, GenerationContext);
 
 						if (!bSuccess)
 						{
diff --git a/Engine/Plugins/Mutable/Source/CustomizableObjectEditor/Private/MuCOE/GenerateMutableSource/GenerateMutableSourceComponent.cpp b/Engine/Plugins/Mutable/Source/CustomizableObjectEditor/Private/MuCOE/GenerateMutableSource/GenerateMutableSourceComponent.cpp
--- a/Engine/Plugins/Mutable/Source/CustomizableObjectEditor/Private/MuCOE/GenerateMutableSource/GenerateMutableSourceComponent.cpp
+++ b/Engine/Plugins/Mutable/Source/CustomizableObjectEditor/Private/MuCOE/GenerateMutableSource/GenerateMutableSourceComponent.cpp
@@ -356,7 +356,7 @@ mu::Ptr<mu::NodeComponent> GenerateMutableSourceComponent(const UEdGraphPin* Pin
 			FString MeshUniqueTags;
 			constexpr bool bIsReference = true;
 			TSoftClassPtr<UAnimInstance> AnimInstance;
-			mu::Ptr<mu::Mesh> MutableMesh = GenerateMutableMesh(SkeletalMesh, AnimInstance, 0, 0, 0, 0, MeshUniqueTags, 
+			mu::Ptr<mu::Mesh> MutableMesh = GenerateMutableMesh(SkeletalMesh, AnimInstance, 0, 0, 0, 0, MeshUniqueTags, 0, 
 																GenerationContext, TypedComponentPassthroughMesh, nullptr, bIsReference);
 
 			MeshNode->SetValue(MutableMesh);
diff --git a/Engine/Plugins/Mutable/Source/CustomizableObjectEditor/Private/MuCOE/GenerateMutableSource/GenerateMutableSourceFloat.cpp b/Engine/Plugins/Mutable/Source/CustomizableObjectEditor/Private/MuCOE/GenerateMutableSource/GenerateMutableSourceFloat.cpp
--- a/Engine/Plugins/Mutable/Source/CustomizableObjectEditor/Private/MuCOE/GenerateMutableSource/GenerateMutableSourceFloat.cpp
+++ b/Engine/Plugins/Mutable/Source/CustomizableObjectEditor/Private/MuCOE/GenerateMutableSource/GenerateMutableSourceFloat.cpp
@@ -354,7 +354,7 @@ mu::NodeScalarPtr GenerateMutableSourceFloat(const UEdGraphPin* Pin, FMutableGra
 					{
 						int32 Dummy = -1; // TODO MTBL-1512
 						bool Dummy2 = false;
-						bSuccess = GenerateTableColumn(TypedNodeTable, Pin, Table, ColumnName, Property, Dummy, Dummy, GenerationContext.CurrentLOD, Dummy, Dummy2, GenerationContext);
+						bSuccess = GenerateTableColumn(TypedNodeTable, Pin, Table, ColumnName, Property, Dummy, Dummy, GenerationContext.CurrentLOD, Dummy, 0, Dummy2, GenerationContext);
 
 						if (!bSuccess)
 						{
diff --git a/Engine/Plugins/Mutable/Source/CustomizableObjectEditor/Private/MuCOE/GenerateMutableSource/GenerateMutableSourceImage.cpp b/Engine/Plugins/Mutable/Source/CustomizableObjectEditor/Private/MuCOE/GenerateMutableSource/GenerateMutableSourceImage.cpp
--- a/Engine/Plugins/Mutable/Source/CustomizableObjectEditor/Private/MuCOE/GenerateMutableSource/GenerateMutableSourceImage.cpp
+++ b/Engine/Plugins/Mutable/Source/CustomizableObjectEditor/Private/MuCOE/GenerateMutableSource/GenerateMutableSourceImage.cpp
@@ -703,7 +703,7 @@ mu::NodeImagePtr GenerateMutableSourceImage(const UEdGraphPin* Pin, FMutableGrap
 			FLayoutGenerationFlags LayoutGenerationFlags;
 			LayoutGenerationFlags.TexturePinModes.Init(EPinMode::Mutable, TEXSTREAM_MAX_NUM_UVCHANNELS);
 			GenerationContext.LayoutGenerationFlags.Push(LayoutGenerationFlags);
-			mu::NodeMeshPtr MeshNode = GenerateMutableSourceMesh(ConnectedPin, GenerationContext, DummyMeshData, false, false);
+			mu::NodeMeshPtr MeshNode = GenerateMutableSourceMesh(ConnectedPin, GenerationContext, DummyMeshData, 0, false, false);
 			GenerationContext.LayoutGenerationFlags.Pop();
 			ImageNode->SetMesh(MeshNode);
 		}
@@ -1092,7 +1092,7 @@ mu::NodeImagePtr GenerateMutableSourceImage(const UEdGraphPin* Pin, FMutableGrap
 						{
 							int32 Dummy = -1; // TODO MTBL-1512
 							bool Dummy2 = false;
-							bSuccess = GenerateTableColumn(TypedNodeTable, Pin, Table, ColumnName, Property, Dummy, Dummy, GenerationContext.CurrentLOD, Dummy, Dummy2, GenerationContext);
+							bSuccess = GenerateTableColumn(TypedNodeTable, Pin, Table, ColumnName, Property, Dummy, Dummy, GenerationContext.CurrentLOD, Dummy, 0, Dummy2, GenerationContext);
 
 							if (!bSuccess)
 							{
diff --git a/Engine/Plugins/Mutable/Source/CustomizableObjectEditor/Private/MuCOE/GenerateMutableSource/GenerateMutableSourceLayout.cpp b/Engine/Plugins/Mutable/Source/CustomizableObjectEditor/Private/MuCOE/GenerateMutableSource/GenerateMutableSourceLayout.cpp
--- a/Engine/Plugins/Mutable/Source/CustomizableObjectEditor/Private/MuCOE/GenerateMutableSource/GenerateMutableSourceLayout.cpp
+++ b/Engine/Plugins/Mutable/Source/CustomizableObjectEditor/Private/MuCOE/GenerateMutableSource/GenerateMutableSourceLayout.cpp
@@ -95,7 +95,7 @@ mu::Ptr<mu::NodeLayout> CreateMutableLayoutNode(FMutableGraphGenerationContext&
 			GenerationContext.MeshGenerationFlags.Push(ShapeFlags);
 
 			GenerationContext.ComponentInfos.Add(FMutableComponentInfo(FName(), SkeletalMesh));
-			MutableMesh = ConvertSkeletalMeshToMutable(SkeletalMesh, TSoftClassPtr<UAnimInstance>(), UnrealLayout->GetLOD(), UnrealLayout->GetMaterial(), UnrealLayout->GetLOD(), UnrealLayout->GetMaterial(), GenerationContext, nullptr, nullptr);
+			MutableMesh = ConvertSkeletalMeshToMutable(SkeletalMesh, TSoftClassPtr<UAnimInstance>(), UnrealLayout->GetLOD(), UnrealLayout->GetMaterial(), UnrealLayout->GetLOD(), UnrealLayout->GetMaterial(), 0, GenerationContext, nullptr, nullptr);
 
 			GenerationContext.MeshGenerationFlags.Pop();
 		}
diff --git a/Engine/Plugins/Mutable/Source/CustomizableObjectEditor/Private/MuCOE/GenerateMutableSource/GenerateMutableSourceMesh.cpp b/Engine/Plugins/Mutable/Source/CustomizableObjectEditor/Private/MuCOE/GenerateMutableSource/GenerateMutableSourceMesh.cpp
--- a/Engine/Plugins/Mutable/Source/CustomizableObjectEditor/Private/MuCOE/GenerateMutableSource/GenerateMutableSourceMesh.cpp
+++ b/Engine/Plugins/Mutable/Source/CustomizableObjectEditor/Private/MuCOE/GenerateMutableSource/GenerateMutableSourceMesh.cpp
@@ -683,9 +683,11 @@ namespace MutablePrivate
 }
 
 
-mu::MeshPtr ConvertSkeletalMeshToMutable(const USkeletalMesh* InSkeletalMesh, const TSoftClassPtr<UAnimInstance>& AnimBp, int32 LODIndexConnected, 
-										 int32 SectionIndexConnected, int32 LODIndex, int32 SectionIndex, FMutableGraphGenerationContext& GenerationContext, 
-										 const UCustomizableObjectNode* CurrentNode, USkeletalMesh* TableReferenceSkeletalMesh)
+
+mu::MeshPtr ConvertSkeletalMeshToMutable(
+		const USkeletalMesh* InSkeletalMesh, const TSoftClassPtr<UAnimInstance>& AnimBp, int32 LODIndexConnected, 
+		int32 SectionIndexConnected, int32 LODIndex, int32 SectionIndex, uint32 SurfaceMetadataId,
+		FMutableGraphGenerationContext& GenerationContext, const UCustomizableObjectNode* CurrentNode, USkeletalMesh* TableReferenceSkeletalMesh)
 {
 	MUTABLE_CPUPROFILER_SCOPE(ConvertSkeletalMeshToMutable);
 
@@ -2204,6 +2206,7 @@ mu::MeshPtr ConvertSkeletalMeshToMutable(const USkeletalMesh* InSkeletalMesh, co
 		// this info can be stored in the CO per mesh.
 		MeshMetadata.MorphMetadataId = 0;
 		MeshMetadata.ClothingMetadataId = 0;
+		MeshMetadata.SurfaceMetadataId = SurfaceMetadataId;
 
 		auto HashMeshMetadataFunc = [](const FMutableMeshMetadata& Data) -> uint32
 		{
@@ -2435,7 +2438,7 @@ mu::MeshPtr ConvertStaticMeshToMutable(const UStaticMesh* StaticMesh, int32 LODI
 // Convert a Mesh constant to a mutable format. UniqueTags are the tags that make this Mesh unique that cannot be merged in the cache 
 // with the exact same Mesh with other tags
 mu::Ptr<mu::Mesh> GenerateMutableMesh(UObject * Mesh, const TSoftClassPtr<UAnimInstance>& AnimInstance, int32 LODIndexConnected, int32 SectionIndexConnected, 
-									  int32 LODIndex, int32 SectionIndex, const FString& UniqueTags, FMutableGraphGenerationContext & GenerationContext, 
+									  int32 LODIndex, int32 SectionIndex, const FString& UniqueTags, uint32 SurfaceMetadataId, FMutableGraphGenerationContext & GenerationContext, 
 									  const UCustomizableObjectNode* CurrentNode, USkeletalMesh* TableReferenceSkeletalMesh, bool bIsReference)
 {
 	// Get the mesh generation flags to use
@@ -2457,7 +2460,7 @@ mu::Ptr<mu::Mesh> GenerateMutableMesh(UObject * Mesh, const TSoftClassPtr<UAnimI
 		}
 		else
 		{
-			MutableMesh = ConvertSkeletalMeshToMutable(SkeletalMesh, AnimInstance, LODIndexConnected, SectionIndexConnected, LODIndex, SectionIndex, GenerationContext, CurrentNode, TableReferenceSkeletalMesh);
+			MutableMesh = ConvertSkeletalMeshToMutable(SkeletalMesh, AnimInstance, LODIndexConnected, SectionIndexConnected, LODIndex, SectionIndex, SurfaceMetadataId, GenerationContext, CurrentNode, TableReferenceSkeletalMesh);
 
 			FSkeletalMeshModel* ImportedModel = SkeletalMesh->GetImportedModel();
 
@@ -2615,7 +2618,7 @@ mu::Ptr<mu::Mesh> BuildMorphedMutableMesh(const UEdGraphPin* BaseSourcePin, cons
 		// Get the base mesh
 		constexpr bool bIsReference = false;
 		mu::Ptr<mu::Mesh> BaseSourceMesh = GenerateMutableMesh(SkeletalMesh, TSoftClassPtr<UAnimInstance>(), LODIndexConnected, SectionIndexConnected,
-			LODIndex, SectionIndex, FString(), GenerationContext, Node, nullptr, bIsReference);
+			LODIndex, SectionIndex, FString(), 0, GenerationContext, Node, nullptr, bIsReference);
 		if (BaseSourceMesh)
 		{
 			MorphedSourceMesh = BuildMorphedMutableMeshFromMesh(BaseSourceMesh, SkeletalMesh, MorphTargetName, LODIndex, SectionIndex);
@@ -3292,9 +3295,10 @@ mu::NodeMeshPtr GenerateMorphMesh(const UEdGraphPin* Pin,
 	return Result;
 }
 
-mu::NodeMeshPtr GenerateMutableSourceMesh(const UEdGraphPin* Pin,
+mu::Ptr<mu::NodeMesh> GenerateMutableSourceMesh(const UEdGraphPin* Pin,
 	FMutableGraphGenerationContext& GenerationContext,
 	FMutableGraphMeshGenerationData& MeshData,
+	uint32 SurfaceMetadataId,
 	const bool bLinkedToExtendMaterial,
 	const bool bOnlyConnectedLOD)
 {
@@ -3405,7 +3409,7 @@ mu::NodeMeshPtr GenerateMutableSourceMesh(const UEdGraphPin* Pin,
 			
 			constexpr bool bIsReference = false;
 			mu::Ptr<mu::Mesh> MutableMesh = GenerateMutableMesh(TypedNodeSkel->SkeletalMesh, TypedNodeSkel->AnimInstance, LODIndexConnected, SectionIndexConnected, 
-																LODIndex, SectionIndex, MeshUniqueTags, GenerationContext, TypedNodeSkel, nullptr, bIsReference);
+																LODIndex, SectionIndex, MeshUniqueTags, SurfaceMetadataId, GenerationContext, TypedNodeSkel, nullptr, bIsReference);
 			if (MutableMesh)
 			{
 				MeshNode->SetValue(MutableMesh);
@@ -3620,7 +3624,7 @@ mu::NodeMeshPtr GenerateMutableSourceMesh(const UEdGraphPin* Pin,
 
 			constexpr bool bIsReference = false;
 			mu::MeshPtr MutableMesh = GenerateMutableMesh(TypedNodeStatic->StaticMesh, TSoftClassPtr<UAnimInstance>(), LODIndex, SectionIndex, LODIndex, SectionIndex, 
-														  FString(), GenerationContext, TypedNodeStatic, nullptr, bIsReference);
+														  FString(), 0, GenerationContext, TypedNodeStatic, nullptr, bIsReference);
 			if (MutableMesh)
 			{
 				MeshNode->SetValue(MutableMesh);
@@ -3671,7 +3675,7 @@ mu::NodeMeshPtr GenerateMutableSourceMesh(const UEdGraphPin* Pin,
 			// Mesh Morph Stack Management
 			FMorphNodeData NewMorphData = { TypedNodeMorph, TypedNodeMorph->MorphTargetName ,TypedNodeMorph->FactorPin(), TypedNodeMorph->MeshPin() };
 			GenerationContext.MeshMorphStack.Push(NewMorphData);
-			Result = GenerateMutableSourceMesh(ConnectedPin, GenerationContext, MeshData, false, bOnlyConnectedLOD);
+			Result = GenerateMutableSourceMesh(ConnectedPin, GenerationContext, MeshData, SurfaceMetadataId, false, bOnlyConnectedLOD);
 			GenerationContext.MeshMorphStack.Pop(EAllowShrinking::Yes);
 		}
 		else
@@ -3733,7 +3737,7 @@ mu::NodeMeshPtr GenerateMutableSourceMesh(const UEdGraphPin* Pin,
 
 				if (const UEdGraphPin* MeshConnectedPin = FollowInputPin(*TypedNodeMeshMorphStackApp->GetMeshPin()))
 				{
-					Result = GenerateMutableSourceMesh(MeshConnectedPin, GenerationContext, MeshData, false, bOnlyConnectedLOD);
+					Result = GenerateMutableSourceMesh(MeshConnectedPin, GenerationContext, MeshData, SurfaceMetadataId, false, bOnlyConnectedLOD);
 				}
 
 				for (int32 MorphIndex = 0; MorphIndex < AddedMorphs; ++MorphIndex)
@@ -3805,7 +3809,7 @@ mu::NodeMeshPtr GenerateMutableSourceMesh(const UEdGraphPin* Pin,
 					if (const UEdGraphPin* ConnectedPin = FollowInputPin(*TypedNodeMeshSwitch->GetElementPin(SelectorIndex)))
 					{
 						FMutableGraphMeshGenerationData ChildMeshData;
-						Result = GenerateMutableSourceMesh(ConnectedPin, GenerationContext, ChildMeshData, false, bOnlyConnectedLOD);
+						Result = GenerateMutableSourceMesh(ConnectedPin, GenerationContext, ChildMeshData, SurfaceMetadataId, false, bOnlyConnectedLOD);
 						SwitchNode->SetOption(SelectorIndex, Result);
 						MeshData.Combine(ChildMeshData);
 					}
@@ -3830,7 +3834,7 @@ mu::NodeMeshPtr GenerateMutableSourceMesh(const UEdGraphPin* Pin,
 		if (const UEdGraphPin* ConnectedPin = FollowInputPin(*TypedNodeMeshVar->DefaultPin()))
 		{
 			FMutableGraphMeshGenerationData ChildMeshData;
-			mu::NodeMeshPtr ChildNode = GenerateMutableSourceMesh(ConnectedPin, GenerationContext, ChildMeshData, false, bOnlyConnectedLOD);
+			mu::NodeMeshPtr ChildNode = GenerateMutableSourceMesh(ConnectedPin, GenerationContext, ChildMeshData, SurfaceMetadataId, false, bOnlyConnectedLOD);
 			if (ChildNode)
 			{
 				MeshNode->SetDefaultMesh(ChildNode.get());
@@ -3853,7 +3857,7 @@ mu::NodeMeshPtr GenerateMutableSourceMesh(const UEdGraphPin* Pin,
 			if (const UEdGraphPin* ConnectedPin = FollowInputPin(*VariationPin))
 			{
 				FMutableGraphMeshGenerationData VariationMeshData;
-				mu::NodeMeshPtr ChildNode = GenerateMutableSourceMesh(ConnectedPin, GenerationContext, VariationMeshData, false, bOnlyConnectedLOD);
+				mu::NodeMeshPtr ChildNode = GenerateMutableSourceMesh(ConnectedPin, GenerationContext, VariationMeshData, SurfaceMetadataId, false, bOnlyConnectedLOD);
 				MeshNode->SetVariationMesh(VariationIndex, ChildNode.get());
 				MeshData.Combine(VariationMeshData);
 			}
@@ -3868,7 +3872,7 @@ mu::NodeMeshPtr GenerateMutableSourceMesh(const UEdGraphPin* Pin,
 		if (const UEdGraphPin* ConnectedPin = FollowInputPin(*TypedNodeGeometry->MeshAPin()))
 		{
 			FMutableGraphMeshGenerationData ChildMeshData;
-			mu::Ptr<mu::NodeMesh> ChildNode = GenerateMutableSourceMesh(ConnectedPin, GenerationContext, ChildMeshData, false, bOnlyConnectedLOD);
+			mu::Ptr<mu::NodeMesh> ChildNode = GenerateMutableSourceMesh(ConnectedPin, GenerationContext, ChildMeshData, SurfaceMetadataId, false, bOnlyConnectedLOD);
 			if (ChildNode)
 			{
 				MeshNode->SetMeshA(ChildNode.get());
@@ -3888,7 +3892,7 @@ mu::NodeMeshPtr GenerateMutableSourceMesh(const UEdGraphPin* Pin,
 		if (const UEdGraphPin* ConnectedPin = FollowInputPin(*TypedNodeGeometry->MeshBPin()))
 		{
 			FMutableGraphMeshGenerationData ChildMeshData;
-			mu::Ptr<mu::NodeMesh> ChildNode = GenerateMutableSourceMesh(ConnectedPin, GenerationContext, ChildMeshData, false, bOnlyConnectedLOD);
+			mu::Ptr<mu::NodeMesh> ChildNode = GenerateMutableSourceMesh(ConnectedPin, GenerationContext, ChildMeshData, SurfaceMetadataId, false, bOnlyConnectedLOD);
 			if (ChildNode)
 			{
 				MeshNode->SetMeshB(ChildNode.get());
@@ -3935,7 +3939,7 @@ mu::NodeMeshPtr GenerateMutableSourceMesh(const UEdGraphPin* Pin,
 		if (const UEdGraphPin* ConnectedPin = FollowInputPin(*TypedNodeReshape->BaseMeshPin()))
 		{
 			FMutableGraphMeshGenerationData ChildMeshData;
-			mu::Ptr<mu::NodeMesh> ChildNode = GenerateMutableSourceMesh(ConnectedPin, GenerationContext, ChildMeshData, false, bOnlyConnectedLOD);
+			mu::Ptr<mu::NodeMesh> ChildNode = GenerateMutableSourceMesh(ConnectedPin, GenerationContext, ChildMeshData, SurfaceMetadataId, false, bOnlyConnectedLOD);
 			if (ChildNode)
 			{
 				MeshNode->SetBaseMesh(ChildNode.get());
@@ -4063,7 +4067,7 @@ mu::NodeMeshPtr GenerateMutableSourceMesh(const UEdGraphPin* Pin,
 		if (const UEdGraphPin* ConnectedPin = FollowInputPin(*TypedNodeReshape->BaseShapePin()))
 		{
 			FMutableGraphMeshGenerationData ChildMeshData;
-			mu::Ptr<mu::NodeMesh> ChildNode = GenerateMutableSourceMesh(ConnectedPin, GenerationContext, ChildMeshData, false, true);
+			mu::Ptr<mu::NodeMesh> ChildNode = GenerateMutableSourceMesh(ConnectedPin, GenerationContext, ChildMeshData, 0, false, true);
 	
 			if (ChildNode)
 			{
@@ -4080,7 +4084,7 @@ mu::NodeMeshPtr GenerateMutableSourceMesh(const UEdGraphPin* Pin,
 		if (const UEdGraphPin* ConnectedPin = FollowInputPin(*TypedNodeReshape->TargetShapePin()))
 		{
 			FMutableGraphMeshGenerationData ChildMeshData;
-			mu::Ptr<mu::NodeMesh> ChildNode = GenerateMutableSourceMesh(ConnectedPin, GenerationContext, ChildMeshData, false, true);
+			mu::Ptr<mu::NodeMesh> ChildNode = GenerateMutableSourceMesh(ConnectedPin, GenerationContext, ChildMeshData, 0, false, true);
 			
 			if (ChildNode)
 			{
@@ -4120,7 +4124,7 @@ mu::NodeMeshPtr GenerateMutableSourceMesh(const UEdGraphPin* Pin,
 	{
 		if (const UEdGraphPin* InputMeshPin = FollowInputPin(*TypedNode->GetInputMeshPin()))
 		{
-			mu::Ptr<mu::NodeMesh> InputMeshNode = GenerateMutableSourceMesh(InputMeshPin, GenerationContext, MeshData, false, bOnlyConnectedLOD);
+			mu::Ptr<mu::NodeMesh> InputMeshNode = GenerateMutableSourceMesh(InputMeshPin, GenerationContext, MeshData, SurfaceMetadataId, false, bOnlyConnectedLOD);
 
 			if (GenerationContext.GetCurrentComponentInfo()->RefSkeletalMesh)
 			{
@@ -4148,7 +4152,7 @@ mu::NodeMeshPtr GenerateMutableSourceMesh(const UEdGraphPin* Pin,
 					if (const UCustomizableObjectNodeTable* TypedNodeTable = Cast<UCustomizableObjectNodeTable>(TablePosePin->GetOwningNode()))
 					{
 						mu::NodeMeshApplyPosePtr NodeMeshApplyPose = new mu::NodeMeshApplyPose();
-						mu::Ptr<mu::NodeMesh> MeshTableNode = GenerateMutableSourceMesh(TablePosePin, GenerationContext, MeshData, false, bOnlyConnectedLOD);
+						mu::Ptr<mu::NodeMesh> MeshTableNode = GenerateMutableSourceMesh(TablePosePin, GenerationContext, MeshData, SurfaceMetadataId, false, bOnlyConnectedLOD);
 
 						NodeMeshApplyPose->SetBase(InputMeshNode);
 						NodeMeshApplyPose->SetPose(MeshTableNode);
@@ -4247,7 +4251,7 @@ mu::NodeMeshPtr GenerateMutableSourceMesh(const UEdGraphPin* Pin,
 					// Generating a new Mesh column if not exists
 					if (Table->FindColumn(MutableColumnName) == INDEX_NONE)
 					{
-						bSuccess = GenerateTableColumn(TypedNodeTable, Pin, Table, DataTableColumnName, Property, LODIndexConnected, SectionIndexConnected, LODIndex, SectionIndex, bOnlyConnectedLOD, GenerationContext);
+						bSuccess = GenerateTableColumn(TypedNodeTable, Pin, Table, DataTableColumnName, Property, LODIndexConnected, SectionIndexConnected, LODIndex, SectionIndex, SurfaceMetadataId, bOnlyConnectedLOD, GenerationContext);
 
 						if (!bSuccess)
 						{
diff --git a/Engine/Plugins/Mutable/Source/CustomizableObjectEditor/Private/MuCOE/GenerateMutableSource/GenerateMutableSourceMesh.h b/Engine/Plugins/Mutable/Source/CustomizableObjectEditor/Private/MuCOE/GenerateMutableSource/GenerateMutableSourceMesh.h
--- a/Engine/Plugins/Mutable/Source/CustomizableObjectEditor/Private/MuCOE/GenerateMutableSource/GenerateMutableSourceMesh.h
+++ b/Engine/Plugins/Mutable/Source/CustomizableObjectEditor/Private/MuCOE/GenerateMutableSource/GenerateMutableSourceMesh.h
@@ -51,7 +51,7 @@ void GetLODAndSectionForAutomaticLODs(const FMutableGraphGenerationContext& Cont
  * @return Mutable Mesh. Nullptr if there has been an error. Empty mesh if the Skeletal Mesh does not contain the requested LOD + Section. */
 mu::Ptr<mu::Mesh> ConvertSkeletalMeshToMutable(const USkeletalMesh* InSkeletalMesh, const TSoftClassPtr<UAnimInstance>& AnimBp,
                                          int32 LODIndexConnected,  int32 SectionIndexConnected,
-                                         int32 LODIndex, int32 SectionIndex,
+                                         int32 LODIndex, int32 SectionIndex, uint32 SurfaceMetadataId,
                                          FMutableGraphGenerationContext& GenerationContext,
                                          const UCustomizableObjectNode* CurrentNode, 
 										 USkeletalMesh* TableReferenceSkeletalMesh);
@@ -76,7 +76,7 @@ mu::Ptr<mu::Mesh> ConvertStaticMeshToMutable(const UStaticMesh* StaticMesh, int3
  * @return Mutable Mesh. Nullptr if there has been an error. Empty mesh if the Skeletal Mesh does not contain the requested LOD + Section. */
 mu::Ptr<mu::Mesh> GenerateMutableMesh(UObject* Mesh, const TSoftClassPtr<UAnimInstance>& AnimBp,
                                 int32 LODIndexConnected, int32 SectionIndexConnected,
-                                int32 LODIndex, int32 SectionIndex, const FString& MeshUniqueTags,
+                                int32 LODIndex, int32 SectionIndex, const FString& MeshUniqueTags, uint32 SurfaceMetadataId, 
                                 FMutableGraphGenerationContext& GenerationContext, 
 								const UCustomizableObjectNode* CurrentNode, USkeletalMesh* TableReferenceSkeletalMesh,
 								bool bIsReference);
@@ -91,7 +91,8 @@ mu::Ptr<mu::Mesh> BuildMorphedMutableMesh(const UEdGraphPin* BaseSourcePin, cons
  * @param Pin 
  * @param GenerationContext 
  * @param MeshData 
+ * @param SurfaceMetadataId
  * @param bLinkedToExtendMaterial 
  * @param bOnlyConnectedLOD Corrected LOD and Section will unconditionally always be the connected ones.
  * @return  Mutable Mesh Node. */
-mu::Ptr<mu::NodeMesh> GenerateMutableSourceMesh(const UEdGraphPin* Pin, FMutableGraphGenerationContext& GenerationContext, FMutableGraphMeshGenerationData& MeshData, bool bLinkedToExtendMaterial, bool bOnlyConnectedLOD);
+mu::Ptr<mu::NodeMesh> GenerateMutableSourceMesh(const UEdGraphPin* Pin, FMutableGraphGenerationContext& GenerationContext, FMutableGraphMeshGenerationData& MeshData, uint32 SurfaceMetadataId, bool bLinkedToExtendMaterial, bool bOnlyConnectedLOD);
diff --git a/Engine/Plugins/Mutable/Source/CustomizableObjectEditor/Private/MuCOE/GenerateMutableSource/GenerateMutableSourceModifier.cpp b/Engine/Plugins/Mutable/Source/CustomizableObjectEditor/Private/MuCOE/GenerateMutableSource/GenerateMutableSourceModifier.cpp
--- a/Engine/Plugins/Mutable/Source/CustomizableObjectEditor/Private/MuCOE/GenerateMutableSource/GenerateMutableSourceModifier.cpp
+++ b/Engine/Plugins/Mutable/Source/CustomizableObjectEditor/Private/MuCOE/GenerateMutableSource/GenerateMutableSourceModifier.cpp
@@ -4,6 +4,7 @@
 
 #include "Engine/StaticMesh.h"
 #include "MuCOE/CustomizableObjectCompiler.h"
+#include "MuCOE/GenerateMutableSource/GenerateMutableSource.h"
 #include "MuCOE/GenerateMutableSource/GenerateMutableSourceMesh.h"
 #include "MuCOE/GenerateMutableSource/GenerateMutableSourceImage.h"
 #include "MuCOE/GenerateMutableSource/GenerateMutableSourceFloat.h"
@@ -23,6 +24,8 @@
 #include "MuCOE/Nodes/CustomizableObjectNodeModifierRemoveMesh.h"
 #include "MuCOE/Nodes/CustomizableObjectNodeModifierRemoveMeshBlocks.h"
 #include "MuCOE/Nodes/CustomizableObjectNodeStaticMesh.h"
+#include "MuCOE/Nodes/CustomizableObjectNodeSkeletalMesh.h"
+#include "MuCOE/Nodes/CustomizableObjectNodeTable.h"
 #include "MuCOE/Nodes/CustomizableObjectNodeFloatParameter.h"
 #include "MuCOE/Nodes/CustomizableObjectNodeFloatConstant.h"
 #include "MuCOE/Nodes/CustomizableObjectNodeModifierTransformInMesh.h"
@@ -36,6 +39,7 @@
 #include "MuT/NodeMeshConstant.h"
 #include "MuT/NodeMeshFragment.h"
 #include "MuT/NodeMeshFormat.h"
+#include "Rendering/SkeletalMeshLODModel.h"
 
 #define LOCTEXT_NAMESPACE "CustomizableObjectEditor"
 
@@ -106,7 +110,7 @@ mu::Ptr<mu::NodeModifier> GenerateMutableSourceModifier(const UEdGraphPin * Pin,
 		if (const UEdGraphPin* ConnectedPin = FollowInputPin(*TypedNodeClipDeform->ClipShapePin()))
 		{
 			FMutableGraphMeshGenerationData DummyMeshData;
-			mu::Ptr<mu::NodeMesh> ClipMesh = GenerateMutableSourceMesh(ConnectedPin, GenerationContext, DummyMeshData, false, true);
+			mu::Ptr<mu::NodeMesh> ClipMesh = GenerateMutableSourceMesh(ConnectedPin, GenerationContext, DummyMeshData, 0, false, true);
 
 			ClipNode->ClipMesh = ClipMesh;
 
@@ -162,7 +166,7 @@ mu::Ptr<mu::NodeModifier> GenerateMutableSourceModifier(const UEdGraphPin * Pin,
 		{
 			FMutableGraphMeshGenerationData DummyMeshData;
 
-			mu::NodeMeshPtr ClipMesh = GenerateMutableSourceMesh(ConnectedPin, GenerationContext, DummyMeshData, false, true);
+			mu::NodeMeshPtr ClipMesh = GenerateMutableSourceMesh(ConnectedPin, GenerationContext, DummyMeshData, 0, false, true);
 
 			FPinDataValue* PinData = GenerationContext.PinData.Find(ConnectedPin);
 			for (const FMeshData& MeshData : PinData->MeshesData)
@@ -288,7 +292,33 @@ mu::Ptr<mu::NodeModifier> GenerateMutableSourceModifier(const UEdGraphPin * Pin,
 
 				GenerationContext.LayoutGenerationFlags.Push(LayoutGenerationFlags);
 
-				AddMeshNode = GenerateMutableSourceMesh(ConnectedPin, GenerationContext, MeshData, true, false);
+				// Generate surface metadata for this fragment.
+				uint32 SurfaceMetadataUniqueHash = 0;	
+				if (ConnectedPin)
+				{
+					//NOTE: This is the same is done in GenerateMutableSourceSurface. 
+					if (const UEdGraphPin* SkeletalMeshPin = FindMeshBaseSource(*ConnectedPin, false))
+					{
+						FSkeletalMaterial* SkeletalMaterial = nullptr;
+						const FSkelMeshSection* ReferenceSkelMeshSection = nullptr;
+						
+						if (const UCustomizableObjectNodeSkeletalMesh* SkeletalMeshNode = Cast<UCustomizableObjectNodeSkeletalMesh>(SkeletalMeshPin->GetOwningNode()))
+						{
+							SkeletalMaterial = SkeletalMeshNode->GetSkeletalMaterialFor(*SkeletalMeshPin);
+							ReferenceSkelMeshSection = SkeletalMeshNode->GetSkeletalMeshSectionFor(*SkeletalMeshPin);
+						}
+
+						else if (const UCustomizableObjectNodeTable* TableNode = Cast<UCustomizableObjectNodeTable>(SkeletalMeshPin->GetOwningNode()))
+						{
+							SkeletalMaterial = TableNode->GetDefaultSkeletalMaterialFor(*SkeletalMeshPin);
+							ReferenceSkelMeshSection = TableNode->GetDefaultSkeletalMeshSectionFor(*SkeletalMeshPin);
+						}
+
+						SurfaceMetadataUniqueHash = AddUniqueSurfaceMetadata(SkeletalMaterial, ReferenceSkelMeshSection, GenerationContext.SurfaceMetadata);
+					}
+				}	
+
+				AddMeshNode = GenerateMutableSourceMesh(ConnectedPin, GenerationContext, MeshData, SurfaceMetadataUniqueHash, true, false);
 
 				GenerationContext.LayoutGenerationFlags.Pop();
 			}
@@ -430,7 +460,7 @@ mu::Ptr<mu::NodeModifier> GenerateMutableSourceModifier(const UEdGraphPin * Pin,
 				GenerationContext.CurrentLOD = LODIndex;
 
 				FMutableGraphMeshGenerationData DummyMeshData;
-				mu::Ptr<mu::NodeMesh> RemoveMeshNode = GenerateMutableSourceMesh(ConnectedPin, GenerationContext, DummyMeshData, false, true);
+				mu::Ptr<mu::NodeMesh> RemoveMeshNode = GenerateMutableSourceMesh(ConnectedPin, GenerationContext, DummyMeshData, 0, false, true);
 				SurfNode->LODs[LODIndex].MeshRemove = RemoveMeshNode;
 			}
 		}
@@ -467,8 +497,8 @@ mu::Ptr<mu::NodeModifier> GenerateMutableSourceModifier(const UEdGraphPin * Pin,
 	else if (UCustomizableObjectNodeModifierEditMeshSection* TypedNodeEdit = Cast<UCustomizableObjectNodeModifierEditMeshSection>(Node))
 	{
 		const EMutableMeshConversionFlags ModifiersMeshFlags =
-			EMutableMeshConversionFlags::IgnoreSkinning |
-			EMutableMeshConversionFlags::IgnorePhysics;
+				EMutableMeshConversionFlags::IgnoreSkinning |
+				EMutableMeshConversionFlags::IgnorePhysics;
 		GenerationContext.MeshGenerationFlags.Push(ModifiersMeshFlags);
 
 		mu::Ptr<mu::NodeModifierSurfaceEdit> SurfNode = new mu::NodeModifierSurfaceEdit();
@@ -632,7 +662,7 @@ mu::Ptr<mu::NodeModifier> GenerateMutableSourceModifier(const UEdGraphPin * Pin,
 		{
 			FMutableGraphMeshGenerationData DummyMeshData;
 
-			mu::NodeMeshPtr BoundingMesh = GenerateMutableSourceMesh(ConnectedPin, GenerationContext, DummyMeshData, false, true);
+			mu::NodeMeshPtr BoundingMesh = GenerateMutableSourceMesh(ConnectedPin, GenerationContext, DummyMeshData, 0, false, true);
 
 			const FPinDataValue* PinData = GenerationContext.PinData.Find(ConnectedPin);
 			for (const FMeshData& MeshData : PinData->MeshesData)
diff --git a/Engine/Plugins/Mutable/Source/CustomizableObjectEditor/Private/MuCOE/GenerateMutableSource/GenerateMutableSourceSurface.cpp b/Engine/Plugins/Mutable/Source/CustomizableObjectEditor/Private/MuCOE/GenerateMutableSource/GenerateMutableSourceSurface.cpp
--- a/Engine/Plugins/Mutable/Source/CustomizableObjectEditor/Private/MuCOE/GenerateMutableSource/GenerateMutableSourceSurface.cpp
+++ b/Engine/Plugins/Mutable/Source/CustomizableObjectEditor/Private/MuCOE/GenerateMutableSource/GenerateMutableSourceSurface.cpp
@@ -54,6 +54,7 @@
 #include "MuT/NodeSurfaceVariation.h"
 #include "MuT/UnrealPixelFormatOverride.h"
 
+
 #define LOCTEXT_NAMESPACE "CustomizableObjectEditor"
 
 
@@ -282,44 +283,7 @@ mu::Ptr<mu::NodeSurface> GenerateMutableSourceSurface(const UEdGraphPin * Pin, F
 						ReferenceSkelMeshSection = TableNode->GetDefaultSkeletalMeshSectionFor(*SkeletalMeshPin);
 					}
 
-					// A reference to the material slot name index will be kept in the surface metadata.
-					int32 MaterialSlotNameIndex = INDEX_NONE;
-					if (SkeletalMaterial)
-					{
-						MaterialSlotNameIndex = 
-								GenerationContext.ReferencedMaterialSlotNames.AddUnique(SkeletalMaterial->MaterialSlotName);
-					}
-
-					if (ReferenceSkelMeshSection)
-					{
-						auto HashSurfaceMetadataFunc = [](const FMutableSurfaceMetadata& Data) -> uint32
-						{
-							return CityHash32(reinterpret_cast<const char*>(&Data), sizeof(FMutableSurfaceMetadata));
-						};
-
-						auto CompareSurfaceMetadataFunc = [](const FMutableSurfaceMetadata& A, const FMutableSurfaceMetadata& B)
-						{
-							return FMemory::Memcmp(&A, &B, sizeof(FMutableSurfaceMetadata)) == 0;
-						};
-
-						FMutableSurfaceMetadata SurfaceMetadata;
-						FMemory::Memzero(SurfaceMetadata);
-						
-						SurfaceMetadata.MaterialSlotIndex = MaterialSlotNameIndex;
-						SurfaceMetadata.bCastShadow = ReferenceSkelMeshSection->bCastShadow;
-
-						SurfaceMetadataUniqueHash = Private::GenerateUniquePersistentHash(
-								SurfaceMetadata, GenerationContext.SurfaceMetadata, HashSurfaceMetadataFunc, CompareSurfaceMetadataFunc);
-
-						if (SurfaceMetadataUniqueHash != 0)
-						{
-							GenerationContext.SurfaceMetadata.FindOrAdd(SurfaceMetadataUniqueHash, SurfaceMetadata);
-						}
-						else
-						{
-							UE_LOG(LogMutable, Error, TEXT("Maximum number of surfaces reached."));
-						}
-					}
+					SurfaceMetadataUniqueHash = AddUniqueSurfaceMetadata(SkeletalMaterial, ReferenceSkelMeshSection, GenerationContext.SurfaceMetadata);
 				}
 			}
 		}
@@ -362,7 +326,7 @@ mu::Ptr<mu::NodeSurface> GenerateMutableSourceSurface(const UEdGraphPin * Pin, F
 				GenerationContext.LayoutGenerationFlags.Push(LayoutGenerationFlags);
 
 				FMutableGraphMeshGenerationData MeshData;
-				MeshNode = GenerateMutableSourceMesh(ConnectedPin, GenerationContext, MeshData, false, false);
+				MeshNode = GenerateMutableSourceMesh(ConnectedPin, GenerationContext, MeshData, SurfaceMetadataUniqueHash, false, false);
 
 				GenerationContext.LayoutGenerationFlags.Pop();
 
diff --git a/Engine/Plugins/Mutable/Source/CustomizableObjectEditor/Private/MuCOE/GenerateMutableSource/GenerateMutableSourceTable.cpp b/Engine/Plugins/Mutable/Source/CustomizableObjectEditor/Private/MuCOE/GenerateMutableSource/GenerateMutableSourceTable.cpp
--- a/Engine/Plugins/Mutable/Source/CustomizableObjectEditor/Private/MuCOE/GenerateMutableSource/GenerateMutableSourceTable.cpp
+++ b/Engine/Plugins/Mutable/Source/CustomizableObjectEditor/Private/MuCOE/GenerateMutableSource/GenerateMutableSourceTable.cpp
@@ -32,7 +32,7 @@
 
 
 bool FillTableColumn(const UCustomizableObjectNodeTable* TableNode, mu::TablePtr MutableTable, const FString& ColumnName, const FString& RowName, const uint32 RowId, uint8* CellData, const FProperty* ColumnProperty,
-	const int LODIndexConnected, const int32 SectionIndexConnected, int32 LODIndex, int32 SectionIndex, const bool bOnlyConnectedLOD, FMutableGraphGenerationContext& GenerationContext)
+	const int LODIndexConnected, const int32 SectionIndexConnected, int32 LODIndex, int32 SectionIndex, uint32 SectionMetadataId, const bool bOnlyConnectedLOD, FMutableGraphGenerationContext& GenerationContext)
 {
 	int32 CurrentColumn;
 	UDataTable* DataTablePtr = GetDataTable(TableNode, GenerationContext);
@@ -280,7 +280,7 @@ bool FillTableColumn(const UCustomizableObjectNodeTable* TableNode, mu::TablePtr
 			//TODO: Add AnimBp physics to Tables.
 			constexpr bool bIsReference = false;
 			mu::Ptr<mu::Mesh> MutableMesh = GenerateMutableMesh(SkeletalMesh, TSoftClassPtr<UAnimInstance>(), LODIndexConnected, SectionIndexConnected, 
-																LODIndex, SectionIndex, MeshUniqueTags, GenerationContext, TableNode, ReferenceSkeletalMesh, bIsReference);
+																LODIndex, SectionIndex, MeshUniqueTags, SectionMetadataId, GenerationContext, TableNode, ReferenceSkeletalMesh, bIsReference);
 
 			if (MutableMesh)
 			{
@@ -391,7 +391,7 @@ bool FillTableColumn(const UCustomizableObjectNodeTable* TableNode, mu::TablePtr
 
 			constexpr bool bIsReference = false;
 			mu::MeshPtr MutableMesh = GenerateMutableMesh(StaticMesh, TSoftClassPtr<UAnimInstance>(), CurrentLOD, SectionIndex, CurrentLOD, SectionIndex, 
-														  FString(), GenerationContext, TableNode, nullptr, bIsReference);
+														  FString(), 0, GenerationContext, TableNode, nullptr, bIsReference);
 
 			if (MutableMesh)
 			{
@@ -916,7 +916,7 @@ TArray<FName> GetRowsToCompile(const UDataTable& DataTable, const UCustomizableO
 
 
 bool GenerateTableColumn(const UCustomizableObjectNodeTable* TableNode, const UEdGraphPin* Pin, mu::TablePtr MutableTable, const FString& DataTableColumnName, const FProperty* ColumnProperty,
-	const int32 LODIndexConnected, const int32 SectionIndexConnected, const int32 LODIndex, const int32 SectionIndex, const bool bOnlyConnectedLOD, FMutableGraphGenerationContext& GenerationContext)
+	const int32 LODIndexConnected, const int32 SectionIndexConnected, const int32 LODIndex, const int32 SectionIndex, uint32 SectionMetadataId, const bool bOnlyConnectedLOD, FMutableGraphGenerationContext& GenerationContext)
 {
 	MUTABLE_CPUPROFILER_SCOPE(GenerateTableColumn);
 
@@ -975,7 +975,7 @@ bool GenerateTableColumn(const UCustomizableObjectNodeTable* TableNode, const UE
 		if (uint8* CellData = GetCellData(RowNames[RowIndex], *DataTable, *ColumnProperty))
 		{
 			bool bCellGenerated = FillTableColumn(TableNode, MutableTable, DataTableColumnName, RowNames[RowIndex].ToString(), RowIds[RowIndex], CellData, ColumnProperty,
-				LODIndexConnected, SectionIndexConnected, LODIndex, SectionIndex, bOnlyConnectedLOD, GenerationContext);
+				LODIndexConnected, SectionIndexConnected, LODIndex, SectionIndex, SectionMetadataId, bOnlyConnectedLOD, GenerationContext);
 
 			if (!bCellGenerated)
 			{
diff --git a/Engine/Plugins/Mutable/Source/CustomizableObjectEditor/Private/MuCOE/GenerateMutableSource/GenerateMutableSourceTable.h b/Engine/Plugins/Mutable/Source/CustomizableObjectEditor/Private/MuCOE/GenerateMutableSource/GenerateMutableSourceTable.h
--- a/Engine/Plugins/Mutable/Source/CustomizableObjectEditor/Private/MuCOE/GenerateMutableSource/GenerateMutableSourceTable.h
+++ b/Engine/Plugins/Mutable/Source/CustomizableObjectEditor/Private/MuCOE/GenerateMutableSource/GenerateMutableSourceTable.h
@@ -28,7 +28,7 @@ struct FMutableGraphGenerationContext;
  * @return  */
 bool GenerateTableColumn(const UCustomizableObjectNodeTable* TableNode, const UEdGraphPin* Pin, mu::TablePtr MutableTable, const FString& DataTableColumnName, const FProperty* ColumnProperty,
 	int32 LODIndexConnected, int32 SectionIndexConnected,
-	int32 LODIndex, int32 SectionIndex,
+	int32 LODIndex, int32 SectionIndex, uint32 SurfaceMetadataId,
 	bool bOnlyConnectedLOD,
 	FMutableGraphGenerationContext& GenerationContext);
 
@@ -52,7 +52,7 @@ bool GenerateTableColumn(const UCustomizableObjectNodeTable* TableNode, const UE
  * @return  */
 bool FillTableColumn(const UCustomizableObjectNodeTable* TableNode, mu::TablePtr MutableTable, const FString& ColumnName,
 	const FString& RowName, uint32 RowId, uint8* CellData, const FProperty* ColumnProperty,
-	int32 LODIndexConnected, int32 SectionIndexConnected, int32 LODIndex, int32 SectionIndex,
+	int32 LODIndexConnected, int32 SectionIndexConnected, int32 LODIndex, int32 SectionIndex, uint32 SectionMetadataId,
 	bool bOnlyConnectedLOD,
 	FMutableGraphGenerationContext& GenerationContext);
 
diff --git a/Engine/Plugins/Mutable/Source/MutableTools/Private/MuT/CodeGenerator.cpp b/Engine/Plugins/Mutable/Source/MutableTools/Private/MuT/CodeGenerator.cpp
--- a/Engine/Plugins/Mutable/Source/MutableTools/Private/MuT/CodeGenerator.cpp
+++ b/Engine/Plugins/Mutable/Source/MutableTools/Private/MuT/CodeGenerator.cpp
@@ -836,16 +836,14 @@ namespace mu
 
 
     //---------------------------------------------------------------------------------------------
-    void CodeGenerator::GenerateSurface( FSurfaceGenerationResult& result, 
+    void CodeGenerator::GenerateSurface( FSurfaceGenerationResult& SurfaceResult, 
 										 const FSurfaceGenerationOptions& Options,
-                                         Ptr<const NodeSurfaceNew> surfaceNode )
+                                         Ptr<const NodeSurfaceNew> SurfaceNode )
     {
         MUTABLE_CPUPROFILER_SCOPE(GenerateSurface);
 
-        const NodeSurfaceNew& node = *surfaceNode;
-
         // Build a series of operations to assemble the surface
-        Ptr<ASTOp> lastSurfOp;
+        Ptr<ASTOp> LastSurfOp;
 
         // Generate the mesh
         //------------------------------------------------------------------------
@@ -858,44 +856,49 @@ namespace mu
 
 		// Do we need to generate the mesh? Or was it already generated for state conditions 
 		// accepting the current state?
-		FirstPassGenerator::FSurface* targetSurface = nullptr;
-		for (FirstPassGenerator::FSurface& its : FirstPass.Surfaces)
+		FirstPassGenerator::FSurface* TargetSurface = nullptr;
+		for (FirstPassGenerator::FSurface& Surface : FirstPass.Surfaces)
 		{
-            if (its.Node == &node)
-            {
-                // Check state conditions
-                bool surfaceValidForThisState =
-                    Options.State >= its.StateCondition.Num()
-                    ||
-                    its.StateCondition[Options.State];
+			if (Surface.Node != SurfaceNode)
+			{
+				continue;
+			}
 
-                if (surfaceValidForThisState)
-                {
-                    if (its.ResultSurfaceOp)
-                    {
-                        // Reuse the entire surface
-                        result.surfaceOp = its.ResultSurfaceOp;
-                        return;
-                    }
-                    else
-                    {
-                        // Not already generated, we will generate this
-                        targetSurface = &its;
-                    }
-                }
-            }
+            // Check state conditions
+            const bool bSurfaceValidForThisState = 
+					Options.State >= Surface.StateCondition.Num() ||
+                    Surface.StateCondition[Options.State];
+
+			if (!bSurfaceValidForThisState)
+			{
+				continue;
+			}
+
+			if (Surface.ResultSurfaceOp)
+			{
+				// Reuse the entire surface
+				SurfaceResult.surfaceOp = Surface.ResultSurfaceOp;
+				return;
+			}
+			else
+			{
+				// Not already generated, we will generate this
+				TargetSurface = &Surface;
+			}
 		}
 
-        if (!targetSurface)
+        if (!TargetSurface)
         {
             return;
         }
 
 		// This assumes that the lods are processed in order. It checks it this way because some platforms may have empty lods at the top.
-		const bool bIsBaseForSharedSurface = node.SharedSurfaceId != INDEX_NONE && !SharedMeshOptionsMap.Contains(node.SharedSurfaceId);
+		const bool bIsBaseForSharedSurface = 
+			SurfaceNode->SharedSurfaceId != INDEX_NONE && 
+			!SharedMeshOptionsMap.Contains(SurfaceNode->SharedSurfaceId);
 
 		// If this is true, we will reuse the surface properties from a higher LOD, se we can skip the generation of material properties and images.
-		const bool bShareSurface = node.SharedSurfaceId != INDEX_NONE && !bIsBaseForSharedSurface;
+		const bool bShareSurface = SurfaceNode->SharedSurfaceId != INDEX_NONE && !bIsBaseForSharedSurface;
 
 		// Gather all modifiers that apply to this surface
 		TArray<FirstPassGenerator::FModifier> Modifiers;
@@ -903,29 +906,27 @@ namespace mu
 
 		// Store the data necessary to apply modifiers for the pre-normal operations stage.
 		// TODO: Should we merge with currently active tags from the InOptions?
-		GetModifiersFor(node.Tags, bModifiersForBeforeOperations, Modifiers);
+		GetModifiersFor(SurfaceNode->Tags, bModifiersForBeforeOperations, Modifiers);
 
 		// This pass on the modifiers is only to detect errors that cannot be detected at the point they are applied.
-		CheckModifiersForSurface( node, Modifiers);
-
-		// Generate the mesh
-        if (node.Mesh)
+		CheckModifiersForSurface(*SurfaceNode, Modifiers);
+        //if (SurfaceNode->Mesh)
         {
             MUTABLE_CPUPROFILER_SCOPE(SurfaceMesh);
 
-            Ptr<ASTOp> lastMeshOp;
+            Ptr<ASTOp> LastMeshOp;
 
             // Generate the mesh
 			FMeshGenerationOptions MeshOptions;
 			MeshOptions.bLayouts = true;
 			MeshOptions.State = Options.State;
-			MeshOptions.ActiveTags = node.Tags;
+			MeshOptions.ActiveTags = SurfaceNode->Tags;
 
 			const FMeshGenerationResult* SharedMeshResults = nullptr;
 			if (bShareSurface)
 			{
 				// Do we have the surface we need to share it with?
-				SharedMeshResults = SharedMeshOptionsMap.Find(node.SharedSurfaceId);
+				SharedMeshResults = SharedMeshOptionsMap.Find(SurfaceNode->SharedSurfaceId);
 				check(SharedMeshResults);
 
 				// Override the layouts with the ones from the surface we share
@@ -934,7 +935,7 @@ namespace mu
 
 			// Normalize UVs if we're going to work with images and layouts.
 			// TODO: This should come from per-layout settings!
-			const bool bNormalizeUVs = false; // !node.Images.IsEmpty();
+			const bool bNormalizeUVs = false; // !SurfaceNode->Images.IsEmpty();
 			MeshOptions.bNormalizeUVs = bNormalizeUVs;
 
 			// Ensure UV islands remain within their main layout block on lower LODs to avoid unexpected reordering 
@@ -942,138 +943,164 @@ namespace mu
 			// that may cause them to fall on a different block.
 			MeshOptions.bClampUVIslands = bShareSurface && bNormalizeUVs;
 
-            GenerateMesh(MeshOptions, MeshResults, node.Mesh);
+            GenerateMesh(MeshOptions, MeshResults, SurfaceNode->Mesh);
 
 			// Apply the modifier for the post-normal operations stage.
-			lastMeshOp = ApplyMeshModifiers(Modifiers, MeshOptions, MeshResults, SharedMeshResults, surfaceNode->GetMessageContext(), nullptr);
+			LastMeshOp = ApplyMeshModifiers(Modifiers, MeshOptions, MeshResults, SharedMeshResults, SurfaceNode->GetMessageContext(), nullptr);
 
-            // Layouts
-            for ( int32 LayoutIndex=0; LayoutIndex < MeshResults.GeneratedLayouts.Num(); ++LayoutIndex)
-            {
-                Ptr<ASTOp> layoutOp;
+			// Base mesh is allowed to be missing, aggregate all layouts and operations per layout indices in the
+			// generated mesh, base and extends.
+			TArray<CodeGenerator::FGeneratedLayout> SurfaceReferenceLayouts;
+			TArray<Ptr<ASTOp>> SurfaceLayoutOps;
 
-				Ptr<const Layout> pLayout = MeshResults.GeneratedLayouts[LayoutIndex].Layout;
-                if ( pLayout )
-                {
-					if (SharedMeshResults)
-					{
-						check(SharedMeshResults->LayoutOps.Num()>LayoutIndex);
-						layoutOp = SharedMeshResults->LayoutOps[LayoutIndex];
-					}
-					else
-					{
-						// Create a new layout expression
+			int32 MaxLayoutNum = MeshResults.GeneratedLayouts.Num();
+			for (const FMeshGenerationResult::FExtraLayouts& ExtraLayoutData : MeshResults.ExtraMeshLayouts)
+			{
+				MaxLayoutNum = FMath::Max(MaxLayoutNum, ExtraLayoutData.GeneratedLayouts.Num());
+			}
 
-						// Constant layout to start with
-						{
-							Ptr<ASTOpConstantResource> op = new ASTOpConstantResource();
-							op->Type = OP_TYPE::LA_CONSTANT;
+			SurfaceReferenceLayouts.SetNum(MaxLayoutNum);
+			SurfaceLayoutOps.SetNum(MaxLayoutNum);
 
-							op->SetValue(pLayout, CompilerOptions->OptimisationOptions.DiskCacheContext);
-							layoutOp = op;
-						}
+			// Add layouts form the base mesh.	
+            for (int32 LayoutIndex = 0; LayoutIndex < MeshResults.GeneratedLayouts.Num(); ++LayoutIndex)
+			{
+				if (!MeshResults.GeneratedLayouts[LayoutIndex].Layout)
+				{
+					continue;
+				}
 
-						// Add children merged meshes layouts
-						for (const FMeshGenerationResult::FExtraLayouts& Data : MeshResults.ExtraMeshLayouts)
-						{
-							if (!Data.MeshFragment)
-							{
-								// No mesh to add, we assume there are no layouts to add either.
-								check(Data.GeneratedLayouts.IsEmpty());
-								continue;
-							}
+				SurfaceReferenceLayouts[LayoutIndex] = MeshResults.GeneratedLayouts[LayoutIndex];
 
-							if (Data.GeneratedLayouts.Num() < MeshResults.GeneratedLayouts.Num())
-							{
-								ErrorLog->GetPrivate()->Add(TEXT("Merged layout has been ignored because the number of layouts is different."), ELMT_ERROR, surfaceNode->GetMessageContext());
-							}
-							else
-							{
-								// Constant layout to start with
-								Ptr<ASTOp> layoutFragmentAd;
-								{
-									Ptr<ASTOpConstantResource> op = new ASTOpConstantResource();
-									op->Type = OP_TYPE::LA_CONSTANT;
+				if (SharedMeshResults)
+				{
+					check(SharedMeshResults->LayoutOps.IsValidIndex(LayoutIndex));
+					SurfaceLayoutOps[LayoutIndex] = SharedMeshResults->LayoutOps[LayoutIndex];
+				}
+				else
+				{
+					Ptr<ASTOpConstantResource> ConstantLayoutOp = new ASTOpConstantResource();
+					ConstantLayoutOp->Type = OP_TYPE::LA_CONSTANT;
 
-									Ptr<const Layout> pCloned = Data.GeneratedLayouts[LayoutIndex].Layout;
-									op->SetValue(pCloned, CompilerOptions->OptimisationOptions.DiskCacheContext);
+					ConstantLayoutOp->SetValue(
+							SurfaceReferenceLayouts[LayoutIndex].Layout, 
+							CompilerOptions->OptimisationOptions.DiskCacheContext);
+					SurfaceLayoutOps[LayoutIndex] = ConstantLayoutOp;
+				}
+			}
 
-									layoutFragmentAd = op;
-								}
+			// Add extra layouts. In case there is a missing reference layout, the first visited will
+			// take the role.
+			for (const FMeshGenerationResult::FExtraLayouts& ExtraLayoutsData : MeshResults.ExtraMeshLayouts)
+			{
+				if (!ExtraLayoutsData.MeshFragment)
+				{
+					// No mesh to add, we assume there are no layouts to add either.
+					check(ExtraLayoutsData.GeneratedLayouts.IsEmpty());
+					continue;
+				}
 
-								// Merge operation
-								Ptr<ASTOpLayoutMerge> mergeAd = new ASTOpLayoutMerge();
-								mergeAd->Base = layoutOp;
-								mergeAd->Added = layoutFragmentAd;
+				const TArray<CodeGenerator::FGeneratedLayout>& ExtraGeneratedLayouts = ExtraLayoutsData.GeneratedLayouts;
+            	for (int32 LayoutIndex = 0; LayoutIndex < ExtraGeneratedLayouts.Num(); ++LayoutIndex)
+				{
+					if (!ExtraGeneratedLayouts[LayoutIndex].Layout)
+					{
+						continue;
+					}
 
-								// Condition to apply
-								if (Data.Condition)
-								{
-									Ptr<ASTOpConditional> conditionalAd = new ASTOpConditional();
-									conditionalAd->type = OP_TYPE::LA_CONDITIONAL;
-									conditionalAd->no = layoutOp;
-									conditionalAd->yes = mergeAd;
-									conditionalAd->condition = Data.Condition;
-									layoutOp = conditionalAd;
-								}
-								else
-								{
-									layoutOp = mergeAd;
-								}
-							}
-						}
+					if (!SurfaceReferenceLayouts[LayoutIndex].Layout)
+					{
+						// This Layout slot is not set by the base surface, set it as reference.
+						SurfaceReferenceLayouts[LayoutIndex] = ExtraGeneratedLayouts[LayoutIndex];
 					}
 
-					bool bIsOverlayLayout = pLayout->GetLayoutPackingStrategy() == mu::EPackStrategy::Overlay;
+					Ptr<ASTOpConstantResource> LayoutFragmentConstantOp = new ASTOpConstantResource();
+					LayoutFragmentConstantOp->Type = OP_TYPE::LA_CONSTANT;
 
-                    if (!bIsOverlayLayout && layoutOp)
-                    {
-                        // Add layout packing instructions
-						if (!SharedMeshResults)
-						{
-                            // Make sure we removed unnecessary blocks
-                            Ptr<ASTOpLayoutFromMesh> ExtractOp = new ASTOpLayoutFromMesh();
-							ExtractOp->Mesh = lastMeshOp;
-							check(LayoutIndex<256);
-							ExtractOp->LayoutIndex = uint8(LayoutIndex);
-
-							Ptr<ASTOpLayoutRemoveBlocks> RemoveOp = new ASTOpLayoutRemoveBlocks();
-							RemoveOp->Source = layoutOp;
-							RemoveOp->ReferenceLayout = ExtractOp;
-							layoutOp = RemoveOp;
-
-                            // Pack uv blocks
-                            Ptr<ASTOpLayoutPack> op = new ASTOpLayoutPack();
-                            op->Source = layoutOp;
-                            layoutOp = op;
-                        }
-
-                        // Create the expression to apply the layout to the mesh
-                        {
-                            Ptr<ASTOpFixed> op = new ASTOpFixed();
-                            op->op.type = OP_TYPE::ME_APPLYLAYOUT;
-                            op->SetChild(op->op.args.MeshApplyLayout.mesh, lastMeshOp );
-                            op->SetChild(op->op.args.MeshApplyLayout.layout, layoutOp );
-                            op->op.args.MeshApplyLayout.channel = (uint16)LayoutIndex;
-                            lastMeshOp = op;
-                        }
-                    }
+					LayoutFragmentConstantOp->SetValue(
+							ExtraLayoutsData.GeneratedLayouts[LayoutIndex].Layout,
+							CompilerOptions->OptimisationOptions.DiskCacheContext);
 
-                }
+					Ptr<ASTOpLayoutMerge> LayoutMergeOp = new ASTOpLayoutMerge();
+					// Base may be null if the base does not have  a mesh with a layout at LayoutIndex.
+					// In that case, when applying the condition this can generate null layouts.
+					LayoutMergeOp->Base = SurfaceLayoutOps[LayoutIndex];
+					LayoutMergeOp->Added = LayoutFragmentConstantOp;
 
-                MeshResults.LayoutOps.Add( layoutOp );
-            }
+					if (ExtraLayoutsData.Condition)
+					{
+						Ptr<ASTOpConditional> ConditionalOp = new ASTOpConditional();
+						ConditionalOp->type = OP_TYPE::LA_CONDITIONAL;
+						ConditionalOp->no = SurfaceLayoutOps[LayoutIndex];
+						ConditionalOp->yes = LayoutMergeOp;
+						ConditionalOp->condition = ExtraLayoutsData.Condition;
+
+						SurfaceLayoutOps[LayoutIndex] = ConditionalOp;
+					}
+					else
+					{
+						SurfaceLayoutOps[LayoutIndex] = LayoutMergeOp;
+					}
+				}
+			}
+
+			check(SurfaceReferenceLayouts.Num() == SurfaceLayoutOps.Num());
+            for (int32 LayoutIndex = 0; LayoutIndex < SurfaceReferenceLayouts.Num(); ++LayoutIndex)
+			{
+				if (!SurfaceReferenceLayouts[LayoutIndex].Layout)
+				{
+					continue;
+				}
+
+				if (SurfaceReferenceLayouts[LayoutIndex].Layout->GetLayoutPackingStrategy() == mu::EPackStrategy::Overlay)
+				{
+					continue;
+				}
+
+				// Add layout packing instructions
+				if (!SharedMeshResults)
+				{
+					// Make sure we removed unnecessary blocks
+					Ptr<ASTOpLayoutFromMesh> ExtractOp = new ASTOpLayoutFromMesh();
+					ExtractOp->Mesh = LastMeshOp;
+					check(LayoutIndex < 256);
+					ExtractOp->LayoutIndex = uint8(LayoutIndex);
+
+					Ptr<ASTOpLayoutRemoveBlocks> RemoveOp = new ASTOpLayoutRemoveBlocks();
+					RemoveOp->Source = SurfaceLayoutOps[LayoutIndex];
+					RemoveOp->ReferenceLayout = ExtractOp;
+					SurfaceLayoutOps[LayoutIndex] = RemoveOp;
+
+					// Pack uv blocks
+					Ptr<ASTOpLayoutPack> LayoutPackOp = new ASTOpLayoutPack();
+					LayoutPackOp->Source = SurfaceLayoutOps[LayoutIndex];
+					SurfaceLayoutOps[LayoutIndex] = LayoutPackOp;
+				}
+
+				// Create the expression to apply the layout to the mesh
+				{
+					Ptr<ASTOpFixed> ApplyLayoutOp = new ASTOpFixed();
+					ApplyLayoutOp->op.type = OP_TYPE::ME_APPLYLAYOUT;
+					ApplyLayoutOp->SetChild(ApplyLayoutOp->op.args.MeshApplyLayout.mesh, LastMeshOp);
+					ApplyLayoutOp->SetChild(ApplyLayoutOp->op.args.MeshApplyLayout.layout, SurfaceLayoutOps[LayoutIndex]);
+					ApplyLayoutOp->op.args.MeshApplyLayout.channel = (uint16)LayoutIndex;
+					
+					LastMeshOp = ApplyLayoutOp;
+				}
+			}
+
+			MeshResults.GeneratedLayouts = MoveTemp(SurfaceReferenceLayouts);
+			MeshResults.LayoutOps = MoveTemp(SurfaceLayoutOps); 
 
             // Store in the surface for later use.
-            targetSurface->ResultMeshOp = lastMeshOp;
+            TargetSurface->ResultMeshOp = LastMeshOp;
         }
 
-
         // Create the expression for each texture, if we are not reusing the surface from another LOD.
         //------------------------------------------------------------------------
 		if (!bShareSurface)
 		{
-			for (int32 ImageIndex = 0; ImageIndex < node.Images.Num(); ++ImageIndex)
+			for (int32 ImageIndex = 0; ImageIndex < SurfaceNode->Images.Num(); ++ImageIndex)
 			{
 				MUTABLE_CPUPROFILER_SCOPE(SurfaceTexture);
 
@@ -1083,7 +1110,7 @@ namespace mu
 				Ptr<NodeImageSwizzle> swizzleNode;
 
 				bool bFound = false;
-				Ptr<NodeImage> pImageNode = node.Images[ImageIndex].Image;
+				Ptr<NodeImage> pImageNode = SurfaceNode->Images[ImageIndex].Image;
 
 				while (!bFound && pImageNode)
 				{
@@ -1137,7 +1164,7 @@ namespace mu
 
 				if (bFound)
 				{
-					const NodeSurfaceNew::FImageData& ImageData = node.Images[ImageIndex];
+					const NodeSurfaceNew::FImageData& ImageData = SurfaceNode->Images[ImageIndex];
 
 					const int32 LayoutIndex = ImageData.LayoutIndex;
 
@@ -1153,7 +1180,7 @@ namespace mu
 						FImageGenerationOptions ImageOptions;
 						ImageOptions.State = Options.State;
 						ImageOptions.ImageLayoutStrategy = ImageLayoutStrategy;
-						ImageOptions.ActiveTags = node.Tags;
+						ImageOptions.ActiveTags = SurfaceNode->Tags;
 						ImageOptions.RectSize = { 0, 0 };
 						FImageGenerationResult Result;
 						GenerateImage(ImageOptions, Result, pImageNode);
@@ -1169,7 +1196,7 @@ namespace mu
 						RectInCells.min = { 0,0 };
 						RectInCells.size = { FakeLayoutSize ,FakeLayoutSize };
 						
-						imageAd = ApplyImageBlockModifiers( Modifiers, ImageOptions, imageAd, ImageData, GridSize, LayoutBlockDesc, RectInCells, surfaceNode->GetMessageContext());
+						imageAd = ApplyImageBlockModifiers(Modifiers, ImageOptions, imageAd, ImageData, GridSize, LayoutBlockDesc, RectInCells, SurfaceNode->GetMessageContext());
 						
 						check(imageAd);
 
@@ -1213,11 +1240,11 @@ namespace mu
 
 						Ptr<ASTOpInstanceAdd> op = new ASTOpInstanceAdd();
 						op->type = OP_TYPE::IN_ADDIMAGE;
-						op->instance = lastSurfOp;
+						op->instance = LastSurfOp;
 						op->value = imageAd;
-						op->name = node.Images[ImageIndex].Name;
+						op->name = SurfaceNode->Images[ImageIndex].Name;
 
-						lastSurfOp = op;
+						LastSurfOp = op;
 					}
 
 					else if (ImageLayoutStrategy == CompilerOptions::TextureLayoutStrategy::Pack) //-V547
@@ -1225,7 +1252,7 @@ namespace mu
 						if (LayoutIndex >= MeshResults.GeneratedLayouts.Num() ||
 							LayoutIndex >= MeshResults.LayoutOps.Num())
 						{
-							ErrorLog->GetPrivate()->Add("Missing layout in object, or its parent.", ELMT_ERROR, surfaceNode->GetMessageContext());
+							ErrorLog->GetPrivate()->Add("Missing layout in object, or its parent.", ELMT_ERROR, SurfaceNode->GetMessageContext());
 						}
 						else
 						{
@@ -1234,7 +1261,7 @@ namespace mu
 
 							Ptr<ASTOpInstanceAdd> op = new ASTOpInstanceAdd();
 							op->type = OP_TYPE::IN_ADDIMAGE;
-							op->instance = lastSurfOp;
+							op->instance = LastSurfOp;
 
 							// Image
 							//-------------------------------------
@@ -1274,17 +1301,17 @@ namespace mu
 								ImageOptions.State = Options.State;
 								ImageOptions.ImageLayoutStrategy = ImageLayoutStrategy;
 								ImageOptions.RectSize = { 0,0 };
-								ImageOptions.ActiveTags = node.Tags;
+								ImageOptions.ActiveTags = SurfaceNode->Tags;
 								ImageOptions.LayoutToApply = pLayout;
 								ImageOptions.LayoutBlockId = pLayout->Blocks[BlockIndex].Id;
-								FImageGenerationResult Result;
-								GenerateImage(ImageOptions, Result, pImageNode);
-								Ptr<ASTOp> blockAd = Result.op;
+								FImageGenerationResult ImageResult;
+								GenerateImage(ImageOptions, ImageResult, pImageNode);
+								Ptr<ASTOp> blockAd = ImageResult.op;
 
 								if (!blockAd)
 								{
 									// The GenerateImage(...) above has failed, skip this block
-									result.surfaceOp = nullptr;
+									SurfaceResult.surfaceOp = nullptr;
 									continue;
 								}
 
@@ -1303,7 +1330,7 @@ namespace mu
 								// Even if we force the size afterwards, we need some size hint in some cases, like image projections.
 								ImageOptions.RectSize = UE::Math::TIntVector2<int32>(BlockDesc.m_size);
 
-								blockAd = ApplyImageBlockModifiers( Modifiers, ImageOptions, blockAd, ImageData, GridSize, LayoutBlockDesc, RectInCells, surfaceNode->GetMessageContext());
+								blockAd = ApplyImageBlockModifiers(Modifiers, ImageOptions, blockAd, ImageData, GridSize, LayoutBlockDesc, RectInCells, SurfaceNode->GetMessageContext());
 
 								// Enforce block size and optimizations
 								blockAd = GenerateImageSize(blockAd, FIntVector2(BlockDesc.m_size));
@@ -1331,10 +1358,10 @@ namespace mu
 
 							FMeshGenerationOptions ModifierOptions;
 							ModifierOptions.State = Options.State;
-							ModifierOptions.ActiveTags = node.Tags;
+							ModifierOptions.ActiveTags = SurfaceNode->Tags;
 							imageAd = ApplyImageExtendModifiers( Modifiers, ModifierOptions, MeshResults, imageAd, ImageLayoutStrategy, 
 								LayoutIndex, ImageData, GridSize, LayoutBlockDesc, 
-								surfaceNode->GetMessageContext());
+								SurfaceNode->GetMessageContext());
 
 							// Complete the base op
 							BlankImageOp->op.args.ImageBlankLayout.blockSize[0] = uint16(LayoutBlockDesc.BlockPixelsX);
@@ -1413,9 +1440,9 @@ namespace mu
 							op->value = imageAd;
 
 							// Name
-							op->name = node.Images[ImageIndex].Name;
+							op->name = SurfaceNode->Images[ImageIndex].Name;
 
-							lastSurfOp = op;
+							LastSurfOp = op;
 						}
 					}
 
@@ -1429,15 +1456,15 @@ namespace mu
 
 			// Create the expression for each vector
 			//------------------------------------------------------------------------
-			for (int32 t = 0; t < node.Vectors.Num(); ++t)
+			for (int32 t = 0; t < SurfaceNode->Vectors.Num(); ++t)
 			{
 				//MUTABLE_CPUPROFILER_SCOPE(SurfaceVector);
 
-				if (Ptr<NodeColour> pVectorNode = node.Vectors[t].Vector)
+				if (Ptr<NodeColour> pVectorNode = SurfaceNode->Vectors[t].Vector)
 				{
 					Ptr<ASTOpInstanceAdd> op = new ASTOpInstanceAdd();
 					op->type = OP_TYPE::IN_ADDVECTOR;
-					op->instance = lastSurfOp;
+					op->instance = LastSurfOp;
 
 					// Vector
 					FColorGenerationResult VectorResult;
@@ -1445,23 +1472,23 @@ namespace mu
 					op->value = VectorResult.op;
 
 					// Name
-					op->name = node.Vectors[t].Name;
+					op->name = SurfaceNode->Vectors[t].Name;
 
-					lastSurfOp = op;
+					LastSurfOp = op;
 				}
 			}
 
 			// Create the expression for each scalar
 			//------------------------------------------------------------------------
-			for (int32 t = 0; t < node.Scalars.Num(); ++t)
+			for (int32 t = 0; t < SurfaceNode->Scalars.Num(); ++t)
 			{
 				// MUTABLE_CPUPROFILER_SCOPE(SurfaceScalar);
 
-				if (NodeScalarPtr pScalarNode = node.Scalars[t].Scalar)
+				if (NodeScalarPtr pScalarNode = SurfaceNode->Scalars[t].Scalar)
 				{
 					Ptr<ASTOpInstanceAdd> op = new ASTOpInstanceAdd();
 					op->type = OP_TYPE::IN_ADDSCALAR;
-					op->instance = lastSurfOp;
+					op->instance = LastSurfOp;
 
 					// Scalar
 					FScalarGenerationResult ScalarResult;
@@ -1469,42 +1496,42 @@ namespace mu
 					op->value = ScalarResult.op;
 
 					// Name
-					op->name = node.Scalars[t].Name;
+					op->name = SurfaceNode->Scalars[t].Name;
 
-					lastSurfOp = op;
+					LastSurfOp = op;
 				}
 			}
 
 			// Create the expression for each string
 			//------------------------------------------------------------------------
-			for (int32 t = 0; t < node.Strings.Num(); ++t)
+			for (int32 t = 0; t < SurfaceNode->Strings.Num(); ++t)
 			{
-				if (NodeStringPtr pStringNode = node.Strings[t].String)
+				if (NodeStringPtr pStringNode = SurfaceNode->Strings[t].String)
 				{
 					Ptr<ASTOpInstanceAdd> op = new ASTOpInstanceAdd();
 					op->type = OP_TYPE::IN_ADDSTRING;
-					op->instance = lastSurfOp;
+					op->instance = LastSurfOp;
 
 					FStringGenerationResult StringResult;
 					GenerateString(StringResult, Options, pStringNode);
 					op->value = StringResult.op;
 
 					// Name
-					op->name = node.Strings[t].Name;
+					op->name = SurfaceNode->Strings[t].Name;
 
-					lastSurfOp = op;
+					LastSurfOp = op;
 				}
 			}
 		}
 
-        result.surfaceOp = lastSurfOp;
-        targetSurface->ResultSurfaceOp = lastSurfOp;
+        SurfaceResult.surfaceOp = LastSurfOp;
+        TargetSurface->ResultSurfaceOp = LastSurfOp;
 
 		// If we are going to share this surface properties, remember it.
 		if (bIsBaseForSharedSurface)
 		{
-			check(!SharedMeshOptionsMap.Contains(node.SharedSurfaceId));
-			SharedMeshOptionsMap.Add( node.SharedSurfaceId, MeshResults );
+			check(!SharedMeshOptionsMap.Contains(SurfaceNode->SharedSurfaceId));
+			SharedMeshOptionsMap.Add(SurfaceNode->SharedSurfaceId, MeshResults);
 		}
     }
 
