[mutable] Applied a patch to an issue with the "Add To Mesh Component" node simultaneously connected to multiple inputs.
- Patch applied to both component nodes for consistency.

diff --git a/Engine/Plugins/Mutable/Source/CustomizableObjectEditor/Private/MuCOE/Nodes/CustomizableObjectNodeComponentMeshBase.cpp b/Engine/Plugins/Mutable/Source/CustomizableObjectEditor/Private/MuCOE/Nodes/CustomizableObjectNodeComponentMeshBase.cpp
--- a/Engine/Plugins/Mutable/Source/CustomizableObjectEditor/Private/MuCOE/Nodes/CustomizableObjectNodeComponentMeshBase.cpp
+++ b/Engine/Plugins/Mutable/Source/CustomizableObjectEditor/Private/MuCOE/Nodes/CustomizableObjectNodeComponentMeshBase.cpp
@@ -30,3 +30,10 @@ FLinearColor UCustomizableObjectNodeComponentMeshBase::GetNodeTitleColor() const
 	return Schema->GetPinTypeColor(Schema->PC_Component);
 }
 
+bool UCustomizableObjectNodeComponentMeshBase::IsSingleOutputNode() const
+{
+	// todo UE-225446 : By limiting the number of connections this node can have we avoid a check failure. However, this method should be
+	// removed in the future and the inherent issue with 1:n output connections should be fixed in its place
+	return false;
+}
+
diff --git a/Engine/Plugins/Mutable/Source/CustomizableObjectEditor/Public/MuCOE/Nodes/CustomizableObjectNodeComponentMeshBase.h b/Engine/Plugins/Mutable/Source/CustomizableObjectEditor/Public/MuCOE/Nodes/CustomizableObjectNodeComponentMeshBase.h
--- a/Engine/Plugins/Mutable/Source/CustomizableObjectEditor/Public/MuCOE/Nodes/CustomizableObjectNodeComponentMeshBase.h
+++ b/Engine/Plugins/Mutable/Source/CustomizableObjectEditor/Public/MuCOE/Nodes/CustomizableObjectNodeComponentMeshBase.h
@@ -29,6 +29,9 @@ public:
 	virtual void AllocateDefaultPins(UCustomizableObjectNodeRemapPins* RemapPins) override;
 	virtual FLinearColor GetNodeTitleColor() const override;
 
+	// UCustomizableObjectNode interface
+	virtual bool IsSingleOutputNode() const override;
+
 	UPROPERTY(EditAnywhere, Category = ComponentMesh)
 	int32 NumLODs = 1;
 	
