[Mutable] Remove uneeded image compose op when extending form a section without mesh.

#rnx

diff --git a/Engine/Plugins/Mutable/Source/MutableTools/Private/MuT/CodeGenerator.cpp b/Engine/Plugins/Mutable/Source/MutableTools/Private/MuT/CodeGenerator.cpp
--- a/Engine/Plugins/Mutable/Source/MutableTools/Private/MuT/CodeGenerator.cpp
+++ b/Engine/Plugins/Mutable/Source/MutableTools/Private/MuT/CodeGenerator.cpp
@@ -910,6 +910,8 @@ namespace mu
 
 		// This pass on the modifiers is only to detect errors that cannot be detected at the point they are applied.
 		CheckModifiersForSurface(*SurfaceNode, Modifiers);
+		
+		TBitArray<> LayoutFromExtension;
         //if (SurfaceNode->Mesh)
         {
             MUTABLE_CPUPROFILER_SCOPE(SurfaceMesh);
@@ -961,6 +963,7 @@ namespace mu
 
 			SurfaceReferenceLayouts.SetNum(MaxLayoutNum);
 			SurfaceLayoutOps.SetNum(MaxLayoutNum);
+			LayoutFromExtension.Init(false, MaxLayoutNum);
 
 			// Add layouts form the base mesh.	
             for (int32 LayoutIndex = 0; LayoutIndex < MeshResults.GeneratedLayouts.Num(); ++LayoutIndex)
@@ -1012,6 +1015,7 @@ namespace mu
 					{
 						// This Layout slot is not set by the base surface, set it as reference.
 						SurfaceReferenceLayouts[LayoutIndex] = ExtraGeneratedLayouts[LayoutIndex];
+						LayoutFromExtension[LayoutIndex] = true;
 					}
 
 					Ptr<ASTOpConstantResource> LayoutFragmentConstantOp = new ASTOpConstantResource();
@@ -1294,65 +1298,69 @@ namespace mu
 								imageAd = BlankImageOp;
 							}
 
-							for (int32 BlockIndex = 0; BlockIndex < pLayout->GetBlockCount(); ++BlockIndex)
+							// Skip the block addition for this image if the layout was from a extension.
+							if (!LayoutFromExtension[LayoutIndex])
 							{
-								// Generate the image
-								FImageGenerationOptions ImageOptions;
-								ImageOptions.State = Options.State;
-								ImageOptions.ImageLayoutStrategy = ImageLayoutStrategy;
-								ImageOptions.RectSize = { 0,0 };
-								ImageOptions.ActiveTags = SurfaceNode->Tags;
-								ImageOptions.LayoutToApply = pLayout;
-								ImageOptions.LayoutBlockId = pLayout->Blocks[BlockIndex].Id;
-								FImageGenerationResult ImageResult;
-								GenerateImage(ImageOptions, ImageResult, pImageNode);
-								Ptr<ASTOp> blockAd = ImageResult.op;
-
-								if (!blockAd)
+								for (int32 BlockIndex = 0; BlockIndex < pLayout->GetBlockCount(); ++BlockIndex)
 								{
-									// The GenerateImage(...) above has failed, skip this block
-									SurfaceResult.surfaceOp = nullptr;
-									continue;
+									// Generate the image
+									FImageGenerationOptions ImageOptions;
+									ImageOptions.State = Options.State;
+									ImageOptions.ImageLayoutStrategy = ImageLayoutStrategy;
+									ImageOptions.RectSize = { 0,0 };
+									ImageOptions.ActiveTags = SurfaceNode->Tags;
+									ImageOptions.LayoutToApply = pLayout;
+									ImageOptions.LayoutBlockId = pLayout->Blocks[BlockIndex].Id;
+									FImageGenerationResult ImageResult;
+									GenerateImage(ImageOptions, ImageResult, pImageNode);
+									Ptr<ASTOp> blockAd = ImageResult.op;
+
+									if (!blockAd)
+									{
+										// The GenerateImage(...) above has failed, skip this block
+										SurfaceResult.surfaceOp = nullptr;
+										continue;
+									}
+
+									// Calculate the desc of the generated block.
+									constexpr bool bReturnBestOption = true;
+									FImageDesc BlockDesc = blockAd->GetImageDesc(bReturnBestOption, nullptr);
+
+									// Block in layout grid units (cells)
+									box< FIntVector2 > RectInCells;
+									RectInCells.min = pLayout->Blocks[BlockIndex].Min;
+									RectInCells.size = pLayout->Blocks[BlockIndex].Size;
+
+									// Try to update the layout block desc if we don't know it yet.
+									UpdateLayoutBlockDesc(LayoutBlockDesc, BlockDesc, RectInCells.size);
+
+									// Even if we force the size afterwards, we need some size hint in some cases, like image projections.
+									ImageOptions.RectSize = UE::Math::TIntVector2<int32>(BlockDesc.m_size);
+
+									blockAd = ApplyImageBlockModifiers(Modifiers, ImageOptions, blockAd, ImageData, GridSize, LayoutBlockDesc, RectInCells, SurfaceNode->GetMessageContext());
+
+									// Enforce block size and optimizations
+									blockAd = GenerateImageSize(blockAd, FIntVector2(BlockDesc.m_size));
+
+									EImageFormat baseFormat = imageAd->GetImageDesc().m_format;
+									// Actually don't do it, it will be propagated from the top format operation.
+									//Ptr<ASTOp> blockAd = GenerateImageFormat(blockAd, baseFormat);
+
+									// Apply tiling to avoid generating chunks of image that are too big.
+									blockAd = ApplyTiling(blockAd, ImageOptions.RectSize, LayoutBlockDesc.FinalFormat);
+
+									// Compose layout operation
+									Ptr<ASTOpImageCompose> composeOp = new ASTOpImageCompose();
+									composeOp->Layout = MeshResults.LayoutOps[LayoutIndex];
+									composeOp->Base = imageAd;
+									composeOp->BlockImage = blockAd;
+
+									// Set the absolute block index.
+									check(pLayout->Blocks[BlockIndex].Id != FLayoutBlock::InvalidBlockId);
+									composeOp->BlockId = pLayout->Blocks[BlockIndex].Id;
+
+									imageAd = composeOp;
 								}
-
-								// Calculate the desc of the generated block.
-								constexpr bool bReturnBestOption = true;
-								FImageDesc BlockDesc = blockAd->GetImageDesc(bReturnBestOption, nullptr);
-
-								// Block in layout grid units (cells)
-								box< FIntVector2 > RectInCells;
-								RectInCells.min = pLayout->Blocks[BlockIndex].Min;
-								RectInCells.size = pLayout->Blocks[BlockIndex].Size;
-
-								// Try to update the layout block desc if we don't know it yet.
-								UpdateLayoutBlockDesc(LayoutBlockDesc, BlockDesc, RectInCells.size);
-
-								// Even if we force the size afterwards, we need some size hint in some cases, like image projections.
-								ImageOptions.RectSize = UE::Math::TIntVector2<int32>(BlockDesc.m_size);
-
-								blockAd = ApplyImageBlockModifiers(Modifiers, ImageOptions, blockAd, ImageData, GridSize, LayoutBlockDesc, RectInCells, SurfaceNode->GetMessageContext());
-
-								// Enforce block size and optimizations
-								blockAd = GenerateImageSize(blockAd, FIntVector2(BlockDesc.m_size));
-
-								EImageFormat baseFormat = imageAd->GetImageDesc().m_format;
-								// Actually don't do it, it will be propagated from the top format operation.
-								//Ptr<ASTOp> blockAd = GenerateImageFormat(blockAd, baseFormat);
-
-								// Apply tiling to avoid generating chunks of image that are too big.
-								blockAd = ApplyTiling(blockAd, ImageOptions.RectSize, LayoutBlockDesc.FinalFormat);
-
-								// Compose layout operation
-								Ptr<ASTOpImageCompose> composeOp = new ASTOpImageCompose();
-								composeOp->Layout = MeshResults.LayoutOps[LayoutIndex];
-								composeOp->Base = imageAd;
-								composeOp->BlockImage = blockAd;
-
-								// Set the absolute block index.
-								check(pLayout->Blocks[BlockIndex].Id != FLayoutBlock::InvalidBlockId);
-								composeOp->BlockId = pLayout->Blocks[BlockIndex].Id;
-
-								imageAd = composeOp;
 							}
 							check(imageAd);
 
