[mutable] Fix problem with typing exact tag name and pressing enter not selecting it in tag widget.
#rnx

diff --git a/Engine/Plugins/Mutable/Source/CustomizableObjectEditor/Private/MuCOE/SMutableSearchComboBox.cpp b/Engine/Plugins/Mutable/Source/CustomizableObjectEditor/Private/MuCOE/SMutableSearchComboBox.cpp
--- a/Engine/Plugins/Mutable/Source/CustomizableObjectEditor/Private/MuCOE/SMutableSearchComboBox.cpp
+++ b/Engine/Plugins/Mutable/Source/CustomizableObjectEditor/Private/MuCOE/SMutableSearchComboBox.cpp
@@ -273,9 +273,20 @@ void SMutableSearchComboBox::OnSearchTextChanged(const FText& ChangedText)
 
 void SMutableSearchComboBox::OnSearchTextCommitted(const FText& InText, ETextCommit::Type InCommitType)
 {
-	if ((InCommitType == ETextCommit::Type::OnEnter) && FilteredOptionsSource.Num() > 0)
+	if ((InCommitType == ETextCommit::Type::OnEnter) && !FilteredOptionsSource.IsEmpty())
 	{
-		ComboTreeView->SetSelection(FilteredOptionsSource[0], ESelectInfo::OnKeyPress);
+		TSharedRef<FFilteredOption> Selected = FilteredOptionsSource[0];
+
+		for (const TSharedRef<FFilteredOption>& Option : *OptionsSource)
+		{
+			if (Option->DisplayOption == InText.ToString())
+			{
+				Selected = Option;
+				break;
+			}
+		}
+
+		ComboTreeView->SetSelection(Selected, ESelectInfo::OnKeyPress);
 	}
 }
 
