Mutable: Fix localization key.

#rnx

diff --git a/Engine/Plugins/Mutable/Source/CustomizableObjectEditor/Private/MuCOE/GenerateMutableSource/GenerateMutableSourceModifier.cpp b/Engine/Plugins/Mutable/Source/CustomizableObjectEditor/Private/MuCOE/GenerateMutableSource/GenerateMutableSourceModifier.cpp
--- a/Engine/Plugins/Mutable/Source/CustomizableObjectEditor/Private/MuCOE/GenerateMutableSource/GenerateMutableSourceModifier.cpp
+++ b/Engine/Plugins/Mutable/Source/CustomizableObjectEditor/Private/MuCOE/GenerateMutableSource/GenerateMutableSourceModifier.cpp
@@ -685,7 +685,7 @@ mu::Ptr<mu::NodeModifier> GenerateMutableSourceModifier(const UEdGraphPin * Pin,
 
 				if (!bClosed)
 				{
-					FText ErrorMsg = FText::Format(LOCTEXT("Clipping mesh", "The bounding [{0}] not watertight (i.e. it does not fully enclose a volume)."), FText::FromName(MeshData.Mesh->GetFName()));
+					FText ErrorMsg = FText::Format(LOCTEXT("Bounding mesh", "The bounding [{0}] not watertight (i.e. it does not fully enclose a volume)."), FText::FromName(MeshData.Mesh->GetFName()));
 					GenerationContext.Log(ErrorMsg, MeshData.Node, EMessageSeverity::Warning);
 				}
 			}
