[mutable] Changed wrong bool value

#rnx

diff --git a/Engine/Plugins/Mutable/Source/CustomizableObjectEditor/Private/MuCOE/Nodes/CustomizableObjectNodeComponentMeshBase.cpp b/Engine/Plugins/Mutable/Source/CustomizableObjectEditor/Private/MuCOE/Nodes/CustomizableObjectNodeComponentMeshBase.cpp
--- a/Engine/Plugins/Mutable/Source/CustomizableObjectEditor/Private/MuCOE/Nodes/CustomizableObjectNodeComponentMeshBase.cpp
+++ b/Engine/Plugins/Mutable/Source/CustomizableObjectEditor/Private/MuCOE/Nodes/CustomizableObjectNodeComponentMeshBase.cpp
@@ -34,6 +34,6 @@ bool UCustomizableObjectNodeComponentMeshBase::IsSingleOutputNode() const
 {
 	// todo UE-225446 : By limiting the number of connections this node can have we avoid a check failure. However, this method should be
 	// removed in the future and the inherent issue with 1:n output connections should be fixed in its place
-	return false;
+	return true;
 }
 
