Mutable: Fix component's private instance storage and usage
- Remove the vestigial instance usage storage on the component.
- Fix Private members on CustomizablebjectInstance and CustomizableSkeletalComponent, since CreateDefaultSubObject doesn't do what is expected here.

#rnx

diff --git a/Engine/Plugins/Mutable/Source/CustomizableObject/Internal/MuCO/CustomizableSkeletalComponentPrivate.h b/Engine/Plugins/Mutable/Source/CustomizableObject/Internal/MuCO/CustomizableSkeletalComponentPrivate.h
--- a/Engine/Plugins/Mutable/Source/CustomizableObject/Internal/MuCO/CustomizableSkeletalComponentPrivate.h
+++ b/Engine/Plugins/Mutable/Source/CustomizableObject/Internal/MuCO/CustomizableSkeletalComponentPrivate.h
@@ -10,12 +10,14 @@ class UPhysicsAsset;
 class USkeletalMesh;
 
 
-UCLASS(Blueprintable, BlueprintType, ClassGroup = (CustomizableObject), meta = (BlueprintSpawnableComponent))
+UCLASS()
 class CUSTOMIZABLEOBJECT_API UCustomizableSkeletalComponentPrivate : public UObject
 {
 	GENERATED_BODY()
 
 public:
+	UCustomizableSkeletalComponentPrivate();
+	
 	void CreateCustomizableObjectInstanceUsage();
 
 	/** Common end point of all updates. Even those which failed. */
@@ -28,6 +30,8 @@ public:
 	void SetSkeletalMesh(USkeletalMesh* SkeletalMesh);
 	
 	void SetPhysicsAsset(UPhysicsAsset* PhysicsAsset);
+	
+	void SetPendingSetSkeletalMesh(bool bIsActive);
 
 #if WITH_EDITOR
 	void EditorUpdateComponent();
diff --git a/Engine/Plugins/Mutable/Source/CustomizableObject/Private/MuCO/CustomizableObjectInstance.cpp b/Engine/Plugins/Mutable/Source/CustomizableObject/Private/MuCO/CustomizableObjectInstance.cpp
--- a/Engine/Plugins/Mutable/Source/CustomizableObject/Private/MuCO/CustomizableObjectInstance.cpp
+++ b/Engine/Plugins/Mutable/Source/CustomizableObject/Private/MuCO/CustomizableObjectInstance.cpp
@@ -287,8 +287,6 @@ const FCustomizableInstanceComponentData* UCustomizableInstancePrivate::GetCompo
 UCustomizableObjectInstance::UCustomizableObjectInstance()
 {
 	SetFlags(RF_Transactional);
-	
-	PrivateData = CreateDefaultSubobject<UCustomizableInstancePrivate>(FName("Private"));
 }
 
 
@@ -440,6 +438,23 @@ bool UCustomizableObjectInstance::IsEditorOnly() const
 	return false;
 }
 
+void UCustomizableObjectInstance::PostInitProperties()
+{
+	UObject::PostInitProperties();
+
+	if (!HasAllFlags(RF_ClassDefaultObject))
+	{
+		if (!PrivateData)
+		{
+			PrivateData = NewObject<UCustomizableInstancePrivate>(this, FName("Private"));
+		}
+		else if (PrivateData->GetOuter() != this)
+		{
+			PrivateData = Cast<UCustomizableInstancePrivate>(StaticDuplicateObject(PrivateData, this, FName("Private")));
+		}
+	}
+}
+
 
 void UCustomizableObjectInstance::BeginDestroy()
 {
diff --git a/Engine/Plugins/Mutable/Source/CustomizableObject/Private/MuCO/CustomizableSkeletalComponent.cpp b/Engine/Plugins/Mutable/Source/CustomizableObject/Private/MuCO/CustomizableSkeletalComponent.cpp
--- a/Engine/Plugins/Mutable/Source/CustomizableObject/Private/MuCO/CustomizableSkeletalComponent.cpp
+++ b/Engine/Plugins/Mutable/Source/CustomizableObject/Private/MuCO/CustomizableSkeletalComponent.cpp
@@ -11,6 +11,16 @@
 #include "MuCO/CustomizableObject.h"
 
 
+UCustomizableSkeletalComponentPrivate::UCustomizableSkeletalComponentPrivate()
+{
+	// This object may get instantiated into a level, as a part of UCustomizableSkeletalComponent, so needs to be public to ensure
+	// it can be serialized out.
+	if (!HasAllFlags(RF_ClassDefaultObject))
+	{
+		SetFlags(RF_Public);
+	}
+}
+
 void UCustomizableSkeletalComponentPrivate::CreateCustomizableObjectInstanceUsage()
 {
 	if (CustomizableObjectInstanceUsage)
@@ -37,7 +47,6 @@ void UCustomizableSkeletalComponentPrivate::CreateCustomizableObjectInstanceUsag
 
 UCustomizableSkeletalComponent::UCustomizableSkeletalComponent()
 {
-	Private = CreateDefaultSubobject<UCustomizableSkeletalComponentPrivate>(FName("Private"));
 }
 
 
@@ -99,6 +108,15 @@ void UCustomizableSkeletalComponentPrivate::SetPhysicsAsset(UPhysicsAsset* Physi
 }
 
 
+void UCustomizableSkeletalComponentPrivate::SetPendingSetSkeletalMesh(bool bIsActive)
+{
+	if (CustomizableObjectInstanceUsage)
+	{
+		CustomizableObjectInstanceUsage->GetPrivate()->SetPendingSetSkeletalMesh(bIsActive);
+	}
+}
+
+
 USkeletalMesh* UCustomizableSkeletalComponentPrivate::GetAttachedSkeletalMesh() const
 {
 	if (CustomizableObjectInstanceUsage)
@@ -168,18 +186,18 @@ bool UCustomizableSkeletalComponent::GetSkipSetSkeletalMeshOnAttach() const
 
 void UCustomizableSkeletalComponent::UpdateSkeletalMeshAsync(bool bNeverSkipUpdate)
 {
-	if (CustomizableObjectInstanceUsage)
+	if (GetPrivate()->CustomizableObjectInstanceUsage)
 	{
-		CustomizableObjectInstanceUsage->UpdateSkeletalMeshAsync(bNeverSkipUpdate);
+		GetPrivate()->CustomizableObjectInstanceUsage->UpdateSkeletalMeshAsync(bNeverSkipUpdate);
 	}
 }
 
 
 void UCustomizableSkeletalComponent::UpdateSkeletalMeshAsyncResult(FInstanceUpdateDelegate Callback, bool bIgnoreCloseDist, bool bForceHighPriority)
 {
-	if (CustomizableObjectInstanceUsage)
+	if (GetPrivate()->CustomizableObjectInstanceUsage)
 	{
-		CustomizableObjectInstanceUsage->UpdateSkeletalMeshAsyncResult(Callback, bIgnoreCloseDist, bForceHighPriority);
+		GetPrivate()->CustomizableObjectInstanceUsage->UpdateSkeletalMeshAsyncResult(Callback, bIgnoreCloseDist, bForceHighPriority);
 	}
 }
 
@@ -231,11 +249,9 @@ void UCustomizableSkeletalComponent::OnAttachmentChanged()
 {
 	Super::OnAttachmentChanged();
 
-	USkeletalMeshComponent* Parent = Cast<USkeletalMeshComponent>(GetAttachParent());
-
-	if (Parent && CustomizableObjectInstanceUsage)
+	if (Cast<USkeletalMeshComponent>(GetAttachParent()))
 	{
-		CustomizableObjectInstanceUsage->GetPrivate()->SetPendingSetSkeletalMesh(true);
+		GetPrivate()->SetPendingSetSkeletalMesh(true);
 	}
 	else if(!GetAttachParent())
 	{
@@ -247,12 +263,18 @@ void UCustomizableSkeletalComponent::OnAttachmentChanged()
 void UCustomizableSkeletalComponent::PostInitProperties()
 {
 	Super::PostInitProperties();
-	GetPrivate()->CreateCustomizableObjectInstanceUsage();
-}
 
-void UCustomizableSkeletalComponent::PostReinitProperties()
-{
-	Super::PostReinitProperties();
-	GetPrivate()->CreateCustomizableObjectInstanceUsage();
+	if (!HasAllFlags(RF_ClassDefaultObject))
+	{
+		if (!Private)
+		{
+			Private = NewObject<UCustomizableSkeletalComponentPrivate>(this, FName("Private"), RF_Public);
+		}
+		else if (Private->GetOuter() != this)
+		{
+			Private = Cast<UCustomizableSkeletalComponentPrivate>(StaticDuplicateObject(Private, this, FName("Private")));
+		}
+	
+		GetPrivate()->CreateCustomizableObjectInstanceUsage();
+	}
 }
-
diff --git a/Engine/Plugins/Mutable/Source/CustomizableObject/Public/MuCO/CustomizableObjectInstance.h b/Engine/Plugins/Mutable/Source/CustomizableObject/Public/MuCO/CustomizableObjectInstance.h
--- a/Engine/Plugins/Mutable/Source/CustomizableObject/Public/MuCO/CustomizableObjectInstance.h
+++ b/Engine/Plugins/Mutable/Source/CustomizableObject/Public/MuCO/CustomizableObjectInstance.h
@@ -249,6 +249,7 @@ public:
 	virtual bool IsReadyForFinishDestroy() override;
 	virtual FString GetDesc() override;
 	virtual bool IsEditorOnly() const override;
+	virtual void PostInitProperties() override;
 
 	/** Set the CustomizableObject this instance will be generated from. 
 	  * It is usually not necessary to call this since instances are already generated from a CustomizableObject. */
diff --git a/Engine/Plugins/Mutable/Source/CustomizableObject/Public/MuCO/CustomizableSkeletalComponent.h b/Engine/Plugins/Mutable/Source/CustomizableObject/Public/MuCO/CustomizableSkeletalComponent.h
--- a/Engine/Plugins/Mutable/Source/CustomizableObject/Public/MuCO/CustomizableSkeletalComponent.h
+++ b/Engine/Plugins/Mutable/Source/CustomizableObject/Public/MuCO/CustomizableSkeletalComponent.h
@@ -24,7 +24,7 @@ public:
 
 	/** This component index refers to the object list of components */
 	UPROPERTY(BlueprintReadWrite, EditAnywhere, Category = CustomizableSkeletalComponent)
-	int32 ComponentIndex;
+	int32 ComponentIndex = 0;
 
 private:
 	/** Only used if the ComponentIndex is INDEX_NONE.
@@ -51,9 +51,8 @@ protected:
 	virtual void PostEditChangeProperty(FPropertyChangedEvent& PropertyChangedEvent) override;
 #endif
 
-	// UActor interface
+	// UObject interface
 	virtual void PostInitProperties() override;
-	virtual void PostReinitProperties() override;
 
 	// USceneComponent interface
 	virtual void OnAttachmentChanged() override;
@@ -99,9 +98,4 @@ public:
 	UCustomizableSkeletalComponentPrivate* GetPrivate();
 
 	const UCustomizableSkeletalComponentPrivate* GetPrivate() const;
-	
-private:
-	UPROPERTY(Transient)
-	TObjectPtr<UCustomizableObjectInstanceUsage> CustomizableObjectInstanceUsage;
 };
-
