[Mutable] Set mutable generated mesh PhysicsAsset preview mesh only if the PhysicsAsset has also been generated.

- In case the PhysicsAsset does not need modification the referenced asset will be used, in that case the preview mesh is not changed.

#rnx

diff --git a/Engine/Plugins/Mutable/Source/CustomizableObject/Private/MuCO/CustomizableObjectInstance.cpp b/Engine/Plugins/Mutable/Source/CustomizableObject/Private/MuCO/CustomizableObjectInstance.cpp
--- a/Engine/Plugins/Mutable/Source/CustomizableObject/Private/MuCO/CustomizableObjectInstance.cpp
+++ b/Engine/Plugins/Mutable/Source/CustomizableObject/Private/MuCO/CustomizableObjectInstance.cpp
@@ -2309,16 +2309,13 @@ bool UCustomizableInstancePrivate::UpdateSkeletalMesh_PostBeginUpdate0(UCustomiz
 
 			SkeletalMesh->SetPhysicsAsset(PhysicsAssetResult);
 
-			if (PhysicsAssetResult)
+#if WITH_EDITORONLY_DATA
+			if (PhysicsAssetResult && PhysicsAssetResult->GetPackage() == GetTransientPackage())
 			{
-				// We are setting the physics asset mesh preview to the generated skeletal mesh.
-				// this is fine if the phyiscs asset is also generated, but not sure about referenced assets.
-				// In any case, don't mark the asset as modified.
-
-				// This is only called in editor, no need to add editor guards.	
 				constexpr bool bMarkAsDirty = false;
 				PhysicsAssetResult->SetPreviewMesh(SkeletalMesh, bMarkAsDirty);
 			}
+#endif
 		}
 
 		const int32 NumAdditionalPhysicsNum = ComponentMesh->AdditionalPhysicsBodies.Num();
