[Mutable] Empty ranged parameter values when setting the default values.

#rnx

diff --git a/Engine/Plugins/Mutable/Source/CustomizableObject/Private/MuCO/CustomizableObjectInstanceDescriptor.cpp b/Engine/Plugins/Mutable/Source/CustomizableObject/Private/MuCO/CustomizableObjectInstanceDescriptor.cpp
--- a/Engine/Plugins/Mutable/Source/CustomizableObject/Private/MuCO/CustomizableObjectInstanceDescriptor.cpp
+++ b/Engine/Plugins/Mutable/Source/CustomizableObject/Private/MuCO/CustomizableObjectInstanceDescriptor.cpp
@@ -2488,11 +2488,7 @@ void FCustomizableObjectInstanceDescriptor::SetDefaultValue(int32 ParamIndex)
 
 		FCustomizableObjectIntParameterValue& IntParameter = IntParameters[TypedIndex];
 		IntParameter.ParameterValueName = DefaultValue;
-
-		for (FString& RangeValue : IntParameter.ParameterRangeValueNames)
-		{
-			RangeValue = DefaultValue;
-		}
+		IntParameter.ParameterRangeValueNames.Empty();
 		break;
 	}
 
@@ -2502,11 +2498,7 @@ void FCustomizableObjectInstanceDescriptor::SetDefaultValue(int32 ParamIndex)
 
 		FCustomizableObjectFloatParameterValue& FloatParameter = FloatParameters[TypedIndex];
 		FloatParameter.ParameterValue = DefaultValue;
-
-		for (float& RangeValue : FloatParameter.ParameterRangeValues)
-		{
-			RangeValue = DefaultValue;
-		}
+		FloatParameter.ParameterRangeValues.Empty();
 		break;
 	}
 
@@ -2530,11 +2522,7 @@ void FCustomizableObjectInstanceDescriptor::SetDefaultValue(int32 ParamIndex)
 
 		FCustomizableObjectProjectorParameterValue& ProjectorParameter = ProjectorParameters[TypedIndex];
 		ProjectorParameter.Value = DefaultValue;
-
-		for (FCustomizableObjectProjector& RangeValue : ProjectorParameter.RangeValues)
-		{
-			RangeValue = DefaultValue;
-		}
+		ProjectorParameter.RangeValues.Empty();
 		break;
 	}
 
@@ -2546,11 +2534,7 @@ void FCustomizableObjectInstanceDescriptor::SetDefaultValue(int32 ParamIndex)
 		{
 			FCustomizableObjectTextureParameterValue& TextureParameter = TextureParameters[TypedIndex];
 			TextureParameter.ParameterValue = DefaultValue;
-
-			for (FName& RangeValue : TextureParameter.ParameterRangeValues)
-			{
-				RangeValue = DefaultValue;
-			}
+			TextureParameter.ParameterRangeValues.Empty();
 		}
 		break;
 	}
