[Mutable] Store compiled resources in multiple bulk data files. Resources will be grouped into categories (meshes, images, high-res images, morphs, and clothing).
Add function to generate up to 'n' BulkData files.
Optimize DDC data distribution to improve put/get times.
Fix SourceId indeterminism due to FNames.

#rnx

diff --git a/Engine/Plugins/Mutable/Source/CustomizableObject/Internal/MuCO/CustomizableObjectCompilerTypes.h b/Engine/Plugins/Mutable/Source/CustomizableObject/Internal/MuCO/CustomizableObjectCompilerTypes.h
--- a/Engine/Plugins/Mutable/Source/CustomizableObject/Internal/MuCO/CustomizableObjectCompilerTypes.h
+++ b/Engine/Plugins/Mutable/Source/CustomizableObject/Internal/MuCO/CustomizableObjectCompilerTypes.h
@@ -53,16 +53,24 @@ struct FCompilationOptions
 	// From 0 to UE_MUTABLE_MAX_OPTIMIZATION
 	int32 OptimizationLevel = UE_MUTABLE_MAX_OPTIMIZATION;
 
-	// Use the disk to store intermediate compilation data. This slows down the object compilation
-	// but it may be necessary for huge objects.
+	/** Use the disk to store intermediate compilation data.This slows down the object compilation
+	 * but it may be necessary for huge objects.
+	 *
+	 * This option does not modify the DDC Key.
+	 * */
 	bool bUseDiskCompilation = false;
 
 	/** High limit of the size in bytes of the packaged data when cooking this object.
-	* This limit is before any pak or filesystem compression. This limit will be broken if a single piece of data is bigger because data is not fragmented for packaging purposes.
-	*/
+	 * This limit is before any pak or filesystem compression. This limit will be broken if a single piece of data is bigger because data is not fragmented for packaging purposes.
+	 *
+	 * This option does not modify the DDC Key.
+	 */
 	uint64 PackagedDataBytesLimit = 256 * 1024 * 1024;
 
-	/** High (inclusive) limit of the size in bytes of a data block to be included into the compiled object directly instead of stored in a streamable file. */
+	/** High (inclusive) limit of the size in bytes of a data block to be included into the compiled object directly instead of stored in a streamable file.
+	 *
+	 * This option does not modify the DDC Key.
+	 */
 	uint64 EmbeddedDataBytesLimit = 1024;
 
 	/** Number of minimum mipmaps that we want to always be available in disk regardless of NumHighResImageLODs. */
@@ -113,22 +121,32 @@ struct FCompilationOptions
 	// Control image tiled generation
 	int32 ImageTiling = 0;
 
-	/** If true, gather all game asset references and save them in the Customizable Object. */
+	/** If true, gather all game asset references and save them in the Customizable Object. 
+	 *
+	 * This option does not modify the DDC Key. 
+	 */
 	bool bGatherReferences = false;
 
-	/** Whether or not the compiler should query a request to load the compiled data from the DDC */
+	/** Whether or not the compiler should query a request to load the compiled data from the DDC.
+	 *
+	 * This option does not modify the DDC Key. 
+	 */
 	bool bQueryCompiledDatafromDDC = false;
 
-	/** Whether or not the compiler should store the compiled data to the DDC*/
-	bool bStoreCompiledDataInDDC = false;
-
-	/** High limit of the size in bytes of the packaged data when storing this object in the DDC.
-	 * This limit is before any pak or filesystem compression. This limit will be broken if a single piece of data is bigger because data is not fragmented for packaging purposes.
+	/** Whether or not the compiler should store the compiled data to the DDC.
+	 *
+	 * This option does not modify the DDC Key. 
 	 */
-	uint64 DDCBytesLimit = 64 * 1024;
+	bool bStoreCompiledDataInDDC = false;
 	
 	/** Stores the only option of an Int Param that should be compiled. */
 	TMap<FString, FString> ParamNamesToSelectedOptions;
+
+	/** Used to know if the object will be using FByteBulkData.
+	 *
+	 * This option does not modify the DDC Key. 
+	 */
+	bool bUseBulkData = false;
 };
 
 
diff --git a/Engine/Plugins/Mutable/Source/CustomizableObject/Internal/MuCO/CustomizableObjectPrivate.h b/Engine/Plugins/Mutable/Source/CustomizableObject/Internal/MuCO/CustomizableObjectPrivate.h
--- a/Engine/Plugins/Mutable/Source/CustomizableObject/Internal/MuCO/CustomizableObjectPrivate.h
+++ b/Engine/Plugins/Mutable/Source/CustomizableObject/Internal/MuCO/CustomizableObjectPrivate.h
@@ -28,6 +28,11 @@ namespace mu
 }
 
 #if WITH_EDITOR
+namespace MutablePrivate
+{
+	struct FClassifyNode;
+}
+
 namespace UE::DerivedData
 {
 	struct FValueId;
@@ -539,13 +544,17 @@ struct FRealTimeMorphStreamable
 	FMutableStreamableBlock Block;
 
 	UPROPERTY()
-	uint32 Size=0;
+	uint32 Size = 0;
+
+	UPROPERTY()
+	uint32 SourceId = 0;
 
 	friend FArchive& operator<<(FArchive& Ar, FRealTimeMorphStreamable& Elem)
 	{
 		Ar << Elem.NameResolutionMap;
 		Ar << Elem.Size;
 		Ar << Elem.Block;
+		Ar << Elem.SourceId;
 
 		return Ar;
 	}
@@ -617,6 +626,9 @@ struct FClothingStreamable
 	UPROPERTY()
 	FMutableStreamableBlock Block;
 
+	UPROPERTY()
+	uint32 SourceId = 0;
+
 	friend FArchive& operator<<(FArchive& Ar, FClothingStreamable& Elem)
 	{
 		Ar << Elem.ClothingAssetIndex;
@@ -624,7 +636,7 @@ struct FClothingStreamable
 		Ar << Elem.PhysicsAssetIndex;
 		Ar << Elem.Size;
 		Ar << Elem.Block;
-
+		Ar << Elem.SourceId;
 		return Ar;
 	}
 };
@@ -753,27 +765,25 @@ namespace MutablePrivate
 
 	struct FBlock
 	{
-		/** Data Type*/
-		EDataType DataType;
-
 		/** Used on some data types as the index to the block stored in the CustomizableObject */
 		uint32 Id;
 
+		/** Used on some data types to group blocks. */
+		uint32 SourceId;
+
 		/** Size of the data block. */
 		uint32 Size;
 
-		/** Data flags, like "high-res". */
-		uint32 Flags;
+		uint32 Padding = 0;
 
 		/** Offset in the full source streamed data file that is created when compiling. */
 		uint64 Offset;
 
 		friend FArchive& operator<<(FArchive& Ar, FBlock& Data)
 		{
-			Ar << Data.DataType;
 			Ar << Data.Id;
+			Ar << Data.SourceId;
 			Ar << Data.Size;
-			Ar << Data.Flags;
 			Ar << Data.Offset;
 			return Ar;
 		};
@@ -784,19 +794,22 @@ namespace MutablePrivate
 	{
 		EDataType DataType = EDataType::None;
 
-		uint32 Padding = 0;
+		/** Rom ResourceType. */
+		uint16 ResourceType = 0;
+
+		/** Common flags of the data stored in this file. See mu::ERomFlags. */
+		uint16 Flags = 0;
 
 		/** Id generated from a hash of the file content + offset to avoid collisions. */
-		uint32 Id;
+		uint32 Id = 0;
 
-		/** Common flags of the data stored in this file. See mu::ERomFlags. */
-		uint32 Flags = 0;
+		uint32 Padding = 0;
 
 		/** List of blocks that are contained in the file, in order. */
 		TArray<FBlock> Blocks;
 
 		/** Get the total size of blocks in this file. */
-		int64 GetSize() const;
+		uint64 GetSize() const;
 
 		/** Copy the requested block to the requested buffer and return its size. */
 		void GetFileData(struct FMutableCachedPlatformData*, TArray64<uint8>& DataDestination, bool bDropData);
@@ -804,14 +817,54 @@ namespace MutablePrivate
 		friend FArchive& operator<<(FArchive& Ar, FFile& Data)
 		{
 			Ar << Data.DataType;
-			Ar << Data.Padding;
-			Ar << Data.Id;
+			Ar << Data.ResourceType;
 			Ar << Data.Flags;
+			Ar << Data.Id;
 			Ar << Data.Blocks;
 			return Ar;
 		};
 	};
 
+	struct FFileCategoryID
+	{
+		// DATATYPE
+		EDataType DataType = EDataType::None;
+
+		/** Rom ResourceType. */
+		uint16 ResourceType = 0;
+
+		/** Rom flags  */
+		uint16 Flags = 0;
+
+		friend uint32 GetTypeHash(const FFileCategoryID& Key);
+		bool operator==(const FFileCategoryID& Other) const = default;
+	};
+
+
+	struct FFileCategory
+	{
+		FFileCategoryID Id;
+
+		// Accumulated size of resources from this category
+		uint64 DataSize = 0;
+
+		// Categories within a bucket with a limited number of files will use sequential ID starting at FirstFile
+		// and up to FirstFile + NumFiles.
+		uint32 FirstFile = 0;
+		uint32 NumFiles = 0;
+	};
+
+
+	/** Group bulk data by categories. */
+	struct FFileBucket
+	{
+		// Resources belonging to these categories will be added to the bucket.
+		TArray<FFileCategory> Categories;
+
+		// Accumulated size of the resources of all categories within this bucket
+		uint64 DataSize = 0;
+	};
+
 	struct FModelStreamableData
 	{
 		void Get(uint32 Key, TArrayView64<uint8> Destination, bool bDropData)
@@ -867,22 +920,39 @@ namespace MutablePrivate
 	};
 
 
-	/** Compute the number of files and sizes the BulkData will be split into and update
-	 * the streamables's FileIds and Offsets.
+	/** Generate the list of BulkData files with a restriction to the number of files to generate per bucket.
+	 *  Resources will be split into two buckets for non-optional and optional BulkData.
 	 */
-	void CUSTOMIZABLEOBJECT_API GenerateBulkDataFilesList(
+	void CUSTOMIZABLEOBJECT_API GenerateBulkDataFilesListWithFileLimit(
 		TSharedPtr<const mu::Model, ESPMode::ThreadSafe> Model,
-		FModelStreamableBulkData& StreamableBulkData,
+		FModelStreamableBulkData& ModelStreamableBulkData,
+		uint32 NumFilesPerBucket,
+		TArray<FFile>& OutBulkDataFiles);
+
+	/** Generate the list of BulkData files with a soft restriction to the size of the files.
+	 */
+	void CUSTOMIZABLEOBJECT_API GenerateBulkDataFilesListWithSizeLimit(
+		TSharedPtr<const mu::Model, ESPMode::ThreadSafe> Model,
+		FModelStreamableBulkData& ModelStreamableBulkData,
 		const ITargetPlatform* TargetPlatform,
 		uint64 TargetBulkDataFileBytes,
 		TArray<FFile>& OutBulkDataFiles);
 
+	/** Compute the number of files and sizes the BulkData will be split into and update
+	 * the streamable's FileIds and Offsets.
+	 */
+	void GenerateBulkDataFilesList(
+		TSharedPtr<const mu::Model, ESPMode::ThreadSafe> Model,
+		FModelStreamableBulkData& StreamableBulkData,
+		bool bUseRomTypeAndFlagsToFilter,
+		TFunctionRef<void(const FFileCategoryID&, const FClassifyNode&, TArray<FFile>&)> CreateFileList,
+		TArray<FFile>& OutBulkDataFiles);
+
 	void CUSTOMIZABLEOBJECT_API SerializeBulkDataFiles(
 		FMutableCachedPlatformData& CachedPlatformData,
 		TArray<FFile>& BulkDataFiles,
-		TFunctionRef<void(FFile&, TArray64<uint8>&)> WriteFile,
-		bool bDropData
-	);
+		TFunctionRef<void(FFile&, TArray64<uint8>&, uint32 FileIndex)> WriteFile,
+		bool bDropData);
 
 	UE::DerivedData::FValueId CUSTOMIZABLEOBJECT_API GetDerivedDataModelId();
 	UE::DerivedData::FValueId CUSTOMIZABLEOBJECT_API GetDerivedDataModelResourcesId();
@@ -901,7 +971,7 @@ struct CUSTOMIZABLEOBJECT_API FModelStreamableBulkData
 
 	TMap<uint32, FRealTimeMorphStreamable> RealTimeMorphStreamables;
 
-	TMap<uint32, FByteBulkData> HashToBulkData;
+	TArray<FByteBulkData> StreamableBulkData;
 
 	void Serialize(FArchive& Ar, UObject* Owner, bool bCooked);
 
@@ -1489,6 +1559,8 @@ public:
 		
 		SurfaceMetadataSlotNameIndexToName,
 
+		BulkDataFilesNumFilesLimit,
+
 		// -----<new versions can be added above this line>--------
 		LastCustomizableObjectVersion
 	};
diff --git a/Engine/Plugins/Mutable/Source/CustomizableObject/Private/MuCO/CustomizableObject.cpp b/Engine/Plugins/Mutable/Source/CustomizableObject/Private/MuCO/CustomizableObject.cpp
--- a/Engine/Plugins/Mutable/Source/CustomizableObject/Private/MuCO/CustomizableObject.cpp
+++ b/Engine/Plugins/Mutable/Source/CustomizableObject/Private/MuCO/CustomizableObject.cpp
@@ -31,6 +31,7 @@
 #include "MuCO/UnrealMutableModelDiskStreamer.h"
 #include "MuCO/UnrealPortabilityHelpers.h"
 #include "MuR/Model.h"
+#include "MuR/Operations.h"
 #include "PhysicsEngine/PhysicsAsset.h"
 #include "Serialization/MemoryReader.h"
 #include "Serialization/MemoryWriter.h"
@@ -197,14 +198,17 @@ void UCustomizableObject::PreSave(FObjectPreSaveContext ObjectSaveContext)
 				TSharedPtr<FModelStreamableBulkData> ModelStreamableBulkData = GetPrivate()->GetModelStreamableBulkData(true);
 
 				const int32 NumBulkDataFiles = CachedPlatformData.BulkDataFiles.Num();
-				ModelStreamableBulkData->HashToBulkData.Empty(NumBulkDataFiles);
-				ModelStreamableBulkData->HashToBulkData.Reserve(NumBulkDataFiles);
 
-				const auto WriteBulkData = [ModelStreamableBulkData](MutablePrivate::FFile& File, TArray64<uint8>& FileBulkData)
+				ModelStreamableBulkData->StreamableBulkData.SetNum(NumBulkDataFiles);
+
+				const auto WriteBulkData = [ModelStreamableBulkData](MutablePrivate::FFile& File, TArray64<uint8>& FileBulkData, uint32 FileIndex)
 					{
 						MUTABLE_CPUPROFILER_SCOPE(WriteBulkData);
 
-						FByteBulkData& ByteBulkData = ModelStreamableBulkData->HashToBulkData.FindOrAdd(File.Id);
+						FByteBulkData& ByteBulkData = ModelStreamableBulkData->StreamableBulkData[FileIndex];
+
+						// BulkData file to store the file to. CookedIndex 0 is used as a default for backwards compatibility, +1 to skip it.
+						ByteBulkData.SetCookedIndex(FBulkDataCookedIndex((File.Id % MAX_uint8) + 1));
 
 						ByteBulkData.Lock(LOCK_READ_WRITE);
 						uint8* Ptr = (uint8*)ByteBulkData.Realloc(FileBulkData.Num());
@@ -212,7 +216,7 @@ void UCustomizableObject::PreSave(FObjectPreSaveContext ObjectSaveContext)
 						ByteBulkData.Unlock();
 
 						uint32 BulkDataFlags = BULKDATA_PayloadInSeperateFile | BULKDATA_Force_NOT_InlinePayload;
-						if (File.Flags == uint32(mu::ERomFlags::HighRes))
+						if (File.Flags == uint16(mu::ERomFlags::HighRes))
 						{
 							BulkDataFlags |= BULKDATA_OptionalPayload;
 						}
@@ -1046,6 +1050,7 @@ void UCustomizableObjectPrivate::CompileForTargetPlatform(UCustomizableObject& C
 	Options.OptimizationLevel = UE_MUTABLE_MAX_OPTIMIZATION;	// Force max optimization when packaging.
 	Options.TextureCompression = ECustomizableObjectTextureCompression::HighQuality;
 	Options.bIsCooking = true;
+	Options.bUseBulkData = CVarMutableUseBulkData.GetValueOnAnyThread();
 	Options.TargetPlatform = &TargetPlatform;
 
 	const int32 DDCUsage = CVarMutableDerivedDataCacheUsage.GetValueOnAnyThread();
@@ -2211,28 +2216,14 @@ void FModelStreamableBulkData::Serialize(FArchive& Ar, UObject* Owner, bool bCoo
 
 	if (bCooked)
 	{
-		int32 NumBulkDatas = HashToBulkData.Num();
+		int32 NumBulkDatas = StreamableBulkData.Num();
 		Ar << NumBulkDatas;
 
-		if (Ar.IsSaving())
-		{
-			for (TPair<uint32, FByteBulkData>& BulkData : HashToBulkData)
-			{
-				Ar << BulkData.Key;
-				BulkData.Value.Serialize(Ar, Owner);
-			}
-		}
-		else
-		{
-			HashToBulkData.Reserve(NumBulkDatas);
-			for (int32 Index = 0; Index < NumBulkDatas; ++Index)
-			{
-				uint32 ResourceId = 0;
-				Ar << ResourceId;
+		StreamableBulkData.SetNum(NumBulkDatas);
 
-				FByteBulkData& BulkData = HashToBulkData.Add(ResourceId);
-				BulkData.Serialize(Ar, Owner);
-			}
+		for (FByteBulkData& BulkData : StreamableBulkData)
+		{
+			BulkData.Serialize(Ar, Owner);
 		}
 	}
 }
@@ -2508,9 +2499,9 @@ FCompilationOptions UCustomizableObjectPrivate::GetCompileOptions() const
 #if WITH_EDITOR
 namespace MutablePrivate
 {
-	int64 FFile::GetSize() const
+	uint64 FFile::GetSize() const
 	{
-		int64 FileSize = 0;
+		uint64 FileSize = 0;
 		for (const FBlock& Block : Blocks)
 		{
 			FileSize += Block.Size;
@@ -2559,21 +2550,45 @@ namespace MutablePrivate
 	}
 
 
-	// TODO:
-	// To avoid influence of the order of the streamed data (their index), classify it recursively based on hash values
-	// until the tree leaves have either a single block, or a sum of blocks below the desired file size.
+	uint32 GetTypeHash(const FFileCategoryID& Key)
+	{
+		uint32 Hash = (uint32)Key.DataType;
+		Hash = HashCombine(Hash, Key.ResourceType);
+		Hash = HashCombine(Hash, Key.Flags);
+		return Hash;
+	}
+
+
+	TPair<FFileBucket&, FFileCategory&> FindOrAddCategory(TArray<FFileBucket>& Buckets, FFileBucket& DefaultBucket, const FFileCategoryID CategoryID)
+	{
+		// Find the category
+		for (FFileBucket& Bucket : Buckets)
+		{
+			for (FFileCategory& Category : Bucket.Categories)
+			{
+				if (Category.Id == CategoryID)
+				{
+					return TPair<FFileBucket&, FFileCategory&>(Bucket, Category);
+				}
+			}
+		}
+
+		// Category not found, add to default bucket
+		FFileCategory& Category = DefaultBucket.Categories.AddDefaulted_GetRef();
+		Category.Id = CategoryID;
+		return TPair<FFileBucket&, FFileCategory&>(DefaultBucket, Category);
+	}
+
+
 	struct FClassifyNode
 	{
 		TArray<FBlock> Blocks;
-
-		//TSharedPtr<FClassifyNode> Child0;
-		//TSharedPtr<FClassifyNode> Child1;
-		//uint32 Depth=0;
 	};
 
-	void AddNode(TMap<uint32, FClassifyNode>& Nodes, int32 Slack, const FBlock& Block)
+
+	void AddNode(TMap<FFileCategoryID, FClassifyNode>& Nodes, int32 Slack, const FFileCategoryID& CategoryID, const FBlock& Block)
 	{
-		FClassifyNode& Root = Nodes.FindOrAdd(Block.Flags);
+		FClassifyNode& Root = Nodes.FindOrAdd(CategoryID);
 		if (Root.Blocks.IsEmpty())
 		{
 			Root.Blocks.Reserve(Slack);
@@ -2581,233 +2596,392 @@ namespace MutablePrivate
 
 		Root.Blocks.Add(Block);
 	}
-
-	void GenerateBulkDataFilesList(
+	
+	void GenerateBulkDataFilesListWithFileLimit(
 		TSharedPtr<const mu::Model, ESPMode::ThreadSafe> Model,
 		FModelStreamableBulkData& ModelStreamableBulkData,
-		const ITargetPlatform* TargetPlatform,
-		uint64 TargetBulkDataFileBytes,
+		uint32 NumFilesPerBucket,
 		TArray<FFile>& OutBulkDataFiles)
 	{
-		MUTABLE_CPUPROFILER_SCOPE(GenerateBulkDataFilesList);
-
-		OutBulkDataFiles.Empty();
+		MUTABLE_CPUPROFILER_SCOPE(GenerateBulkDataFilesListWithFileLimit);
 
 		if (!Model)
 		{
 			return;
 		}
 
-		const uint64 MaxChunkSize = UCustomizableObjectSystem::GetInstance()->GetMaxChunkSizeForPlatform(TargetPlatform);
-		TargetBulkDataFileBytes = FMath::Min(TargetBulkDataFileBytes, MaxChunkSize);
+		/* Overview.
+		 *	1. Add categories to the different buckets and accumulate the size of its resources 
+		 *	   to know the total size of each category and the size of the buckets.
+		 *	2. Use the accumulated sizes to distribute the NumFilesPerBucket between the bucket's categories.
+		 *	3. Generate the list of BulkData files based on the number of files per category.
+		 */
 
-		// TODO: Temp. Remove after unifying generated output files code between editor an package. UE-222777 
-		const bool bRequiresCookedData = TargetPlatform->RequiresCookedData();
+		// Two buckets. One for non-optional data and one for optional data.
+		TArray<FFileBucket> FileBuckets;
 
-		// Root nodes by flags.
-		const int32 NumRoms = Model->GetRomCount();
-		TMap<uint32, FClassifyNode> RootNode;
+		// DefaultBucket is for non-optional BulkData
+		FFileBucket& DefaultBucket = FileBuckets.AddDefaulted_GetRef();
+		FFileBucket& OptionalBucket = FileBuckets.AddDefaulted_GetRef();
 
-		// Create blocks data.
+		// Model Roms. Iterate all Model roms to distribute them in categories.
 		{
+			// Add meshes and low-res textures to the Default bucket 
+			DefaultBucket.Categories.Add({ FFileCategoryID(EDataType::Model, mu::DATATYPE::DT_MESH, 0) });
+			DefaultBucket.Categories.Add({ FFileCategoryID(EDataType::Model, mu::DATATYPE::DT_IMAGE, 0) });
+
+			// Add High-res textures to the Optional bucket
+			OptionalBucket.Categories.Add({ FFileCategoryID(EDataType::Model, mu::DATATYPE::DT_IMAGE, (uint16)mu::ERomFlags::HighRes) });
+
+			const int32 NumRoms = Model->GetRomCount();
 			for (int32 RomIndex = 0; RomIndex < NumRoms; ++RomIndex)
 			{
 				uint32 BlockId = Model->GetRomId(RomIndex);
 				const uint32 BlockSize = Model->GetRomSize(RomIndex);
-				const mu::ERomFlags BlockFlags = bRequiresCookedData ? Model->GetRomFlags(RomIndex) : mu::ERomFlags::None;
+				const uint16 BlockResourceType = Model->GetRomType(RomIndex);
+				const mu::ERomFlags BlockFlags = Model->GetRomFlags(RomIndex);
 
-				FBlock CurrentBlock = { EDataType::Model, BlockId, BlockSize, uint32(BlockFlags), 0 };
-				AddNode(RootNode, NumRoms, CurrentBlock);
+				FFileCategoryID CategoryID = { EDataType::Model, BlockResourceType, uint16(BlockFlags) };
+				TPair<FFileBucket&, FFileCategory&> It = FindOrAddCategory(FileBuckets, DefaultBucket, CategoryID); // Add block to an existing or new category
+				It.Key.DataSize += BlockSize;
+				It.Value.DataSize += BlockSize;
 			}
 		}
 
-		// TODO: This should create a new classification branch when the tree is implemented.
-		// For now append after Model roms. 
+		// RealTime Morphs. Iterate RealTimeMorph streamables to accumulate their sizes.
 		{
-			uint64 SourceOffset = 0;
-
-			TArray<FMutableStreamableBlock> RealTimeMorphTargetsBlocks;
+			// Add RealTimeMorphs to the Default bucket
+			FFileCategory& RealTimeMorphCategory = DefaultBucket.Categories.AddDefaulted_GetRef();
+			RealTimeMorphCategory.Id.DataType = EDataType::RealTimeMorph;
 
 			const TMap<uint32, FRealTimeMorphStreamable>& RealTimeMorphStreamables = ModelStreamableBulkData.RealTimeMorphStreamables;
-
-			const int32 NumBlocks = RealTimeMorphTargetsBlocks.Num();
 			for (const TPair<uint32, FRealTimeMorphStreamable>& MorphStreamable : RealTimeMorphStreamables)
 			{
-				const uint32 BlockSize = MorphStreamable.Value.Size;
-
-				uint32 Flags = 0;
-				FBlock CurrentBlock = { EDataType::RealTimeMorph, MorphStreamable.Key, BlockSize, Flags, SourceOffset };
-				AddNode(RootNode, NumRoms, CurrentBlock);
-
-				SourceOffset += BlockSize;
+				RealTimeMorphCategory.DataSize += MorphStreamable.Value.Size;
 			}
+
+			DefaultBucket.DataSize += RealTimeMorphCategory.DataSize;
 		}
 
-		// TODO: This should create a new classification branch when the tree is implemented.
-		// For now append after Model roms. 
+		// Clothing. Iterate clothing streamables to accumulate their sizes.
 		{
-			uint64 SourceOffset = 0;
-
-			TArray<FMutableStreamableBlock> ClothingBlocks;
-
+			// Add Clothing to the Default bucket
+			FFileCategory& ClothingCategory = DefaultBucket.Categories.AddDefaulted_GetRef();
+			ClothingCategory.Id.DataType = EDataType::Clothing;
+			
 			const TMap<uint32, FClothingStreamable>& ClothingStreamables = ModelStreamableBulkData.ClothingStreamables;
-
-			const int32 NumBlocks = ClothingBlocks.Num();
 			for (const TPair<uint32, FClothingStreamable>& ClothStreamable : ClothingStreamables)
 			{
-				const uint32 BlockSize = ClothStreamable.Value.Size;
-
-				//check(SourceOffset == ClothStreamable.Value.Block.Offset);
-				uint32 Flags = 0;
-				FBlock CurrentBlock = { EDataType::Clothing, ClothStreamable.Key, BlockSize, Flags, SourceOffset };
-				AddNode(RootNode, NumRoms, CurrentBlock);
-
-				SourceOffset += BlockSize;
+				ClothingCategory.DataSize += ClothStreamable.Value.Size;
 			}
+
+			DefaultBucket.DataSize += ClothingCategory.DataSize;
 		}
+		
+		// Limited number of files in each bucket. Find the ideal file distribution between categories based on the accumulated size of their resources.
+		TArray<FFileCategory> Categories;
 
-		for (int32 FlagClassIndex = 0; FlagClassIndex < RootNode.Num(); ++FlagClassIndex)
+		for (FFileBucket& Bucket : FileBuckets)
 		{
-			FClassifyNode& Root = RootNode[FlagClassIndex];
+			uint32 NumFiles = 0;
 
-			// Temp: Group by order in the array
-			for (int32 BlockIndex = 0; BlockIndex < Root.Blocks.Num(); )
+			for (FFileCategory& Category : Bucket.Categories)
 			{
-				int32 CurrentFileSize = 0;
+				if (Category.DataSize > 0)
+				{
+					double DataDistribution = (double)Category.DataSize / Bucket.DataSize;
+					Category.NumFiles = FMath::Max(DataDistribution * NumFilesPerBucket, 1);  // At least one file if size > 0
+					Category.FirstFile = NumFiles;
+
+					NumFiles += Category.NumFiles;
+				}
+			}
+
+			Categories.Append(Bucket.Categories);
+		}
+		
+		// Function to create the list of bulk data files. Blocks will be grouped by source Id.
+		const auto CreateFileList = [Categories](const FFileCategoryID& CategoryID, const FClassifyNode& Node, TArray<FFile>& OutBulkDataFiles)
+			{
+				const FFileCategory* Category = Categories.FindByPredicate(
+					[CategoryID](const FFileCategory& C) { return C.Id == CategoryID; });
+				check(Category);
+
+				int32 NumBulkDataFiles = OutBulkDataFiles.Num();
+				OutBulkDataFiles.Reserve(NumBulkDataFiles + Category->NumFiles);
 
-				FFile CurrentFile;
-				CurrentFile.DataType = Root.Blocks[BlockIndex].DataType;
-				CurrentFile.Flags = Root.Blocks[BlockIndex].Flags;
+				// FileID (File Index) to BulkData file index.
+				TArray<int64> BulkDataFileIndex;
+				BulkDataFileIndex.Init(INDEX_NONE, Category->NumFiles);
 
-				while (BlockIndex < Root.Blocks.Num())
+				const int32 NumBlocks = Node.Blocks.Num();
+				for (const FBlock& Block : Node.Blocks)
 				{
-					FBlock CurrentBlock = Root.Blocks[BlockIndex];
+					// Use the module of the source id to determine the file id (FileIndex)
+					const uint32 FileID = Block.SourceId % Category->NumFiles;
+					int64& FileIndex = BulkDataFileIndex[FileID];
 
-					// Next file?
-					// Store different data types in different files. Blocks should be sorted by DataType so data is properly packeted
-					if (CurrentFile.DataType != CurrentBlock.DataType)
+					// Add new file
+					if (FileIndex == INDEX_NONE)
 					{
-						break;
+						FFile& NewFile = OutBulkDataFiles.AddDefaulted_GetRef();
+						NewFile.DataType = CategoryID.DataType;
+						NewFile.ResourceType = CategoryID.ResourceType;
+						NewFile.Flags = CategoryID.Flags;
+						NewFile.Id = FileID; 
+
+						FileIndex = NumBulkDataFiles;
+						++NumBulkDataFiles;
 					}
 
-					// Different flags go to different files
-					check(CurrentFile.Flags == CurrentBlock.Flags)
+					// Add block to the file 
+					OutBulkDataFiles[FileIndex].Blocks.Add(Block);
+				}
+			};
+
+		// Generate the list of BulkData files.
+		GenerateBulkDataFilesList(Model, ModelStreamableBulkData, true /* bUseRomTypeAndFlagsToFilter */, CreateFileList, OutBulkDataFiles);
+	}
+	
+
+	void GenerateBulkDataFilesListWithSizeLimit(
+		TSharedPtr<const mu::Model, ESPMode::ThreadSafe> Model,
+		FModelStreamableBulkData& ModelStreamableBulkData,
+		const ITargetPlatform* TargetPlatform,
+		uint64 TargetBulkDataFileBytes,
+		TArray<FFile>& OutBulkDataFiles)
+	{
+		MUTABLE_CPUPROFILER_SCOPE(GenerateBulkDataFilesListWithSizeLimit);
+
+		if (!Model)
+		{
+			return;
+		}
 
-						if (CurrentFileSize > 0 && CurrentFileSize + CurrentBlock.Size > TargetBulkDataFileBytes)
+		const uint64 MaxChunkSize = UCustomizableObjectSystem::GetInstance()->GetMaxChunkSizeForPlatform(TargetPlatform);
+		TargetBulkDataFileBytes = FMath::Min(TargetBulkDataFileBytes, MaxChunkSize);
+
+		// Unlimited number of files, limited file size. Add blocks to the file if the size limit won't be surpassed. Add at least one block to each file. 
+		const auto CreateFileList = [TargetBulkDataFileBytes](const FFileCategoryID& CategoryID, const FClassifyNode& Node, TArray<FFile>& OutBulkDataFiles)
+			{
+				// Temp: Group by order in the array
+				for (int32 BlockIndex = 0; BlockIndex < Node.Blocks.Num(); )
+				{
+					FFile File;
+					File.DataType = CategoryID.DataType;
+					File.ResourceType = CategoryID.ResourceType;
+					File.Flags = CategoryID.Flags;
+
+					uint64 FileSize = 0;
+					uint32 FileId = uint32(CategoryID.DataType);
+
+					while (BlockIndex < Node.Blocks.Num())
+					{
+						const FBlock& CurrentBlock = Node.Blocks[BlockIndex];
+
+						if (FileSize > 0 &&
+							FileSize + CurrentBlock.Size > TargetBulkDataFileBytes &&
+							TargetBulkDataFileBytes > 0)
 						{
 							break;
 						}
 
-					// Add the block to the current file
-					CurrentFile.Blocks.Add(CurrentBlock);
-					CurrentFileSize += CurrentBlock.Size;
+						// Block added to file. Set offset and increase file size.
+						FileSize += CurrentBlock.Size;
 
-					// Next block
-					++BlockIndex;
+						// Generate cumulative id for this file
+						FileId = HashCombine(FileId, CurrentBlock.Id);
+
+						// Add the block to the current file
+						File.Blocks.Add(CurrentBlock);
+
+						// Next block
+						++BlockIndex;
+					}
+
+					const int32 NumFiles = OutBulkDataFiles.Num();
+
+					// Ensure the FileId is unique
+					bool bUnique = false;
+					while (!bUnique)
+					{
+						bUnique = true;
+						for (int32 PreviousFileIndex = 0; PreviousFileIndex < NumFiles; ++PreviousFileIndex)
+						{
+							if (OutBulkDataFiles[PreviousFileIndex].Id == FileId)
+							{
+								bUnique = false;
+								++FileId;
+								break;
+							}
+						}
+					}
+
+					// Set it to the editor-only file descriptor
+					File.Id = FileId;
+
+					OutBulkDataFiles.Add(MoveTemp(File));
 				}
+			};
+
+		// TODO: Temp. Remove after unifying generated output files code between editor an package. UE-222777
+		const bool bUseRomTypeAndFlagsToFilter = TargetPlatform->RequiresCookedData();
+
+		GenerateBulkDataFilesList(Model, ModelStreamableBulkData, bUseRomTypeAndFlagsToFilter, CreateFileList, OutBulkDataFiles);
+	}
+
+
+	void GenerateBulkDataFilesList(
+		TSharedPtr<const mu::Model, ESPMode::ThreadSafe> Model,
+		FModelStreamableBulkData& ModelStreamableBulkData,
+		bool bUseRomTypeAndFlagsToFilter,
+		TFunctionRef<void(const FFileCategoryID&, const FClassifyNode&, TArray<FFile>&)> CreateFileList,
+		TArray<FFile>& OutBulkDataFiles)
+	{
+		MUTABLE_CPUPROFILER_SCOPE(GenerateBulkDataFilesList);
+
+		OutBulkDataFiles.Empty();
+
+		if (!Model)
+		{
+			return;
+		}
 
-				OutBulkDataFiles.Add(MoveTemp(CurrentFile));
+		// TODO: Temp. Remove after unifying generated output files code between editor an package. UE-222777 
+		const uint16 IgnoreMask = bUseRomTypeAndFlagsToFilter ? MAX_uint16 : 0;
+
+		// Root nodes by flags.
+		const int32 NumRoms = Model->GetRomCount();
+		TMap<FFileCategoryID, FClassifyNode> RootNode;
+
+		// Create blocks data.
+		{
+			for (int32 RomIndex = 0; RomIndex < NumRoms; ++RomIndex)
+			{
+				uint32 BlockId = Model->GetRomId(RomIndex);
+				uint32 SourceBlockId = Model->GetRomSourceId(RomIndex);
+				const uint32 BlockSize = Model->GetRomSize(RomIndex);
+				const uint16 BlockResourceType = IgnoreMask & Model->GetRomType(RomIndex);
+				const mu::ERomFlags BlockFlags = (mu::ERomFlags)(IgnoreMask & (uint16)Model->GetRomFlags(RomIndex));
+
+				FFileCategoryID CurrentCategory = { EDataType::Model, BlockResourceType, uint16(BlockFlags) };
+				FBlock CurrentBlock = { BlockId, SourceBlockId, BlockSize, 0 };
+				AddNode(RootNode, NumRoms, CurrentCategory, CurrentBlock);
 			}
 		}
 
-		// Create the file list
-		for (int32 FileIndex = 0; FileIndex < OutBulkDataFiles.Num(); ++FileIndex)
 		{
-			// Generate the id for this file
-			FFile& CurrentFile = OutBulkDataFiles[FileIndex];
-			uint32 FileId = uint32(CurrentFile.DataType);
-			for (int32 FileBlockIndex = 0; FileBlockIndex < CurrentFile.Blocks.Num(); ++FileBlockIndex)
+			uint64 SourceOffset = 0;
+
+			const TMap<uint32, FRealTimeMorphStreamable>& RealTimeMorphStreamables = ModelStreamableBulkData.RealTimeMorphStreamables;
+
+			const FFileCategoryID RealTimeMorphCategory = { EDataType::RealTimeMorph, (uint16)mu::DATATYPE::DT_NONE, (uint16)mu::ERomFlags::None };
+
+			for (const TPair<uint32, FRealTimeMorphStreamable>& MorphStreamable : RealTimeMorphStreamables)
 			{
-				FBlock ThisBlock = CurrentFile.Blocks[FileBlockIndex];
-				FileId = HashCombine(FileId, ThisBlock.Id);
+				const uint32 BlockSize = MorphStreamable.Value.Size;
+
+				FBlock CurrentBlock = { MorphStreamable.Key, MorphStreamable.Value.SourceId, BlockSize, SourceOffset };
+				AddNode(RootNode, NumRoms, RealTimeMorphCategory, CurrentBlock);
+
+				SourceOffset += BlockSize;
 			}
+		}
+
+		{
+			uint64 SourceOffset = 0;
+
+			const TMap<uint32, FClothingStreamable>& ClothingStreamables = ModelStreamableBulkData.ClothingStreamables;
 
-			// Ensure the FileId is unique
-			bool bUnique = false;
-			while (!bUnique)
+			const FFileCategoryID ClothingCategory = { EDataType::Clothing, (uint16)mu::DATATYPE::DT_NONE, (uint16)mu::ERomFlags::None };
+
+			for (const TPair<uint32, FClothingStreamable>& ClothStreamable : ClothingStreamables)
 			{
-				bUnique = true;
-				for (int32 PreviousFileIndex = 0; PreviousFileIndex < FileIndex; ++PreviousFileIndex)
-				{
-					if (OutBulkDataFiles[PreviousFileIndex].Id == FileId)
-					{
-						bUnique = false;
-						++FileId;
-						break;
-					}
-				}
+				const uint32 BlockSize = ClothStreamable.Value.Size;
+
+				FBlock CurrentBlock = { ClothStreamable.Key, ClothStreamable.Value.SourceId, BlockSize, SourceOffset };
+				AddNode(RootNode, NumRoms, ClothingCategory, CurrentBlock);
+
+				SourceOffset += BlockSize;
 			}
+		}
+
+		// Create Files list
+		for (TPair<FFileCategoryID, FClassifyNode>& Node : RootNode)
+		{
+			CreateFileList(Node.Key, Node.Value, OutBulkDataFiles);
+		}
+
+		// Update streamable blocks data
+		const int32 NumBulkDataFiles = OutBulkDataFiles.Num();
+		for (int32 FileIndex = 0; FileIndex < NumBulkDataFiles; ++FileIndex)
+		{
+			FFile& File = OutBulkDataFiles[FileIndex];
 
-			// Set it to the editor-only file descriptor
-			CurrentFile.Id = FileId;
+			uint64 SourceOffset = 0;
 
-			if (CurrentFile.DataType == EDataType::Model)
+			switch (File.DataType)
 			{
-				// Set it to all streamable blocks
-				uint64 OffsetInFile = 0;
-				for (int32 FileBlockIndex = 0; FileBlockIndex < CurrentFile.Blocks.Num(); ++FileBlockIndex)
+			case EDataType::Model:
+			{
+				for (FBlock& Block : File.Blocks)
 				{
-					FBlock& ThisBlock = CurrentFile.Blocks[FileBlockIndex];
-					ThisBlock.Offset = OffsetInFile;
-
-					FMutableStreamableBlock* StreamableBlock = ModelStreamableBulkData.ModelStreamables.Find(ThisBlock.Id);
-					StreamableBlock->FileId = FileId;
-					StreamableBlock->Offset = OffsetInFile;
-					check(StreamableBlock->Flags == CurrentFile.Flags);
-					OffsetInFile += ThisBlock.Size;
+					Block.Offset = SourceOffset;
+					SourceOffset += Block.Size;
+
+					FMutableStreamableBlock& StreamableBlock = ModelStreamableBulkData.ModelStreamables[Block.Id];
+					StreamableBlock.FileId = FileIndex;
+					StreamableBlock.Offset = Block.Offset;
 				}
+				break;
 			}
-			else if (CurrentFile.DataType == EDataType::RealTimeMorph)
+			case EDataType::RealTimeMorph:
 			{
-				TMap<uint32, FRealTimeMorphStreamable>& MorphBlocks = ModelStreamableBulkData.RealTimeMorphStreamables;
-				// Set it to all streamable blocks
-				uint64 OffsetInFile = 0;
-				for (int32 FileBlockIndex = 0; FileBlockIndex < CurrentFile.Blocks.Num(); ++FileBlockIndex)
+				for (FBlock& Block : File.Blocks)
 				{
-					FBlock& ThisBlock = CurrentFile.Blocks[FileBlockIndex];
-					ThisBlock.Offset = OffsetInFile;
+					Block.Offset = SourceOffset;
+					SourceOffset += Block.Size;
 
-					FMutableStreamableBlock& StreamableBlock = MorphBlocks[ThisBlock.Id].Block;
-					StreamableBlock.FileId = FileId;
-					StreamableBlock.Offset = OffsetInFile;
-					OffsetInFile += ThisBlock.Size;
+					FMutableStreamableBlock& StreamableBlock = ModelStreamableBulkData.RealTimeMorphStreamables[Block.Id].Block;
+					StreamableBlock.FileId = FileIndex;
+					StreamableBlock.Offset = Block.Offset;
 				}
+				break;
 			}
-			else if (CurrentFile.DataType == EDataType::Clothing)
+			case EDataType::Clothing:
 			{
-				TMap<uint32, FClothingStreamable>& ClothBlocks = ModelStreamableBulkData.ClothingStreamables;
-				// Set it to all streamable blocks
-				uint64 OffsetInFile = 0;
-				for (int32 FileBlockIndex = 0; FileBlockIndex < CurrentFile.Blocks.Num(); ++FileBlockIndex)
+				for (FBlock& Block : File.Blocks)
 				{
-					FBlock& ThisBlock = CurrentFile.Blocks[FileBlockIndex];
-					ThisBlock.Offset = OffsetInFile;
+					Block.Offset = SourceOffset;
+					SourceOffset += Block.Size;
 
-					FMutableStreamableBlock& StreamableBlock = ClothBlocks[ThisBlock.Id].Block;
-					StreamableBlock.FileId = FileId;
-					StreamableBlock.Offset = OffsetInFile;
-					OffsetInFile += ThisBlock.Size;
+					FMutableStreamableBlock& StreamableBlock = ModelStreamableBulkData.ClothingStreamables[Block.Id].Block;
+					StreamableBlock.FileId = FileIndex;
+					StreamableBlock.Offset = Block.Offset;
 				}
+				break;
 			}
-			else
-			{
+			default:
 				UE_LOG(LogMutable, Error, TEXT("Unknown DataType found while fixing streaming block files ids."));
-				check(false);
+				unimplemented();
+				break;
 			}
 		}
 	}
 
-
+	
 	void CUSTOMIZABLEOBJECT_API SerializeBulkDataFiles(
 		FMutableCachedPlatformData& CachedPlatformData,
 		TArray<FFile>& BulkDataFiles,
-		TFunctionRef<void(FFile&, TArray64<uint8>&)> WriteFile,
+		TFunctionRef<void(FFile&, TArray64<uint8>&, uint32)> WriteFile,
 		bool bDropData)
 	{
+		MUTABLE_CPUPROFILER_SCOPE(SerializeBulkDataFiles);
+
 		TArray64<uint8> FileBulkData;
 
-		const int32 NumBulkDataFiles = BulkDataFiles.Num();
-		for (int32 FileIndex = 0; FileIndex < NumBulkDataFiles; ++FileIndex)
+		const uint32 NumBulkDataFiles = BulkDataFiles.Num();
+		for (uint32 FileIndex = 0; FileIndex < NumBulkDataFiles; ++FileIndex)
 		{
 			MutablePrivate::FFile& CurrentFile = BulkDataFiles[FileIndex];
 
@@ -2817,35 +2991,35 @@ namespace MutablePrivate
 			// Get the file data in memory
 			CurrentFile.GetFileData(&CachedPlatformData, FileBulkData, bDropData);
 
-			WriteFile(CurrentFile, FileBulkData);
+			WriteFile(CurrentFile, FileBulkData, FileIndex);
 		}
 	}
-
+	
 	UE::DerivedData::FValueId GetDerivedDataModelId()
 	{
 		UE::DerivedData::FValueId::ByteArray ValueIdBytes = {};
-		ValueIdBytes[0] = 1;
+		FMemory::Memset(&ValueIdBytes[0], 1, sizeof(ValueIdBytes));
 		return UE::DerivedData::FValueId(ValueIdBytes);
 	}
 
 	UE::DerivedData::FValueId GetDerivedDataModelResourcesId()
 	{
-		UE::DerivedData::FValueId::ByteArray ValueIdBytes = {};
-		ValueIdBytes[0] = 2;
+		UE::DerivedData::FValueId::ByteArray ValueIdBytes = {};		
+		FMemory::Memset(&ValueIdBytes[0], 2, sizeof(ValueIdBytes));
 		return UE::DerivedData::FValueId(ValueIdBytes);
 	}
 
 	UE::DerivedData::FValueId GetDerivedDataModelStreamableBulkDataId()
 	{
 		UE::DerivedData::FValueId::ByteArray ValueIdBytes = {};
-		ValueIdBytes[0] = 3;
+		FMemory::Memset(&ValueIdBytes[0], 3, sizeof(ValueIdBytes));
 		return UE::DerivedData::FValueId(ValueIdBytes);
 	}
 
 	UE::DerivedData::FValueId GetDerivedDataBulkDataFilesId()
 	{
 		UE::DerivedData::FValueId::ByteArray ValueIdBytes = {};
-		ValueIdBytes[0] = 4;
+		FMemory::Memset(&ValueIdBytes[0], 4, sizeof(ValueIdBytes));
 		return UE::DerivedData::FValueId(ValueIdBytes);
 	}
 }
@@ -2858,7 +3032,6 @@ void SerializeCompilationOptionsForDDC(FArchive& Ar, FCompilationOptions& Option
 	Ar << PlatformName;
 	Ar << Options.TextureCompression;
 	Ar << Options.OptimizationLevel;
-	Ar << Options.DDCBytesLimit;
 	Ar << Options.CustomizableObjectNumBoneInfluences;
 	Ar << Options.bRealTimeMorphTargetsEnabled;
 	Ar << Options.bClothingEnabled;
@@ -3073,8 +3246,15 @@ void UCustomizableObjectPrivate::LoadCompiledDataFromDDC(FCompilationOptions Opt
 				FValueId::ByteArray ValueIdBytes = {};
 				for (MutablePrivate::FFile& File : BulkDataFiles)
 				{
-					FMemory::Memcpy(&ValueIdBytes, &File.DataType, sizeof(File.DataType));
-					FMemory::Memcpy(&ValueIdBytes[4], &File.Id, sizeof(File.Id));
+					int8 ValueIdOffset = 0;
+					FMemory::Memcpy(&ValueIdBytes[ValueIdOffset], &File.DataType, sizeof(File.DataType));
+					ValueIdOffset += sizeof(File.DataType);
+					FMemory::Memcpy(&ValueIdBytes[ValueIdOffset], &File.Id, sizeof(File.Id));
+					ValueIdOffset += sizeof(File.Id);
+					FMemory::Memcpy(&ValueIdBytes[ValueIdOffset], &File.ResourceType, sizeof(File.ResourceType));
+					ValueIdOffset += sizeof(File.ResourceType);
+					FMemory::Memcpy(&ValueIdBytes[ValueIdOffset], &File.Flags, sizeof(File.Flags));
+
 					MutablePrivate::FFile& DestFile = ValueIdToFile.Add(FValueId(ValueIdBytes));
 					DestFile = MoveTemp(File);
 				}
@@ -3166,7 +3346,15 @@ void UCustomizableObjectPrivate::LoadCompiledDataFromDDC(FCompilationOptions Opt
 			RequestOwner.Wait();
 
 			// Generate list of files and update streamable blocks ids and offsets
-			MutablePrivate::GenerateBulkDataFilesList(Model, *ModelStreamables.Get(), Options.TargetPlatform, Options.PackagedDataBytesLimit, CachedData.BulkDataFiles);
+			if (CVarMutableUseBulkData.GetValueOnAnyThread())
+			{
+				const uint32 NumBulkDataFilesPerBucket = MAX_uint8;
+				MutablePrivate::GenerateBulkDataFilesListWithFileLimit(Model, *ModelStreamables.Get(), NumBulkDataFilesPerBucket, CachedData.BulkDataFiles);
+			}
+			else
+			{
+				MutablePrivate::GenerateBulkDataFilesListWithSizeLimit(Model, *ModelStreamables.Get(), Options.TargetPlatform, Options.PackagedDataBytesLimit, CachedData.BulkDataFiles);
+			}
 
 			MutablePrivate::FMutableCachedPlatformData& CachedPlatformData = CachedPlatformsData.Add(Options.TargetPlatform->PlatformName(), {});
 			CachedPlatformData = MoveTemp(CachedData);
@@ -3176,6 +3364,7 @@ void UCustomizableObjectPrivate::LoadCompiledDataFromDDC(FCompilationOptions Opt
 	return;
 }
 
+
 #endif // WITH_EDITOR
 
 //-------------------------------------------------------------------------------------------------
@@ -3228,7 +3417,7 @@ void UCustomizableObjectBulk::CookAdditionalFilesOverride(const TCHAR* PackageFi
 
 	const FString CookedBulkFileName = FString::Printf(TEXT("%s/%s"), *FPaths::GetPath(PackageFilename), *CustomizableObject->GetName());
 
-	const auto WriteFile = [WriteAdditionalFile, CookedBulkFileName](MutablePrivate::FFile& File, TArray64<uint8>& FileBulkData)
+	const auto WriteFile = [WriteAdditionalFile, CookedBulkFileName](MutablePrivate::FFile& File, TArray64<uint8>& FileBulkData, uint32 FileIndex)
 		{
 			FString FileName = CookedBulkFileName + FString::Printf(TEXT("-%08x.mut"), File.Id);
 
diff --git a/Engine/Plugins/Mutable/Source/CustomizableObject/Private/MuCO/CustomizableObjectInstance.cpp b/Engine/Plugins/Mutable/Source/CustomizableObject/Private/MuCO/CustomizableObjectInstance.cpp
--- a/Engine/Plugins/Mutable/Source/CustomizableObject/Private/MuCO/CustomizableObjectInstance.cpp
+++ b/Engine/Plugins/Mutable/Source/CustomizableObject/Private/MuCO/CustomizableObjectInstance.cpp
@@ -5583,7 +5583,7 @@ UE::Tasks::FTask UCustomizableInstancePrivate::LoadAdditionalAssetsAndData(
 			uint32 FileId;
 		};
 
-		const bool bUseFBulkData = !ModelStreamableBulkData->HashToBulkData.IsEmpty();
+		const bool bUseFBulkData = !ModelStreamableBulkData->StreamableBulkData.IsEmpty();
 
 		TArray<FBlockReadInfo> BlockReadInfos;
 		BlockReadInfos.Reserve(16);
@@ -5729,10 +5729,10 @@ UE::Tasks::FTask UCustomizableInstancePrivate::LoadAdditionalAssetsAndData(
 								ReadDataReadyEvent.Trigger();
 							};
 
-						FByteBulkData* ByteBulkData = ModelStreamableBulkData->HashToBulkData.Find(Block.FileId);
-						check(ByteBulkData);
+						check(ModelStreamableBulkData->StreamableBulkData.IsValidIndex(Block.FileId));
+						FByteBulkData& ByteBulkData = ModelStreamableBulkData->StreamableBulkData[Block.FileId];
 
-						return TUniquePtr<IBulkDataIORequest>(ByteBulkData->CreateStreamingRequest(
+						return TUniquePtr<IBulkDataIORequest>(ByteBulkData.CreateStreamingRequest(
 							Block.Offset,
 							(int64)Block.AllocatedMemoryView.Num(),
 							Priority,
diff --git a/Engine/Plugins/Mutable/Source/CustomizableObject/Private/MuCO/UnrealMutableModelDiskStreamer.cpp b/Engine/Plugins/Mutable/Source/CustomizableObject/Private/MuCO/UnrealMutableModelDiskStreamer.cpp
--- a/Engine/Plugins/Mutable/Source/CustomizableObject/Private/MuCO/UnrealMutableModelDiskStreamer.cpp
+++ b/Engine/Plugins/Mutable/Source/CustomizableObject/Private/MuCO/UnrealMutableModelDiskStreamer.cpp
@@ -116,7 +116,7 @@ bool FUnrealMutableModelBulkReader::PrepareStreamingForObject(UCustomizableObjec
 		NewData.DDCKey = ModelResources.DDCKey;
 		NewData.DDCPolicy = ModelResources.DDCDefaultPolicy;
 #else
-		if (NewData.ModelStreamableBulkData->HashToBulkData.IsEmpty())
+		if (NewData.ModelStreamableBulkData->StreamableBulkData.IsEmpty())
 		{
 			const UCustomizableObjectBulk* BulkData = CustomizableObject->GetPrivate()->GetStreamableBulkData();
 			if (!BulkData)
@@ -251,8 +251,10 @@ mu::ModelReader::OPERATION_ID FUnrealMutableModelBulkReader::BeginReadBlock(cons
 
 	OPERATION_ID Result = ++LastOperationID;
 
-	if (FByteBulkData* BulkData = ObjectData->ModelStreamableBulkData->HashToBulkData.Find(Block->FileId))
+	if(ObjectData->ModelStreamableBulkData->StreamableBulkData.IsValidIndex(Block->FileId))
 	{
+		FByteBulkData& BulkData = ObjectData->ModelStreamableBulkData->StreamableBulkData[Block->FileId];
+
 		FBulkDataIORequestCallBack IOCallback;
 
 		if (CompletionCallback)
@@ -265,7 +267,7 @@ mu::ModelReader::OPERATION_ID FUnrealMutableModelBulkReader::BeginReadBlock(cons
 
 		FReadRequest& Request = ObjectData->CurrentReadRequests.Add(Result);
 
-		Request.BulkReadRequest = MakeShareable(BulkData->CreateStreamingRequest(
+		Request.BulkReadRequest = MakeShareable(BulkData.CreateStreamingRequest(
 			Block->Offset,
 			size,
 			(EAsyncIOPriorityAndFlags)StreamPriority,
diff --git a/Engine/Plugins/Mutable/Source/CustomizableObjectEditor/Private/MuCOE/CustomizableObjectCompileRunnable.cpp b/Engine/Plugins/Mutable/Source/CustomizableObjectEditor/Private/MuCOE/CustomizableObjectCompileRunnable.cpp
--- a/Engine/Plugins/Mutable/Source/CustomizableObjectEditor/Private/MuCOE/CustomizableObjectCompileRunnable.cpp
+++ b/Engine/Plugins/Mutable/Source/CustomizableObjectEditor/Private/MuCOE/CustomizableObjectCompileRunnable.cpp
@@ -394,10 +394,16 @@ void FCustomizableObjectSaveDDRunnable::CachePlatfromData()
 	{
 		PlatformData.ModelStreamables = ModelStreamables;
 
-		const uint64 PackageDataBytesLimit = Options.bIsCooking ? Options.PackagedDataBytesLimit : MAX_uint64;
-
 		// Generate list of files and update streamable blocks ids and offsets
-		MutablePrivate::GenerateBulkDataFilesList(Model, *PlatformData.ModelStreamables.Get(), Options.TargetPlatform, PackageDataBytesLimit, PlatformData.BulkDataFiles);
+		if (Options.bUseBulkData)
+		{
+			MutablePrivate::GenerateBulkDataFilesListWithFileLimit(Model, *PlatformData.ModelStreamables.Get(), MAX_uint8, PlatformData.BulkDataFiles);
+		}
+		else
+		{
+			const uint64 PackageDataBytesLimit = Options.bIsCooking ? Options.PackagedDataBytesLimit : MAX_uint64;
+			MutablePrivate::GenerateBulkDataFilesListWithSizeLimit(Model, *PlatformData.ModelStreamables.Get(), Options.TargetPlatform, PackageDataBytesLimit, PlatformData.BulkDataFiles);
+		}
 	}
 
 
@@ -435,7 +441,7 @@ void FCustomizableObjectSaveDDRunnable::StoreCachedPlatformDataInDDC(bool& bStor
 		FModelStreamableBulkData ModelStreamablesDDC = *ModelStreamables.Get();
 
 		// Generate list of files and update streamable blocks ids and offsets
-		MutablePrivate::GenerateBulkDataFilesList(Model, ModelStreamablesDDC, Options.TargetPlatform, Options.DDCBytesLimit, BulkDataFilesDDC);
+		MutablePrivate::GenerateBulkDataFilesListWithFileLimit(Model, ModelStreamablesDDC, MAX_int16, BulkDataFilesDDC);
 
 		TArray64<uint8> ModelStreamablesBytesDDC;
 		FMemoryWriter64 MemoryWriterDDC(ModelStreamablesBytesDDC);
@@ -448,12 +454,24 @@ void FCustomizableObjectSaveDDRunnable::StoreCachedPlatformDataInDDC(bool& bStor
 	// Store streamable resources as FValues
 	{
 		MUTABLE_CPUPROFILER_SCOPE(SerializeBulkDataForDDC);
-		
+	
+		{
+			MutablePrivate::FFile File;
+			check(sizeof(FValueId::ByteArray) >= sizeof(File.DataType) + sizeof(File.Id) + sizeof(File.ResourceType) + sizeof(File.Flags));
+		}
+
 		FValueId::ByteArray ValueIdBytes = {};
-		const auto WriteBulkDataDDC = [&RecordBuilder, &ValueIdBytes](MutablePrivate::FFile& File, TArray64<uint8>& FileBulkData)
+		const auto WriteBulkDataDDC = [&RecordBuilder, &ValueIdBytes](MutablePrivate::FFile& File, TArray64<uint8>& FileBulkData, uint32 FileIndex)
 			{
-				FMemory::Memcpy(&ValueIdBytes, &File.DataType, sizeof(File.DataType));
-				FMemory::Memcpy(&ValueIdBytes[4], &File.Id, sizeof(File.Id));
+				int8 ValueIdOffset = 0;
+				FMemory::Memcpy(&ValueIdBytes[ValueIdOffset], &File.DataType, sizeof(File.DataType));
+				ValueIdOffset += sizeof(File.DataType);
+				FMemory::Memcpy(&ValueIdBytes[ValueIdOffset], &File.Id, sizeof(File.Id));
+				ValueIdOffset += sizeof(File.Id);
+				FMemory::Memcpy(&ValueIdBytes[ValueIdOffset], &File.ResourceType, sizeof(File.ResourceType));
+				ValueIdOffset += sizeof(File.ResourceType);
+				FMemory::Memcpy(&ValueIdBytes[ValueIdOffset], &File.Flags, sizeof(File.Flags));
+
 				const FValueId ValueId(ValueIdBytes);
 				const FValue Value = FValue::Compress(FSharedBuffer::MakeView(FileBulkData.GetData(), FileBulkData.Num()));
 
@@ -561,7 +579,7 @@ void FCustomizableObjectSaveDDRunnable::StoreCachedPlatformDataToDisk(bool& bSto
 		// Serailize headers to validate data
 		*StreamableMemoryWriter << CustomizableObjectHeader;
 
-		const auto WriteBulkDataToDisk = [&StreamableMemoryWriter](MutablePrivate::FFile& File, TArray64<uint8>& FileBulkData)
+		const auto WriteBulkDataToDisk = [&StreamableMemoryWriter](MutablePrivate::FFile& File, TArray64<uint8>& FileBulkData, uint32 FileIndex)
 			{
 				switch (File.DataType)
 				{
diff --git a/Engine/Plugins/Mutable/Source/CustomizableObjectEditor/Private/MuCOE/CustomizableObjectCompiler.cpp b/Engine/Plugins/Mutable/Source/CustomizableObjectEditor/Private/MuCOE/CustomizableObjectCompiler.cpp
--- a/Engine/Plugins/Mutable/Source/CustomizableObjectEditor/Private/MuCOE/CustomizableObjectCompiler.cpp
+++ b/Engine/Plugins/Mutable/Source/CustomizableObjectEditor/Private/MuCOE/CustomizableObjectCompiler.cpp
@@ -920,6 +920,7 @@ void FCustomizableObjectCompiler::CompileInternal(bool bAsync)
 			ResourceMeshData.Size = DataSizeInBytes;
 			mu::ERomFlags Flags = mu::ERomFlags::None;
 			ResourceMeshData.Block = FMutableStreamableBlock { uint32(0), uint32(Flags), RealTimeMorphDataOffsetInBytes };
+			ResourceMeshData.SourceId = MeshData.Value.SourceId;
 
 			RealTimeMorphDataOffsetInBytes += DataSizeInBytes;
 			ModelResources.EditorOnlyMorphTargetReconstructionData.Append(MeshData.Value.Data);
@@ -951,6 +952,7 @@ void FCustomizableObjectCompiler::CompileInternal(bool bAsync)
 			ResourceMeshData.Size = DataSizeInBytes;
 			mu::ERomFlags Flags = mu::ERomFlags::None;
 			ResourceMeshData.Block = FMutableStreamableBlock{ uint32(0), uint32(Flags), ClothingDataOffsetInBytes };
+			ResourceMeshData.SourceId = MeshData.Value.SourceId;
 
 			ClothingDataOffsetInBytes += DataSizeInBytes;
 			ModelResources.EditorOnlyClothingMeshToMeshVertData.Append(MeshData.Value.Data);
diff --git a/Engine/Plugins/Mutable/Source/CustomizableObjectEditor/Private/MuCOE/GenerateMutableSource/GenerateMutableSource.h b/Engine/Plugins/Mutable/Source/CustomizableObjectEditor/Private/MuCOE/GenerateMutableSource/GenerateMutableSource.h
--- a/Engine/Plugins/Mutable/Source/CustomizableObjectEditor/Private/MuCOE/GenerateMutableSource/GenerateMutableSource.h
+++ b/Engine/Plugins/Mutable/Source/CustomizableObjectEditor/Private/MuCOE/GenerateMutableSource/GenerateMutableSource.h
@@ -299,6 +299,11 @@ struct FRealTimeMorphMeshData
 {
 	TArray<FName> NameResolutionMap;
 	TArray<FMorphTargetVertexData> Data;
+
+	/* Used to group data when generating bulk data files.
+	 * This property should not be taken into consideration when comparing structs.
+	 */
+	uint32 SourceId = 0;
 };
 
 struct FClothingMeshData
@@ -307,6 +312,11 @@ struct FClothingMeshData
 	int32 ClothingAssetLOD = INDEX_NONE;
 	int32 PhysicsAssetIndex = INDEX_NONE;
 	TArray<FCustomizableObjectMeshToMeshVertData> Data;
+
+	/* Used to group data when generating bulk data files. 
+	 * This property should not be taken into consideration when comparing structs.
+	 */
+	uint32 SourceId = 0;
 };
 
 struct FGroupProjectorTempData
diff --git a/Engine/Plugins/Mutable/Source/CustomizableObjectEditor/Private/MuCOE/GenerateMutableSource/GenerateMutableSourceGroupProjector.cpp b/Engine/Plugins/Mutable/Source/CustomizableObjectEditor/Private/MuCOE/GenerateMutableSource/GenerateMutableSourceGroupProjector.cpp
--- a/Engine/Plugins/Mutable/Source/CustomizableObjectEditor/Private/MuCOE/GenerateMutableSource/GenerateMutableSourceGroupProjector.cpp
+++ b/Engine/Plugins/Mutable/Source/CustomizableObjectEditor/Private/MuCOE/GenerateMutableSource/GenerateMutableSourceGroupProjector.cpp
@@ -505,7 +505,7 @@ bool GenerateMutableSourceGroupProjector(const UEdGraphPin* Pin, FMutableGraphGe
 						int32 HighResMipsForThisImage = FMath::Min(NumMipsBeyondMin, GenerationContext.Options.NumHighResImageMips);
 						ImageNode->SourceDataDescriptor.SourceHighResMips = HighResMipsForThisImage;
 
-						const FString TextureName = GetNameSafe(Texture);
+						const FString TextureName = GetNameSafe(Texture).ToLower();
 						ImageNode->SourceDataDescriptor.SourceId = CityHash32(reinterpret_cast<const char*>(*TextureName), TextureName.Len() * sizeof(FString::ElementType));
 					}
 				}
diff --git a/Engine/Plugins/Mutable/Source/CustomizableObjectEditor/Private/MuCOE/GenerateMutableSource/GenerateMutableSourceImage.cpp b/Engine/Plugins/Mutable/Source/CustomizableObjectEditor/Private/MuCOE/GenerateMutableSource/GenerateMutableSourceImage.cpp
--- a/Engine/Plugins/Mutable/Source/CustomizableObjectEditor/Private/MuCOE/GenerateMutableSource/GenerateMutableSourceImage.cpp
+++ b/Engine/Plugins/Mutable/Source/CustomizableObjectEditor/Private/MuCOE/GenerateMutableSource/GenerateMutableSourceImage.cpp
@@ -128,7 +128,7 @@ mu::NodeImagePtr GenerateMutableSourceImage(const UEdGraphPin* Pin, FMutableGrap
 				int32 HighResMipsForThisImage = FMath::Min(NumMipsBeyondMin, GenerationContext.Options.NumHighResImageMips);
 				ConstantImageNode->SourceDataDescriptor.SourceHighResMips = HighResMipsForThisImage;
 
-				const FString TextureName = GetNameSafe(BaseTexture);
+				const FString TextureName = GetNameSafe(BaseTexture).ToLower();
 				ConstantImageNode->SourceDataDescriptor.SourceId = CityHash32(reinterpret_cast<const char*>(*TextureName), TextureName.Len() * sizeof(FString::ElementType));
 
 				ImageNode = ConstantImageNode;
@@ -232,7 +232,7 @@ mu::NodeImagePtr GenerateMutableSourceImage(const UEdGraphPin* Pin, FMutableGrap
 			int32 HighResMipsForThisImage = FMath::Min(NumMipsBeyondMin, GenerationContext.Options.NumHighResImageMips);
 			ImageNode->SourceDataDescriptor.SourceHighResMips = HighResMipsForThisImage;
 
-			const FString TextureName = GetNameSafe(Texture);
+			const FString TextureName = GetNameSafe(Texture).ToLower();
 			ImageNode->SourceDataDescriptor.SourceId = CityHash32(reinterpret_cast<const char*>(*TextureName), TextureName.Len() * sizeof(FString::ElementType));
 		}
 		else
@@ -1003,7 +1003,7 @@ mu::NodeImagePtr GenerateMutableSourceImage(const UEdGraphPin* Pin, FMutableGrap
 		}
 		else
 		{
-			const FString TableName = TypedNodeTable->Table ? GetNameSafe(TypedNodeTable->Table) : GetNameSafe(TypedNodeTable->Structure);
+			const FString TableName = TypedNodeTable->Table ? GetNameSafe(TypedNodeTable->Table).ToLower() : GetNameSafe(TypedNodeTable->Structure).ToLower();
 			const uint32 TableId = CityHash32(reinterpret_cast<const char*>(*TableName), TableName.Len() * sizeof(FString::ElementType));
 
 			// This node will add a checker texture in case of error
diff --git a/Engine/Plugins/Mutable/Source/CustomizableObjectEditor/Private/MuCOE/GenerateMutableSource/GenerateMutableSourceMesh.cpp b/Engine/Plugins/Mutable/Source/CustomizableObjectEditor/Private/MuCOE/GenerateMutableSource/GenerateMutableSourceMesh.cpp
--- a/Engine/Plugins/Mutable/Source/CustomizableObjectEditor/Private/MuCOE/GenerateMutableSource/GenerateMutableSourceMesh.cpp
+++ b/Engine/Plugins/Mutable/Source/CustomizableObjectEditor/Private/MuCOE/GenerateMutableSource/GenerateMutableSourceMesh.cpp
@@ -697,7 +697,10 @@ mu::MeshPtr ConvertSkeletalMeshToMutable(
 	}
 
 	GenerationContext.AddParticipatingObject(*InSkeletalMesh);
-	
+
+	const FString MeshName = *GetNameSafe(InSkeletalMesh).ToLower();
+	const uint32 MeshId = CityHash32(reinterpret_cast<const char*>(*MeshName), MeshName.Len() * sizeof(FString::ElementType));
+
 	const FSkeletalMeshModel* ImportedModel = InSkeletalMesh->GetImportedModel();
 	if (!ImportedModel)
 	{
@@ -1472,6 +1475,7 @@ mu::MeshPtr ConvertSkeletalMeshToMutable(
 
 						RealTimeMorphMeshData.NameResolutionMap = MoveTemp(NameResolutionMap);
 						RealTimeMorphMeshData.Data = MoveTemp(MorphsMeshData);
+						RealTimeMorphMeshData.SourceId = MeshId;
 					}
 				}
 			}
@@ -1796,6 +1800,7 @@ mu::MeshPtr ConvertSkeletalMeshToMutable(
 				NewClothingMeshData.PhysicsAssetIndex = ClothingMeshData.PhysicsAssetIndex;
 
 				NewClothingMeshData.Data = MoveTemp(ClothingMeshData.Data);
+				NewClothingMeshData.SourceId = MeshId;
 			}
 			
 			NextBufferIndex += 2;
@@ -3162,7 +3167,7 @@ mu::NodeMeshPtr GenerateMorphMesh(const UEdGraphPin* Pin,
 	
 	if (const UCustomizableObjectNodeTable* TypedNodeTable = Cast<UCustomizableObjectNodeTable>(Pin->GetOwningNode()))
 	{
-		const FString TableName = TypedNodeTable->Table ? GetNameSafe(TypedNodeTable->Table) : GetNameSafe(TypedNodeTable->Structure);
+		const FString TableName = TypedNodeTable->Table ? GetNameSafe(TypedNodeTable->Table).ToLower() : GetNameSafe(TypedNodeTable->Structure).ToLower();
 		const uint32 TableId = CityHash32(reinterpret_cast<const char*>(*TableName), TableName.Len() * sizeof(FString::ElementType));
 
 
@@ -3562,7 +3567,7 @@ mu::Ptr<mu::NodeMesh> GenerateMutableSourceMesh(const UEdGraphPin* Pin,
 			}
 
 
-			const FString MeshName = GetNameSafe(TypedNodeSkel->SkeletalMesh);
+			const FString MeshName = GetNameSafe(TypedNodeSkel->SkeletalMesh).ToLower();
 			MeshNode->SourceDataDescriptor.SourceId = CityHash32(reinterpret_cast<const char*>(*MeshName), MeshName.Len() * sizeof(FString::ElementType));
 			MeshNode->SourceDataDescriptor.SourceHighResMips = 0;
 
@@ -3657,7 +3662,7 @@ mu::Ptr<mu::NodeMesh> GenerateMutableSourceMesh(const UEdGraphPin* Pin,
 				LayoutNode->SetMessageContext(Node);  // We need it here because we create multiple nodes.
 
 
-				const FString MeshName = GetNameSafe(TypedNodeStatic->StaticMesh);
+				const FString MeshName = GetNameSafe(TypedNodeStatic->StaticMesh).ToLower();
 				MeshNode->SourceDataDescriptor.SourceId = CityHash32(reinterpret_cast<const char*>(*MeshName), MeshName.Len() * sizeof(FString::ElementType));
 				MeshNode->SourceDataDescriptor.SourceHighResMips = 0;
 			}
@@ -4177,7 +4182,7 @@ mu::Ptr<mu::NodeMesh> GenerateMutableSourceMesh(const UEdGraphPin* Pin,
 
 	else if (const UCustomizableObjectNodeTable* TypedNodeTable = Cast<UCustomizableObjectNodeTable>(Node))
 	{
-		const FString TableName = TypedNodeTable->Table ? GetNameSafe(TypedNodeTable->Table) : GetNameSafe(TypedNodeTable->Structure);
+		const FString TableName = TypedNodeTable->Table ? GetNameSafe(TypedNodeTable->Table).ToLower() : GetNameSafe(TypedNodeTable->Structure).ToLower();
 		const uint32 TableId = CityHash32(reinterpret_cast<const char*>(*TableName), TableName.Len() * sizeof(FString::ElementType));
 
 		mu::NodeMeshConstantPtr EmptyNode = new mu::NodeMeshConstant();
diff --git a/Engine/Plugins/Mutable/Source/CustomizableObjectEditor/Private/MuCOE/GenerateMutableSource/GenerateMutableSourceSurface.cpp b/Engine/Plugins/Mutable/Source/CustomizableObjectEditor/Private/MuCOE/GenerateMutableSource/GenerateMutableSourceSurface.cpp
--- a/Engine/Plugins/Mutable/Source/CustomizableObjectEditor/Private/MuCOE/GenerateMutableSource/GenerateMutableSourceSurface.cpp
+++ b/Engine/Plugins/Mutable/Source/CustomizableObjectEditor/Private/MuCOE/GenerateMutableSource/GenerateMutableSourceSurface.cpp
@@ -590,7 +590,7 @@ mu::Ptr<mu::NodeSurface> GenerateMutableSourceSurface(const UEdGraphPin * Pin, F
 										int32 HighResMipsForThisImage = FMath::Min(NumMipsBeyondMin, GenerationContext.Options.NumHighResImageMips);
 										ConstImageNode->SourceDataDescriptor.SourceHighResMips = HighResMipsForThisImage;
 
-										const FString TextureName = GetNameSafe(Texture2D);
+										const FString TextureName = GetNameSafe(Texture2D).ToLower();
 										ConstImageNode->SourceDataDescriptor.SourceId = CityHash32(reinterpret_cast<const char*>(*TextureName), TextureName.Len() * sizeof(FString::ElementType));
 
 										return Result;
@@ -690,7 +690,7 @@ mu::Ptr<mu::NodeSurface> GenerateMutableSourceSurface(const UEdGraphPin * Pin, F
 									int32 HighResMipsForThisImage = FMath::Min(NumMipsBeyondMin, GenerationContext.Options.NumHighResImageMips);
 									CompositeNormalImage->SourceDataDescriptor.SourceHighResMips = HighResMipsForThisImage;
 
-									const FString TextureName = GetNameSafe(ReferenceCompositeNormalTexture);
+									const FString TextureName = GetNameSafe(ReferenceCompositeNormalTexture).ToLower();
 									CompositeNormalImage->SourceDataDescriptor.SourceId = CityHash32(reinterpret_cast<const char*>(*TextureName), TextureName.Len() * sizeof(FString::ElementType));
 								}
 
