[Mutable] Fix Modifiers and Group Projector nodes not compiling correctly with multiple components.

- Fixed Mesh Component nodes being compiled even if they were not connected.
- Removed unused ParameterNamesMap/AddParameterNameUnique.

#rnx

diff --git a/Engine/Plugins/Mutable/Source/CustomizableObject/Internal/MuCO/CustomizableObjectPrivate.h b/Engine/Plugins/Mutable/Source/CustomizableObject/Internal/MuCO/CustomizableObjectPrivate.h
--- a/Engine/Plugins/Mutable/Source/CustomizableObject/Internal/MuCO/CustomizableObjectPrivate.h
+++ b/Engine/Plugins/Mutable/Source/CustomizableObject/Internal/MuCO/CustomizableObjectPrivate.h
@@ -1561,6 +1561,8 @@ public:
 
 		BulkDataFilesNumFilesLimit,
 
+		RemoveModifiersHack,
+
 		// -----<new versions can be added above this line>--------
 		LastCustomizableObjectVersion
 	};
diff --git a/Engine/Plugins/Mutable/Source/CustomizableObjectEditor/Private/MuCOE/CustomizableObjectCompiler.cpp b/Engine/Plugins/Mutable/Source/CustomizableObjectEditor/Private/MuCOE/CustomizableObjectCompiler.cpp
--- a/Engine/Plugins/Mutable/Source/CustomizableObjectEditor/Private/MuCOE/CustomizableObjectCompiler.cpp
+++ b/Engine/Plugins/Mutable/Source/CustomizableObjectEditor/Private/MuCOE/CustomizableObjectCompiler.cpp
@@ -9,6 +9,7 @@
 #include "Interfaces/ITargetPlatformManagerModule.h"
 #include "Materials/MaterialInterface.h"
 #include "MessageLogModule.h"
+#include "GenerateMutableSource/GenerateMutableSourceComponent.h"
 #include "MuCO/CustomizableObject.h"
 #include "MuCO/CustomizableObjectPrivate.h"
 #include "MuCO/CustomizableObjectInstance.h"
@@ -485,7 +486,7 @@ mu::Ptr<mu::NodeObject> GenerateMutableRoot(
 	}
 
 	bool bMultipleBaseObjectsFound;
-	UCustomizableObjectNodeObject* Root = GetRootNode(Object, bMultipleBaseObjectsFound);
+	UCustomizableObjectNodeObject* LocalRootNodeObject = GetRootNode(Object, bMultipleBaseObjectsFound);
 
 	if (bMultipleBaseObjectsFound)
 	{
@@ -493,19 +494,19 @@ mu::Ptr<mu::NodeObject> GenerateMutableRoot(
 		return nullptr;
 	}
 
-	if (!Root)
+	if (!LocalRootNodeObject)
 	{
 		GenerationContext.Log(LOCTEXT("NoRootBase","No base object node found. Object not built."));
 		return nullptr;
 	}
 	
-	const UCustomizableObject* ActualRootObject = GraphTraversal::GetRootObject(Object);
-	check(ActualRootObject);
+	const UCustomizableObject* RootObject = GraphTraversal::GetRootObject(Object);
+	check(RootObject);
 
-	GenerationContext.RootVersionBridge = ActualRootObject->VersionBridge;
+	GenerationContext.RootVersionBridge = RootObject->VersionBridge;
 
-	UCustomizableObjectNodeObject* ActualRoot = GetRootNode(ActualRootObject, bMultipleBaseObjectsFound);
-	GenerationContext.Root = ActualRoot;
+	UCustomizableObjectNodeObject* RootNodeObject = GetRootNode(RootObject, bMultipleBaseObjectsFound);
+	GenerationContext.Root = RootNodeObject;
 	
 	if (bMultipleBaseObjectsFound)
 	{
@@ -513,20 +514,20 @@ mu::Ptr<mu::NodeObject> GenerateMutableRoot(
 		return nullptr;
 	}
 
-	if (!ActualRoot)
+	if (!RootNodeObject)
 	{
 		GenerationContext.Log(LOCTEXT("NoActualRootBase", "No base object node found in root Customizable Object. Object not built."));
 		return nullptr;
 	}
 	
-	if (Root->ObjectName.IsEmpty())
+	if (LocalRootNodeObject->ObjectName.IsEmpty())
 	{
-		GenerationContext.NoNameNodeObjectArray.AddUnique(Root);
+		GenerationContext.NoNameNodeObjectArray.AddUnique(LocalRootNodeObject);
 	}
 
 	if ((Object->MeshCompileType == EMutableCompileMeshType::Full) || GenerationContext.Options.bIsCooking)
 	{
-		if (Root->ParentObject!=nullptr && GenerationContext.Options.bIsCooking)
+		if (LocalRootNodeObject->ParentObject!=nullptr && GenerationContext.Options.bIsCooking)
 		{
 			// This happens while packaging.
 			return nullptr;
@@ -539,7 +540,7 @@ mu::Ptr<mu::NodeObject> GenerateMutableRoot(
 
 			// The object doesn't reference a root object but is a root object, look for all the objects that reference it and get their root nodes
 			FAssetRegistryModule& AssetRegistryModule = FModuleManager::LoadModuleChecked<FAssetRegistryModule>("AssetRegistry");
-			ProcessChildObjectsRecursively(ActualRootObject, GenerationContext);
+			ProcessChildObjectsRecursively(RootObject, GenerationContext);
 			UE_LOG(LogMutable, Verbose, TEXT("PROFILE: [ %16.8f ] End search for children."), FPlatformTime::Seconds());
 		}
 	}
@@ -551,7 +552,7 @@ mu::Ptr<mu::NodeObject> GenerateMutableRoot(
 		
 		if (!GetParentsUntilRoot(Object, ArrayNodeObject, ArrayCustomizableObject))
 		{
-			GenerationContext.Log(LOCTEXT("SkeletalMeshCycleFound", "Error! Cycle detected in the Customizable Object hierarchy."), Root);
+			GenerationContext.Log(LOCTEXT("SkeletalMeshCycleFound", "Error! Cycle detected in the Customizable Object hierarchy."), LocalRootNodeObject);
 			return nullptr;
 		}
 
@@ -567,7 +568,7 @@ mu::Ptr<mu::NodeObject> GenerateMutableRoot(
 
 					if (!GetParentsUntilRoot(WorkingSetObject, ArrayNodeObject, ArrayCustomizableObject))
 					{
-						GenerationContext.Log(LOCTEXT("NoReferenceMesh", "Error! Cycle detected in the Customizable Object hierarchy."), Root);
+						GenerationContext.Log(LOCTEXT("NoReferenceMesh", "Error! Cycle detected in the Customizable Object hierarchy."), LocalRootNodeObject);
 						return nullptr;
 					}
 				}
@@ -594,80 +595,30 @@ mu::Ptr<mu::NodeObject> GenerateMutableRoot(
 			}
 		}
 	}
-	
-	// Find all component nodes
-	TSet<const UCustomizableObject*> CustomizableObjects; // All COs in the hierarchy
-	{
-		auto GetParents = [&CustomizableObjects](const UCustomizableObject* Object)
-		{
-			TArray<UCustomizableObjectNodeObject*> ArrayNodeObject;
-			TArray<const UCustomizableObject*> ArrayCustomizableObject;
-
-			GetParentsUntilRoot(Object, ArrayNodeObject, ArrayCustomizableObject);
-
-			CustomizableObjects.Append(ArrayCustomizableObject);
-		};
 
-		// Parents
-		GetParents(Object);
-
-		// Working set and their parents
-		for (TSoftObjectPtr<UCustomizableObject> SoftObject :  Object->WorkingSet)
-		{
-			if (UCustomizableObject* WorkingSetObject = SoftObject.LoadSynchronous())
-			{
-				GetParents(WorkingSetObject);
-			}
-		}
-
-		// Children
-		TArray<FAssetData> ReferencingAssets;
-		GetReferencingPackages(*Object, ReferencingAssets);
-		for (FAssetData& AssetData : ReferencingAssets)
-		{
-			if (UCustomizableObject* ChildObject = Cast<UCustomizableObject>(AssetData.GetAsset()))
-			{
-				CustomizableObjects.Add(ChildObject);
-			}
-		}
-	}
-	
-	for (const UCustomizableObject* ParticipatingObject : CustomizableObjects)
+	// First pass. Only used to recollect info required for the primary pass.
+	// Notice that the traversal is different form the primary pass. Here we follow all pins indiscriminately,
+	// while the primary pass follows the Mutable Source structure (which may cut branches).
+	GraphTraversal::VisitNodes(*RootNodeObject, GenerationContext.GroupIdToExternalNodeMap, [&GenerationContext](UCustomizableObjectNode& Node)
 	{
-		for (UEdGraphNode* Node : ParticipatingObject->GetPrivate()->GetSource()->Nodes)
+		if (UCustomizableObjectNodeComponentMesh* NodeComponentMesh = Cast<UCustomizableObjectNodeComponentMesh>(&Node))
 		{
-			UCustomizableObjectNodeComponentMesh* NodeComponentMesh = Cast<UCustomizableObjectNodeComponentMesh>(Node);
-			if (!NodeComponentMesh)
-			{
-				continue;
-			}
-
-			if (UCustomizableObjectNodeComponentMesh** Result = GenerationContext.MeshComponents.Find(NodeComponentMesh->ComponentName))
-			{
-				UCustomizableObjectNodeComponentMesh* NodeMeshComponent = *Result;
-				check(NodeMeshComponent);
-				
-				FText Msg = FText::Format(LOCTEXT("ComponentNodeWithSameNameExists", "Already exists a Mesh Component node with the same name in Customizable Object [{0}]"), FText::FromString(GetRootObject(*NodeMeshComponent)->GetName()));
-				GenerationContext.Log(Msg, Node, EMessageSeverity::Error);
-				return nullptr;
-			}
-
-			GenerationContext.MeshComponents.Add(NodeComponentMesh->ComponentName, NodeComponentMesh);
+			FirstPass(*NodeComponentMesh, GenerationContext);
 		}
-	}
+	});
 	
-    GenerationContext.RealTimeMorphTargetsOverrides = ActualRoot->RealTimeMorphSelectionOverrides;
+    GenerationContext.RealTimeMorphTargetsOverrides = RootNodeObject->RealTimeMorphSelectionOverrides;
 
 	if (!GenerationContext.Options.ParamNamesToSelectedOptions.IsEmpty())
 	{
 		GenerationContext.TableToParamNames = Object->GetPrivate()->GetModelResources().TableToParamNames;
 	}
 
-	GenerationContext.bPartialCompilation = Root->ParentObject != nullptr;
+	GenerationContext.bPartialCompilation = LocalRootNodeObject->ParentObject != nullptr;
 
 	// Generate the object expression
 	UE_LOG(LogMutable, Verbose, TEXT("PROFILE: [ %16.8f ] GenerateMutableSource start."), FPlatformTime::Seconds());
-	mu::NodeObjectPtr MutableRoot = GenerateMutableSource(ActualRoot->OutputPin(), GenerationContext);
+	mu::NodeObjectPtr MutableRoot = GenerateMutableSource(RootNodeObject->OutputPin(), GenerationContext);
 	UE_LOG(LogMutable, Verbose, TEXT("PROFILE: [ %16.8f ] GenerateMutableSource end."), FPlatformTime::Seconds());
 
 	GenerationContext.GenerateSharedSurfacesUniqueIds();
diff --git a/Engine/Plugins/Mutable/Source/CustomizableObjectEditor/Private/MuCOE/GenerateMutableSource/GenerateMutableSource.cpp b/Engine/Plugins/Mutable/Source/CustomizableObjectEditor/Private/MuCOE/GenerateMutableSource/GenerateMutableSource.cpp
--- a/Engine/Plugins/Mutable/Source/CustomizableObjectEditor/Private/MuCOE/GenerateMutableSource/GenerateMutableSource.cpp
+++ b/Engine/Plugins/Mutable/Source/CustomizableObjectEditor/Private/MuCOE/GenerateMutableSource/GenerateMutableSource.cpp
@@ -56,6 +56,7 @@
 #include "Math/NumericLimits.h"
 #include "Hash/CityHash.h"
 #include "MuCOE/Nodes/CustomizableObjectNodeCopyMaterial.h"
+#include "MuCOE/Nodes/CustomizableObjectNodeGroupProjectorParameter.h"
 
 #define LOCTEXT_NAMESPACE "CustomizableObjectEditor"
 
@@ -122,6 +123,16 @@ void CheckNumOutputs(const UEdGraphPin& Pin, const FMutableGraphGenerationContex
 }
 
 
+uint32 GetTypeHash(const FGeneratedGroupProjectorsKey& Key)
+{
+	uint32 Hash = GetTypeHash(Key.Node);
+
+	Hash = HashCombine(Hash, GetTypeHash(Key.CurrentComponent));
+	
+	return Hash; 
+}
+
+
 FMutableGraphGenerationContext::FMutableGraphGenerationContext(const UCustomizableObject* InObject, FCustomizableObjectCompiler* InCompiler, const FCompilationOptions& InOptions)
 	: Object(InObject), Compiler(InCompiler), Options(InOptions)
 	, ExtensionDataCompilerInterface(*this)
@@ -254,21 +265,6 @@ int32 FMutableGraphGenerationContext::AddAssetUserDataToStreamedResources(UAsset
 }
 
 
-/** Adds to ParameterNamesMap the node Node to the array of elements with name Name */
-void FMutableGraphGenerationContext::AddParameterNameUnique(const UCustomizableObjectNode* Node, FString Name)
-{
-	if (TArray<const UObject*>* ArrayResult = ParameterNamesMap.Find(Name))
-	{
-		ArrayResult->AddUnique(Node);
-	}
-	else
-	{
-		TArray<const UObject*> ArrayTemp;
-		ArrayTemp.Add(Node);
-		ParameterNamesMap.Add(Name, ArrayTemp);
-	}
-}
-
 uint32 FMutableGraphGenerationContext::GetSkinWeightProfileIdUnique(const FName InProfileName)
 {
 	const FString ProfileNameString = InProfileName.ToString();
@@ -436,20 +432,10 @@ mu::FBoneName FMutableGraphGenerationContext::GetBoneUnique(const FName& InBoneN
 
 FMutableComponentInfo* FMutableGraphGenerationContext::GetCurrentComponentInfo()
 {
-	FMutableComponentInfo* CurrentComponentInfo = ComponentInfos.FindByPredicate(
-		[this](const FMutableComponentInfo& Component) { return Component.ComponentName == CurrentMeshComponent; });
-	
-	// Temp workaround to the problem of modifiers that generate meshes (like ExtendMeshSection) not having a known component at generation time.
-	// TODO: Actually detect all possible components and change this query in GenerationContext state to accomodate a set of "current components" instead of one.
-	//check(CurrentComponentInfo);
-
-	if (!CurrentComponentInfo)
+	return ComponentInfos.FindByPredicate([this](const FMutableComponentInfo& Component)
 	{
-		check(!ComponentInfos.IsEmpty());
-		CurrentComponentInfo = &ComponentInfos[0];
-	}
-
-	return CurrentComponentInfo;
+		return Component.ComponentName == CurrentMeshComponent;
+	});
 }
 
 
@@ -497,18 +483,6 @@ FGeneratedKey::FGeneratedKey(void* InFunctionAddress, const UEdGraphPin& InPin,
 }
 
 
-bool FGeneratedKey::operator==(const FGeneratedKey& Other) const
-{
-	return FunctionAddress == Other.FunctionAddress &&
-		Pin == Other.Pin &&
-		LOD == Other.LOD &&
-		Flags == Other.Flags &&
-		LayoutFlags == Other.LayoutFlags &&
-		MeshMorphStack == Other.MeshMorphStack &&
-		bOnlyConnectedLOD == Other.bOnlyConnectedLOD;
-}
-
-
 uint32 GetTypeHash(const FGeneratedKey& Key)
 {
 	uint32 Hash = GetTypeHash(Key.FunctionAddress);
@@ -518,6 +492,7 @@ uint32 GetTypeHash(const FGeneratedKey& Key)
 	//Hash = HashCombine(Hash, GetTypeHash(Key.LayoutFlags)); // Does not support array
 	//Hash = HashCombine(Hash, GetTypeHash(Key.MeshMorphStack)); // Does not support array
 	Hash = HashCombine(Hash, GetTypeHash(Key.bOnlyConnectedLOD));
+	Hash = HashCombine(Hash, GetTypeHash(Key.CurrentMeshComponent));
 	
 	return Hash;
 }
@@ -999,17 +974,18 @@ mu::Ptr<mu::NodeObject> GenerateMutableSource(const UEdGraphPin * Pin, FMutableG
 				{
 					continue;
 				}
-
-				// Temp hack
-				if (GenerationContext.ComponentInfos.Num())
+				
+				for (const FMutableComponentInfo& Pair : GenerationContext.ComponentInfos)
 				{
-					GenerationContext.CurrentMeshComponent = GenerationContext.ComponentInfos[0].ComponentName;
+					GenerationContext.CurrentMeshComponent = Pair.ComponentName;
+					GenerationContext.CurrentAutoLODStrategy = Pair.Node->AutoLODStrategy;
+					
+					mu::Ptr<mu::NodeModifier> ModifierNode = GenerateMutableSourceModifier(ChildNodePin, GenerationContext);
+					ObjectNode->Modifiers.Add(ModifierNode);
+
+					GenerationContext.CurrentMeshComponent = {};
+					GenerationContext.CurrentAutoLODStrategy = {};
 				}
-				
-				mu::Ptr<mu::NodeModifier> ModifierNode = GenerateMutableSourceModifier(ChildNodePin, GenerationContext);
-				ObjectNode->Modifiers.Add(ModifierNode);
-
-				GenerationContext.CurrentMeshComponent = FName();
 			}
 		}
 		
@@ -1068,22 +1044,27 @@ mu::Ptr<mu::NodeObject> GenerateMutableSource(const UEdGraphPin * Pin, FMutableG
 		// All sockets from all mesh parts plugged into this group node will have the following priority when there's a socket name clash
 		GenerationContext.SocketPriorityStack.Push(TypedNodeGroup->SocketPriority);
 
-		GenerationContext.AddParameterNameUnique(TypedNodeGroup, TypedNodeGroup->GroupName);
 		GroupNode->SetName(TypedNodeGroup->GroupName);
 		GroupNode->SetUid(TypedNodeGroup->NodeGuid.ToString());
-
-		// Get all group projectors and put them in the generation context so that they are available to the child material nodes of this group node
-		uint32 NumProjectorCountBeforeNode = GenerationContext.ProjectorGroupMap.Num();
-		UEdGraphPin* ProjectorsPin = TypedNodeGroup->GroupProjectorsPin();
 		
-		if (ProjectorsPin && TypedNodeGroup)
+		TArray<UCustomizableObjectNodeGroupProjectorParameter*> GroupProjectors;
+		if (UEdGraphPin* ProjectorsPin = TypedNodeGroup->GroupProjectorsPin())
 		{
-			for (const UEdGraphPin* GroupProjectorNodePin : FollowInputPinArray(*ProjectorsPin))
+			for (const UEdGraphPin* ConnectedPin : FollowInputPinArray(*ProjectorsPin))
 			{
-				GenerateMutableSourceGroupProjector(GroupProjectorNodePin, GenerationContext, TypedNodeGroup);
+				if (UCustomizableObjectNodeGroupProjectorParameter* GroupProjectorNode = Cast<UCustomizableObjectNodeGroupProjectorParameter>(ConnectedPin->GetOwningNode()))
+				{
+					GroupProjectors.Add(GroupProjectorNode);
+				}
 			}
 		}
 
+		GenerationContext.CurrentGroupProjectors.Push(GroupProjectors);
+		ON_SCOPE_EXIT
+		{
+			GenerationContext.CurrentGroupProjectors.Pop();
+		};
+
 		mu::NodeObjectGroup::CHILD_SELECTION Type = mu::NodeObjectGroup::CS_ALWAYS_ALL;
 		switch (TypedNodeGroup->GroupType)
 		{
@@ -1150,8 +1131,6 @@ mu::Ptr<mu::NodeObject> GenerateMutableSource(const UEdGraphPin * Pin, FMutableG
 
 					if (TypedNodeGroup->GroupType == ECustomizableObjectGroupType::COGT_TOGGLE)
 					{
-						GenerationContext.AddParameterNameUnique(CustomizableObjectNodeObject, CustomizableObjectNodeObject->ObjectName);
-						
 						// UI Data is only relevant when the group node is set to Toggle
 						GenerationContext.ParameterUIDataMap.Add(CustomizableObjectNodeObject->ObjectName, FMutableParameterData(
 							CustomizableObjectNodeObject->ParamUIMetadata,
@@ -1263,8 +1242,6 @@ mu::Ptr<mu::NodeObject> GenerateMutableSource(const UEdGraphPin * Pin, FMutableG
 
 					if (TypedNodeGroup->GroupType == ECustomizableObjectGroupType::COGT_TOGGLE)
 					{
-						GenerationContext.AddParameterNameUnique(CustomizableObjectNodeObject, CustomizableObjectNodeObject->ObjectName);
-
 						// UI Data is only relevant when the group node is set to Toggle
 						GenerationContext.ParameterUIDataMap.Add(CustomizableObjectNodeObject->ObjectName, FMutableParameterData(
 							CustomizableObjectNodeObject->ParamUIMetadata,
@@ -1314,14 +1291,9 @@ mu::Ptr<mu::NodeObject> GenerateMutableSource(const UEdGraphPin * Pin, FMutableG
 			GenerationContext.ParameterUIDataMap.Add(TypedNodeGroup->GroupName, ParameterUIData);
 		}
 
-		// Remove the projectors from this node
-		GenerationContext.ProjectorGroupMap.Remove(TypedNodeGroup);
-
 		// Go back to the parent group node's socket priority if it exists
 		ensure(GenerationContext.SocketPriorityStack.Num() > 0);
 		GenerationContext.SocketPriorityStack.Pop();
-
-		check(NumProjectorCountBeforeNode == GenerationContext.ProjectorGroupMap.Num());
 	}
 	else
 	{
diff --git a/Engine/Plugins/Mutable/Source/CustomizableObjectEditor/Private/MuCOE/GenerateMutableSource/GenerateMutableSource.h b/Engine/Plugins/Mutable/Source/CustomizableObjectEditor/Private/MuCOE/GenerateMutableSource/GenerateMutableSource.h
--- a/Engine/Plugins/Mutable/Source/CustomizableObjectEditor/Private/MuCOE/GenerateMutableSource/GenerateMutableSource.h
+++ b/Engine/Plugins/Mutable/Source/CustomizableObjectEditor/Private/MuCOE/GenerateMutableSource/GenerateMutableSource.h
@@ -30,6 +30,7 @@
 #include "MuR/Ptr.h"
 #include "UObject/Package.h"
 
+class UCustomizableObjectNodeGroupProjectorParameter;
 class UCustomizableObjectNodeMaterialBase;
 class UCustomizableObjectNodeModifierClipWithMesh;
 class UCustomizableObjectNodeTable;
@@ -149,9 +150,8 @@ class FGeneratedKey
 public:
 	FGeneratedKey(void* InFunctionAddress, const UEdGraphPin& InPin, const UCustomizableObjectNode& Node, FMutableGraphGenerationContext& GenerationContext, const bool UseMesh = false, bool bOnlyConnectedLOD = false);
 	
-	bool operator==(const FGeneratedKey& Other) const;
+	bool operator==(const FGeneratedKey& Other) const = default;
 
-private:
 	/** Used to differentiate pins being cached from different functions (e.g. a PC_Color pin cached from GenerateMutableSourceImage and GenerateMutableSourceColor). */
 	void* FunctionAddress;
 	
@@ -167,7 +167,9 @@ private:
 
 	/** UV Layout modes */
 	FLayoutGenerationFlags LayoutFlags;
-
+	
+	FName CurrentMeshComponent;
+	
 	/** When caching a generated mesh, true if we force to generate the connected LOD when using Automatic LODs From Mesh. */
 	bool bOnlyConnectedLOD = false;
 };
@@ -441,6 +443,8 @@ struct FMutableComponentInfo
 
 	// Bones to remove on each LOD, include bones on previous LODs. FName (BoneToRemove) - bool (bOnlyRemoveChildren)
 	TArray<TMap<FName, bool>> BonesToRemovePerLOD;
+	
+	UCustomizableObjectNodeComponentMesh* Node = nullptr;
 };
 
 
@@ -551,6 +555,19 @@ private:
 		return {}; \
 	} \
 
+
+
+struct FGeneratedGroupProjectorsKey
+{
+	UCustomizableObjectNodeGroupProjectorParameter* Node = nullptr;
+	FName CurrentComponent;
+
+	bool operator==(const FGeneratedGroupProjectorsKey&) const = default;
+};
+
+uint32 GetTypeHash(const FGeneratedGroupProjectorsKey& Key);
+
+
 struct FMutableGraphGenerationContext
 {
 	FMutableGraphGenerationContext(const UCustomizableObject* CustomizableObject, class FCustomizableObjectCompiler* InCompiler, const FCompilationOptions& InOptions);
@@ -624,6 +641,14 @@ public:
 	// Cache of generated Node Tables
 	TMap<FString, FGeneratedDataTablesData> GeneratedTables;
 
+	TMap<FGeneratedGroupProjectorsKey, FGroupProjectorTempData> GeneratedGroupProjectors;
+
+	/** Key is the Node Uid. */
+	TMap<FString, mu::Ptr<mu::NodeScalarParameter>> GeneratedScalarParameters;
+
+	/** Key is the Node Uid. */
+	TMap<FString, mu::Ptr<mu::NodeScalarEnumParameter>> GeneratedEnumParameters;
+	
 	struct FGeneratedCompositeDataTablesData
 	{
 		UScriptStruct* ParentStruct = nullptr;
@@ -736,6 +761,10 @@ public:
 	// Stack of Layout generation flags. The last one is the currently valid.
 	TArray<FLayoutGenerationFlags> LayoutGenerationFlags;
 
+	/** Stack of Group Projector nodes. Each time a Group Object node is visited, a set of Group Projector nodes get pushed
+	 * When a Mesh Section node is found, it will compile all Group Projector nodes in the stack. */
+	TArray<TArray<UCustomizableObjectNodeGroupProjectorParameter*>> CurrentGroupProjectors;
+
 	/** Find a mesh if already generated for a given source and flags. */
 	mu::Ptr<mu::Mesh> FindGeneratedMesh(const FGeneratedMeshData::FKey& Key);
 
@@ -748,9 +777,6 @@ public:
 	  * Returns resource index in the array of streamed resources. */
 	int32 AddAssetUserDataToStreamedResources(UAssetUserData* AssetUserData);
 
-	/** Adds to ParameterNamesMap the node Node to the array of elements with name Name */
-	void AddParameterNameUnique(const UCustomizableObjectNode* Node, FString Name);
-
 	uint32 GetSkinWeightProfileIdUnique(const FName ProfileName);
 
 	// Check if the Id of the node Node already exists, if it's new adds it to NodeIds array, otherwise, returns new Id
@@ -796,12 +822,10 @@ private:
 	void AddParticipatingObjectChecked(const FName& PackageName, const FGuid& PackageGuid);
 
 public:
-	TMap<FName, UCustomizableObjectNodeComponentMesh*> MeshComponents;
-
 	/** Only Mesh Components (no passthrough). */
 	TArray<FMutableComponentInfo> ComponentInfos;
 
-	/** All components. Index is the ObjectComponentIndex. */
+	/** Only compiled components. All components types. Index is the ObjectComponentIndex. */
 	TArray<FName> ComponentNames;
 	
 	TArray<FMutableRefSkeletalMeshData> ReferenceSkeletalMeshesData;
@@ -809,14 +833,12 @@ public:
 	TArray<UMaterialInterface*> ReferencedMaterials;
 	TArray<FName> ReferencedMaterialSlotNames;
 	TMap<FGeneratedImagePropertiesKey, FGeneratedImageProperties> ImageProperties;
-	TMap<FString, TArray<const UObject*>> ParameterNamesMap;
 	TArray<const UCustomizableObjectNode*> NoNameNodeObjectArray;
 	TMap<FString, FCustomizableObjectIdPair> GroupNodeMap;
 	TMap<FString, FString> CustomizableObjectPathMap;
 	TMap<FString, FMutableParameterData> ParameterUIDataMap;
 	TMap<FString, FMutableStateData> StateUIDataMap;
 	TMap<TTuple<FString, FString>, TSet<TSoftObjectPtr<UDataTable>>> IntParameterOptionDataTable;
-	TMultiMap<const UCustomizableObjectNodeObjectGroup*, FGroupProjectorTempData> ProjectorGroupMap;
 	//TMap<UPhysicsAsset*, uint32> DiscartedPhysicsAssetMap;
 
 	TArray<const USkeleton*> ReferencedSkeletons;
diff --git a/Engine/Plugins/Mutable/Source/CustomizableObjectEditor/Private/MuCOE/GenerateMutableSource/GenerateMutableSourceColor.cpp b/Engine/Plugins/Mutable/Source/CustomizableObjectEditor/Private/MuCOE/GenerateMutableSource/GenerateMutableSourceColor.cpp
--- a/Engine/Plugins/Mutable/Source/CustomizableObjectEditor/Private/MuCOE/GenerateMutableSource/GenerateMutableSourceColor.cpp
+++ b/Engine/Plugins/Mutable/Source/CustomizableObjectEditor/Private/MuCOE/GenerateMutableSource/GenerateMutableSourceColor.cpp
@@ -67,8 +67,6 @@ mu::Ptr<mu::NodeColour> GenerateMutableSourceColor(const UEdGraphPin* Pin, FMuta
 		mu::Ptr<mu::NodeColourParameter> ColorNode = new mu::NodeColourParameter();
 		Result = ColorNode;
 
-		GenerationContext.AddParameterNameUnique(Node, TypedNodeColorParam->ParameterName);
-
 		ColorNode->Name = TypedNodeColorParam->ParameterName;
 		ColorNode->Uid = GenerationContext.GetNodeIdUnique(Node).ToString();
 		ColorNode->DefaultValue = TypedNodeColorParam->DefaultValue;
diff --git a/Engine/Plugins/Mutable/Source/CustomizableObjectEditor/Private/MuCOE/GenerateMutableSource/GenerateMutableSourceComponent.cpp b/Engine/Plugins/Mutable/Source/CustomizableObjectEditor/Private/MuCOE/GenerateMutableSource/GenerateMutableSourceComponent.cpp
--- a/Engine/Plugins/Mutable/Source/CustomizableObjectEditor/Private/MuCOE/GenerateMutableSource/GenerateMutableSourceComponent.cpp
+++ b/Engine/Plugins/Mutable/Source/CustomizableObjectEditor/Private/MuCOE/GenerateMutableSource/GenerateMutableSourceComponent.cpp
@@ -40,7 +40,7 @@
 /** Generate LOD pins of the given NodeComponentBase (NodeComponent, NodeComponentExtend...).
  * @param TypedComponentMesh Given component node.
  * @param NodeComponent Core node to connect LOD generated pins. */
-void GenerateMutableSourceComponentMesh(FMutableGraphGenerationContext& GenerationContext, const UCustomizableObjectNodeComponentMeshBase& TypedComponentMesh, mu::Ptr<mu::NodeComponent> NodeComponent, mu::Ptr<mu::NodeObjectNew> NodeObject)
+void GenerateMutableSourceComponentMesh(FMutableGraphGenerationContext& GenerationContext, const UCustomizableObjectNodeComponentMeshBase& TypedComponentMesh, mu::Ptr<mu::NodeComponent> NodeComponent)
 {
 	int32 FirstLOD = -1;
 
@@ -130,168 +130,37 @@ mu::Ptr<mu::NodeComponent> GenerateMutableSourceComponent(const UEdGraphPin* Pin
 	
 	if (const UCustomizableObjectNodeComponentMesh* TypedComponentMesh = Cast<UCustomizableObjectNodeComponentMesh>(Node))
 	{
-		FName ComponentName = TypedComponentMesh->ComponentName;
-		
-		if (TypedComponentMesh->ComponentName.IsNone())
-		{
-			GenerationContext.Log(LOCTEXT("EmptyComponentNameError", "Error! Missing name in a component of the Customizable Object."), TypedComponentMesh, EMessageSeverity::Error);
-			return nullptr;
-		}
-
-		if (GenerationContext.ComponentInfos.ContainsByPredicate([&](const FMutableComponentInfo& ComponentInfo)
+		if (!GenerationContext.ComponentInfos.ContainsByPredicate([&](const FMutableComponentInfo& ComponentInfo)
 		{
-			return ComponentInfo.ComponentName == ComponentName;
+			return ComponentInfo.Node == TypedComponentMesh;
 		}))
 		{
-			GenerationContext.Log(FText::Format(LOCTEXT("RepeatedComponentName", "Error! Repeated name [{0}] used in more than one Component"),
-				FText::FromName(ComponentName)), TypedComponentMesh, EMessageSeverity::Error);
-			return nullptr;
-		}
-			
-		USkeletalMesh* RefSkeletalMesh = TypedComponentMesh->ReferenceSkeletalMesh;
-		if (!RefSkeletalMesh)
-		{
-			GenerationContext.Log(LOCTEXT("NoReferenceMeshObjectTab", "Error! Missing reference Skeletal Mesh"), TypedComponentMesh, EMessageSeverity::Error);
-			return nullptr;
-		}
-
-		USkeleton* RefSkeleton = RefSkeletalMesh->GetSkeleton();
-		if(!RefSkeleton)
-		{
-			FText Msg = FText::Format(LOCTEXT("NoReferenceSkeleton", "Error! Missing skeleton in the reference mesh [{0}]"), FText::FromString(GenerationContext.CustomizableObjectWithCycle->GetPathName()));
-
-			GenerationContext.Log(Msg, TypedComponentMesh, EMessageSeverity::Error);
-			return nullptr;
-		}
-
-		// Add a new entry to the list of Component Infos
-		FMutableComponentInfo& ComponentInfo = GenerationContext.ComponentInfos.Add_GetRef(FMutableComponentInfo(ComponentName, RefSkeletalMesh));
-
-		ComponentInfo.AccumulateBonesToRemovePerLOD(TypedComponentMesh->LODReductionSettings, TypedComponentMesh->NumLODs);
-
-		// Make sure the Skeleton from the reference mesh is added to the list of referenced Skeletons.
-		GenerationContext.ReferencedSkeletons.Add(RefSkeleton);
-
-		// Add reference meshes to the participating objects
-		GenerationContext.AddParticipatingObject(*RefSkeletalMesh);
-		
-		// Ensure that the CO has a valid AutoLODStrategy on the Component node.
-		if (TypedComponentMesh->AutoLODStrategy == ECustomizableObjectAutomaticLODStrategy::Inherited)
-		{
-			GenerationContext.Log(LOCTEXT("RootInheritsFromParent", "Error! Component LOD Strategy can't be set to 'Inherit from parent object'"), TypedComponentMesh, EMessageSeverity::Error);
-			return nullptr;
-		}
-		GenerationContext.CurrentAutoLODStrategy = TypedComponentMesh->AutoLODStrategy;
-		
-		mu::Ptr<mu::NodeObjectNew> ObjectNode = new mu::NodeObjectNew();
-
-		ObjectNode->SetName(TypedComponentMesh->ComponentName.ToString());
-		FGuid FinalGuid = GenerationContext.GetNodeIdUnique(TypedComponentMesh);
-		if (FinalGuid != TypedComponentMesh->NodeGuid)
-		{
-			GenerationContext.Log(FText::FromString(TEXT("Warning: Node has a duplicated GUID. A new ID has been generated, but cooked data will not be deterministic.")), Node, EMessageSeverity::Warning);
-		}
-		ObjectNode->SetUid(FinalGuid.ToString());
-
-
-		// Process the pins defining components
-		//-------------------------------------------------------------------
-
-		// LOD
-		const int32 NumLODs = TypedComponentMesh->LODPins.Num();
-
-		// Fill the basic LOD Settings
-		if (!GenerationContext.NumLODsInRoot)
-		{
-			check(!GenerationContext.ComponentInfos.IsEmpty());
-
-			// NumLODsInRoot
-			int32 MaxRefMeshLODs = TypedComponentMesh->ReferenceSkeletalMesh->GetLODNum();
-			if (MaxRefMeshLODs < NumLODs)
-			{
-				FString Msg = FString::Printf(TEXT("The object has %d LODs but the reference mesh only %d. Resulting objects will have %d LODs."),
-					NumLODs, MaxRefMeshLODs, MaxRefMeshLODs);
-				GenerationContext.Log(FText::FromString(Msg), Node, EMessageSeverity::Warning);
-				GenerationContext.NumLODsInRoot = MaxRefMeshLODs;
-			}
-			else
-			{
-				GenerationContext.NumLODsInRoot = FMath::Max(GenerationContext.NumLODsInRoot, static_cast<uint8>(NumLODs));
-			}
-		
-			const FMutableLODSettings& LODSettings = GenerationContext.Object->LODSettings;
-
-			// Find the MinLOD available for the target platform
-			if (RefSkeletalMesh->IsMinLodQualityLevelEnable())
-			{
-				FSupportedQualityLevelArray SupportedQualityLevels = LODSettings.MinQualityLevelLOD.GetSupportedQualityLevels(*GenerationContext.Options.TargetPlatform->GetPlatformInfo().IniPlatformName.ToString());
-				
-				int32 MinValue = GenerationContext.NumLODsInRoot - 1;
-				for (int32& QL : SupportedQualityLevels)
-				{
-					// check if have data for the supported quality level or set to default.
-					if (LODSettings.MinQualityLevelLOD.IsQualityLevelValid(QL))
-					{
-						MinValue = FMath::Min(LODSettings.MinQualityLevelLOD.GetValueForQualityLevel(QL), MinValue);
-					}
-					else 
-					{
-						MinValue = LODSettings.MinQualityLevelLOD.GetDefault();
-						break;
-					}
-				}
-
-				GenerationContext.FirstLODAvailable = FMath::Max(0, MinValue);
-			}
-			else
-			{
-				GenerationContext.FirstLODAvailable = LODSettings.MinLOD.GetValueForPlatform(*GenerationContext.Options.TargetPlatform->IniPlatformName());
-			}
-
-			GenerationContext.FirstLODAvailable = FMath::Clamp(GenerationContext.FirstLODAvailable, 0, GenerationContext.NumLODsInRoot - 1);
-
-			// Find the streaming settings for the target platform
-			if (LODSettings.bOverrideLODStreamingSettings)
-			{
-				GenerationContext.bEnableLODStreaming = LODSettings.bEnableLODStreaming.GetValueForPlatform(*GenerationContext.Options.TargetPlatform->IniPlatformName());
-				GenerationContext.NumMaxLODsToStream = LODSettings.NumMaxStreamedLODs.GetValueForPlatform(*GenerationContext.Options.TargetPlatform->IniPlatformName());
-			}
-			else
-			{
-				for (int32 MeshIndex = 0; MeshIndex < GenerationContext.ComponentInfos.Num(); ++MeshIndex)
-				{
-					RefSkeletalMesh = GenerationContext.ComponentInfos[MeshIndex].RefSkeletalMesh;
-					check(RefSkeletalMesh);
-
-					GenerationContext.bEnableLODStreaming = GenerationContext.bEnableLODStreaming &&
-						RefSkeletalMesh->GetEnableLODStreaming(GenerationContext.Options.TargetPlatform);
-
-					GenerationContext.NumMaxLODsToStream = FMath::Min(static_cast<int32>(GenerationContext.NumMaxLODsToStream),
-						RefSkeletalMesh->GetMaxNumStreamedLODs(GenerationContext.Options.TargetPlatform));
-				}
-			}
-
-			GenerationContext.NumMaxLODsToStream = FMath::Clamp(GenerationContext.NumMaxLODsToStream, 0, GenerationContext.NumLODsInRoot - 1);
+			return nullptr; // Not generated in the first pass.
 		}
 		
 		mu::Ptr<mu::NodeComponentNew> NodeComponentNew = new mu::NodeComponentNew();
-		NodeComponentNew->Id = GenerationContext.ComponentNames.Num();
-		GenerationContext.ComponentNames.Add(TypedComponentMesh->ComponentName);
-	
+		NodeComponentNew->Id = GenerationContext.ComponentNames.Add(TypedComponentMesh->ComponentName);
 		NodeComponentNew->SetMessageContext(Node);
-		ObjectNode->Components.Add(NodeComponentNew);
+		
+		Result = NodeComponentNew;
 
 		GenerationContext.CurrentMeshComponent = TypedComponentMesh->ComponentName;
-		GenerateMutableSourceComponentMesh(GenerationContext, *TypedComponentMesh, NodeComponentNew, ObjectNode);
+		GenerationContext.CurrentAutoLODStrategy = TypedComponentMesh->AutoLODStrategy;
 		
-		Result = NodeComponentNew;
+		GenerateMutableSourceComponentMesh(GenerationContext, *TypedComponentMesh, NodeComponentNew);
+
+		GenerationContext.CurrentMeshComponent = {};
+		GenerationContext.CurrentAutoLODStrategy = {};
 	}
 
 	else if (const UCustomizableObjectNodeComponentMeshAddTo* TypedComponentMeshExtend = Cast<UCustomizableObjectNodeComponentMeshAddTo>(Node))
 	{
-		if (UCustomizableObjectNodeComponentMesh** FindResult = GenerationContext.MeshComponents.Find(TypedComponentMeshExtend->ParentComponentName))
+		if (FMutableComponentInfo* FindResult = GenerationContext.ComponentInfos.FindByPredicate([&](const FMutableComponentInfo& Element)
 		{
-			UCustomizableObjectNodeComponentMesh* TypedParentComponentMesh = *FindResult;
+			return Element.ComponentName == TypedComponentMeshExtend->ParentComponentName;
+		}))
+		{
+			UCustomizableObjectNodeComponentMesh* TypedParentComponentMesh = FindResult->Node;
 
 			if (TypedComponentMeshExtend->NumLODs > TypedParentComponentMesh->NumLODs)
 			{
@@ -306,13 +175,15 @@ mu::Ptr<mu::NodeComponent> GenerateMutableSourceComponent(const UEdGraphPin* Pin
 			NodeComponentEdit->SetMessageContext(TypedComponentMeshExtend);
 		
 			GenerationContext.CurrentMeshComponent = TypedParentComponentMesh->ComponentName;
-
 			GenerationContext.CurrentAutoLODStrategy = TypedComponentMeshExtend->AutoLODStrategy == ECustomizableObjectAutomaticLODStrategy::Inherited ?
 				TypedParentComponentMesh->AutoLODStrategy :
 				TypedComponentMeshExtend->AutoLODStrategy;
 			
-			GenerateMutableSourceComponentMesh(GenerationContext, *TypedComponentMeshExtend, NodeComponentEdit, nullptr);
+			GenerateMutableSourceComponentMesh(GenerationContext, *TypedComponentMeshExtend, NodeComponentEdit);
 
+			GenerationContext.CurrentMeshComponent = {};
+			GenerationContext.CurrentAutoLODStrategy = {};
+			
 			Result = NodeComponentEdit;	
 		}
 		else
@@ -535,5 +406,136 @@ mu::Ptr<mu::NodeComponent> GenerateMutableSourceComponent(const UEdGraphPin* Pin
 	return Result;
 }
 
+
+void FirstPass(UCustomizableObjectNodeComponentMesh& Node, FMutableGraphGenerationContext& GenerationContext)
+{
+	if (Node.ComponentName.IsNone())
+	{
+		GenerationContext.Log(LOCTEXT("EmptyComponentNameError", "Error! Missing name in a component of the Customizable Object."), &Node, EMessageSeverity::Error);
+		return;
+	}
+	
+	if (FMutableComponentInfo* Result = GenerationContext.ComponentInfos.FindByPredicate([&](const FMutableComponentInfo& Element)
+	{
+		return Element.ComponentName == Node.ComponentName;
+	}))
+	{
+		FText Msg = FText::Format(LOCTEXT("ComponentNodeWithSameNameExists", "Error! It already exists a Mesh Component node with the same name in Customizable Object [{0}]"), FText::FromString(GetRootObject(*Result->Node)->GetName()));
+		GenerationContext.Log(Msg, &Node, EMessageSeverity::Error);
+		return;
+	}
+	
+	USkeletalMesh* RefSkeletalMesh = Node.ReferenceSkeletalMesh;
+	if (!RefSkeletalMesh)
+	{
+		GenerationContext.Log(LOCTEXT("NoReferenceMeshObjectTab", "Error! Missing reference Skeletal Mesh"), &Node, EMessageSeverity::Error);
+		return;
+	}
+	
+	USkeleton* RefSkeleton = RefSkeletalMesh->GetSkeleton();
+	if (!RefSkeleton)
+	{
+		FText Msg = FText::Format(LOCTEXT("NoReferenceSkeleton", "Error! Missing skeleton in the reference mesh [{0}]"), FText::FromString(GenerationContext.CustomizableObjectWithCycle->GetPathName()));
+
+		GenerationContext.Log(Msg, &Node, EMessageSeverity::Error);
+		return;
+	}
+	
+	// Ensure that the CO has a valid AutoLODStrategy on the Component node.
+	if (Node.AutoLODStrategy == ECustomizableObjectAutomaticLODStrategy::Inherited)
+	{
+		GenerationContext.Log(LOCTEXT("RootInheritsFromParent", "Error! Component LOD Strategy can't be set to 'Inherit from parent object'"), &Node, EMessageSeverity::Error);
+		return;
+	}
+
+	// Fill the basic LOD Settings
+	const int32 NumLODs = Node.LODPins.Num();
+
+	if (!GenerationContext.NumLODsInRoot)
+	{
+		// NumLODsInRoot
+		int32 MaxRefMeshLODs = Node.ReferenceSkeletalMesh->GetLODNum();
+		if (MaxRefMeshLODs < NumLODs)
+		{
+			FString Msg = FString::Printf(TEXT("The object has %d LODs but the reference mesh only %d. Resulting objects will have %d LODs."),
+				NumLODs, MaxRefMeshLODs, MaxRefMeshLODs);
+			GenerationContext.Log(FText::FromString(Msg), &Node, EMessageSeverity::Warning);
+			GenerationContext.NumLODsInRoot = MaxRefMeshLODs;
+		}
+		else
+		{
+			GenerationContext.NumLODsInRoot = FMath::Max(GenerationContext.NumLODsInRoot, static_cast<uint8>(NumLODs));
+		}
+	
+		const FMutableLODSettings& LODSettings = GenerationContext.Object->LODSettings;
+
+		// Find the MinLOD available for the target platform
+		if (RefSkeletalMesh->IsMinLodQualityLevelEnable())
+		{
+			FSupportedQualityLevelArray SupportedQualityLevels = LODSettings.MinQualityLevelLOD.GetSupportedQualityLevels(*GenerationContext.Options.TargetPlatform->GetPlatformInfo().IniPlatformName.ToString());
+			
+			int32 MinValue = GenerationContext.NumLODsInRoot - 1;
+			for (int32& QL : SupportedQualityLevels)
+			{
+				// check if have data for the supported quality level or set to default.
+				if (LODSettings.MinQualityLevelLOD.IsQualityLevelValid(QL))
+				{
+					MinValue = FMath::Min(LODSettings.MinQualityLevelLOD.GetValueForQualityLevel(QL), MinValue);
+				}
+				else 
+				{
+					MinValue = LODSettings.MinQualityLevelLOD.GetDefault();
+					break;
+				}
+			}
+
+			GenerationContext.FirstLODAvailable = FMath::Max(0, MinValue);
+		}
+		else
+		{
+			GenerationContext.FirstLODAvailable = LODSettings.MinLOD.GetValueForPlatform(*GenerationContext.Options.TargetPlatform->IniPlatformName());
+		}
+
+		GenerationContext.FirstLODAvailable = FMath::Clamp(GenerationContext.FirstLODAvailable, 0, GenerationContext.NumLODsInRoot - 1);
+
+		// Find the streaming settings for the target platform
+		if (LODSettings.bOverrideLODStreamingSettings)
+		{
+			GenerationContext.bEnableLODStreaming = LODSettings.bEnableLODStreaming.GetValueForPlatform(*GenerationContext.Options.TargetPlatform->IniPlatformName());
+			GenerationContext.NumMaxLODsToStream = LODSettings.NumMaxStreamedLODs.GetValueForPlatform(*GenerationContext.Options.TargetPlatform->IniPlatformName());
+		}
+		else
+		{
+			for (int32 MeshIndex = 0; MeshIndex < GenerationContext.ComponentInfos.Num(); ++MeshIndex)
+			{
+				RefSkeletalMesh = GenerationContext.ComponentInfos[MeshIndex].RefSkeletalMesh;
+				check(RefSkeletalMesh);
+
+				GenerationContext.bEnableLODStreaming = GenerationContext.bEnableLODStreaming &&
+					RefSkeletalMesh->GetEnableLODStreaming(GenerationContext.Options.TargetPlatform);
+
+				GenerationContext.NumMaxLODsToStream = FMath::Min(static_cast<int32>(GenerationContext.NumMaxLODsToStream),
+					RefSkeletalMesh->GetMaxNumStreamedLODs(GenerationContext.Options.TargetPlatform));
+			}
+		}
+
+		GenerationContext.NumMaxLODsToStream = FMath::Clamp(GenerationContext.NumMaxLODsToStream, 0, GenerationContext.NumLODsInRoot - 1);
+	}
+	
+	// Add a new entry to the list of Component Infos
+	FMutableComponentInfo ComponentInfo(Node.ComponentName, RefSkeletalMesh);
+	ComponentInfo.Node = &Node;
+	ComponentInfo.AccumulateBonesToRemovePerLOD(Node.LODReductionSettings, Node.NumLODs);
+
+	GenerationContext.ComponentInfos.Add(ComponentInfo);
+
+	// Make sure the Skeleton from the reference mesh is added to the list of referenced Skeletons.
+	GenerationContext.ReferencedSkeletons.Add(RefSkeleton);
+
+	// Add reference meshes to the participating objects
+	GenerationContext.AddParticipatingObject(*RefSkeletalMesh);
+}
+
+
 #undef LOCTEXT_NAMESPACE
 
diff --git a/Engine/Plugins/Mutable/Source/CustomizableObjectEditor/Private/MuCOE/GenerateMutableSource/GenerateMutableSourceComponent.h b/Engine/Plugins/Mutable/Source/CustomizableObjectEditor/Private/MuCOE/GenerateMutableSource/GenerateMutableSourceComponent.h
--- a/Engine/Plugins/Mutable/Source/CustomizableObjectEditor/Private/MuCOE/GenerateMutableSource/GenerateMutableSourceComponent.h
+++ b/Engine/Plugins/Mutable/Source/CustomizableObjectEditor/Private/MuCOE/GenerateMutableSource/GenerateMutableSourceComponent.h
@@ -2,6 +2,7 @@
 
 #pragma once
 
+#include "MuCOE/Nodes/CustomizableObjectNodeComponentMesh.h"
 #include "MuR/Ptr.h"
 
 class UEdGraphPin;
@@ -15,3 +16,6 @@ namespace mu
 
 mu::Ptr<mu::NodeComponent> GenerateMutableSourceComponent(const UEdGraphPin* Pin, FMutableGraphGenerationContext& GenerationContext);
 
+
+void FirstPass(UCustomizableObjectNodeComponentMesh& Node, FMutableGraphGenerationContext& GenerationContext);
+
diff --git a/Engine/Plugins/Mutable/Source/CustomizableObjectEditor/Private/MuCOE/GenerateMutableSource/GenerateMutableSourceFloat.cpp b/Engine/Plugins/Mutable/Source/CustomizableObjectEditor/Private/MuCOE/GenerateMutableSource/GenerateMutableSourceFloat.cpp
--- a/Engine/Plugins/Mutable/Source/CustomizableObjectEditor/Private/MuCOE/GenerateMutableSource/GenerateMutableSourceFloat.cpp
+++ b/Engine/Plugins/Mutable/Source/CustomizableObjectEditor/Private/MuCOE/GenerateMutableSource/GenerateMutableSourceFloat.cpp
@@ -68,8 +68,6 @@ mu::NodeScalarPtr GenerateMutableSourceFloat(const UEdGraphPin* Pin, FMutableGra
 		mu::NodeScalarParameterPtr ScalarNode = new mu::NodeScalarParameter();
 		Result = ScalarNode;
 
-		GenerationContext.AddParameterNameUnique(Node, FloatParameterNode->ParameterName);
-
 		ScalarNode->SetName(FloatParameterNode->ParameterName);
 		ScalarNode->SetUid(GenerationContext.GetNodeIdUnique(Node).ToString());
 		ScalarNode->SetDefaultValue(FloatParameterNode->DefaultValue);
@@ -87,8 +85,6 @@ mu::NodeScalarPtr GenerateMutableSourceFloat(const UEdGraphPin* Pin, FMutableGra
 
 		int32 DefaultValue = FMath::Clamp(EnumParamNode->DefaultIndex, 0, NumSelectors - 1);
 
-		GenerationContext.AddParameterNameUnique(Node, EnumParamNode->ParameterName);
-
 		EnumParameterNode->SetName(EnumParamNode->ParameterName);
 		EnumParameterNode->SetUid(GenerationContext.GetNodeIdUnique(Node).ToString());
 		EnumParameterNode->SetValueCount(NumSelectors);
diff --git a/Engine/Plugins/Mutable/Source/CustomizableObjectEditor/Private/MuCOE/GenerateMutableSource/GenerateMutableSourceGroupProjector.cpp b/Engine/Plugins/Mutable/Source/CustomizableObjectEditor/Private/MuCOE/GenerateMutableSource/GenerateMutableSourceGroupProjector.cpp
--- a/Engine/Plugins/Mutable/Source/CustomizableObjectEditor/Private/MuCOE/GenerateMutableSource/GenerateMutableSourceGroupProjector.cpp
+++ b/Engine/Plugins/Mutable/Source/CustomizableObjectEditor/Private/MuCOE/GenerateMutableSource/GenerateMutableSourceGroupProjector.cpp
@@ -32,13 +32,13 @@ class UPoseAsset;
 #define LOCTEXT_NAMESPACE "CustomizableObjectEditor"
 
 
-mu::NodeImagePtr GenerateMutableGroupProjection(const int32 NodeLOD, const int32 ImageIndex, mu::NodeMeshPtr MeshNode, FMutableGraphGenerationContext& GenerationContext,
+mu::NodeImagePtr GenerateMutableSourceGroupProjector(const int32 NodeLOD, const int32 ImageIndex, mu::NodeMeshPtr MeshNode, FMutableGraphGenerationContext& GenerationContext,
 	UCustomizableObjectNodeMaterialBase* TypedNodeMat, UCustomizableObjectNodeModifierExtendMeshSection* TypedNodeExt, bool& bShareProjectionTexturesBetweenLODs, bool& bIsGroupProjectorImage,
 	UTexture2D*& GroupProjectionReferenceTexture, TMap<FString, float>& TextureNameToProjectionResFactor, FString& AlternateResStateName)
 {
 	check(static_cast<bool>(TypedNodeMat) != static_cast<bool>(TypedNodeExt)); // XOr. TypedNodeMat valid or TypedNodeExt valid. At least one valid.
 
-	if (!GenerationContext.ProjectorGroupMap.Num() || !MeshNode.get())
+	if (!MeshNode.get())
 	{
 		return mu::NodeImagePtr();
 	}
@@ -47,156 +47,165 @@ mu::NodeImagePtr GenerateMutableGroupProjection(const int32 NodeLOD, const int32
 	TArray<FGroupProjectorTempData> ImageNodes_ProjectorTempData;
 
 	int32 TextureSize = 512;
-
-	for (TPair<const UCustomizableObjectNodeObjectGroup*, FGroupProjectorTempData>& Elem : GenerationContext.ProjectorGroupMap)
+	
+	for (TArray<UCustomizableObjectNodeGroupProjectorParameter*> GroupProjectors : GenerationContext.CurrentGroupProjectors)
 	{
-		FGroupProjectorTempData& ProjectorTempData = Elem.Value;
-
-		const int32& DropProjectionTextureAtLOD = ProjectorTempData.CustomizableObjectNodeGroupProjectorParameter->DropProjectionTextureAtLOD;
-		if (DropProjectionTextureAtLOD >= 0 && NodeLOD >= DropProjectionTextureAtLOD)
-		{
-			continue;
-		}
-
-		bShareProjectionTexturesBetweenLODs |= ProjectorTempData.CustomizableObjectNodeGroupProjectorParameter->bShareProjectionTexturesBetweenLODs;
-
-		if (!GroupProjectionReferenceTexture)
+		for (UCustomizableObjectNodeGroupProjectorParameter* ProjParamNode : GroupProjectors)
 		{
-			GroupProjectionReferenceTexture = ProjectorTempData.CustomizableObjectNodeGroupProjectorParameter->ReferenceTexture;
-			
-			if (GroupProjectionReferenceTexture)
+			TOptional<FGroupProjectorTempData> OptionalProjectorTempData = GenerateMutableGroupProjector(ProjParamNode, GenerationContext);
+			if (!OptionalProjectorTempData.IsSet())
 			{
-				GenerationContext.AddParticipatingObject(*GroupProjectionReferenceTexture);
+				continue;
 			}
-		}
 
-		const bool bProjectToImage = [&]
-		{
-			FString ParameterName;
-			
-			if (TypedNodeMat)
+			FGroupProjectorTempData& ProjectorTempData = OptionalProjectorTempData.GetValue();
+				
+			const int32& DropProjectionTextureAtLOD = ProjectorTempData.CustomizableObjectNodeGroupProjectorParameter->DropProjectionTextureAtLOD;
+			if (DropProjectionTextureAtLOD >= 0 && NodeLOD >= DropProjectionTextureAtLOD)
 			{
-				ParameterName = TypedNodeMat->GetParameterName(EMaterialParameterType::Texture, ImageIndex).ToString();
+				continue;
 			}
-			else
+
+			bShareProjectionTexturesBetweenLODs |= ProjectorTempData.CustomizableObjectNodeGroupProjectorParameter->bShareProjectionTexturesBetweenLODs;
+
+			if (!GroupProjectionReferenceTexture)
 			{
-				const FNodeMaterialParameterId ImageId = TypedNodeExt->GetParameterId(EMaterialParameterType::Texture, ImageIndex);
-				return TypedNodeExt->UsesImage(ImageId);
+				GroupProjectionReferenceTexture = ProjectorTempData.CustomizableObjectNodeGroupProjectorParameter->ReferenceTexture;
+				
+				if (GroupProjectionReferenceTexture)
+				{
+					GenerationContext.AddParticipatingObject(*GroupProjectionReferenceTexture);
+				}
 			}
 
-			return ParameterName == ProjectorTempData.CustomizableObjectNodeGroupProjectorParameter->MaterialChannelNameToConnect;
-		}();
-		
-		if (bProjectToImage)
-		{
-			const bool bWarningReplacedImage = [&]
+			const bool bProjectToImage = [&]
 			{
+				FString ParameterName;
+				
 				if (TypedNodeMat)
 				{
-					return TypedNodeMat->IsImageMutableMode(ImageIndex);
+					ParameterName = TypedNodeMat->GetParameterName(EMaterialParameterType::Texture, ImageIndex).ToString();
 				}
 				else
 				{
 					const FNodeMaterialParameterId ImageId = TypedNodeExt->GetParameterId(EMaterialParameterType::Texture, ImageIndex);
 					return TypedNodeExt->UsesImage(ImageId);
 				}
+
+				return ParameterName == ProjectorTempData.CustomizableObjectNodeGroupProjectorParameter->MaterialChannelNameToConnect;
 			}();
 			
-			if (bWarningReplacedImage)
+			if (bProjectToImage)
 			{
-				const FString ImageName = [&]
+				const bool bWarningReplacedImage = [&]
 				{
 					if (TypedNodeMat)
 					{
-						return TypedNodeMat->GetParameterName(EMaterialParameterType::Texture, ImageIndex).ToString();
+						return TypedNodeMat->IsImageMutableMode(ImageIndex);
 					}
 					else
 					{
-						return TypedNodeExt->GetParameterName(EMaterialParameterType::Texture, ImageIndex).ToString();
+						const FNodeMaterialParameterId ImageId = TypedNodeExt->GetParameterId(EMaterialParameterType::Texture, ImageIndex);
+						return TypedNodeExt->UsesImage(ImageId);
 					}
 				}();
 				
-				FString msg = FString::Printf(TEXT("Material image [%s] is connected to an image but will be replaced by a Group Projector."), *ImageName);
-				GenerationContext.Log(FText::FromString(msg), TypedNodeMat);
-				continue;
-			}
-			
-			mu::Ptr<mu::NodeImageProject> ImageNode = new mu::NodeImageProject();
-			bIsGroupProjectorImage = true;
-			ImageNode->SetLayout(ProjectorTempData.CustomizableObjectNodeGroupProjectorParameter->UVLayout);
-
-			{
-				if (TypedNodeMat &&
-					!ProjectorTempData.CustomizableObjectNodeGroupProjectorParameter->AlternateProjectionResolutionStateName.IsEmpty()
-					&& ProjectorTempData.CustomizableObjectNodeGroupProjectorParameter->AlternateProjectionResolutionFactor > 0)
+				if (bWarningReplacedImage)
 				{
-					TextureNameToProjectionResFactor.Add(
-						TypedNodeMat->GetParameterName(EMaterialParameterType::Texture, ImageIndex).ToString(),
-						ProjectorTempData.CustomizableObjectNodeGroupProjectorParameter->AlternateProjectionResolutionFactor);
+					const FString ImageName = [&]
+					{
+						if (TypedNodeMat)
+						{
+							return TypedNodeMat->GetParameterName(EMaterialParameterType::Texture, ImageIndex).ToString();
+						}
+						else
+						{
+							return TypedNodeExt->GetParameterName(EMaterialParameterType::Texture, ImageIndex).ToString();
+						}
+					}();
+					
+					FString msg = FString::Printf(TEXT("Material image [%s] is connected to an image but will be replaced by a Group Projector."), *ImageName);
+					GenerationContext.Log(FText::FromString(msg), TypedNodeMat);
+					continue;
+				}
+				
+				mu::Ptr<mu::NodeImageProject> ImageNode = new mu::NodeImageProject();
+				bIsGroupProjectorImage = true;
+				ImageNode->SetLayout(ProjectorTempData.CustomizableObjectNodeGroupProjectorParameter->UVLayout);
 
-					if (!AlternateResStateName.IsEmpty() &&
-						AlternateResStateName != ProjectorTempData.CustomizableObjectNodeGroupProjectorParameter->AlternateProjectionResolutionStateName &&
-						!ProjectorTempData.bAlternateResStateNameWarningDisplayed)
+				{
+					if (TypedNodeMat &&
+						!ProjectorTempData.CustomizableObjectNodeGroupProjectorParameter->AlternateProjectionResolutionStateName.IsEmpty()
+						&& ProjectorTempData.CustomizableObjectNodeGroupProjectorParameter->AlternateProjectionResolutionFactor > 0)
 					{
-						FString msg = FString::Printf(TEXT("All 'Alternate Projection Resolution State Name' properties in Group Projector Parameter nodes connected to same Group node must have the same value or be blank. Only the value of the last connected node will be used."));
-						GenerationContext.Log(FText::FromString(msg), ProjectorTempData.CustomizableObjectNodeGroupProjectorParameter);
-						ProjectorTempData.bAlternateResStateNameWarningDisplayed = true;
+						TextureNameToProjectionResFactor.Add(
+							TypedNodeMat->GetParameterName(EMaterialParameterType::Texture, ImageIndex).ToString(),
+							ProjectorTempData.CustomizableObjectNodeGroupProjectorParameter->AlternateProjectionResolutionFactor);
+
+						if (!AlternateResStateName.IsEmpty() &&
+							AlternateResStateName != ProjectorTempData.CustomizableObjectNodeGroupProjectorParameter->AlternateProjectionResolutionStateName &&
+							!ProjectorTempData.bAlternateResStateNameWarningDisplayed)
+						{
+							FString msg = FString::Printf(TEXT("All 'Alternate Projection Resolution State Name' properties in Group Projector Parameter nodes connected to same Group node must have the same value or be blank. Only the value of the last connected node will be used."));
+							GenerationContext.Log(FText::FromString(msg), ProjectorTempData.CustomizableObjectNodeGroupProjectorParameter);
+							ProjectorTempData.bAlternateResStateNameWarningDisplayed = true;
+						}
+
+						AlternateResStateName = ProjectorTempData.CustomizableObjectNodeGroupProjectorParameter->AlternateProjectionResolutionStateName;
 					}
-
-					AlternateResStateName = ProjectorTempData.CustomizableObjectNodeGroupProjectorParameter->AlternateProjectionResolutionStateName;
 				}
-			}
 
-			{
-				mu::NodeScalarConstantPtr ScalarNode = new mu::NodeScalarConstant;
-				ScalarNode->SetValue(120.f);
-				ImageNode->SetAngleFadeStart(ScalarNode);
-			}
+				{
+					mu::NodeScalarConstantPtr ScalarNode = new mu::NodeScalarConstant;
+					ScalarNode->SetValue(120.f);
+					ImageNode->SetAngleFadeStart(ScalarNode);
+				}
 
-			{
-				mu::NodeScalarConstantPtr ScalarNode = new mu::NodeScalarConstant;
-				ScalarNode->SetValue(150.f);
-				ImageNode->SetAngleFadeEnd(ScalarNode);
-			}
+				{
+					mu::NodeScalarConstantPtr ScalarNode = new mu::NodeScalarConstant;
+					ScalarNode->SetValue(150.f);
+					ImageNode->SetAngleFadeEnd(ScalarNode);
+				}
 
-			mu::NodeMeshSwitchPtr MeshSwitchNode = new mu::NodeMeshSwitch;
-			MeshSwitchNode->SetParameter(ProjectorTempData.PoseOptionsParameter);
-			MeshSwitchNode->SetOptionCount(ProjectorTempData.CustomizableObjectNodeGroupProjectorParameter->OptionPoses.Num() + 1);
-			MeshSwitchNode->SetOption(0, MeshNode);
+				mu::NodeMeshSwitchPtr MeshSwitchNode = new mu::NodeMeshSwitch;
+				MeshSwitchNode->SetParameter(ProjectorTempData.PoseOptionsParameter);
+				MeshSwitchNode->SetOptionCount(ProjectorTempData.CustomizableObjectNodeGroupProjectorParameter->OptionPoses.Num() + 1);
+				MeshSwitchNode->SetOption(0, MeshNode);
 
-			for (int32 SelectorIndex = 0; SelectorIndex < ProjectorTempData.CustomizableObjectNodeGroupProjectorParameter->OptionPoses.Num(); ++SelectorIndex)
-			{
-				if (ProjectorTempData.CustomizableObjectNodeGroupProjectorParameter->OptionPoses[SelectorIndex].OptionPose)
+				for (int32 SelectorIndex = 0; SelectorIndex < ProjectorTempData.CustomizableObjectNodeGroupProjectorParameter->OptionPoses.Num(); ++SelectorIndex)
 				{
-					mu::NodeMeshApplyPosePtr NodeMeshApplyPose = CreateNodeMeshApplyPose(GenerationContext, MeshNode,
-						ProjectorTempData.PoseBoneDataArray[SelectorIndex].ArrayBoneName, ProjectorTempData.PoseBoneDataArray[SelectorIndex].ArrayTransform);
+					if (ProjectorTempData.CustomizableObjectNodeGroupProjectorParameter->OptionPoses[SelectorIndex].OptionPose)
+					{
+						mu::NodeMeshApplyPosePtr NodeMeshApplyPose = CreateNodeMeshApplyPose(GenerationContext, MeshNode,
+							ProjectorTempData.PoseBoneDataArray[SelectorIndex].ArrayBoneName, ProjectorTempData.PoseBoneDataArray[SelectorIndex].ArrayTransform);
+
+						if (!NodeMeshApplyPose)
+						{
+							FString msg = FString::Printf(TEXT("Couldn't get bone transform information from a Pose Asset."));
+							GenerationContext.Log(FText::FromString(msg), TypedNodeMat);
+						}
 
-					if (!NodeMeshApplyPose)
+						MeshSwitchNode->SetOption(SelectorIndex + 1, NodeMeshApplyPose);
+					}
+					else
 					{
-						FString msg = FString::Printf(TEXT("Couldn't get bone transform information from a Pose Asset."));
-						GenerationContext.Log(FText::FromString(msg), TypedNodeMat);
+						MeshSwitchNode->SetOption(SelectorIndex + 1, MeshNode);
 					}
-
-					MeshSwitchNode->SetOption(SelectorIndex + 1, NodeMeshApplyPose);
 				}
-				else
-				{
-					MeshSwitchNode->SetOption(SelectorIndex + 1, MeshNode);
-				}
-			}
 
-			ImageNode->SetMesh(MeshSwitchNode);
-			ImageNode->SetProjector(ProjectorTempData.NodeProjectorParameterPtr);
-			ImageNode->SetImage(ProjectorTempData.NodeImagePtr);
+				ImageNode->SetMesh(MeshSwitchNode);
+				ImageNode->SetProjector(ProjectorTempData.NodeProjectorParameterPtr);
+				ImageNode->SetImage(ProjectorTempData.NodeImagePtr);
 
-			TextureSize = ProjectorTempData.TextureSize;
-			ImageNode->SetImageSize( FUintVector2(TextureSize, TextureSize) );
+				TextureSize = ProjectorTempData.TextureSize;
+				ImageNode->SetImageSize( FUintVector2(TextureSize, TextureSize) );
 
-			ImageNodes.Add(ImageNode);
-			ImageNodes_ProjectorTempData.Add(ProjectorTempData);
+				ImageNodes.Add(ImageNode);
+				ImageNodes_ProjectorTempData.Add(ProjectorTempData);
+			}
 		}
 	}
-
+	
 	if (ImageNodes.Num() == 0)
 	{
 		return mu::NodeImagePtr();
@@ -310,239 +319,246 @@ mu::NodeImagePtr GenerateMutableGroupProjection(const int32 NodeLOD, const int32
 }
 
 
-bool GenerateMutableSourceGroupProjector(const UEdGraphPin* Pin, FMutableGraphGenerationContext& GenerationContext, const UCustomizableObjectNodeObjectGroup* originalGroup)
+TOptional<FGroupProjectorTempData> GenerateMutableGroupProjector(UCustomizableObjectNodeGroupProjectorParameter* ProjParamNode, FMutableGraphGenerationContext& GenerationContext)
 {
-	check(Pin)
-	check(originalGroup);
-	RETURN_ON_CYCLE(*Pin, GenerationContext)
-
-	CheckNumOutputs(*Pin, GenerationContext);
+	FGeneratedGroupProjectorsKey Key;
+	Key.Node = ProjParamNode;
+	Key.CurrentComponent = GenerationContext.CurrentMeshComponent;
 	
-	UCustomizableObjectNode* Node = CastChecked<UCustomizableObjectNode>(Pin->GetOwningNode());
-
-	const FGeneratedKey Key(reinterpret_cast<void*>(&GenerateMutableSourceGroupProjector), *Pin, *Node, GenerationContext, true);
-	if (const FGeneratedData* Generated = GenerationContext.Generated.Find(Key))
+	if (FGroupProjectorTempData* Result = GenerationContext.GeneratedGroupProjectors.Find(Key))
 	{
-		return true;
+		return *Result;
 	}
 
-	mu::NodeProjectorPtr Result;
+	FGroupProjectorTempData GroupProjectorTempData;
 	
-	if (UCustomizableObjectNodeGroupProjectorParameter* ProjParamNode = Cast<UCustomizableObjectNodeGroupProjectorParameter>(Node))
-	{
-		FGroupProjectorTempData GroupProjectorTempData;
-		// The static cast works here because it's already known to be a mu::NodeProjectorParameter* because of the UE5 Cast in the previous line
-		GroupProjectorTempData.NodeProjectorParameterPtr = static_cast<mu::NodeProjectorParameter*>(GenerateMutableSourceProjector(Pin, GenerationContext).get());
+	// The static cast works here because it's already known to be a mu::NodeProjectorParameter* because of the UE5 Cast in the previous line
+	GroupProjectorTempData.NodeProjectorParameterPtr = static_cast<mu::NodeProjectorParameter*>(GenerateMutableSourceProjector(&ProjParamNode->OutputPin(), GenerationContext).get());
 
-		if (GroupProjectorTempData.NodeProjectorParameterPtr)
+	if (GroupProjectorTempData.NodeProjectorParameterPtr)
+	{
+		// Use the projector parameter uid + num to identify parameters derived from this node
+		FGuid NumLayersParamUid = ProjParamNode->NodeGuid;
+		NumLayersParamUid.D += 1;
+		FGuid SelectedPoseParamUid = ProjParamNode->NodeGuid;
+		SelectedPoseParamUid.D += 2;
+		FGuid OpacityParamUid = ProjParamNode->NodeGuid;
+		OpacityParamUid.D += 3;
+		FGuid SelectedImageParamUid = ProjParamNode->NodeGuid;
+		SelectedImageParamUid.D += 4;
+
+		// Add to UCustomizableObjectNodeGroupProjectorParameter::OptionTextures those textures that are present in
+		// UCustomizableObjectNodeGroupProjectorParameter::OptionTexturesDataTable avoiding any repeated element
+		TArray<FGroupProjectorParameterImage> ArrayOptionTexture = ProjParamNode->GetFinalOptionTexturesNoRepeat();
+
+		if ((ProjParamNode->OptionTexturesDataTable != nullptr) &&
+			(ProjParamNode->DataTableTextureColumnName.ToString().IsEmpty() || (ProjParamNode->DataTableTextureColumnName.ToString() == "None")))
 		{
-			// Use the projector parameter uid + num to identify parameters derived from this node
-			FGuid NumLayersParamUid = Node->NodeGuid;
-			NumLayersParamUid.D += 1;
-			FGuid SelectedPoseParamUid = Node->NodeGuid;
-			SelectedPoseParamUid.D += 2;
-			FGuid OpacityParamUid = Node->NodeGuid;
-			OpacityParamUid.D += 3;
-			FGuid SelectedImageParamUid = Node->NodeGuid;
-			SelectedImageParamUid.D += 4;
-
-			// Add to UCustomizableObjectNodeGroupProjectorParameter::OptionTextures those textures that are present in
-			// UCustomizableObjectNodeGroupProjectorParameter::OptionTexturesDataTable avoiding any repeated element
-			TArray<FGroupProjectorParameterImage> ArrayOptionTexture = ProjParamNode->GetFinalOptionTexturesNoRepeat();
-
-			if ((ProjParamNode->OptionTexturesDataTable != nullptr) &&
-				(ProjParamNode->DataTableTextureColumnName.ToString().IsEmpty() || (ProjParamNode->DataTableTextureColumnName.ToString() == "None")))
-			{
-				FString msg = FString::Printf(TEXT("The group projection node has a table assigned to the Option Images Data Table property, but no column to read textures is specified at the Data Table Texture Column Name property."));
-				GenerationContext.Log(FText::FromString(msg), ProjParamNode, EMessageSeverity::Error, true);
-			}
+			FString msg = FString::Printf(TEXT("The group projection node has a table assigned to the Option Images Data Table property, but no column to read textures is specified at the Data Table Texture Column Name property."));
+			GenerationContext.Log(FText::FromString(msg), ProjParamNode, EMessageSeverity::Error, true);
+		}
 
-			GroupProjectorTempData.CustomizableObjectNodeGroupProjectorParameter = ProjParamNode;
+		GroupProjectorTempData.CustomizableObjectNodeGroupProjectorParameter = ProjParamNode;
 
-			mu::NodeScalarParameterPtr NodeScalarParam = new mu::NodeScalarParameter;
+		mu::NodeScalarParameterPtr NodeScalarParam;
+		if (mu::Ptr<mu::NodeScalarParameter>* FindResult = GenerationContext.GeneratedScalarParameters.Find(NumLayersParamUid.ToString()))
+		{
+			NodeScalarParam = *FindResult;
+		}
+		else
+		{
+			NodeScalarParam = new mu::NodeScalarParameter;
 			FString NodeScalarParamName = ProjParamNode->ParameterName + NUM_LAYERS_PARAMETER_POSTFIX;
 			NodeScalarParam->SetName(NodeScalarParamName);
 			NodeScalarParam->SetUid(NumLayersParamUid.ToString());
-			GenerationContext.AddParameterNameUnique(originalGroup, NodeScalarParamName);
 			NodeScalarParam->SetDefaultValue(0.f);
 
 			GenerationContext.ParameterUIDataMap.Add(NodeScalarParamName,
-				FMutableParameterData(ProjParamNode->ParamUIMetadata, EMutableParameterType::Int));
-
-			mu::Ptr<mu::NodeRangeFromScalar> NodeRangeFromScalar = new mu::NodeRangeFromScalar;
-			NodeRangeFromScalar->SetSize(NodeScalarParam);
-			GroupProjectorTempData.NodeRange = NodeRangeFromScalar;
-			GroupProjectorTempData.NodeProjectorParameterPtr->SetRangeCount(1);
-			GroupProjectorTempData.NodeProjectorParameterPtr->SetRange(0, NodeRangeFromScalar);
+		FMutableParameterData(ProjParamNode->ParamUIMetadata, EMutableParameterType::Int));
+					
+			GenerationContext.GeneratedScalarParameters.Add(NumLayersParamUid.ToString(), NodeScalarParam);
+		}
 
-			
-			mu::NodeScalarEnumParameterPtr PoseEnumParameterNode = new mu::NodeScalarEnumParameter;
+		mu::Ptr<mu::NodeRangeFromScalar> NodeRangeFromScalar = new mu::NodeRangeFromScalar;
+		NodeRangeFromScalar->SetSize(NodeScalarParam);
+		GroupProjectorTempData.NodeRange = NodeRangeFromScalar;
+		GroupProjectorTempData.NodeProjectorParameterPtr->SetRangeCount(1);
+		GroupProjectorTempData.NodeProjectorParameterPtr->SetRange(0, NodeRangeFromScalar);
+		
+		mu::NodeScalarEnumParameterPtr PoseEnumParameterNode;
+		if (mu::Ptr<mu::NodeScalarEnumParameter>* FindResult = GenerationContext.GeneratedEnumParameters.Find(SelectedPoseParamUid.ToString()))
+		{
+			PoseEnumParameterNode = *FindResult;
+		}
+		else
+		{
+			PoseEnumParameterNode = new mu::NodeScalarEnumParameter;
 			FString PoseNodeEnumParamName = ProjParamNode->ParameterName + POSE_PARAMETER_POSTFIX;
 			PoseEnumParameterNode->SetName(PoseNodeEnumParamName);
 			PoseEnumParameterNode->SetUid(SelectedPoseParamUid.ToString());
-			GenerationContext.AddParameterNameUnique(originalGroup, PoseNodeEnumParamName);
 			PoseEnumParameterNode->SetValueCount(ProjParamNode->OptionPoses.Num() + 1);
 			PoseEnumParameterNode->SetDefaultValueIndex(0);
 			GroupProjectorTempData.PoseOptionsParameter = PoseEnumParameterNode;
 
-			
+			GenerationContext.GeneratedEnumParameters.Add(SelectedPoseParamUid.ToString(), PoseEnumParameterNode);
+
 			GenerationContext.ParameterUIDataMap.Add(PoseNodeEnumParamName,
-				FMutableParameterData(ProjParamNode->ParamUIMetadata, EMutableParameterType::Int));
-			
-			mu::NodeScalarParameterPtr OpacityParameterNode = new mu::NodeScalarParameter;
+		FMutableParameterData(ProjParamNode->ParamUIMetadata, EMutableParameterType::Int));
+		}
+		
+		mu::NodeScalarParameterPtr OpacityParameterNode;
+		if (mu::Ptr<mu::NodeScalarParameter>* FindResult = GenerationContext.GeneratedScalarParameters.Find(OpacityParamUid.ToString()))
+		{
+			OpacityParameterNode = *FindResult;
+		}
+		else
+		{
+			OpacityParameterNode = new mu::NodeScalarParameter;
 			FString OpacityParameterNodeName = ProjParamNode->ParameterName + OPACITY_PARAMETER_POSTFIX;
 			OpacityParameterNode->SetName(OpacityParameterNodeName);
 			OpacityParameterNode->SetUid(OpacityParamUid.ToString());
-			GenerationContext.AddParameterNameUnique(originalGroup, OpacityParameterNodeName);
 			OpacityParameterNode->SetDefaultValue(0.75f);
 			OpacityParameterNode->SetRangeCount(1);
 			OpacityParameterNode->SetRange(0, NodeRangeFromScalar);
+
+			GenerationContext.GeneratedScalarParameters.Add(OpacityParamUid.ToString(), OpacityParameterNode);
+					
 			FMutableParamUIMetadata OpacityMetadata = ProjParamNode->ParamUIMetadata;
 			OpacityMetadata.ObjectFriendlyName = FString("Opacity");
 			GroupProjectorTempData.NodeOpacityParameter = OpacityParameterNode;
+					
+			FMutableParameterData ParameterUIData;
+			ParameterUIData.ParamUIMetadata = OpacityMetadata;
+			ParameterUIData.Type = EMutableParameterType::Float;
+				
+			GenerationContext.ParameterUIDataMap.Add(OpacityParameterNodeName, ParameterUIData);
+		}
+		
+		if (ArrayOptionTexture.Num() == 0)
+		{
+			FString msg = FString::Printf(TEXT("The group projection node must have at least one option image connected to a texture or at least one valid element in Option Images Data Table."));
+			GenerationContext.Log(FText::FromString(msg), ProjParamNode, EMessageSeverity::Error, true);
+			return {};
+		}
 
-			{
-				FMutableParameterData ParameterUIData;
-				ParameterUIData.ParamUIMetadata = OpacityMetadata;
-				ParameterUIData.Type = EMutableParameterType::Float;
-			
-				GenerationContext.ParameterUIDataMap.Add(OpacityParameterNodeName, ParameterUIData);
-			}
-			
-			if (ArrayOptionTexture.Num() == 0)
-			{
-				FString msg = FString::Printf(TEXT("The group projection node must have at least one option image connected to a texture or at least one valid element in Option Images Data Table."));
-				GenerationContext.Log(FText::FromString(msg), ProjParamNode, EMessageSeverity::Error, true);
-				return false;
-			}
-			
-			if (FMutableComponentInfo* CurrentComponentInfo = GenerationContext.GetCurrentComponentInfo())
-			{
-				// Poses will only affect component 0 of the CO,
-				// TODO UE-206803
-				FName OldCurrentMeshComponent = GenerationContext.CurrentMeshComponent;
-				GenerationContext.CurrentMeshComponent = GenerationContext.ComponentInfos[0].ComponentName;
-
-				PoseEnumParameterNode->SetValue(0, 0.f, "Default pose");
+		FMutableComponentInfo* CurrentComponentInfo = GenerationContext.GetCurrentComponentInfo();
+		check(CurrentComponentInfo);
+		
+		PoseEnumParameterNode->SetValue(0, 0.f, "Default pose");
 
-				for (int32 PoseIndex = 0; PoseIndex < ProjParamNode->OptionPoses.Num(); ++PoseIndex)
-				{
-					PoseEnumParameterNode->SetValue(PoseIndex + 1, (float)PoseIndex + 1.f, ProjParamNode->OptionPoses[PoseIndex].PoseName);
+		for (int32 PoseIndex = 0; PoseIndex < ProjParamNode->OptionPoses.Num(); ++PoseIndex)
+		{
+			PoseEnumParameterNode->SetValue(PoseIndex + 1, (float)PoseIndex + 1.f, ProjParamNode->OptionPoses[PoseIndex].PoseName);
 
-					TArray<FString> ArrayBoneName;
-					TArray<FTransform> ArrayTransform;
-					UPoseAsset* PoseAsset = ProjParamNode->OptionPoses[PoseIndex].OptionPose;
+			TArray<FString> ArrayBoneName;
+			TArray<FTransform> ArrayTransform;
+			UPoseAsset* PoseAsset = ProjParamNode->OptionPoses[PoseIndex].OptionPose;
 
-					if (PoseAsset == nullptr) // Check if the slot has a selected pose. Could be left empty by the user
-					{
-						FString msg = FString::Printf(TEXT("The group projection node must have a pose assigned on each Option Poses element."));
-						GenerationContext.Log(FText::FromString(msg), ProjParamNode, EMessageSeverity::Error, true);
-						return false;
-					}
+			if (PoseAsset == nullptr) // Check if the slot has a selected pose. Could be left empty by the user
+			{
+				FString msg = FString::Printf(TEXT("The group projection node must have a pose assigned on each Option Poses element."));
+				GenerationContext.Log(FText::FromString(msg), ProjParamNode, EMessageSeverity::Error, true);
+				return {};
+			}
 
-					check(GroupProjectorTempData.PoseBoneDataArray.Num() == PoseIndex);
-					GroupProjectorTempData.PoseBoneDataArray.AddDefaulted(1);
-					UCustomizableObjectNodeAnimationPose::StaticRetrievePoseInformation(PoseAsset, CurrentComponentInfo->RefSkeletalMesh,
-						GroupProjectorTempData.PoseBoneDataArray[PoseIndex].ArrayBoneName, GroupProjectorTempData.PoseBoneDataArray[PoseIndex].ArrayTransform);
-				}
+			check(GroupProjectorTempData.PoseBoneDataArray.Num() == PoseIndex);
+			GroupProjectorTempData.PoseBoneDataArray.AddDefaulted(1);
+			UCustomizableObjectNodeAnimationPose::StaticRetrievePoseInformation(PoseAsset, CurrentComponentInfo->RefSkeletalMesh,
+				GroupProjectorTempData.PoseBoneDataArray[PoseIndex].ArrayBoneName, GroupProjectorTempData.PoseBoneDataArray[PoseIndex].ArrayTransform);
+		}
+	
+		mu::NodeScalarEnumParameterPtr EnumParameterNode;
+		if (mu::Ptr<mu::NodeScalarEnumParameter>* FindResult = GenerationContext.GeneratedEnumParameters.Find(SelectedImageParamUid.ToString()))
+		{
+			EnumParameterNode = *FindResult;
+		}
+		else
+		{
+			EnumParameterNode = new mu::NodeScalarEnumParameter;
 
-				GenerationContext.CurrentMeshComponent = OldCurrentMeshComponent;
-			}
-		
-			mu::NodeScalarEnumParameterPtr EnumParameterNode = new mu::NodeScalarEnumParameter;
 			FString NodeEnumParamName = ProjParamNode->ParameterName + IMAGE_PARAMETER_POSTFIX;
 			EnumParameterNode->SetName(NodeEnumParamName);
 			EnumParameterNode->SetUid(SelectedImageParamUid.ToString());
-			GenerationContext.AddParameterNameUnique(originalGroup, NodeEnumParamName);
 			EnumParameterNode->SetValueCount(ArrayOptionTexture.Num());
 			EnumParameterNode->SetDefaultValueIndex(0);
 			EnumParameterNode->SetRangeCount(1);
 			EnumParameterNode->SetRange(0, NodeRangeFromScalar);
 
-			{
-				FMutableParameterData ParameterUIData;
-				ParameterUIData.ParamUIMetadata = ProjParamNode->ParamUIMetadata;
-				ParameterUIData.Type = EMutableParameterType::Int;
-				ParameterUIData.IntegerParameterGroupType = ECustomizableObjectGroupType::COGT_ONE;
-				ParameterUIData.ParamUIMetadata.ExtraInformation.Add(FString("UseThumbnails"));
-
-				for (int ImageIndex = 0; ImageIndex < ArrayOptionTexture.Num(); ++ImageIndex)
-				{
-					EnumParameterNode->SetValue(ImageIndex, (float)ImageIndex, ArrayOptionTexture[ImageIndex].OptionName);
+			GenerationContext.GeneratedEnumParameters.Add(SelectedImageParamUid.ToString(), EnumParameterNode);
+					
+			FMutableParameterData ParameterUIData;
+			ParameterUIData.ParamUIMetadata = ProjParamNode->ParamUIMetadata;
+			ParameterUIData.Type = EMutableParameterType::Int;
+			ParameterUIData.IntegerParameterGroupType = ECustomizableObjectGroupType::COGT_ONE;
+			ParameterUIData.ParamUIMetadata.ExtraInformation.Add(FString("UseThumbnails"));
 
-					FMutableParamUIMetadata OptionMetadata = ParameterUIData.ParamUIMetadata;
-					OptionMetadata.UIThumbnail = ArrayOptionTexture[ImageIndex].OptionTexture;
-					ParameterUIData.ArrayIntegerParameterOption.Add(
-						ArrayOptionTexture[ImageIndex].OptionName,
-						FIntegerParameterUIData(OptionMetadata));
-				}
+			for (int OptionTextureIndex = 0; OptionTextureIndex < ArrayOptionTexture.Num(); ++OptionTextureIndex)
+			{
+				EnumParameterNode->SetValue(OptionTextureIndex, (float)OptionTextureIndex, ArrayOptionTexture[OptionTextureIndex].OptionName);
 
-				GenerationContext.ParameterUIDataMap.Add(NodeEnumParamName, ParameterUIData);
+				FMutableParamUIMetadata OptionMetadata = ParameterUIData.ParamUIMetadata;
+				OptionMetadata.UIThumbnail = ArrayOptionTexture[OptionTextureIndex].OptionTexture;
+				ParameterUIData.ArrayIntegerParameterOption.Add(
+					ArrayOptionTexture[OptionTextureIndex].OptionName,
+					FIntegerParameterUIData(OptionMetadata));
 			}
 
-			mu::NodeImageSwitchPtr SwitchNode = new mu::NodeImageSwitch;
-			SwitchNode->SetParameter(EnumParameterNode);
-			SwitchNode->SetOptionCount(ArrayOptionTexture.Num());
+			GenerationContext.ParameterUIDataMap.Add(NodeEnumParamName, ParameterUIData);
+		}
 
-			const uint32 AdditionalLODBias = GenerationContext.Options.bUseLODAsBias ? GenerationContext.FirstLODAvailable : 0;
-			for (int32 SelectorIndex = 0; SelectorIndex < ArrayOptionTexture.Num(); ++SelectorIndex)
-			{
-				if (const TObjectPtr<UTexture2D>& Texture = ArrayOptionTexture[SelectorIndex].OptionTexture)
-				{
-					mu::Ptr<mu::Image> ImageConstant = GenerateImageConstant(ArrayOptionTexture[SelectorIndex].OptionTexture, GenerationContext, false);
+		mu::NodeImageSwitchPtr SwitchNode = new mu::NodeImageSwitch;
+		SwitchNode->SetParameter(EnumParameterNode);
+		SwitchNode->SetOptionCount(ArrayOptionTexture.Num());
 
-					mu::Ptr<mu::NodeImageConstant> ImageNode = new mu::NodeImageConstant();
-					ImageNode->SetValue(ImageConstant.get());
+		const uint32 AdditionalLODBias = GenerationContext.Options.bUseLODAsBias ? GenerationContext.FirstLODAvailable : 0;
+		for (int32 SelectorIndex = 0; SelectorIndex < ArrayOptionTexture.Num(); ++SelectorIndex)
+		{
+			if (const TObjectPtr<UTexture2D>& Texture = ArrayOptionTexture[SelectorIndex].OptionTexture)
+			{
+				mu::Ptr<mu::Image> ImageConstant = GenerateImageConstant(ArrayOptionTexture[SelectorIndex].OptionTexture, GenerationContext, false);
 
-					const uint32 MipsToSkip = ComputeLODBiasForTexture(GenerationContext, *Texture, ProjParamNode->ReferenceTexture) + AdditionalLODBias;
-					SwitchNode->SetOption(SelectorIndex, ResizeTextureByNumMips(ImageNode, MipsToSkip));
+				mu::Ptr<mu::NodeImageConstant> ImageNode = new mu::NodeImageConstant();
+				ImageNode->SetValue(ImageConstant.get());
 
-					// Calculate the number of mips to tag as high res for this image.
-					if (ProjParamNode->ReferenceTexture)
-					{
-						int32 TotalMips = mu::Image::GetMipmapCount(ProjParamNode->ReferenceTexture->GetSizeX(), ProjParamNode->ReferenceTexture->GetSizeY());
-						int32 NumMipsBeyondMin = FMath::Max(0, TotalMips - GenerationContext.Options.MinDiskMips);
-						int32 HighResMipsForThisImage = FMath::Min(NumMipsBeyondMin, GenerationContext.Options.NumHighResImageMips);
-						ImageNode->SourceDataDescriptor.SourceHighResMips = HighResMipsForThisImage;
+				const uint32 MipsToSkip = ComputeLODBiasForTexture(GenerationContext, *Texture, ProjParamNode->ReferenceTexture) + AdditionalLODBias;
+				SwitchNode->SetOption(SelectorIndex, ResizeTextureByNumMips(ImageNode, MipsToSkip));
 
-						const FString TextureName = GetNameSafe(Texture).ToLower();
-						ImageNode->SourceDataDescriptor.SourceId = CityHash32(reinterpret_cast<const char*>(*TextureName), TextureName.Len() * sizeof(FString::ElementType));
-					}
-				}
-				else
+				// Calculate the number of mips to tag as high res for this image.
+				if (ProjParamNode->ReferenceTexture)
 				{
-					FString msg = FString::Printf(TEXT("The group projection node must have a texture for all the options. Please set a texture for all the options."));
-					GenerationContext.Log(FText::FromString(msg), ProjParamNode);
+					int32 TotalMips = mu::Image::GetMipmapCount(ProjParamNode->ReferenceTexture->GetSizeX(), ProjParamNode->ReferenceTexture->GetSizeY());
+					int32 NumMipsBeyondMin = FMath::Max(0, TotalMips - GenerationContext.Options.MinDiskMips);
+					int32 HighResMipsForThisImage = FMath::Min(NumMipsBeyondMin, GenerationContext.Options.NumHighResImageMips);
+					ImageNode->SourceDataDescriptor.SourceHighResMips = HighResMipsForThisImage;
+
+					const FString TextureName = GetNameSafe(Texture).ToLower();
+					ImageNode->SourceDataDescriptor.SourceId = CityHash32(reinterpret_cast<const char*>(*TextureName), TextureName.Len() * sizeof(FString::ElementType));
 				}
 			}
-
-			int32 TextureSize = ProjParamNode->ProjectionTextureSize > 0 ? ProjParamNode->ProjectionTextureSize : 512;
-
-			// If TextureSize is not power of two, round up to the next power of two 
-			if (!FMath::IsPowerOfTwo(TextureSize))
+			else
 			{
-				TextureSize = FMath::RoundUpToPowerOfTwo(TextureSize);
+				FString msg = FString::Printf(TEXT("The group projection node must have a texture for all the options. Please set a texture for all the options."));
+				GenerationContext.Log(FText::FromString(msg), ProjParamNode);
 			}
+		}
 
-			// Apply additional LODBias if necessary
-			GroupProjectorTempData.TextureSize = FMath::Max(TextureSize >> AdditionalLODBias, 1);
-
-			GroupProjectorTempData.NodeImagePtr = SwitchNode;
+		int32 TextureSize = ProjParamNode->ProjectionTextureSize > 0 ? ProjParamNode->ProjectionTextureSize : 512;
 
-			GenerationContext.ProjectorGroupMap.Add(originalGroup, GroupProjectorTempData);
-		}
-		else
+		// If TextureSize is not power of two, round up to the next power of two 
+		if (!FMath::IsPowerOfTwo(TextureSize))
 		{
-			GenerationContext.Log(FText::FromString(TEXT("Error getting group projection properties.")), ProjParamNode, EMessageSeverity::Error, true);
-			return false;
+			TextureSize = FMath::RoundUpToPowerOfTwo(TextureSize);
 		}
-	}
-	
 
-	GenerationContext.Generated.Add(Key, FGeneratedData(Node, Result));
-	GenerationContext.GeneratedNodes.Add(Node);
+		// Apply additional LODBias if necessary
+		GroupProjectorTempData.TextureSize = FMath::Max(TextureSize >> AdditionalLODBias, 1);
+
+		GroupProjectorTempData.NodeImagePtr = SwitchNode;
+	}
 
-	return true;
+	GenerationContext.GeneratedGroupProjectors.Add(Key, GroupProjectorTempData);
+	
+	return GroupProjectorTempData;
 }
 
 #undef LOCTEXT_NAMESPACE
diff --git a/Engine/Plugins/Mutable/Source/CustomizableObjectEditor/Private/MuCOE/GenerateMutableSource/GenerateMutableSourceGroupProjector.h b/Engine/Plugins/Mutable/Source/CustomizableObjectEditor/Private/MuCOE/GenerateMutableSource/GenerateMutableSourceGroupProjector.h
--- a/Engine/Plugins/Mutable/Source/CustomizableObjectEditor/Private/MuCOE/GenerateMutableSource/GenerateMutableSourceGroupProjector.h
+++ b/Engine/Plugins/Mutable/Source/CustomizableObjectEditor/Private/MuCOE/GenerateMutableSource/GenerateMutableSourceGroupProjector.h
@@ -5,6 +5,8 @@
 #include "Containers/Map.h"
 #include "MuT/NodeSurfaceNew.h"
 
+class UCustomizableObjectNodeGroupProjectorParameter;
+struct FGroupProjectorTempData;
 class UCustomizableObjectNodeModifierExtendMeshSection;
 class FString;
 class UCustomizableObjectNodeMaterial;
@@ -15,10 +17,10 @@ class UTexture2D;
 struct FMutableGraphGenerationContext;
 
 
-mu::NodeImagePtr GenerateMutableGroupProjection(const int32 NodeLOD, const int32 ImageIndex, mu::NodeMeshPtr MeshNode, FMutableGraphGenerationContext& GenerationContext,
+mu::NodeImagePtr GenerateMutableSourceGroupProjector(const int32 NodeLOD, const int32 ImageIndex, mu::NodeMeshPtr MeshNode, FMutableGraphGenerationContext& GenerationContext,
 	UCustomizableObjectNodeMaterialBase* TypedNodeMat, UCustomizableObjectNodeModifierExtendMeshSection* TypedNodeExt, bool& bShareProjectionTexturesBetweenLODs, bool& bIsGroupProjectorImage,
 	UTexture2D*& GroupProjectionReferenceTexture, TMap<FString, float>& TextureNameToProjectionResFactor, FString& AlternateResStateName);
 
 
-/** Convert a CustomizableObject Source Graph into a mutable source graph. */
-bool GenerateMutableSourceGroupProjector(const UEdGraphPin* Pin, FMutableGraphGenerationContext& GenerationContext, const UCustomizableObjectNodeObjectGroup* originalGroup);
+TOptional<FGroupProjectorTempData> GenerateMutableGroupProjector(UCustomizableObjectNodeGroupProjectorParameter* ProjParamNode, FMutableGraphGenerationContext& GenerationContext);
+
diff --git a/Engine/Plugins/Mutable/Source/CustomizableObjectEditor/Private/MuCOE/GenerateMutableSource/GenerateMutableSourceImage.cpp b/Engine/Plugins/Mutable/Source/CustomizableObjectEditor/Private/MuCOE/GenerateMutableSource/GenerateMutableSourceImage.cpp
--- a/Engine/Plugins/Mutable/Source/CustomizableObjectEditor/Private/MuCOE/GenerateMutableSource/GenerateMutableSourceImage.cpp
+++ b/Engine/Plugins/Mutable/Source/CustomizableObjectEditor/Private/MuCOE/GenerateMutableSource/GenerateMutableSourceImage.cpp
@@ -148,8 +148,6 @@ mu::NodeImagePtr GenerateMutableSourceImage(const UEdGraphPin* Pin, FMutableGrap
 	{
 		mu::Ptr<mu::NodeImageParameter> TextureNode = new mu::NodeImageParameter();
 
-		GenerationContext.AddParameterNameUnique(Node, TypedNodeParam->ParameterName);
-
 		TextureNode->SetName(TypedNodeParam->ParameterName);
 		TextureNode->SetUid(GenerationContext.GetNodeIdUnique(Node).ToString());
 
diff --git a/Engine/Plugins/Mutable/Source/CustomizableObjectEditor/Private/MuCOE/GenerateMutableSource/GenerateMutableSourceLayout.cpp b/Engine/Plugins/Mutable/Source/CustomizableObjectEditor/Private/MuCOE/GenerateMutableSource/GenerateMutableSourceLayout.cpp
--- a/Engine/Plugins/Mutable/Source/CustomizableObjectEditor/Private/MuCOE/GenerateMutableSource/GenerateMutableSourceLayout.cpp
+++ b/Engine/Plugins/Mutable/Source/CustomizableObjectEditor/Private/MuCOE/GenerateMutableSource/GenerateMutableSourceLayout.cpp
@@ -94,7 +94,6 @@ mu::Ptr<mu::NodeLayout> CreateMutableLayoutNode(FMutableGraphGenerationContext&
 
 			GenerationContext.MeshGenerationFlags.Push(ShapeFlags);
 
-			GenerationContext.ComponentInfos.Add(FMutableComponentInfo(FName(), SkeletalMesh));
 			MutableMesh = ConvertSkeletalMeshToMutable(SkeletalMesh, TSoftClassPtr<UAnimInstance>(), UnrealLayout->GetLOD(), UnrealLayout->GetMaterial(), UnrealLayout->GetLOD(), UnrealLayout->GetMaterial(), 0, GenerationContext, nullptr, nullptr);
 
 			GenerationContext.MeshGenerationFlags.Pop();
diff --git a/Engine/Plugins/Mutable/Source/CustomizableObjectEditor/Private/MuCOE/GenerateMutableSource/GenerateMutableSourceMesh.cpp b/Engine/Plugins/Mutable/Source/CustomizableObjectEditor/Private/MuCOE/GenerateMutableSource/GenerateMutableSourceMesh.cpp
--- a/Engine/Plugins/Mutable/Source/CustomizableObjectEditor/Private/MuCOE/GenerateMutableSource/GenerateMutableSourceMesh.cpp
+++ b/Engine/Plugins/Mutable/Source/CustomizableObjectEditor/Private/MuCOE/GenerateMutableSource/GenerateMutableSourceMesh.cpp
@@ -781,11 +781,8 @@ mu::MeshPtr ConvertSkeletalMeshToMutable(
 		GenerationContext.AddParticipatingObject(*InSkeleton);
 
 		FMutableComponentInfo* MutComponentInfo = GenerationContext.GetCurrentComponentInfo();
-		if (!MutComponentInfo)
-		{
-			return {};
-		}
-
+		check(MutComponentInfo);
+		
 		USkeletalMesh* ComponentRefSkeletalMesh = MutComponentInfo->RefSkeletalMesh;
 		USkeleton* ComponentRefSkeleton = MutComponentInfo->RefSkeleton;
 		check(ComponentRefSkeletalMesh);
diff --git a/Engine/Plugins/Mutable/Source/CustomizableObjectEditor/Private/MuCOE/GenerateMutableSource/GenerateMutableSourceModifier.cpp b/Engine/Plugins/Mutable/Source/CustomizableObjectEditor/Private/MuCOE/GenerateMutableSource/GenerateMutableSourceModifier.cpp
--- a/Engine/Plugins/Mutable/Source/CustomizableObjectEditor/Private/MuCOE/GenerateMutableSource/GenerateMutableSourceModifier.cpp
+++ b/Engine/Plugins/Mutable/Source/CustomizableObjectEditor/Private/MuCOE/GenerateMutableSource/GenerateMutableSourceModifier.cpp
@@ -53,7 +53,9 @@ mu::Ptr<mu::NodeModifier> GenerateMutableSourceModifier(const UEdGraphPin * Pin,
 	
 	UCustomizableObjectNode* Node = CastChecked<UCustomizableObjectNode>(Pin->GetOwningNode());
 
-	const FGeneratedKey Key(reinterpret_cast<void*>(&GenerateMutableSourceModifier), *Pin, *Node, GenerationContext, true);
+	FGeneratedKey Key(reinterpret_cast<void*>(&GenerateMutableSourceModifier), *Pin, *Node, GenerationContext, true);
+	Key.CurrentMeshComponent = GenerationContext.CurrentMeshComponent;
+	
 	if (const FGeneratedData* Generated = GenerationContext.Generated.Find(Key))
 	{
 		return static_cast<mu::NodeModifier*>(Generated->Node.get());
@@ -391,7 +393,7 @@ mu::Ptr<mu::NodeModifier> GenerateMutableSourceModifier(const UEdGraphPin * Pin,
 					TMap<FString, float> TextureNameToProjectionResFactor;
 					FString AlternateResStateName;
 
-					ImageNode = GenerateMutableGroupProjection(LOD, ImageIndex, AddMeshNode, GenerationContext,
+					ImageNode = GenerateMutableSourceGroupProjector(LOD, ImageIndex, AddMeshNode, GenerationContext,
 						nullptr, TypedNodeExt, bShareProjectionTexturesBetweenLODs, bIsGroupProjectorImage,
 						GroupProjectionReferenceTexture, TextureNameToProjectionResFactor, AlternateResStateName);
 				}
diff --git a/Engine/Plugins/Mutable/Source/CustomizableObjectEditor/Private/MuCOE/GenerateMutableSource/GenerateMutableSourceProjector.cpp b/Engine/Plugins/Mutable/Source/CustomizableObjectEditor/Private/MuCOE/GenerateMutableSource/GenerateMutableSourceProjector.cpp
--- a/Engine/Plugins/Mutable/Source/CustomizableObjectEditor/Private/MuCOE/GenerateMutableSource/GenerateMutableSourceProjector.cpp
+++ b/Engine/Plugins/Mutable/Source/CustomizableObjectEditor/Private/MuCOE/GenerateMutableSource/GenerateMutableSourceProjector.cpp
@@ -83,8 +83,6 @@ mu::NodeProjectorPtr GenerateMutableSourceProjector(const UEdGraphPin* Pin, FMut
 		mu::NodeProjectorParameterPtr ProjectorNode = new mu::NodeProjectorParameter();
 		Result = ProjectorNode;
 
-		GenerationContext.AddParameterNameUnique(Node, TypedNodeParam->ParameterName);
-
 		ProjectorNode->SetName(TypedNodeParam->ParameterName);
 		ProjectorNode->SetUid(GenerationContext.GetNodeIdUnique(Node).ToString());
 		switch ((int)TypedNodeParam->DefaultValue.ProjectionType)
diff --git a/Engine/Plugins/Mutable/Source/CustomizableObjectEditor/Private/MuCOE/GenerateMutableSource/GenerateMutableSourceSurface.cpp b/Engine/Plugins/Mutable/Source/CustomizableObjectEditor/Private/MuCOE/GenerateMutableSource/GenerateMutableSourceSurface.cpp
--- a/Engine/Plugins/Mutable/Source/CustomizableObjectEditor/Private/MuCOE/GenerateMutableSource/GenerateMutableSourceSurface.cpp
+++ b/Engine/Plugins/Mutable/Source/CustomizableObjectEditor/Private/MuCOE/GenerateMutableSource/GenerateMutableSourceSurface.cpp
@@ -453,7 +453,7 @@ mu::Ptr<mu::NodeSurface> GenerateMutableSourceSurface(const UEdGraphPin * Pin, F
 				{
 					bool bIsGroupProjectorImage = false;
 
-					GroupProjectionImg = GenerateMutableGroupProjection(LOD, ImageIndex, MeshNode, GenerationContext,
+					GroupProjectionImg = GenerateMutableSourceGroupProjector(LOD, ImageIndex, MeshNode, GenerationContext,
 						TypedNodeMat, nullptr, bShareProjectionTexturesBetweenLODs, bIsGroupProjectorImage,
 						GroupProjectionReferenceTexture, TextureNameToProjectionResFactor, AlternateResStateName);
 
diff --git a/Engine/Plugins/Mutable/Source/CustomizableObjectEditor/Private/MuCOE/GenerateMutableSource/GenerateMutableSourceTable.cpp b/Engine/Plugins/Mutable/Source/CustomizableObjectEditor/Private/MuCOE/GenerateMutableSource/GenerateMutableSourceTable.cpp
--- a/Engine/Plugins/Mutable/Source/CustomizableObjectEditor/Private/MuCOE/GenerateMutableSource/GenerateMutableSourceTable.cpp
+++ b/Engine/Plugins/Mutable/Source/CustomizableObjectEditor/Private/MuCOE/GenerateMutableSource/GenerateMutableSourceTable.cpp
@@ -990,9 +990,6 @@ bool GenerateTableColumn(const UCustomizableObjectNodeTable* TableNode, const UE
 
 void GenerateTableParameterUIData(const UDataTable* DataTable, const UCustomizableObjectNodeTable* TableNode, FMutableGraphGenerationContext& GenerationContext)
 {
-	// Checking if the parameter name already exists
-	GenerationContext.AddParameterNameUnique(TableNode, TableNode->ParameterName);
-
 	TArray<uint32> RowIds;
 	TArray<FName> RowNames = GetRowsToCompile(*DataTable, *TableNode, GenerationContext, RowIds);
 
diff --git a/Engine/Plugins/Mutable/Source/CustomizableObjectEditor/Private/MuCOE/GenerateMutableSource/GenerateMutableSourceTransform.cpp b/Engine/Plugins/Mutable/Source/CustomizableObjectEditor/Private/MuCOE/GenerateMutableSource/GenerateMutableSourceTransform.cpp
--- a/Engine/Plugins/Mutable/Source/CustomizableObjectEditor/Private/MuCOE/GenerateMutableSource/GenerateMutableSourceTransform.cpp
+++ b/Engine/Plugins/Mutable/Source/CustomizableObjectEditor/Private/MuCOE/GenerateMutableSource/GenerateMutableSourceTransform.cpp
@@ -43,9 +43,7 @@ mu::Ptr<mu::NodeMatrix> GenerateMutableSourceTransform(const UEdGraphPin* Pin, F
 	{
 		mu::Ptr<mu::NodeMatrixParameter> MatrixNode = new mu::NodeMatrixParameter();
 		Result = MatrixNode;
-
-		GenerationContext.AddParameterNameUnique(Node, TypedNodeTransformParameter->ParameterName);
-
+		
 		MatrixNode->Name = TypedNodeTransformParameter->ParameterName;
 		MatrixNode->Uid = GenerationContext.GetNodeIdUnique(Node).ToString();
 		MatrixNode->DefaultValue = FMatrix44f(TypedNodeTransformParameter->DefaultValue.ToMatrixWithScale());
diff --git a/Engine/Plugins/Mutable/Source/CustomizableObjectEditor/Private/MuCOE/GraphTraversal.cpp b/Engine/Plugins/Mutable/Source/CustomizableObjectEditor/Private/MuCOE/GraphTraversal.cpp
--- a/Engine/Plugins/Mutable/Source/CustomizableObjectEditor/Private/MuCOE/GraphTraversal.cpp
+++ b/Engine/Plugins/Mutable/Source/CustomizableObjectEditor/Private/MuCOE/GraphTraversal.cpp
@@ -398,6 +398,58 @@ const UCustomizableObject* GraphTraversal::GetRootObject(const UCustomizableObje
 }
 
 
+void GraphTraversal::VisitNodes(UCustomizableObjectNode& StartNode, const TMultiMap<FGuid, UCustomizableObjectNodeObject*>& ObjectGroupMap, const TFunction<void(UCustomizableObjectNode&)>& VisitFunction)
+{
+	TSet<UCustomizableObjectNode*> VisitedNodes;
+	
+	TArray<UCustomizableObjectNode*> NodesToVisit;
+	NodesToVisit.Add(&StartNode);
+
+	while (!NodesToVisit.IsEmpty())
+	{
+		UCustomizableObjectNode* CurrentNode = NodesToVisit.Pop();
+
+		if (VisitedNodes.Contains(CurrentNode))
+		{
+			continue;
+		}
+
+		VisitedNodes.Add(CurrentNode);
+
+		VisitFunction(*CurrentNode);
+
+		for (UEdGraphPin* Pin : CurrentNode->GetAllNonOrphanPins())
+		{
+			if (Pin->Direction != EGPD_Input)
+			{
+				continue;
+			}
+
+			for (UEdGraphPin* ConnectedPin : FollowInputPinArray(*Pin))
+			{
+				UEdGraphNode* ConnectedNode = ConnectedPin->GetOwningNode();
+				
+				if (UCustomizableObjectNode* Node = Cast<UCustomizableObjectNode>(ConnectedNode))
+				{
+					NodesToVisit.Add(Node);
+				}
+
+				if (UCustomizableObjectNodeObjectGroup* ObjectGroupNode = Cast<UCustomizableObjectNodeObjectGroup>(ConnectedNode))
+				{
+					TArray<UCustomizableObjectNodeObject*> ChildObjectNodes;
+					ObjectGroupMap.MultiFind(ObjectGroupNode->NodeGuid, ChildObjectNodes);
+						
+					for (UCustomizableObjectNodeObject* ChildObjectNode : ChildObjectNodes)
+					{
+						NodesToVisit.Add(ChildObjectNode);
+					}
+				}
+			}
+		}
+	}
+}
+
+
 UCustomizableObjectNodeObject* GetFullGraphRootNodeObject(const UCustomizableObjectNodeObject* Node, TArray<const UCustomizableObject*>& VisitedObjects)
 {
 	if (Node->ParentObject != nullptr)
diff --git a/Engine/Plugins/Mutable/Source/CustomizableObjectEditor/Private/MuCOE/Nodes/CustomizableObjectNodeGroupProjectorParameter.cpp b/Engine/Plugins/Mutable/Source/CustomizableObjectEditor/Private/MuCOE/Nodes/CustomizableObjectNodeGroupProjectorParameter.cpp
--- a/Engine/Plugins/Mutable/Source/CustomizableObjectEditor/Private/MuCOE/Nodes/CustomizableObjectNodeGroupProjectorParameter.cpp
+++ b/Engine/Plugins/Mutable/Source/CustomizableObjectEditor/Private/MuCOE/Nodes/CustomizableObjectNodeGroupProjectorParameter.cpp
@@ -141,4 +141,10 @@ void UCustomizableObjectNodeGroupProjectorParameter::AllocateDefaultPins(UCustom
 }
 
 
+UEdGraphPin& UCustomizableObjectNodeGroupProjectorParameter::OutputPin() const
+{
+	return *FindPin(TEXT("Value"));
+}
+
+
 #undef LOCTEXT_NAMESPACE
diff --git a/Engine/Plugins/Mutable/Source/CustomizableObjectEditor/Public/MuCOE/GraphTraversal.h b/Engine/Plugins/Mutable/Source/CustomizableObjectEditor/Public/MuCOE/GraphTraversal.h
--- a/Engine/Plugins/Mutable/Source/CustomizableObjectEditor/Public/MuCOE/GraphTraversal.h
+++ b/Engine/Plugins/Mutable/Source/CustomizableObjectEditor/Public/MuCOE/GraphTraversal.h
@@ -77,6 +77,13 @@ namespace GraphTraversal
 		 * object does not have any parent.	*/
 	CUSTOMIZABLEOBJECTEDITOR_API UCustomizableObject* GetRootObject(UCustomizableObject* ChildObject);
 	CUSTOMIZABLEOBJECTEDITOR_API const UCustomizableObject* GetRootObject(const UCustomizableObject* ChildObject);
+
+	/** From the StartNode visit all connected nodes in the hierarchy.
+	 *
+	 *   @param StartNode Node to start the visit.
+	 *   @param ObjectGroupMap Key is the Object Group node id, values are the attached Child Object nodes. Used to prune the traversal.
+	 *   @param VisitFunction Called for each UCustomizableObjectNode node found. Order is not guaranteed. */
+	void VisitNodes(UCustomizableObjectNode& StartNode, const TMultiMap<FGuid, UCustomizableObjectNodeObject*>& ObjectGroupMap, const TFunction<void(UCustomizableObjectNode&)>& VisitFunction);
 }
 
 /** Return the full graph Customizable Object Node root of the node given as parameter */
diff --git a/Engine/Plugins/Mutable/Source/CustomizableObjectEditor/Public/MuCOE/Nodes/CustomizableObjectNodeGroupProjectorParameter.h b/Engine/Plugins/Mutable/Source/CustomizableObjectEditor/Public/MuCOE/Nodes/CustomizableObjectNodeGroupProjectorParameter.h
--- a/Engine/Plugins/Mutable/Source/CustomizableObjectEditor/Public/MuCOE/Nodes/CustomizableObjectNodeGroupProjectorParameter.h
+++ b/Engine/Plugins/Mutable/Source/CustomizableObjectEditor/Public/MuCOE/Nodes/CustomizableObjectNodeGroupProjectorParameter.h
@@ -104,5 +104,8 @@ public:
 	// UCustomizableObjectNode interface.
 	virtual void BackwardsCompatibleFixup(int32 CustomizableObjectCustomVersion) override;
 	virtual void AllocateDefaultPins(UCustomizableObjectNodeRemapPins* RemapPins);
+
+	// Own interface
+	UEdGraphPin& OutputPin() const;
 };
 
