[Mutable] Add missing FFileCategoryID constructors.

#rnx

diff --git a/Engine/Plugins/Mutable/Source/CustomizableObject/Internal/MuCO/CustomizableObjectPrivate.h b/Engine/Plugins/Mutable/Source/CustomizableObject/Internal/MuCO/CustomizableObjectPrivate.h
--- a/Engine/Plugins/Mutable/Source/CustomizableObject/Internal/MuCO/CustomizableObjectPrivate.h
+++ b/Engine/Plugins/Mutable/Source/CustomizableObject/Internal/MuCO/CustomizableObjectPrivate.h
@@ -827,6 +827,10 @@ namespace MutablePrivate
 
 	struct FFileCategoryID
 	{
+		FFileCategoryID(EDataType DataType, uint16 ResourceType, uint16 Flags);
+
+		FFileCategoryID() = default;
+
 		// DATATYPE
 		EDataType DataType = EDataType::None;
 
diff --git a/Engine/Plugins/Mutable/Source/CustomizableObject/Private/MuCO/CustomizableObject.cpp b/Engine/Plugins/Mutable/Source/CustomizableObject/Private/MuCO/CustomizableObject.cpp
--- a/Engine/Plugins/Mutable/Source/CustomizableObject/Private/MuCO/CustomizableObject.cpp
+++ b/Engine/Plugins/Mutable/Source/CustomizableObject/Private/MuCO/CustomizableObject.cpp
@@ -2549,6 +2549,12 @@ namespace MutablePrivate
 		}
 	}
 
+	FFileCategoryID::FFileCategoryID(EDataType InDataType, uint16 InResourceType, uint16 InFlags)
+	{
+		DataType = InDataType;
+		ResourceType = InResourceType;
+		Flags = InFlags;
+	}
 
 	uint32 GetTypeHash(const FFileCategoryID& Key)
 	{
@@ -2627,11 +2633,11 @@ namespace MutablePrivate
 		// Model Roms. Iterate all Model roms to distribute them in categories.
 		{
 			// Add meshes and low-res textures to the Default bucket 
-			DefaultBucket.Categories.Add({ FFileCategoryID(EDataType::Model, mu::DATATYPE::DT_MESH, 0) });
-			DefaultBucket.Categories.Add({ FFileCategoryID(EDataType::Model, mu::DATATYPE::DT_IMAGE, 0) });
+			DefaultBucket.Categories.Add({ FFileCategoryID(EDataType::Model, mu::DATATYPE::DT_MESH, 0), 0, 0, 0 });
+			DefaultBucket.Categories.Add({ FFileCategoryID(EDataType::Model, mu::DATATYPE::DT_IMAGE, 0), 0, 0, 0 });
 
 			// Add High-res textures to the Optional bucket
-			OptionalBucket.Categories.Add({ FFileCategoryID(EDataType::Model, mu::DATATYPE::DT_IMAGE, (uint16)mu::ERomFlags::HighRes) });
+			OptionalBucket.Categories.Add({ FFileCategoryID(EDataType::Model, mu::DATATYPE::DT_IMAGE, (uint16)mu::ERomFlags::HighRes), 0, 0, 0});
 
 			const int32 NumRoms = Model->GetRomCount();
 			for (int32 RomIndex = 0; RomIndex < NumRoms; ++RomIndex)
