[Mutable] Enable the use of FByteBulkData on cooked builds by default.

#rnx

diff --git a/Engine/Plugins/Mutable/Source/CustomizableObject/Private/MuCO/CustomizableObject.cpp b/Engine/Plugins/Mutable/Source/CustomizableObject/Private/MuCO/CustomizableObject.cpp
--- a/Engine/Plugins/Mutable/Source/CustomizableObject/Private/MuCO/CustomizableObject.cpp
+++ b/Engine/Plugins/Mutable/Source/CustomizableObject/Private/MuCO/CustomizableObject.cpp
@@ -71,7 +71,7 @@ TAutoConsoleVariable<int32> CVarPackagedDataBytesLimitOverride(
 
 TAutoConsoleVariable<bool> CVarMutableUseBulkData(
 	TEXT("Mutable.UseBulkData"),
-	false,
+	true,
 	TEXT("Switch between .utoc/.ucas (FBulkData) and .mut files (CookAdditionalFiles).\n")
 	TEXT("True - Use FBulkData to store streamable data.\n")
 	TEXT("False - Use Mut files to store streamable data\n"));
