[Mutable] Fix UCustomizableSkeletalComponent not working properly in PIE.

diff --git a/Engine/Plugins/Mutable/Source/CustomizableObject/Internal/MuCO/CustomizableObjectInstanceUsagePrivate.h b/Engine/Plugins/Mutable/Source/CustomizableObject/Internal/MuCO/CustomizableObjectInstanceUsagePrivate.h
--- a/Engine/Plugins/Mutable/Source/CustomizableObject/Internal/MuCO/CustomizableObjectInstanceUsagePrivate.h
+++ b/Engine/Plugins/Mutable/Source/CustomizableObject/Internal/MuCO/CustomizableObjectInstanceUsagePrivate.h
@@ -29,6 +29,8 @@ public:
 	
 	// Own interface
 	
+	UCustomizableObjectInstanceUsagePrivate();
+	
 	/** Common end point of all updates. Even those which failed. */
 	void Callbacks() const;
 
diff --git a/Engine/Plugins/Mutable/Source/CustomizableObject/Internal/MuCO/CustomizableSkeletalComponentPrivate.h b/Engine/Plugins/Mutable/Source/CustomizableObject/Internal/MuCO/CustomizableSkeletalComponentPrivate.h
--- a/Engine/Plugins/Mutable/Source/CustomizableObject/Internal/MuCO/CustomizableSkeletalComponentPrivate.h
+++ b/Engine/Plugins/Mutable/Source/CustomizableObject/Internal/MuCO/CustomizableSkeletalComponentPrivate.h
@@ -47,7 +47,7 @@ public:
 	UPROPERTY()
 	bool bSkipSkipSetSkeletalMeshOnAttach = false;
 
-	UPROPERTY(Transient)
+	UPROPERTY()
 	TObjectPtr<UCustomizableObjectInstanceUsage> CustomizableObjectInstanceUsage;
 };
 
diff --git a/Engine/Plugins/Mutable/Source/CustomizableObject/Private/MuCO/CustomizableObjectInstanceUsage.cpp b/Engine/Plugins/Mutable/Source/CustomizableObject/Private/MuCO/CustomizableObjectInstanceUsage.cpp
--- a/Engine/Plugins/Mutable/Source/CustomizableObject/Private/MuCO/CustomizableObjectInstanceUsage.cpp
+++ b/Engine/Plugins/Mutable/Source/CustomizableObject/Private/MuCO/CustomizableObjectInstanceUsage.cpp
@@ -52,9 +52,26 @@ void UCustomizableObjectInstanceUsagePrivate::Callbacks() const
 }
 
 
+void UCustomizableObjectInstanceUsage::PostInitProperties()
+{
+	Super::PostInitProperties();
+	
+	if (!HasAllFlags(RF_ClassDefaultObject))
+	{
+		if (!Private)
+		{
+			Private = NewObject<UCustomizableObjectInstanceUsagePrivate>(this, FName("Private"), RF_Public);
+		}
+		else if (Private->GetOuter() != this)
+		{
+			Private = Cast<UCustomizableObjectInstanceUsagePrivate>(StaticDuplicateObject(Private, this, FName("Private")));
+		}
+	}
+}
+
+
 UCustomizableObjectInstanceUsage::UCustomizableObjectInstanceUsage()
 {
-	Private = CreateDefaultSubobject<UCustomizableObjectInstanceUsagePrivate>(FName("Private"));
 }
 
 
@@ -681,6 +698,17 @@ bool UCustomizableObjectInstanceUsagePrivate::IsTickable() const
 }
 
 
+UCustomizableObjectInstanceUsagePrivate::UCustomizableObjectInstanceUsagePrivate()
+{
+	// This object may get instantiated into a level, as a part of UCustomizableSkeletalComponent, so needs to be public to ensure
+	// it can be serialized out.
+	if (!HasAllFlags(RF_ClassDefaultObject))
+	{
+		SetFlags(RF_Public);
+	}
+}
+
+
 bool UCustomizableObjectInstanceUsagePrivate::IsNetMode(ENetMode InNetMode) const
 {
 	if (GetPublic()->CustomizableSkeletalComponent)
diff --git a/Engine/Plugins/Mutable/Source/CustomizableObject/Private/MuCO/CustomizableSkeletalComponent.cpp b/Engine/Plugins/Mutable/Source/CustomizableObject/Private/MuCO/CustomizableSkeletalComponent.cpp
--- a/Engine/Plugins/Mutable/Source/CustomizableObject/Private/MuCO/CustomizableSkeletalComponent.cpp
+++ b/Engine/Plugins/Mutable/Source/CustomizableObject/Private/MuCO/CustomizableSkeletalComponent.cpp
@@ -39,7 +39,7 @@ void UCustomizableSkeletalComponentPrivate::CreateCustomizableObjectInstanceUsag
 
 	if (!CustomizableObjectInstanceUsage && !HasAnyFlags(RF_ClassDefaultObject) && !bIsDefaultActor)
 	{
-		CustomizableObjectInstanceUsage = NewObject<UCustomizableObjectInstanceUsage>(this, TEXT("InstanceUsage"), RF_Transient);
+		CustomizableObjectInstanceUsage = NewObject<UCustomizableObjectInstanceUsage>(this, TEXT("InstanceUsage"));
 		CustomizableObjectInstanceUsage->GetPrivate()->SetCustomizableSkeletalComponent(GetPublic());
 	}
 }
diff --git a/Engine/Plugins/Mutable/Source/CustomizableObject/Public/MuCO/CustomizableObjectInstanceUsage.h b/Engine/Plugins/Mutable/Source/CustomizableObject/Public/MuCO/CustomizableObjectInstanceUsage.h
--- a/Engine/Plugins/Mutable/Source/CustomizableObject/Public/MuCO/CustomizableObjectInstanceUsage.h
+++ b/Engine/Plugins/Mutable/Source/CustomizableObject/Public/MuCO/CustomizableObjectInstanceUsage.h
@@ -29,6 +29,9 @@ class CUSTOMIZABLEOBJECT_API UCustomizableObjectInstanceUsage : public UObject
 public:
 	GENERATED_BODY()
 
+	// UObject interface
+	virtual void PostInitProperties() override;
+	
 	// Own interface
 	UCustomizableObjectInstanceUsage();
 	
