[mutable] Tag selection widget: ensure a tag can be selected with "Enter" even if an object exists with the same name.
#rnx

diff --git a/Engine/Plugins/Mutable/Source/CustomizableObjectEditor/Private/MuCOE/SMutableSearchComboBox.cpp b/Engine/Plugins/Mutable/Source/CustomizableObjectEditor/Private/MuCOE/SMutableSearchComboBox.cpp
--- a/Engine/Plugins/Mutable/Source/CustomizableObjectEditor/Private/MuCOE/SMutableSearchComboBox.cpp
+++ b/Engine/Plugins/Mutable/Source/CustomizableObjectEditor/Private/MuCOE/SMutableSearchComboBox.cpp
@@ -38,7 +38,6 @@ void SMutableSearchComboBox::Construct(const FArguments& InArgs)
 	this->OnSelectionChanged = InArgs._OnSelectionChanged;
 
 	OptionsSource = InArgs._OptionsSource;
-	RefreshOptions();
 
 	TAttribute<EVisibility> SearchVisibility = InArgs._SearchVisibility;
 	const EVisibility CurrentSearchVisibility = SearchVisibility.Get();
@@ -225,7 +224,11 @@ void SMutableSearchComboBox::OnGetChildren(TSharedRef<FFilteredOption> InItem, T
 
 void SMutableSearchComboBox::OnMenuOpenChanged(bool bOpen)
 {
-	if (bOpen == false)
+	if (bOpen)
+	{
+		RefreshOptions();
+	}
+	else
 	{
 		// Set focus back to ComboBox for users focusing the ListView that just closed
 		FSlateApplication::Get().ForEachUser([this](FSlateUser& User)
@@ -275,18 +278,35 @@ void SMutableSearchComboBox::OnSearchTextCommitted(const FText& InText, ETextCom
 {
 	if ((InCommitType == ETextCommit::Type::OnEnter) && !FilteredOptionsSource.IsEmpty())
 	{
-		TSharedRef<FFilteredOption> Selected = FilteredOptionsSource[0];
+		TSharedPtr<FFilteredOption> Selected;
 
-		for (const TSharedRef<FFilteredOption>& Option : *OptionsSource)
+		for (const TSharedRef<FFilteredOption>& Option : FilteredOptionsSource)
 		{
-			if (Option->DisplayOption == InText.ToString())
+			if (Option->ActualOption == InText.ToString())
 			{
 				Selected = Option;
 				break;
 			}
 		}
 
-		ComboTreeView->SetSelection(Selected, ESelectInfo::OnKeyPress);
+		if (!Selected)
+		{
+			for (const TSharedRef<FFilteredOption>& Option : FilteredOptionsSource)
+			{
+				if (Option->DisplayOption == InText.ToString())
+				{
+					Selected = Option;
+					break;
+				}
+			}
+		}
+
+		if (!Selected)
+		{
+			Selected = FilteredOptionsSource[0];
+		}
+
+		ComboTreeView->SetSelection(Selected.ToSharedRef(), ESelectInfo::OnKeyPress);
 	}
 }
 
