[mutable] Fixed crash when adding grooms to an invalid minimal graph.
#tests Editor
#rnx

diff --git a/Engine/Plugins/Mutable/Source/MutableTools/Private/MuT/ASTOpAddExtensionData.cpp b/Engine/Plugins/Mutable/Source/MutableTools/Private/MuT/ASTOpAddExtensionData.cpp
--- a/Engine/Plugins/Mutable/Source/MutableTools/Private/MuT/ASTOpAddExtensionData.cpp
+++ b/Engine/Plugins/Mutable/Source/MutableTools/Private/MuT/ASTOpAddExtensionData.cpp
@@ -78,7 +78,12 @@ void ASTOpAddExtensionData::Link(FProgram& Program, FLinkerOptions*)
 	OP::InstanceAddExtensionDataArgs Args;
 	FMemory::Memset(Args, 0);
 
-	check(Instance->linkedAddress);
+	if (!Instance || !Instance->linkedAddress)
+	{
+		// Can happen if there's no reference skeletal mesh in the first component
+		return;
+	}
+
 	Args.Instance = Instance->linkedAddress;
 
 	check(ExtensionData->linkedAddress);
