[Mutable] Fixed texture analyzer. Now it uses the name of the component to get the textures of a SKM instead of the index of the component.
+Added metadata to remove the "SetDefaultValue" buttons of the texture analyser's elements.

#rnx

diff --git a/Engine/Plugins/Mutable/Source/CustomizableObjectEditor/Private/MuCOE/CustomizableObjectEditor.h b/Engine/Plugins/Mutable/Source/CustomizableObjectEditor/Private/MuCOE/CustomizableObjectEditor.h
--- a/Engine/Plugins/Mutable/Source/CustomizableObjectEditor/Private/MuCOE/CustomizableObjectEditor.h
+++ b/Engine/Plugins/Mutable/Source/CustomizableObjectEditor/Private/MuCOE/CustomizableObjectEditor.h
@@ -88,61 +88,61 @@ class UCustomizableObjectEditorTextureStats : public UObject
 public:
 
 	/** Texture - double click to open */
-	UPROPERTY(VisibleAnywhere, AssetRegistrySearchable, Category = "Stats", meta = (DisplayName = "Texture", ColumnWidth = "40"))
+	UPROPERTY(VisibleAnywhere, AssetRegistrySearchable, Category = "Stats", meta = (DisplayName = "Texture", ColumnWidth = "40", NoResetToDefault))
 	FString TextureName;
 
 	/** Texture - double click to open */
-	UPROPERTY(VisibleAnywhere, AssetRegistrySearchable, Category = "Stats", meta = (DisplayName = "Name", ColumnWidth = "50"))
+	UPROPERTY(VisibleAnywhere, AssetRegistrySearchable, Category = "Stats", meta = (DisplayName = "Name", ColumnWidth = "50", NoResetToDefault))
 	FString TextureParameterName;
 
 	/** Material - double click to open */
-	UPROPERTY(VisibleAnywhere, AssetRegistrySearchable, Category = "Stats", meta = (DisplayName = "Material", ColumnWidth = "50"))
+	UPROPERTY(VisibleAnywhere, AssetRegistrySearchable, Category = "Stats", meta = (DisplayName = "Material", ColumnWidth = "50", NoResetToDefault))
 	FString MaterialName;
 
 	/** Parent Material - double click to open */
-	UPROPERTY(VisibleAnywhere, AssetRegistrySearchable, Category = "Stats", meta = (DisplayName = "Parent", ColumnWidth = "50"))
+	UPROPERTY(VisibleAnywhere, AssetRegistrySearchable, Category = "Stats", meta = (DisplayName = "Parent", ColumnWidth = "50", NoResetToDefault))
 	FString MaterialParameterName;
 
 	/** Used to open the texture in the editor*/
-	UPROPERTY()
+	UPROPERTY(meta = (NoResetToDefault))
 	TObjectPtr<UTexture> Texture;
 
 	/** Used to open the material in the editor*/
-	UPROPERTY()
+	UPROPERTY(meta = (NoResetToDefault))
 	TObjectPtr<UMaterialInterface> Material;
 
 	/** Used to open the parent material in the editor*/
-	UPROPERTY()
+	UPROPERTY(meta = (NoResetToDefault))
 	TObjectPtr<UMaterialInterface> ParentMaterial;
 
 	/** Resolution of the texture */
-	UPROPERTY(VisibleAnywhere, AssetRegistrySearchable, Category = "Stats", meta = (DisplayName = "Resolution X", ColumnWidth = "40", DisplayRight = "true"))
+	UPROPERTY(VisibleAnywhere, AssetRegistrySearchable, Category = "Stats", meta = (DisplayName = "Resolution X", ColumnWidth = "40", DisplayRight = "true", NoResetToDefault))
 	int32 ResolutionX;
-	UPROPERTY(VisibleAnywhere, AssetRegistrySearchable, Category = "Stats", meta = (DisplayName = "Resolution Y", ColumnWidth = "40", DisplayRight = "true"))
+	UPROPERTY(VisibleAnywhere, AssetRegistrySearchable, Category = "Stats", meta = (DisplayName = "Resolution Y", ColumnWidth = "40", DisplayRight = "true", NoResetToDefault))
 	int32 ResolutionY;
 
 	/** The memory used in KB */
-	UPROPERTY(VisibleAnywhere, AssetRegistrySearchable, Category = "Stats", meta = (DisplayName = "Size Kb", ColumnWidth = "90"))
+	UPROPERTY(VisibleAnywhere, AssetRegistrySearchable, Category = "Stats", meta = (DisplayName = "Size Kb", ColumnWidth = "90", NoResetToDefault))
 	FString Size;
 
 	/** The texture format, e.g. PF_DXT1 */
-	UPROPERTY(VisibleAnywhere, AssetRegistrySearchable, Category = "Stats", meta = (ColumnWidth = "96"))
+	UPROPERTY(VisibleAnywhere, AssetRegistrySearchable, Category = "Stats", meta = (ColumnWidth = "96", NoResetToDefault))
 	TEnumAsByte<EPixelFormat> Format;
 
 	/** LOD Bias for this texture. (Texture LODBias + Texture group) */
-	UPROPERTY(VisibleAnywhere, AssetRegistrySearchable, Category = "Stats", meta = (DisplayName = "LOD Bias", ColumnWidth = "70"))
+	UPROPERTY(VisibleAnywhere, AssetRegistrySearchable, Category = "Stats", meta = (DisplayName = "LOD Bias", ColumnWidth = "70", NoResetToDefault))
 	int32 LODBias;
 	
 	/** Says if the texture is being streamed */
-	UPROPERTY(VisibleAnywhere, AssetRegistrySearchable, Category = "Stats", meta = (DisplayName = "Streamed", ColumnWidth = "70"))
+	UPROPERTY(VisibleAnywhere, AssetRegistrySearchable, Category = "Stats", meta = (DisplayName = "Streamed", ColumnWidth = "70", NoResetToDefault))
 	FString IsStreamed;
 
 	/** The Level of detail group of the texture */
-	UPROPERTY(VisibleAnywhere, AssetRegistrySearchable, Category = "Stats", meta = (DisplayName = "Texture Group", ColumnWidth = "70"))
+	UPROPERTY(VisibleAnywhere, AssetRegistrySearchable, Category = "Stats", meta = (DisplayName = "Texture Group", ColumnWidth = "70", NoResetToDefault))
 	TEnumAsByte<enum TextureGroup> LODGroup;
 	
 	/** The Component of the texture */
-	UPROPERTY(VisibleAnywhere, AssetRegistrySearchable, Category = "Stats", meta = (DisplayName = "Component", ColumnWidth = "70"))
+	UPROPERTY(VisibleAnywhere, AssetRegistrySearchable, Category = "Stats", meta = (DisplayName = "Component", ColumnWidth = "70", NoResetToDefault))
 	int32 Component;
 
 };
diff --git a/Engine/Plugins/Mutable/Source/CustomizableObjectEditor/Private/MuCOE/SCustomizableObjectEditorTextureAnalyzer.cpp b/Engine/Plugins/Mutable/Source/CustomizableObjectEditor/Private/MuCOE/SCustomizableObjectEditorTextureAnalyzer.cpp
--- a/Engine/Plugins/Mutable/Source/CustomizableObjectEditor/Private/MuCOE/SCustomizableObjectEditorTextureAnalyzer.cpp
+++ b/Engine/Plugins/Mutable/Source/CustomizableObjectEditor/Private/MuCOE/SCustomizableObjectEditorTextureAnalyzer.cpp
@@ -130,22 +130,25 @@ void SCustomizableObjecEditorTextureAnalyzer::FillTextureAnalyzerTable(UCustomiz
 
 	if (PreviewInstance && PreviewInstance->HasAnySkeletalMesh())
 	{
+		const UCustomizableObject* CustomizableObject = PreviewInstance->GetCustomizableObject();
+		check(CustomizableObject)
+
 		for (int32 ComponentIndex = 0; ComponentIndex < PreviewInstance->GetNumComponents(); ++ComponentIndex)
 		{
-			if (PreviewInstance->GetSkeletalMesh(ComponentIndex))
+			const FName ComponentName = CustomizableObject->GetComponentName(ComponentIndex);
+			
+			if (const USkeletalMesh* SkeletalMesh = PreviewInstance->GetComponentMeshSkeletalMesh(ComponentName))
 			{
-				for (int i = 0; i < PreviewInstance->GetSkeletalMesh(ComponentIndex)->GetMaterials().Num(); ++i)
+				for (int MaterialIndex = 0; MaterialIndex < SkeletalMesh->GetMaterials().Num(); ++MaterialIndex)
 				{
-					UMaterialInterface* MaterialInterface = PreviewInstance->GetSkeletalMesh(ComponentIndex)->GetMaterials()[i].MaterialInterface;
-					UMaterialInstance* Material = Cast<UMaterialInstance>(MaterialInterface);
-					FString MaterialPath = MaterialInterface->GetPathName();
+					UMaterialInstance* Material = Cast<UMaterialInstance>(SkeletalMesh->GetMaterials()[MaterialIndex].MaterialInterface);
 
 					// We just want the Transient materials
-					if (Material && MaterialPath.Contains("Transient"))
+					if (Material && Material->GetPathName().Contains("Transient"))
 					{
-						for (int32 j = 0; j < Material->TextureParameterValues.Num(); ++j)
+						for (int32 TextureIndex = 0; TextureIndex < Material->TextureParameterValues.Num(); ++TextureIndex)
 						{
-							UTexture* Texture = Material->TextureParameterValues[j].ParameterValue;
+							UTexture* Texture = Material->TextureParameterValues[TextureIndex].ParameterValue;
 							if (!Texture) continue;
 
 							FString TextPath = Texture->GetPathName();
@@ -160,7 +163,7 @@ void SCustomizableObjecEditorTextureAnalyzer::FillTextureAnalyzerTable(UCustomiz
 								//Texture Info
 								Entry->Texture = Texture;
 								Entry->TextureName = Texture->GetName();
-								Entry->TextureParameterName = Material->TextureParameterValues[j].ParameterInfo.Name.ToString();
+								Entry->TextureParameterName = Material->TextureParameterValues[TextureIndex].ParameterInfo.Name.ToString();
 								Entry->ResolutionX = Cast<UTexture2D>(Texture)->GetSizeX();
 								Entry->ResolutionY = Cast<UTexture2D>(Texture)->GetSizeY();
 								Entry->LODBias = Texture->GetCachedLODBias();
@@ -171,8 +174,8 @@ void SCustomizableObjecEditorTextureAnalyzer::FillTextureAnalyzerTable(UCustomiz
 								Entry->LODGroup = Texture->LODGroup;
 
 								//Material Info
-								Entry->Material = MaterialInterface;
-								Entry->MaterialName = MaterialInterface->GetName();
+								Entry->Material = Material;
+								Entry->MaterialName = Material->GetName();
 								Entry->ParentMaterial = Material->Parent;
 								Entry->MaterialParameterName = Material->Parent ? Material->Parent->GetName() : FString();
 								Entry->Component = ComponentIndex;
