[Mutable] Fix posible discrepacies between constant resolution output and runtime results for some image operations.

- Operations where the source image can be requested to load a higher mip (IM_TRANSFORM and IM_RASTERMESH with minfilter method set) now resize the image if the source expected size does not match the loaded source size.

#rnx

diff --git a/Engine/Plugins/Mutable/Source/MutableRuntime/Private/MuR/CodeRunner.cpp b/Engine/Plugins/Mutable/Source/MutableRuntime/Private/MuR/CodeRunner.cpp
--- a/Engine/Plugins/Mutable/Source/MutableRuntime/Private/MuR/CodeRunner.cpp
+++ b/Engine/Plugins/Mutable/Source/MutableRuntime/Private/MuR/CodeRunner.cpp
@@ -4763,6 +4763,39 @@ namespace mu
 					Release(Source);
 					Source = Formatted;
 				}
+			
+				EMinFilterMethod MinFilterMethod = Invoke([&]() -> EMinFilterMethod
+				{
+					if (ForcedProjectionMode == 0)
+					{
+						return EMinFilterMethod::None;
+					}
+					else if (ForcedProjectionMode == 1)
+					{
+						return EMinFilterMethod::TotalAreaHeuristic;
+					}
+						
+					return static_cast<EMinFilterMethod>(args.MinFilterMethod);
+				});
+
+				if (MinFilterMethod == EMinFilterMethod::TotalAreaHeuristic)
+				{
+					const uint16 Mip = Data.RasterMesh.Mip;
+					const FImageSize ExpectedSourceSize = FImageSize(
+							FMath::Max<uint16>(args.SourceSizeX >> Mip, 1), 
+							FMath::Max<uint16>(args.SourceSizeY >> Mip, 1));
+
+					if (Source->GetSize() != ExpectedSourceSize)
+					{
+						MUTABLE_CPUPROFILER_SCOPE(RunCode_ImageRasterMesh_SizeFixup);	
+						
+						Ptr<Image> Resized = CreateImage(ExpectedSourceSize.X, ExpectedSourceSize.Y, 1, Format, EInitializationType::NotInitialized);
+						ImOp.ImageResizeLinear(Resized.get(), m_pSettings->ImageCompressionQuality, Source.get());
+						
+						Release(Source);
+						Source = Resized;	
+					}
+				}
 
 				// Allocate memory for the temporary buffers
 				FScratchImageProject Scratch;
@@ -5108,6 +5141,20 @@ namespace mu
 					Source = Formatted;
 				}
 
+				const FImageSize ExpectedSourceSize = FImageSize(
+						FMath::Max<uint16>(Args.SourceSizeX >> (uint16)Mip, 1), 
+						FMath::Max<uint16>(Args.SourceSizeY >> (uint16)Mip, 1));
+				if (Source->GetSize() != ExpectedSourceSize)
+				{
+					MUTABLE_CPUPROFILER_SCOPE(RunCode_ImageTransform_SizeFixup);	
+					
+					Ptr<Image> Resized = CreateImage(ExpectedSourceSize.X, ExpectedSourceSize.Y, 1, Format, EInitializationType::NotInitialized);
+					ImOp.ImageResizeLinear(Resized.get(), m_pSettings->ImageCompressionQuality, Source.get());
+					
+					Release(Source);
+					Source = Resized;	
+				}
+
 				if (Source->GetLODCount() < 2 && Source->GetSizeX() > 1 && Source->GetSizeY() > 1)
 				{
 					MUTABLE_CPUPROFILER_SCOPE(RunCode_ImageTransform_BilinearMipGen);
