Fixed an assert in backwards compatibility fixup code for old Customizable Objects that use the Material Variations node

diff --git a/Engine/Plugins/Mutable/Source/CustomizableObjectEditor/Private/MuCOE/GraphTraversal.cpp b/Engine/Plugins/Mutable/Source/CustomizableObjectEditor/Private/MuCOE/GraphTraversal.cpp
--- a/Engine/Plugins/Mutable/Source/CustomizableObjectEditor/Private/MuCOE/GraphTraversal.cpp
+++ b/Engine/Plugins/Mutable/Source/CustomizableObjectEditor/Private/MuCOE/GraphTraversal.cpp
@@ -119,7 +119,7 @@ TArray<UEdGraphPin*> FollowInputPinArray(const UEdGraphPin& Pin, bool* bOutCycle
 UEdGraphPin* FollowInputPin(const UEdGraphPin& Pin, bool* CycleDetected)
 {
 	TArray<UEdGraphPin*> Result = FollowInputPinArray(Pin, CycleDetected);
-	check(Result.Num() <= 1); // Use FollowPinImmersive if the pin can have more than one input.
+	check(Result.Num() <= 1); // Use FollowInputPinArray if the pin can have more than one input.
 
 	if (!Result.IsEmpty())
 	{
@@ -142,7 +142,7 @@ TArray<UEdGraphPin*> FollowOutputPinArray(const UEdGraphPin& Pin, bool* bOutCycl
 UEdGraphPin* FollowOutputPin(const UEdGraphPin& Pin, bool* CycleDetected)
 {
 	TArray<UEdGraphPin*> Result = FollowOutputPinArray(Pin, CycleDetected);
-	check(Result.Num() <= 1); // Use FollowPinImmersive if the pin can have more than one input.
+	check(Result.Num() <= 1); // Use FollowInputPinArray if the pin can have more than one input.
 
 	if (!Result.IsEmpty())
 	{
diff --git a/Engine/Plugins/Mutable/Source/CustomizableObjectEditor/Private/MuCOE/Nodes/CustomizableObjectNodeObject.cpp b/Engine/Plugins/Mutable/Source/CustomizableObjectEditor/Private/MuCOE/Nodes/CustomizableObjectNodeObject.cpp
--- a/Engine/Plugins/Mutable/Source/CustomizableObjectEditor/Private/MuCOE/Nodes/CustomizableObjectNodeObject.cpp
+++ b/Engine/Plugins/Mutable/Source/CustomizableObjectEditor/Private/MuCOE/Nodes/CustomizableObjectNodeObject.cpp
@@ -392,9 +392,11 @@ void UCustomizableObjectNodeObject::BackwardsCompatibleFixup(int32 CustomizableO
 						FName ComponentName;
 						bool bFirst = true;
 
-						if (UEdGraphPin* ConnectedPin = FollowInputPin(*NodeMaterialVariation->DefaultPin()))
+						TArray<UEdGraphPin*> ConnectedDefaultPins = FollowInputPinArray(*NodeMaterialVariation->DefaultPin(), nullptr);
+						if (ConnectedDefaultPins.Num() > 0)
 						{
-							if (UCustomizableObjectNodeMaterialBase* FirstNodeMaterialBase = Cast<UCustomizableObjectNodeMaterialBase>(ConnectedPin->GetOwningNode()))
+							UEdGraphPin* FirstDefaultPin = ConnectedDefaultPins[0];
+							if (UCustomizableObjectNodeMaterialBase* FirstNodeMaterialBase = Cast<UCustomizableObjectNodeMaterialBase>(FirstDefaultPin->GetOwningNode()))
 							{
 								if (UCustomizableObjectNodeMaterial* FirstNodeMaterial = FirstNodeMaterialBase->GetMaterialNode())
 								{
@@ -409,8 +411,10 @@ void UCustomizableObjectNodeObject::BackwardsCompatibleFixup(int32 CustomizableO
 							UEdGraphPin* VariationPin = NodeMaterialVariation->VariationPin(ElementIndex);
 							if (VariationPin)
 							{
-								if (UEdGraphPin* ConnectedPin = FollowInputPin(*VariationPin))
+								TArray<UEdGraphPin*> ConnectedPins = FollowInputPinArray(*VariationPin, nullptr);
+								if (ConnectedPins.Num() > 0)
 								{
+									UEdGraphPin* ConnectedPin = ConnectedPins[0];
 									if (UCustomizableObjectNodeMaterialBase* FirstNodeMaterialBase = Cast<UCustomizableObjectNodeMaterialBase>(ConnectedPin->GetOwningNode()))
 									{
 										if (UCustomizableObjectNodeMaterial* FirstNodeMaterial = FirstNodeMaterialBase->GetMaterialNode())
