Mutable: Add missing GetSourceDataDescriptor to ASTOpMeshTransformWithBoundingMesh

#rnx

diff --git a/Engine/Plugins/Mutable/Source/MutableTools/Private/MuT/ASTOpMeshTransformWithBoundingMesh.cpp b/Engine/Plugins/Mutable/Source/MutableTools/Private/MuT/ASTOpMeshTransformWithBoundingMesh.cpp
--- a/Engine/Plugins/Mutable/Source/MutableTools/Private/MuT/ASTOpMeshTransformWithBoundingMesh.cpp
+++ b/Engine/Plugins/Mutable/Source/MutableTools/Private/MuT/ASTOpMeshTransformWithBoundingMesh.cpp
@@ -77,3 +77,9 @@ void mu::ASTOpMeshTransformWithBoundingMesh::Link(FProgram& program, FLinkerOpti
 		AppendCode(program.m_byteCode, args);
 	}
 }
+
+mu::FSourceDataDescriptor mu::ASTOpMeshTransformWithBoundingMesh::GetSourceDataDescriptor(
+	FGetSourceDataDescriptorContext* GetSourceDataDescriptorContext) const
+{
+	return source->GetSourceDataDescriptor(GetSourceDataDescriptorContext);
+}
diff --git a/Engine/Plugins/Mutable/Source/MutableTools/Private/MuT/ASTOpMeshTransformWithBoundingMesh.h b/Engine/Plugins/Mutable/Source/MutableTools/Private/MuT/ASTOpMeshTransformWithBoundingMesh.h
--- a/Engine/Plugins/Mutable/Source/MutableTools/Private/MuT/ASTOpMeshTransformWithBoundingMesh.h
+++ b/Engine/Plugins/Mutable/Source/MutableTools/Private/MuT/ASTOpMeshTransformWithBoundingMesh.h
@@ -36,6 +36,7 @@ public:
 	Ptr<ASTOp> Clone(MapChildFuncRef mapChild) const override;
 	void ForEachChild(const TFunctionRef<void(ASTChild&)>) override;
 	void Link(FProgram& program, FLinkerOptions* Options) override;
+	FSourceDataDescriptor GetSourceDataDescriptor(FGetSourceDataDescriptorContext*) const override;
 };
 
 }
