[Mutable] Fix cooked data indeterminism due to streamed resources disappearing. FCustomizableObjectCookPackageSplitter was looking for streamed resources at the wrong place.

#rnx

diff --git a/Engine/Plugins/Mutable/Source/CustomizableObject/Private/MuCO/CustomizableObject.cpp b/Engine/Plugins/Mutable/Source/CustomizableObject/Private/MuCO/CustomizableObject.cpp
--- a/Engine/Plugins/Mutable/Source/CustomizableObject/Private/MuCO/CustomizableObject.cpp
+++ b/Engine/Plugins/Mutable/Source/CustomizableObject/Private/MuCO/CustomizableObject.cpp
@@ -86,6 +86,13 @@ TAutoConsoleVariable<int32> CVarMutableDerivedDataCacheUsage(
 	TEXT("2 - Default. Allow cache requests to query and store records and values in any caches."),
 	ECVF_Default);
 
+
+TAutoConsoleVariable<bool> CVarMutableAsyncCook(
+	TEXT("Mutable.CookAsync"),
+	false,
+	TEXT("True - Customizable Objects will be compiled asynchronously during cook.\n")
+	TEXT("False - Sync compilation.\n"));
+
 #endif
 
 #if WITH_EDITORONLY_DATA
@@ -1043,7 +1050,7 @@ void UCustomizableObjectPrivate::CompileForTargetPlatform(UCustomizableObject& C
 		return;
 	}
 
-	const bool bAsync = false; // TODO PERE
+	const bool bAsync = CVarMutableAsyncCook.GetValueOnAnyThread();
 
 	TSharedRef<FCompilationRequest> CompileRequest = MakeShared<FCompilationRequest>(CustomizableObject, bAsync);
 	FCompilationOptions& Options = CompileRequest->GetCompileOptions();
diff --git a/Engine/Plugins/Mutable/Source/CustomizableObjectEditor/Private/MuCOE/CustomizableObjectCompiler.cpp b/Engine/Plugins/Mutable/Source/CustomizableObjectEditor/Private/MuCOE/CustomizableObjectCompiler.cpp
--- a/Engine/Plugins/Mutable/Source/CustomizableObjectEditor/Private/MuCOE/CustomizableObjectCompiler.cpp
+++ b/Engine/Plugins/Mutable/Source/CustomizableObjectEditor/Private/MuCOE/CustomizableObjectCompiler.cpp
@@ -762,8 +762,7 @@ void FCustomizableObjectCompiler::CompileInternal(bool bAsync)
 	}
 	else
 	{
-		// Always work with the ModelResources (Editor) when compiling. They'll be copied to the cooked version during PreSave.
-		FModelResources& ModelResources = CurrentObject->GetPrivate()->GetModelResources(false);
+		FModelResources& ModelResources = CurrentObject->GetPrivate()->GetModelResources(CurrentOptions.bIsCooking);
 		ModelResources = FModelResources();
 		ModelStreamableBulkData = MakeShared<FModelStreamableBulkData>();
 		
@@ -1301,7 +1300,7 @@ void FCustomizableObjectCompiler::FinishCompilationTask()
 
 	// Create SaveDD task
 	TRACE_BEGIN_REGION(UE_MUTABLE_SAVEDD_REGION);
-	SaveDDTask = MakeShareable(new FCustomizableObjectSaveDDRunnable(CurrentRequest, Model, CurrentObject->GetPrivate()->GetModelResources(false), ModelStreamableBulkData));
+	SaveDDTask = MakeShareable(new FCustomizableObjectSaveDDRunnable(CurrentRequest, Model, CurrentObject->GetPrivate()->GetModelResources(CurrentOptions.bIsCooking), ModelStreamableBulkData));
 }
 
 
diff --git a/Engine/Plugins/Mutable/Source/CustomizableObjectEditor/Private/MuCOE/CustomizableObjectCookPackageSplitter.h b/Engine/Plugins/Mutable/Source/CustomizableObjectEditor/Private/MuCOE/CustomizableObjectCookPackageSplitter.h
--- a/Engine/Plugins/Mutable/Source/CustomizableObjectEditor/Private/MuCOE/CustomizableObjectCookPackageSplitter.h
+++ b/Engine/Plugins/Mutable/Source/CustomizableObjectEditor/Private/MuCOE/CustomizableObjectCookPackageSplitter.h
@@ -13,6 +13,7 @@ public:
 	/** ICookPackageSplitter interface */
 	static bool ShouldSplit(UObject* SplitData);
 	static FString GetSplitterDebugName() { return TEXT("FCustomizableObjectCookPackageSplitter"); }
+	static bool RequiresCachedCookedPlatformDataBeforeSplit() { return true; }
 
 	virtual TArray<ICookPackageSplitter::FGeneratedPackage> GetGenerateList(
 		const UPackage* OwnerPackage,
