[mutable] Fix node texture preview being incorrect in pass-through texture nodes.
#rnx

diff --git a/Engine/Plugins/Mutable/Source/CustomizableObjectEditor/Private/MuCOE/Nodes/CustomizableObjectNodePassThroughTexture.cpp b/Engine/Plugins/Mutable/Source/CustomizableObjectEditor/Private/MuCOE/Nodes/CustomizableObjectNodePassThroughTexture.cpp
--- a/Engine/Plugins/Mutable/Source/CustomizableObjectEditor/Private/MuCOE/Nodes/CustomizableObjectNodePassThroughTexture.cpp
+++ b/Engine/Plugins/Mutable/Source/CustomizableObjectEditor/Private/MuCOE/Nodes/CustomizableObjectNodePassThroughTexture.cpp
@@ -14,14 +14,14 @@ void UCustomizableObjectNodePassThroughTexture::BackwardsCompatibleFixup(int32 C
 	
 	if (CustomizableObjectCustomVersion == FCustomizableObjectCustomVersion::AddedAnyTextureTypeToPassThroughTextures)
 	{
-		if (Texture)
+		if (Texture_DEPRECATED)
 		{
 			if (!PassThroughTexture)
 			{
-				PassThroughTexture = Texture;
+				PassThroughTexture = Texture_DEPRECATED;
 			}
 
-			Texture = nullptr;
+			Texture_DEPRECATED = nullptr;
 		}
 	}
 }
@@ -39,10 +39,10 @@ void UCustomizableObjectNodePassThroughTexture::AllocateDefaultPins(UCustomizabl
 
 FText UCustomizableObjectNodePassThroughTexture::GetNodeTitle(ENodeTitleType::Type TitleType) const
 {
-	if (Texture)
+	if (PassThroughTexture)
 	{
 		FFormatNamedArguments Args;
-		Args.Add(TEXT("TextureName"), FText::FromString(Texture->GetName()));
+		Args.Add(TEXT("TextureName"), FText::FromString(PassThroughTexture->GetName()));
 
 		return FText::Format(LOCTEXT("Passthrough Texture_Title", "{TextureName}\nPassthrough Texture"), Args);
 	}
@@ -65,4 +65,11 @@ FText UCustomizableObjectNodePassThroughTexture::GetTooltipText() const
 	return LOCTEXT("PassThrough_Texture_Tooltip", "Defines a pass-through texture. It will not be modified by Mutable in any way, just referenced as a UE asset. It's much cheaper than a Mutable texture, but you cannot make any operations on it, just switch it.");
 }
 
+
+TObjectPtr<UTexture> UCustomizableObjectNodePassThroughTexture::GetTexture()
+{
+	return PassThroughTexture;
+}
+
+
 #undef LOCTEXT_NAMESPACE
diff --git a/Engine/Plugins/Mutable/Source/CustomizableObjectEditor/Private/MuCOE/Nodes/CustomizableObjectNodeTexture.cpp b/Engine/Plugins/Mutable/Source/CustomizableObjectEditor/Private/MuCOE/Nodes/CustomizableObjectNodeTexture.cpp
--- a/Engine/Plugins/Mutable/Source/CustomizableObjectEditor/Private/MuCOE/Nodes/CustomizableObjectNodeTexture.cpp
+++ b/Engine/Plugins/Mutable/Source/CustomizableObjectEditor/Private/MuCOE/Nodes/CustomizableObjectNodeTexture.cpp
@@ -2,15 +2,16 @@
 
 #include "MuCOE/Nodes/CustomizableObjectNodeTexture.h"
 
-#include "ISinglePropertyView.h"
-#include "Modules/ModuleManager.h"
+#include "MuCOE/Nodes/CustomizableObjectNodePassThroughTexture.h"
 #include "MuCO/CustomizableObjectCustomVersion.h"
 #include "MuCOE/CustomizableObjectEditorStyle.h"
 #include "MuCOE/EdGraphSchema_CustomizableObject.h"
+#include "ISinglePropertyView.h"
 #include "Widgets/Images/SImage.h"
 #include "Widgets/Input/SCheckBox.h"
 #include "Widgets/SBoxPanel.h"
 #include "Widgets/SOverlay.h"
+#include "Modules/ModuleManager.h"
 
 class UCustomizableObjectNodeRemapPins;
 struct FGeometry;
@@ -71,6 +72,12 @@ FLinearColor UCustomizableObjectNodeTexture::GetNodeTitleColor() const
 }
 
 
+TObjectPtr<UTexture> UCustomizableObjectNodeTexture::GetTexture()
+{
+	return Texture;
+}
+
+
 FText UCustomizableObjectNodeTexture::GetTooltipText() const
 {
 	return LOCTEXT("Texture_Tooltip", "Defines a texture.");
@@ -87,9 +94,21 @@ void SGraphNodeTexture::Construct(const FArguments& InArgs, UEdGraphNode* InGrap
 	SingleDetails.NamePlacement = EPropertyNamePlacement::Hidden;
 	SingleDetails.bHideAssetThumbnail = true;
 
-	TextureSelector = PropPlugin.CreateSingleProperty(NodeTexture, "Texture", SingleDetails);
+	if (Cast<UCustomizableObjectNodeTexture>(InGraphNode))
+	{
+		TextureSelector = PropPlugin.CreateSingleProperty(NodeTexture, "Texture", SingleDetails);
+	}
+	else if (Cast<UCustomizableObjectNodePassThroughTexture>(InGraphNode))
+	{
+		TextureSelector = PropPlugin.CreateSingleProperty(NodeTexture, "PassThroughTexture", SingleDetails);
+	}
+	else
+	{
+		// Node type not supported.
+		ensure(false);
+	}
 
-	TextureBrush.SetResourceObject(NodeTexture->Texture);
+	TextureBrush.SetResourceObject(NodeTexture->GetTexture());
 	TextureBrush.ImageSize.X = 128.0f;
 	TextureBrush.ImageSize.Y = 128.0f;
 	TextureBrush.DrawAs = ESlateBrushDrawType::Image;
@@ -194,13 +213,11 @@ EVisibility SGraphNodeTexture::ExpressionPreviewVisibility() const
 
 void SGraphNodeTexture::Tick(const FGeometry& AllottedGeometry, const double InCurrentTime, const float InDeltaTime)
 {
-	UTexture2D* BrushTexture = Cast<UTexture2D>(TextureBrush.GetResourceObject());
-	
-	if (NodeTexture && BrushTexture)
+	if (NodeTexture)
 	{
-		if (NodeTexture->Texture != BrushTexture)
+		if (NodeTexture->GetTexture() != TextureBrush.GetResourceObject())
 		{
-			TextureBrush.SetResourceObject(NodeTexture->Texture);
+			TextureBrush.SetResourceObject(NodeTexture->GetTexture());
 		}
 	}
 }
diff --git a/Engine/Plugins/Mutable/Source/CustomizableObjectEditor/Public/MuCOE/Nodes/CustomizableObjectNodePassThroughTexture.h b/Engine/Plugins/Mutable/Source/CustomizableObjectEditor/Public/MuCOE/Nodes/CustomizableObjectNodePassThroughTexture.h
--- a/Engine/Plugins/Mutable/Source/CustomizableObjectEditor/Public/MuCOE/Nodes/CustomizableObjectNodePassThroughTexture.h
+++ b/Engine/Plugins/Mutable/Source/CustomizableObjectEditor/Public/MuCOE/Nodes/CustomizableObjectNodePassThroughTexture.h
@@ -18,12 +18,20 @@ public:
 
 	// UCustomizableObjectNode interface
 	virtual void BackwardsCompatibleFixup(int32 CustomizableObjectCustomVersion) override;
-
-	// UCustomizableObjectNode interface
 	void AllocateDefaultPins(UCustomizableObjectNodeRemapPins* RemapPins) override;
 
+	// UCustomizableObjectNodeTextureBase interface
+	virtual TObjectPtr<UTexture> GetTexture() override;
+
 	// Begin EdGraphNode interface
 	FText GetNodeTitle(ENodeTitleType::Type TitleType) const override;
 	FLinearColor GetNodeTitleColor() const override;
 	FText GetTooltipText() const override;
+
+private:
+
+	// For backwards compatibility
+	UPROPERTY()
+	TObjectPtr<UTexture2D> Texture_DEPRECATED = nullptr;
+
 };
diff --git a/Engine/Plugins/Mutable/Source/CustomizableObjectEditor/Public/MuCOE/Nodes/CustomizableObjectNodeTexture.h b/Engine/Plugins/Mutable/Source/CustomizableObjectEditor/Public/MuCOE/Nodes/CustomizableObjectNodeTexture.h
--- a/Engine/Plugins/Mutable/Source/CustomizableObjectEditor/Public/MuCOE/Nodes/CustomizableObjectNodeTexture.h
+++ b/Engine/Plugins/Mutable/Source/CustomizableObjectEditor/Public/MuCOE/Nodes/CustomizableObjectNodeTexture.h
@@ -58,20 +58,20 @@ public:
 };
 
 
-UCLASS()
+UCLASS(Abstract)
 class CUSTOMIZABLEOBJECTEDITOR_API UCustomizableObjectNodeTextureBase : public UCustomizableObjectNode
 {
 public:
 	GENERATED_BODY()
 
-	UPROPERTY(EditAnywhere, Category = Texture2D)
-	TObjectPtr<UTexture2D> Texture = nullptr;
-
 	// Determines if the Node is collapsed or not
 	bool bCollapsed = true;
 
 	// Creates the SGraph Node widget for the thumbnail
 	TSharedPtr<SGraphNode> CreateVisualWidget() override;
+
+	/** */
+	virtual TObjectPtr<UTexture> GetTexture() PURE_VIRTUAL(UCustomizableObjectNodeTextureBase::GetTexture, return {};);
 };
 
 
@@ -81,10 +81,16 @@ class CUSTOMIZABLEOBJECTEDITOR_API UCustomizableObjectNodeTexture : public UCust
 public:
 	GENERATED_BODY()
 
+	UPROPERTY(EditAnywhere, Category = Texture2D)
+	TObjectPtr<UTexture2D> Texture = nullptr;
+
 	// UCustomizableObjectNode interface
 	virtual void AllocateDefaultPins(UCustomizableObjectNodeRemapPins* RemapPins) override;
 	virtual void BackwardsCompatibleFixup(int32 CustomizableObjectCustomVersion) override;
 
+	// UCustomizableObjectNodeTextureBase interface
+	virtual TObjectPtr<UTexture> GetTexture() override;
+
 	// Begin EdGraphNode interface
 	FText GetNodeTitle(ENodeTitleType::Type TitleType) const override;
 	FLinearColor GetNodeTitleColor() const override;
