[Mutable] Fix shared layouts when the layout operations are set from mesh extensions.

#rnx

diff --git a/Engine/Plugins/Mutable/Source/MutableTools/Private/MuT/CodeGenerator.cpp b/Engine/Plugins/Mutable/Source/MutableTools/Private/MuT/CodeGenerator.cpp
--- a/Engine/Plugins/Mutable/Source/MutableTools/Private/MuT/CodeGenerator.cpp
+++ b/Engine/Plugins/Mutable/Source/MutableTools/Private/MuT/CodeGenerator.cpp
@@ -1011,39 +1011,53 @@ namespace mu
 						continue;
 					}
 
+					bool bLayoutSetByThisExtension = false;
 					if (!SurfaceReferenceLayouts[LayoutIndex].Layout)
 					{
 						// This Layout slot is not set by the base surface, set it as reference.
 						SurfaceReferenceLayouts[LayoutIndex] = ExtraGeneratedLayouts[LayoutIndex];
-						LayoutFromExtension[LayoutIndex] = true;
-					}
-
-					Ptr<ASTOpConstantResource> LayoutFragmentConstantOp = new ASTOpConstantResource();
-					LayoutFragmentConstantOp->Type = OP_TYPE::LA_CONSTANT;
-
-					LayoutFragmentConstantOp->SetValue(
-							ExtraLayoutsData.GeneratedLayouts[LayoutIndex].Layout,
-							CompilerOptions->OptimisationOptions.DiskCacheContext);
-
-					Ptr<ASTOpLayoutMerge> LayoutMergeOp = new ASTOpLayoutMerge();
-					// Base may be null if the base does not have  a mesh with a layout at LayoutIndex.
-					// In that case, when applying the condition this can generate null layouts.
-					LayoutMergeOp->Base = SurfaceLayoutOps[LayoutIndex];
-					LayoutMergeOp->Added = LayoutFragmentConstantOp;
+						bLayoutSetByThisExtension = true;
 
-					if (ExtraLayoutsData.Condition)
+						LayoutFromExtension[LayoutIndex] = bLayoutSetByThisExtension;
+					}
+					
+					if (SharedMeshResults)
 					{
-						Ptr<ASTOpConditional> ConditionalOp = new ASTOpConditional();
-						ConditionalOp->type = OP_TYPE::LA_CONDITIONAL;
-						ConditionalOp->no = SurfaceLayoutOps[LayoutIndex];
-						ConditionalOp->yes = LayoutMergeOp;
-						ConditionalOp->condition = ExtraLayoutsData.Condition;
-
-						SurfaceLayoutOps[LayoutIndex] = ConditionalOp;
+						if (!SurfaceLayoutOps[LayoutIndex] && bLayoutSetByThisExtension)
+						{
+							check(SharedMeshResults->LayoutOps.IsValidIndex(LayoutIndex));
+							SurfaceLayoutOps[LayoutIndex] = SharedMeshResults->LayoutOps[LayoutIndex];
+						}
 					}
 					else
 					{
-						SurfaceLayoutOps[LayoutIndex] = LayoutMergeOp;
+						Ptr<ASTOpConstantResource> LayoutFragmentConstantOp = new ASTOpConstantResource();
+						LayoutFragmentConstantOp->Type = OP_TYPE::LA_CONSTANT;
+
+						LayoutFragmentConstantOp->SetValue(
+								ExtraLayoutsData.GeneratedLayouts[LayoutIndex].Layout,
+								CompilerOptions->OptimisationOptions.DiskCacheContext);
+
+						Ptr<ASTOpLayoutMerge> LayoutMergeOp = new ASTOpLayoutMerge();
+						// Base may be null if the base does not have  a mesh with a layout at LayoutIndex.
+						// In that case, when applying the condition this can generate null layouts.
+						LayoutMergeOp->Base = SurfaceLayoutOps[LayoutIndex];
+						LayoutMergeOp->Added = LayoutFragmentConstantOp;
+
+						if (ExtraLayoutsData.Condition)
+						{
+							Ptr<ASTOpConditional> ConditionalOp = new ASTOpConditional();
+							ConditionalOp->type = OP_TYPE::LA_CONDITIONAL;
+							ConditionalOp->no = SurfaceLayoutOps[LayoutIndex];
+							ConditionalOp->yes = LayoutMergeOp;
+							ConditionalOp->condition = ExtraLayoutsData.Condition;
+
+							SurfaceLayoutOps[LayoutIndex] = ConditionalOp;
+						}
+						else
+						{
+							SurfaceLayoutOps[LayoutIndex] = LayoutMergeOp;
+						}
 					}
 				}
 			}
