[mutable] Fixed texture streaming taking 5 seconds because now mesh components have to be explicitely notified of any texture changes to them.
#rnx
#tests Editor

diff --git a/Engine/Plugins/Mutable/Source/CustomizableObject/Private/MuCO/CustomizableObjectInstance.cpp b/Engine/Plugins/Mutable/Source/CustomizableObject/Private/MuCO/CustomizableObjectInstance.cpp
--- a/Engine/Plugins/Mutable/Source/CustomizableObject/Private/MuCO/CustomizableObjectInstance.cpp
+++ b/Engine/Plugins/Mutable/Source/CustomizableObject/Private/MuCO/CustomizableObjectInstance.cpp
@@ -78,6 +78,13 @@ static FAutoConsoleVariableRef CVarDisableClothingPhysicsEditsPropagation(
 	bDisableClothingPhysicsEditsPropagation,
 	TEXT("If set to true, disables clothing physics edits propagation from the render mesh."),
 	ECVF_Default);
+
+bool bDisableNotifyComponentsOfTextureUpdates = false;
+static FAutoConsoleVariableRef CVarDisableNotifyComponentsOfTextureUpdates(
+	TEXT("mutable.DisableNotifyComponentsOfTextureUpdates"),
+	bDisableNotifyComponentsOfTextureUpdates,
+	TEXT("If set to true, disables Mutable notifying the streaming system that a component has had a change in at least one texture of its components."),
+	ECVF_Default);
 }
 
 
@@ -6155,6 +6162,9 @@ void UCustomizableInstancePrivate::BuildMaterials(const TSharedRef<FUpdateContex
 	TArray<bool> RecreateRenderStateOnInstanceComponent;
 	RecreateRenderStateOnInstanceComponent.Init(false, OperationData->NumInstanceComponents);
 
+	TArray<bool> NotifyUpdateOnInstanceComponent;
+	NotifyUpdateOnInstanceComponent.Init(false, OperationData->NumInstanceComponents);
+
 	for (int32 InstanceComponentIndex = 0; InstanceComponentIndex < OperationData->NumInstanceComponents; ++InstanceComponentIndex)
 	{
 		const FInstanceUpdateData::FComponent& Component = OperationData->InstanceUpdateData.Components[InstanceComponentIndex];
@@ -6741,6 +6751,11 @@ void UCustomizableInstancePrivate::BuildMaterials(const TSharedRef<FUpdateContex
 									MaterialInstance->SetTextureParameterValueByInfo(ParameterInfo, Param.Texture.Texture);
 								}
 
+								if (!bDisableNotifyComponentsOfTextureUpdates)
+								{
+									NotifyUpdateOnInstanceComponent[InstanceComponentIndex] = true;
+								}
+
 								Material.Textures.Add(Param.Texture);
 
 								break;
@@ -6799,7 +6814,7 @@ void UCustomizableInstancePrivate::BuildMaterials(const TSharedRef<FUpdateContex
 
 	// Force recreate render state if the mesh is reused and the materials have changed.
 	// TODO: MTBL-1697 Remove after merging ConvertResources and Callbacks.
-	if (RecreateRenderStateOnInstanceComponent.Find(true) != INDEX_NONE)
+	if (RecreateRenderStateOnInstanceComponent.Find(true) != INDEX_NONE || NotifyUpdateOnInstanceComponent.Find(true) != INDEX_NONE)
 	{
 		MUTABLE_CPUPROFILER_SCOPE(BuildMaterials_RecreateRenderState);
 
@@ -6832,7 +6847,10 @@ void UCustomizableInstancePrivate::BuildMaterials(const TSharedRef<FUpdateContex
 				}
 			}
 
-			if (!RecreateRenderStateOnInstanceComponent.IsValidIndex(InstanceComponentIndex) || !RecreateRenderStateOnInstanceComponent[InstanceComponentIndex])
+			bool bDoRecreateRenderStateOnComponent = RecreateRenderStateOnInstanceComponent.IsValidIndex(InstanceComponentIndex) && RecreateRenderStateOnInstanceComponent[InstanceComponentIndex];
+			bool bDoNotifyUpdateOnComponent = NotifyUpdateOnInstanceComponent.IsValidIndex(InstanceComponentIndex) && NotifyUpdateOnInstanceComponent[InstanceComponentIndex];
+
+			if (!bDoRecreateRenderStateOnComponent && !bDoNotifyUpdateOnComponent)
 			{
 				continue;
 			}
@@ -6844,7 +6862,14 @@ void UCustomizableInstancePrivate::BuildMaterials(const TSharedRef<FUpdateContex
 				continue;
 			}
 
-			AttachedParent->RecreateRenderState_Concurrent();
+			if (bDoRecreateRenderStateOnComponent)
+			{
+				AttachedParent->RecreateRenderState_Concurrent();
+			}
+			else if (bDoNotifyUpdateOnComponent)
+			{
+				IStreamingManager::Get().NotifyPrimitiveUpdated(AttachedParent);
+			}
 		}
 	}
 
