[mutable] Ensure modifiers are applied to only the relevant component.
- Add component id filter to modifiers.
- Implement missing method for the transform in mesh modifier.
- Ensure the transform mesh in mesh operation is highlighted as expensive in the debugger.
- Optimize the mesh-transform-in-bounding-mesh versus the mesh-format
#rnx

diff --git a/Engine/Plugins/Mutable/Source/CustomizableObjectEditor/Private/MuCOE/GenerateMutableSource/GenerateMutableSourceModifier.cpp b/Engine/Plugins/Mutable/Source/CustomizableObjectEditor/Private/MuCOE/GenerateMutableSource/GenerateMutableSourceModifier.cpp
--- a/Engine/Plugins/Mutable/Source/CustomizableObjectEditor/Private/MuCOE/GenerateMutableSource/GenerateMutableSourceModifier.cpp
+++ b/Engine/Plugins/Mutable/Source/CustomizableObjectEditor/Private/MuCOE/GenerateMutableSource/GenerateMutableSourceModifier.cpp
@@ -234,7 +234,6 @@ mu::Ptr<mu::NodeModifier> GenerateMutableSourceModifier(const UEdGraphPin * Pin,
 		mu::Ptr<mu::NodeModifierMeshClipWithUVMask> ClipNode = new mu::NodeModifierMeshClipWithUVMask();
 		Result = ClipNode;
 
-		ClipNode->SetMessageContext(Node);
 		ClipNode->FaceCullStrategy = TypedNodeClipUVMask->FaceCullStrategy;
 
 		if (const UEdGraphPin* ConnectedPin = FollowInputPin(*TypedNodeClipUVMask->ClipMaskPin()))
@@ -266,7 +265,6 @@ mu::Ptr<mu::NodeModifier> GenerateMutableSourceModifier(const UEdGraphPin * Pin,
 		GenerationContext.MeshGenerationFlags.Push(ModifiersMeshFlags);
 
 		mu::Ptr<mu::NodeModifierSurfaceEdit> SurfNode = new mu::NodeModifierSurfaceEdit();
-		SurfNode->SetMessageContext(Node);
 		Result = SurfNode;
 
 		// TODO: This was used in the non-modifier version for group projectors. It may affect the "drop projection from LOD" feature.
@@ -444,7 +442,6 @@ mu::Ptr<mu::NodeModifier> GenerateMutableSourceModifier(const UEdGraphPin * Pin,
 		mu::Ptr<mu::NodeModifierSurfaceEdit> SurfNode = new mu::NodeModifierSurfaceEdit();
 		Result = SurfNode;
 
-		SurfNode->SetMessageContext(Node);
 		SurfNode->MultipleTagsPolicy = TypedNodeRem->MultipleTagPolicy;
 		SurfNode->RequiredTags = TypedNodeRem->RequiredTags;
 
@@ -482,7 +479,6 @@ mu::Ptr<mu::NodeModifier> GenerateMutableSourceModifier(const UEdGraphPin * Pin,
 		mu::Ptr<mu::NodeModifierMeshClipWithUVMask> ClipNode = new mu::NodeModifierMeshClipWithUVMask();
 		Result = ClipNode;
 
-		ClipNode->SetMessageContext(Node);
 		ClipNode->FaceCullStrategy = TypedNodeRemBlocks->FaceCullStrategy;
 
 		ClipNode->MultipleTagsPolicy = TypedNodeRemBlocks->MultipleTagPolicy;
@@ -506,7 +502,6 @@ mu::Ptr<mu::NodeModifier> GenerateMutableSourceModifier(const UEdGraphPin * Pin,
 		mu::Ptr<mu::NodeModifierSurfaceEdit> SurfNode = new mu::NodeModifierSurfaceEdit();
 		Result = SurfNode;
 
-		SurfNode->SetMessageContext(Node);
 		SurfNode->MultipleTagsPolicy = TypedNodeEdit->MultipleTagPolicy;
 		SurfNode->RequiredTags = TypedNodeEdit->RequiredTags;
 
@@ -584,8 +579,6 @@ mu::Ptr<mu::NodeModifier> GenerateMutableSourceModifier(const UEdGraphPin * Pin,
 		mu::Ptr<mu::NodeModifierSurfaceEdit> SurfNode = new mu::NodeModifierSurfaceEdit();
 		Result = SurfNode;
 
-		SurfNode->SetMessageContext(Node);
-
 		// This modifier needs to be applied right after the mesh constant is generated
 		SurfNode->bApplyBeforeNormalOperations = true;
 
@@ -651,7 +644,6 @@ mu::Ptr<mu::NodeModifier> GenerateMutableSourceModifier(const UEdGraphPin * Pin,
 		bDoNotAddToGeneratedCache = true;
 
 		mu::Ptr<mu::NodeModifierMeshTransformInMesh> TransformNode = new mu::NodeModifierMeshTransformInMesh();
-		TransformNode->SetMessageContext(Node);
 		Result = TransformNode;
 
 		if (const UEdGraphPin* ConnectedPin = FollowInputPin(*TypedNodeTransformMesh->TransformPin()))
@@ -715,6 +707,15 @@ mu::Ptr<mu::NodeModifier> GenerateMutableSourceModifier(const UEdGraphPin * Pin,
 		GenerationContext.Log(LOCTEXT("UnimplementedNode", "Node type not implemented yet."), Node);
 	}
 
+	if (Result)
+	{
+		Result->SetMessageContext(Node);
+
+		int32 ComponentId = GenerationContext.ComponentNames.IndexOfByKey(GenerationContext.CurrentMeshComponent);
+		check(ComponentId>=0);
+		Result->RequiredComponentId = ComponentId;
+	}
+
 	if (!bDoNotAddToGeneratedCache)
 	{
 		GenerationContext.Generated.Add(Key, FGeneratedData(Node, Result));
diff --git a/Engine/Plugins/Mutable/Source/CustomizableObjectEditor/Private/MuCOE/SMutableCodeViewer.h b/Engine/Plugins/Mutable/Source/CustomizableObjectEditor/Private/MuCOE/SMutableCodeViewer.h
--- a/Engine/Plugins/Mutable/Source/CustomizableObjectEditor/Private/MuCOE/SMutableCodeViewer.h
+++ b/Engine/Plugins/Mutable/Source/CustomizableObjectEditor/Private/MuCOE/SMutableCodeViewer.h
@@ -232,6 +232,7 @@ private:
 		mu::OP_TYPE::IM_PIXELFORMAT,
 		mu::OP_TYPE::IM_SWIZZLE,
 		mu::OP_TYPE::ME_PROJECT,
+		mu::OP_TYPE::ME_TRANSFORMWITHMESH,
 	};
 	
 	/** Enum designed to be able to notify the row generation of the type of operation being generated */
diff --git a/Engine/Plugins/Mutable/Source/MutableTools/Private/MuT/ASTOpMeshFormat.cpp b/Engine/Plugins/Mutable/Source/MutableTools/Private/MuT/ASTOpMeshFormat.cpp
--- a/Engine/Plugins/Mutable/Source/MutableTools/Private/MuT/ASTOpMeshFormat.cpp
+++ b/Engine/Plugins/Mutable/Source/MutableTools/Private/MuT/ASTOpMeshFormat.cpp
@@ -14,6 +14,7 @@
 #include "MuT/ASTOpConditional.h"
 #include "MuT/ASTOpConstantResource.h"
 #include "MuT/ASTOpMeshClipMorphPlane.h"
+#include "MuT/ASTOpMeshTransformWithBoundingMesh.h"
 #include "MuT/ASTOpMeshRemoveMask.h"
 #include "MuT/ASTOpMeshMorph.h"
 #include "MuT/ASTOpMeshAddTags.h"
@@ -267,12 +268,24 @@ mu::Ptr<ASTOp> Sink_MeshFormatAST::Visit(const mu::Ptr<ASTOp>& at, const ASTOpMe
 
 	case OP_TYPE::ME_CLIPMORPHPLANE:
 	{
-		auto newOp = mu::Clone<ASTOpMeshClipMorphPlane>(at);
+		Ptr<ASTOpMeshClipMorphPlane> newOp = mu::Clone<ASTOpMeshClipMorphPlane>(at);
 		newOp->source = Visit(newOp->source.child(), currentFormatOp);
 		newAt = newOp;
 		break;
 	}
 
+	case OP_TYPE::ME_TRANSFORMWITHMESH:
+	{
+		Ptr<ASTOpMeshTransformWithBoundingMesh> NewOp = mu::Clone<ASTOpMeshTransformWithBoundingMesh>(at);
+		NewOp->source = Visit(NewOp->source.child(), currentFormatOp);
+
+		// Don't transform the bounding mesh: it should be optimized with a different specific format elsewhere (TODO).
+		// NewOp->boundingMesh = Visit(NewOp->boundingMesh.child(), currentFormatOp);
+
+		newAt = NewOp;
+		break;
+	}
+
 	case OP_TYPE::ME_MORPH:
 	{
 		// Move the format down the base of the morph
diff --git a/Engine/Plugins/Mutable/Source/MutableTools/Private/MuT/ASTOpMeshTransformWithBoundingMesh.cpp b/Engine/Plugins/Mutable/Source/MutableTools/Private/MuT/ASTOpMeshTransformWithBoundingMesh.cpp
--- a/Engine/Plugins/Mutable/Source/MutableTools/Private/MuT/ASTOpMeshTransformWithBoundingMesh.cpp
+++ b/Engine/Plugins/Mutable/Source/MutableTools/Private/MuT/ASTOpMeshTransformWithBoundingMesh.cpp
@@ -78,8 +78,12 @@ void mu::ASTOpMeshTransformWithBoundingMesh::Link(FProgram& program, FLinkerOpti
 	}
 }
 
-mu::FSourceDataDescriptor mu::ASTOpMeshTransformWithBoundingMesh::GetSourceDataDescriptor(
-	FGetSourceDataDescriptorContext* GetSourceDataDescriptorContext) const
+mu::FSourceDataDescriptor mu::ASTOpMeshTransformWithBoundingMesh::GetSourceDataDescriptor(FGetSourceDataDescriptorContext* Context) const
 {
-	return source->GetSourceDataDescriptor(GetSourceDataDescriptorContext);
+	if (source)
+	{
+		return source->GetSourceDataDescriptor(Context);
+	}
+
+	return {};
 }
diff --git a/Engine/Plugins/Mutable/Source/MutableTools/Private/MuT/ASTOpMeshTransformWithBoundingMesh.h b/Engine/Plugins/Mutable/Source/MutableTools/Private/MuT/ASTOpMeshTransformWithBoundingMesh.h
--- a/Engine/Plugins/Mutable/Source/MutableTools/Private/MuT/ASTOpMeshTransformWithBoundingMesh.h
+++ b/Engine/Plugins/Mutable/Source/MutableTools/Private/MuT/ASTOpMeshTransformWithBoundingMesh.h
@@ -36,7 +36,8 @@ public:
 	Ptr<ASTOp> Clone(MapChildFuncRef mapChild) const override;
 	void ForEachChild(const TFunctionRef<void(ASTChild&)>) override;
 	void Link(FProgram& program, FLinkerOptions* Options) override;
-	FSourceDataDescriptor GetSourceDataDescriptor(FGetSourceDataDescriptorContext*) const override;
+	virtual FSourceDataDescriptor GetSourceDataDescriptor(FGetSourceDataDescriptorContext*) const override;
+
 };
 
 }
diff --git a/Engine/Plugins/Mutable/Source/MutableTools/Private/MuT/CodeGenerator.cpp b/Engine/Plugins/Mutable/Source/MutableTools/Private/MuT/CodeGenerator.cpp
--- a/Engine/Plugins/Mutable/Source/MutableTools/Private/MuT/CodeGenerator.cpp
+++ b/Engine/Plugins/Mutable/Source/MutableTools/Private/MuT/CodeGenerator.cpp
@@ -466,7 +466,7 @@ namespace mu
 				Ptr<NodeImageConstant> pNode = new NodeImageConstant();
 				pNode->SetValue(PatchMask.get());
 
-				FImageGenerationOptions ConstantOptions;
+				FImageGenerationOptions ConstantOptions(-1);
 				FImageGenerationResult ConstantResult;
 				GenerateImage(ConstantOptions, ConstantResult, pNode);
 				RectConstantOp = ConstantResult.op;
@@ -906,7 +906,8 @@ namespace mu
 
 		// Store the data necessary to apply modifiers for the pre-normal operations stage.
 		// TODO: Should we merge with currently active tags from the InOptions?
-		GetModifiersFor(SurfaceNode->Tags, bModifiersForBeforeOperations, Modifiers);
+		int32 ComponentId = Options.Component ? Options.Component->Id : -1;
+		GetModifiersFor(ComponentId, SurfaceNode->Tags, bModifiersForBeforeOperations, Modifiers);
 
 		// This pass on the modifiers is only to detect errors that cannot be detected at the point they are applied.
 		CheckModifiersForSurface(*SurfaceNode, Modifiers);
@@ -919,7 +920,7 @@ namespace mu
             Ptr<ASTOp> LastMeshOp;
 
             // Generate the mesh
-			FMeshGenerationOptions MeshOptions;
+			FMeshGenerationOptions MeshOptions(ComponentId);
 			MeshOptions.bLayouts = true;
 			MeshOptions.State = Options.State;
 			MeshOptions.ActiveTags = SurfaceNode->Tags;
@@ -1195,7 +1196,7 @@ namespace mu
 					if (ImageLayoutStrategy == CompilerOptions::TextureLayoutStrategy::None)
 					{
 						// Generate the image
-						FImageGenerationOptions ImageOptions;
+						FImageGenerationOptions ImageOptions(ComponentId);
 						ImageOptions.State = Options.State;
 						ImageOptions.ImageLayoutStrategy = ImageLayoutStrategy;
 						ImageOptions.ActiveTags = SurfaceNode->Tags;
@@ -1318,7 +1319,7 @@ namespace mu
 								for (int32 BlockIndex = 0; BlockIndex < pLayout->GetBlockCount(); ++BlockIndex)
 								{
 									// Generate the image
-									FImageGenerationOptions ImageOptions;
+									FImageGenerationOptions ImageOptions(ComponentId);
 									ImageOptions.State = Options.State;
 									ImageOptions.ImageLayoutStrategy = ImageLayoutStrategy;
 									ImageOptions.RectSize = { 0,0 };
@@ -1378,10 +1379,10 @@ namespace mu
 							}
 							check(imageAd);
 
-							FMeshGenerationOptions ModifierOptions;
+							FMeshGenerationOptions ModifierOptions(ComponentId);
 							ModifierOptions.State = Options.State;
 							ModifierOptions.ActiveTags = SurfaceNode->Tags;
-							imageAd = ApplyImageExtendModifiers( Modifiers, ModifierOptions, MeshResults, imageAd, ImageLayoutStrategy, 
+							imageAd = ApplyImageExtendModifiers( Modifiers, ModifierOptions, ComponentId, MeshResults, imageAd, ImageLayoutStrategy, 
 								LayoutIndex, ImageData, GridSize, LayoutBlockDesc, 
 								SurfaceNode->GetMessageContext());
 
@@ -1923,6 +1924,7 @@ namespace mu
 
 	//---------------------------------------------------------------------------------------------
 	void CodeGenerator::GetModifiersFor(
+		int32 ComponentId,
 		const TArray<FString>& SurfaceTags,
 		bool bModifiersForBeforeOperations,
 		TArray<FirstPassGenerator::FModifier>& OutModifiers)
@@ -1936,12 +1938,23 @@ namespace mu
 
 		for (const FirstPassGenerator::FModifier& m: FirstPass.Modifiers)
 		{
+			if (!m.Node)
+			{
+				continue;
+			}
+
 			// Correct stage?
 			if (m.Node->bApplyBeforeNormalOperations != bModifiersForBeforeOperations)
 			{
 				continue;
 			}
 
+			// Correct component?
+			if (m.Node->RequiredComponentId>=0 && m.Node->RequiredComponentId!=ComponentId)
+			{
+				continue;
+			}
+
 			// Already there?
 			bool bAlreadyAdded = 
 				OutModifiers.FindByPredicate( [&m](const FirstPassGenerator::FModifier& c) {return c.Node == m.Node; })
@@ -2047,7 +2060,7 @@ namespace mu
 
 					TArray<FirstPassGenerator::FModifier> ChildModifiers;
 					constexpr bool bModifiersForBeforeOperations = false;
-					GetModifiersFor(ModifierOptions.ActiveTags, bModifiersForBeforeOperations, ChildModifiers);
+					GetModifiersFor(Options.ComponentId, ModifierOptions.ActiveTags, bModifiersForBeforeOperations, ChildModifiers);
 
 					ModifiersToIgnore.Push(m);
 					Ptr<ASTOp> AddedMeshOp = ApplyMeshModifiers(ChildModifiers, ModifierOptions, AddResults, SharedMeshResults, ErrorContext, nullptr);
@@ -2108,7 +2121,7 @@ namespace mu
 					Ptr<NodeMesh> pRemove = Edit->LODs[CurrentLOD].MeshRemove;
 
 					FMeshGenerationResult removeResults;
-					FMeshGenerationOptions RemoveMeshOptions;
+					FMeshGenerationOptions RemoveMeshOptions(Options.ComponentId);
 					RemoveMeshOptions.bLayouts = false;
 					RemoveMeshOptions.State = Options.State;
 					RemoveMeshOptions.ActiveTags = Edit->EnableTags;
@@ -2239,7 +2252,7 @@ namespace mu
 				op->source = PreModifiersMesh;
 
 				// Parameters
-				FMeshGenerationOptions ClipOptions;
+				FMeshGenerationOptions ClipOptions( Options.ComponentId );
 				ClipOptions.bLayouts = false;
 				ClipOptions.State = Options.State;
 
@@ -2297,7 +2310,7 @@ namespace mu
 				if (TypedClipNode->ClipMask)
 				{
 					// Parameters to generate the mask image
-					FImageGenerationOptions ClipOptions;
+					FImageGenerationOptions ClipOptions(Options.ComponentId);
 					ClipOptions.ImageLayoutStrategy = CompilerOptions::TextureLayoutStrategy::None;
 					ClipOptions.LayoutBlockId = FLayoutBlock::InvalidBlockId;
 					ClipOptions.State = Options.State;
@@ -2430,7 +2443,7 @@ namespace mu
 
 				ClipOp->FaceCullStrategy = TypedClipNode->FaceCullStrategy;
 
-				FMeshGenerationOptions ClipOptions;
+				FMeshGenerationOptions ClipOptions(Options.ComponentId);
 				ClipOptions.bLayouts = false;
 				ClipOptions.State = Options.State;				
 
@@ -2492,7 +2505,7 @@ namespace mu
 					if (TypedTransformNode->BoundingMesh)
 					{
 						// Parameters
-						FMeshGenerationOptions MeshOptions;
+						FMeshGenerationOptions MeshOptions(Options.ComponentId);
 						MeshOptions.bLayouts = false;
 						MeshOptions.State = Options.State;
 
@@ -2609,7 +2622,8 @@ namespace mu
 
 	Ptr<ASTOp> CodeGenerator::ApplyImageExtendModifiers(
 		const TArray<FirstPassGenerator::FModifier>& Modifiers,
-		const FGenericGenerationOptions& Options, 
+		const FGenericGenerationOptions& Options,
+		int32 ComponentId,
 		const FMeshGenerationResult& BaseMeshResults,
 		Ptr<ASTOp> BaseImageOp, 
 		CompilerOptions::TextureLayoutStrategy ImageLayoutStrategy,
@@ -2674,7 +2688,7 @@ namespace mu
 					for (int32 b = 0; b < pExtendLayout->GetBlockCount(); ++b)
 					{
 						// Generate the image block
-						FImageGenerationOptions ImageOptions;
+						FImageGenerationOptions ImageOptions(ComponentId);
 						ImageOptions.State = Options.State;
 						ImageOptions.ImageLayoutStrategy = ImageLayoutStrategy;
 						ImageOptions.ActiveTags = Edit->EnableTags; // TODO: Merge with current tags?
diff --git a/Engine/Plugins/Mutable/Source/MutableTools/Private/MuT/CodeGenerator.h b/Engine/Plugins/Mutable/Source/MutableTools/Private/MuT/CodeGenerator.h
--- a/Engine/Plugins/Mutable/Source/MutableTools/Private/MuT/CodeGenerator.h
+++ b/Engine/Plugins/Mutable/Source/MutableTools/Private/MuT/CodeGenerator.h
@@ -380,6 +380,14 @@ namespace mu
 		/** Options that affect the generation of images. It is like list of what required data we want while parsing down the image node graph. */
 		struct FImageGenerationOptions : public FGenericGenerationOptions
 		{
+			FImageGenerationOptions(int32 InComponentId)
+				: ComponentId(InComponentId)
+			{
+			}
+
+			/** The id of the component that we are currently generating. */
+			int32 ComponentId = -1;
+
 			/** */
 			CompilerOptions::TextureLayoutStrategy ImageLayoutStrategy = CompilerOptions::TextureLayoutStrategy::None;
 
@@ -392,29 +400,16 @@ namespace mu
 
 			friend FORCEINLINE uint32 GetTypeHash(const FImageGenerationOptions& InKey)
 			{
-				uint32 KeyHash = 0;
+				uint32 KeyHash = GetTypeHash((FGenericGenerationOptions&)InKey);
+				KeyHash = HashCombineFast(KeyHash, ::GetTypeHash(InKey.ComponentId));
 				KeyHash = HashCombineFast(KeyHash, ::GetTypeHash(InKey.ImageLayoutStrategy));
 				KeyHash = HashCombineFast(KeyHash, GetTypeHash(InKey.RectSize));
-				KeyHash = HashCombineFast(KeyHash, ::GetTypeHash(InKey.State));
 				KeyHash = HashCombineFast(KeyHash, ::GetTypeHash(InKey.LayoutBlockId));
 				KeyHash = HashCombineFast(KeyHash, ::GetTypeHash(InKey.LayoutToApply.get()));
 				return KeyHash;
 			}
 
-			FORCEINLINE bool operator==(const FImageGenerationOptions& Other) const
-			{
-				return ImageLayoutStrategy == Other.ImageLayoutStrategy
-					&&
-					RectSize == Other.RectSize
-					&&
-					State == Other.State
-					&&
-					LayoutBlockId == Other.LayoutBlockId
-					&&
-					LayoutToApply == Other.LayoutToApply
-					&&
-					ActiveTags == Other.ActiveTags;
-			}
+			FORCEINLINE bool operator==(const FImageGenerationOptions& Other) const = default;
 
 		};
 
@@ -427,6 +422,12 @@ namespace mu
 		/** */
 		struct FGeneratedImageCacheKey
 		{
+			FGeneratedImageCacheKey(const FImageGenerationOptions& InOptions, const NodeImagePtrConst& InNode)
+				: Node(InNode)
+				, Options(InOptions)
+			{
+			}
+
 			NodePtrConst Node;
 			FImageGenerationOptions Options;
 
@@ -438,10 +439,7 @@ namespace mu
 				return KeyHash;
 			}
 
-			FORCEINLINE bool operator==(const FGeneratedImageCacheKey& Other) const
-			{
-				return Node == Other.Node && Options == Other.Options;
-			}
+			FORCEINLINE bool operator==(const FGeneratedImageCacheKey& Other) const = default;
 		};
 
 		typedef TMap<FGeneratedImageCacheKey, FImageGenerationResult> GeneratedImagesMap;
@@ -521,6 +519,14 @@ namespace mu
 		*/
 		struct FMeshGenerationOptions : public FGenericGenerationOptions
 		{
+			FMeshGenerationOptions(int32 InComponentId)
+				: ComponentId(InComponentId)
+			{
+			}
+
+			/** The id of the component that we are currently generating. */
+			int32 ComponentId = -1;
+
 			/** The meshes at the leaves will need their own layout block data. */
 			bool bLayouts = false;
 
@@ -547,23 +553,14 @@ namespace mu
 			
 			friend FORCEINLINE uint32 GetTypeHash(const FMeshGenerationOptions& InKey)
 			{
-				uint32 KeyHash = 0;
+				uint32 KeyHash = GetTypeHash((FGenericGenerationOptions&)InKey);
+				KeyHash = HashCombineFast(KeyHash, ::GetTypeHash(InKey.ComponentId));
 				KeyHash = HashCombineFast(KeyHash, ::GetTypeHash(InKey.bLayouts));
 				KeyHash = HashCombineFast(KeyHash, ::GetTypeHash(InKey.OverrideLayouts.Num()));
-				KeyHash = HashCombineFast(KeyHash, ::GetTypeHash(InKey.State));
-				KeyHash = HashCombineFast(KeyHash, ::GetTypeHash(InKey.ActiveTags.Num()));
 				return KeyHash;
 			}
 
-			FORCEINLINE bool operator==(const FMeshGenerationOptions& Other) const
-			{
-				return State==Other.State 
-					&& bLayouts==Other.bLayouts 
-					&& bClampUVIslands == Other.bClampUVIslands && bNormalizeUVs == Other.bNormalizeUVs
-					&& bEnsureAllVerticesHaveLayoutBlock == Other.bEnsureAllVerticesHaveLayoutBlock
-					&& ActiveTags==Other.ActiveTags
-					&& OverrideLayouts ==Other.OverrideLayouts;
-			}
+			FORCEINLINE bool operator==(const FMeshGenerationOptions& Other) const = default;
 		};
 
 		//! Store the results of the code generation of a mesh.
@@ -826,7 +823,7 @@ namespace mu
 		void UpdateLayoutBlockDesc(FLayoutBlockDesc& Out, FImageDesc BlockDesc, FIntVector2 LayoutCellSize);
 
 		// Get the modifiers that have to be applied to elements with a specific tag.
-		void GetModifiersFor(const TArray<FString>& SurfaceTags, bool bModifiersForBeforeOperations, TArray<FirstPassGenerator::FModifier>& OutModifiers);
+		void GetModifiersFor(int32 ComponentId, const TArray<FString>& SurfaceTags, bool bModifiersForBeforeOperations, TArray<FirstPassGenerator::FModifier>& OutModifiers);
 
 		// Used to avoid recursion when generating modifiers.
 		TArray<FirstPassGenerator::FModifier> ModifiersToIgnore;
@@ -853,6 +850,7 @@ namespace mu
 		Ptr<ASTOp> ApplyImageExtendModifiers(
 			const TArray<FirstPassGenerator::FModifier>&,
 			const FGenericGenerationOptions& Options, 
+			int32 ComponentId,
 			const FMeshGenerationResult& BaseMeshResults,
 			Ptr<ASTOp> ImageAd, 
 			CompilerOptions::TextureLayoutStrategy ImageLayoutStrategy, 
diff --git a/Engine/Plugins/Mutable/Source/MutableTools/Private/MuT/CodeGenerator_Color.cpp b/Engine/Plugins/Mutable/Source/MutableTools/Private/MuT/CodeGenerator_Color.cpp
--- a/Engine/Plugins/Mutable/Source/MutableTools/Private/MuT/CodeGenerator_Color.cpp
+++ b/Engine/Plugins/Mutable/Source/MutableTools/Private/MuT/CodeGenerator_Color.cpp
@@ -301,7 +301,8 @@ namespace mu
 		op->op.type = OP_TYPE::CO_SAMPLEIMAGE;
 
 		// Source image
-		FImageGenerationOptions ImageOptions;
+		int32 ComponentId = -1; // TODO.
+		FImageGenerationOptions ImageOptions(ComponentId);
 		ImageOptions.State = Options.State;
 		ImageOptions.ActiveTags = Options.ActiveTags;
 
diff --git a/Engine/Plugins/Mutable/Source/MutableTools/Private/MuT/CodeGenerator_Image.cpp b/Engine/Plugins/Mutable/Source/MutableTools/Private/MuT/CodeGenerator_Image.cpp
--- a/Engine/Plugins/Mutable/Source/MutableTools/Private/MuT/CodeGenerator_Image.cpp
+++ b/Engine/Plugins/Mutable/Source/MutableTools/Private/MuT/CodeGenerator_Image.cpp
@@ -118,9 +118,7 @@ namespace mu
 		}
 
 		// See if it was already generated
-		FGeneratedImageCacheKey key;
-		key.Options = Options;
-		key.Node = Untyped;
+		FGeneratedImageCacheKey key(Options,Untyped);
 		GeneratedImagesMap::ValueType* CachedPtr = GeneratedImages.Find(key);
 		if (CachedPtr)
 		{
@@ -1423,7 +1421,7 @@ namespace mu
         {
 			// TODO: This will probably Result in a duplicated mesh subgraph, with the original mesh but new layout block ids.
 			// See if it can be optimized and try to reuse the existing layout block ids instead of generating new ones.
-			FMeshGenerationOptions MeshOptions;
+			FMeshGenerationOptions MeshOptions(Options.ComponentId);
 			MeshOptions.State = Options.State;
 			MeshOptions.ActiveTags = Options.ActiveTags;
 			MeshOptions.bLayouts = true;			// We need the layout that we will use to render
diff --git a/Engine/Plugins/Mutable/Source/MutableTools/Private/MuT/CodeGenerator_Mesh.cpp b/Engine/Plugins/Mutable/Source/MutableTools/Private/MuT/CodeGenerator_Mesh.cpp
--- a/Engine/Plugins/Mutable/Source/MutableTools/Private/MuT/CodeGenerator_Mesh.cpp
+++ b/Engine/Plugins/Mutable/Source/MutableTools/Private/MuT/CodeGenerator_Mesh.cpp
@@ -589,9 +589,7 @@ namespace mu
         }
 
         // See if it was already generated
-		FGeneratedMeshCacheKey Key;
-		Key.Node = InUntypedNode;
-		Key.Options = InOptions;
+		FGeneratedMeshCacheKey Key( InUntypedNode, InOptions );
         GeneratedMeshMap::ValueType* it = GeneratedMeshes.Find(Key);
         if ( it )
         {
@@ -1324,7 +1322,7 @@ namespace mu
 		// Apply the modifier for the pre-normal operations stage.
 		TArray<FirstPassGenerator::FModifier> Modifiers;
 		constexpr bool bModifiersForBeforeOperations = true;
-		GetModifiersFor(InOptions.ActiveTags, bModifiersForBeforeOperations, Modifiers);
+		GetModifiersFor(InOptions.ComponentId, InOptions.ActiveTags, bModifiersForBeforeOperations, Modifiers);
 
 		OutResult.MeshOp = ApplyMeshModifiers(Modifiers, InOptions, OutResult, nullptr, InNode->GetMessageContext(), InNode);
     }
diff --git a/Engine/Plugins/Mutable/Source/MutableTools/Public/MuT/NodeModifier.h b/Engine/Plugins/Mutable/Source/MutableTools/Public/MuT/NodeModifier.h
--- a/Engine/Plugins/Mutable/Source/MutableTools/Public/MuT/NodeModifier.h
+++ b/Engine/Plugins/Mutable/Source/MutableTools/Public/MuT/NodeModifier.h
@@ -36,6 +36,9 @@ namespace mu
 	{
 	public:
 
+		/** If not negative, this modifier will only be applied to the nodes of the component with matching id. Otherwise it will be applied to all components. */
+		int32 RequiredComponentId = -1;
+
 		/** Tags that target surface need to have enabled to receive this modifier. */
 		TArray<FString> RequiredTags;
 
