[Mutable] Fix UCustomizableSkeletalComponent not working properly in PIE.

#tests PIE, Editing Level
#rnx

diff --git a/Engine/Plugins/Mutable/Source/CustomizableObject/Internal/MuCO/CustomizableObjectInstanceUsagePrivate.h b/Engine/Plugins/Mutable/Source/CustomizableObject/Internal/MuCO/CustomizableObjectInstanceUsagePrivate.h
--- a/Engine/Plugins/Mutable/Source/CustomizableObject/Internal/MuCO/CustomizableObjectInstanceUsagePrivate.h
+++ b/Engine/Plugins/Mutable/Source/CustomizableObject/Internal/MuCO/CustomizableObjectInstanceUsagePrivate.h
@@ -65,14 +65,4 @@ public:
 
 	// Returns true if the NetMode of the associated UCustomizableSkeletalComponent (or the associated SkeletalMeshComponent if the former does not exist) is equal to InNetMode
 	bool IsNetMode(ENetMode InNetMode) const;
-	
-	// Used to replace the SkeletalMesh of the parent component by the ReferenceSkeletalMesh or the generated SkeletalMesh 
-	bool bUsedPendingSetSkeletalMesh = false;
-
-	// Used to avoid replacing the SkeletalMesh of the parent component by the ReferenceSkeletalMesh if bPendingSetSkeletalMesh is true
-	UPROPERTY()
-	bool bUsedSkipSetReferenceSkeletalMesh = false;
-
-	UPROPERTY()
-	bool bUsedSkipSetSkeletalMeshOnAttach = false;
 };
diff --git a/Engine/Plugins/Mutable/Source/CustomizableObject/Internal/MuCO/CustomizableSkeletalComponentPrivate.h b/Engine/Plugins/Mutable/Source/CustomizableObject/Internal/MuCO/CustomizableSkeletalComponentPrivate.h
--- a/Engine/Plugins/Mutable/Source/CustomizableObject/Internal/MuCO/CustomizableSkeletalComponentPrivate.h
+++ b/Engine/Plugins/Mutable/Source/CustomizableObject/Internal/MuCO/CustomizableSkeletalComponentPrivate.h
@@ -37,17 +37,10 @@ public:
 	void EditorUpdateComponent();
 #endif
 
+	bool& PendingSetSkeletalMesh();
+	
 	UCustomizableSkeletalComponent* GetPublic();
 	
 	const UCustomizableSkeletalComponent* GetPublic() const;
-	
-	// Used to replace the SkeletalMesh of the parent component by the ReferenceSkeletalMesh or the generated SkeletalMesh 
-	bool bPendingSetSkeletalMesh = false;
-
-	UPROPERTY()
-	bool bSkipSkipSetSkeletalMeshOnAttach = false;
-
-	UPROPERTY(Transient)
-	TObjectPtr<UCustomizableObjectInstanceUsage> CustomizableObjectInstanceUsage;
 };
 
diff --git a/Engine/Plugins/Mutable/Source/CustomizableObject/Private/MuCO/CustomizableObjectInstanceUsage.cpp b/Engine/Plugins/Mutable/Source/CustomizableObject/Private/MuCO/CustomizableObjectInstanceUsage.cpp
--- a/Engine/Plugins/Mutable/Source/CustomizableObject/Private/MuCO/CustomizableObjectInstanceUsage.cpp
+++ b/Engine/Plugins/Mutable/Source/CustomizableObject/Private/MuCO/CustomizableObjectInstanceUsage.cpp
@@ -148,11 +148,11 @@ void UCustomizableObjectInstanceUsagePrivate::SetPendingSetSkeletalMesh(bool bIs
 {
 	if (GetPublic()->CustomizableSkeletalComponent)
 	{
-		GetPublic()->CustomizableSkeletalComponent->GetPrivate()->bPendingSetSkeletalMesh = bIsActive;
+		GetPublic()->CustomizableSkeletalComponent->GetPrivate()->PendingSetSkeletalMesh() = bIsActive;
 	}
 	else
 	{
-		bUsedPendingSetSkeletalMesh = bIsActive;
+		GetPublic()->bUsedPendingSetSkeletalMesh = bIsActive;
 	}
 }
 
@@ -161,11 +161,11 @@ bool UCustomizableObjectInstanceUsagePrivate::GetPendingSetSkeletalMesh() const
 {
 	if (GetPublic()->CustomizableSkeletalComponent)
 	{
-		return GetPublic()->CustomizableSkeletalComponent->GetPrivate()->bPendingSetSkeletalMesh;
+		return GetPublic()->CustomizableSkeletalComponent->GetPrivate()->PendingSetSkeletalMesh();
 	}
 	else
 	{
-		return bUsedPendingSetSkeletalMesh;
+		return GetPublic()->bUsedPendingSetSkeletalMesh;
 	}
 }
 
@@ -208,7 +208,7 @@ void UCustomizableObjectInstanceUsage::SetSkipSetReferenceSkeletalMesh(bool bSki
 	}
 	else
 	{
-		GetPrivate()->bUsedSkipSetReferenceSkeletalMesh = bSkip;
+		bUsedSkipSetReferenceSkeletalMesh = bSkip;
 	}
 }
 
@@ -221,7 +221,7 @@ bool UCustomizableObjectInstanceUsage::GetSkipSetReferenceSkeletalMesh() const
 	}
 	else
 	{
-		return GetPrivate()->bUsedSkipSetReferenceSkeletalMesh;
+		return bUsedSkipSetReferenceSkeletalMesh;
 	}
 }
 
@@ -234,7 +234,7 @@ void UCustomizableObjectInstanceUsage::SetSkipSetSkeletalMeshOnAttach(bool bSkip
 	}
 	else
 	{
-		GetPrivate()->bUsedSkipSetSkeletalMeshOnAttach = bSkip;
+		bUsedSkipSetSkeletalMeshOnAttach = bSkip;
 	}
 }
 
@@ -247,7 +247,7 @@ bool UCustomizableObjectInstanceUsage::GetSkipSetSkeletalMeshOnAttach() const
 	}
 	else
 	{
-		return GetPrivate()->bUsedSkipSetSkeletalMeshOnAttach;
+		return bUsedSkipSetSkeletalMeshOnAttach;
 	}
 }
 
diff --git a/Engine/Plugins/Mutable/Source/CustomizableObject/Private/MuCO/CustomizableSkeletalComponent.cpp b/Engine/Plugins/Mutable/Source/CustomizableObject/Private/MuCO/CustomizableSkeletalComponent.cpp
--- a/Engine/Plugins/Mutable/Source/CustomizableObject/Private/MuCO/CustomizableSkeletalComponent.cpp
+++ b/Engine/Plugins/Mutable/Source/CustomizableObject/Private/MuCO/CustomizableSkeletalComponent.cpp
@@ -23,12 +23,12 @@ UCustomizableSkeletalComponentPrivate::UCustomizableSkeletalComponentPrivate()
 
 void UCustomizableSkeletalComponentPrivate::CreateCustomizableObjectInstanceUsage()
 {
-	if (CustomizableObjectInstanceUsage)
+	if (GetPublic()->CustomizableObjectInstanceUsage)
 	{
 		// CustomizableObjectInstanceUsage may already exist if duplicated from an existing Customizable Skeletal Component
-		if (CustomizableObjectInstanceUsage->GetPrivate()->GetCustomizableSkeletalComponent() != GetPublic())
+		if (GetPublic()->CustomizableObjectInstanceUsage->GetPrivate()->GetCustomizableSkeletalComponent() != GetPublic())
 		{
-			CustomizableObjectInstanceUsage = nullptr;
+			GetPublic()->CustomizableObjectInstanceUsage = nullptr;
 		}
 	}
 
@@ -37,10 +37,10 @@ void UCustomizableSkeletalComponentPrivate::CreateCustomizableObjectInstanceUsag
 		RootActor->HasAnyFlags(RF_ClassDefaultObject) :
 		false;
 
-	if (!CustomizableObjectInstanceUsage && !HasAnyFlags(RF_ClassDefaultObject) && !bIsDefaultActor)
+	if (!GetPublic()->CustomizableObjectInstanceUsage && !HasAnyFlags(RF_ClassDefaultObject) && !bIsDefaultActor)
 	{
-		CustomizableObjectInstanceUsage = NewObject<UCustomizableObjectInstanceUsage>(this, TEXT("InstanceUsage"), RF_Transient);
-		CustomizableObjectInstanceUsage->GetPrivate()->SetCustomizableSkeletalComponent(GetPublic());
+		GetPublic()->CustomizableObjectInstanceUsage = NewObject<UCustomizableObjectInstanceUsage>(this, TEXT("InstanceUsage"), RF_Transient);
+		GetPublic()->CustomizableObjectInstanceUsage->GetPrivate()->SetCustomizableSkeletalComponent(GetPublic());
 	}
 }
 
@@ -72,18 +72,18 @@ void UCustomizableSkeletalComponent::PostEditChangeProperty(FPropertyChangedEven
 
 void UCustomizableSkeletalComponentPrivate::Callbacks() const
 {
-	if (CustomizableObjectInstanceUsage)
+	if (GetPublic()->CustomizableObjectInstanceUsage)
 	{
-		CustomizableObjectInstanceUsage->GetPrivate()->Callbacks();
+		GetPublic()->CustomizableObjectInstanceUsage->GetPrivate()->Callbacks();
 	}
 }
 
 
 USkeletalMesh* UCustomizableSkeletalComponentPrivate::GetSkeletalMesh() const
 {
-	if (CustomizableObjectInstanceUsage)
+	if (GetPublic()->CustomizableObjectInstanceUsage)
 	{
-		return CustomizableObjectInstanceUsage->GetPrivate()->GetSkeletalMesh();
+		return GetPublic()->CustomizableObjectInstanceUsage->GetPrivate()->GetSkeletalMesh();
 	}
 
 	return nullptr;
@@ -92,36 +92,36 @@ USkeletalMesh* UCustomizableSkeletalComponentPrivate::GetSkeletalMesh() const
 
 void UCustomizableSkeletalComponentPrivate::SetSkeletalMesh(USkeletalMesh* SkeletalMesh)
 {
-	if (CustomizableObjectInstanceUsage)
+	if (GetPublic()->CustomizableObjectInstanceUsage)
 	{
-		CustomizableObjectInstanceUsage->GetPrivate()->SetSkeletalMesh(SkeletalMesh);
+		GetPublic()->CustomizableObjectInstanceUsage->GetPrivate()->SetSkeletalMesh(SkeletalMesh);
 	}
 }
 
 
 void UCustomizableSkeletalComponentPrivate::SetPhysicsAsset(UPhysicsAsset* PhysicsAsset)
 {
-	if (CustomizableObjectInstanceUsage)
+	if (GetPublic()->CustomizableObjectInstanceUsage)
 	{
-		CustomizableObjectInstanceUsage->GetPrivate()->SetPhysicsAsset(PhysicsAsset);
+		GetPublic()->CustomizableObjectInstanceUsage->GetPrivate()->SetPhysicsAsset(PhysicsAsset);
 	}
 }
 
 
 void UCustomizableSkeletalComponentPrivate::SetPendingSetSkeletalMesh(bool bIsActive)
 {
-	if (CustomizableObjectInstanceUsage)
+	if (GetPublic()->CustomizableObjectInstanceUsage)
 	{
-		CustomizableObjectInstanceUsage->GetPrivate()->SetPendingSetSkeletalMesh(bIsActive);
+		GetPublic()->CustomizableObjectInstanceUsage->GetPrivate()->SetPendingSetSkeletalMesh(bIsActive);
 	}
 }
 
 
 USkeletalMesh* UCustomizableSkeletalComponentPrivate::GetAttachedSkeletalMesh() const
 {
-	if (CustomizableObjectInstanceUsage)
+	if (GetPublic()->CustomizableObjectInstanceUsage)
 	{
-		return CustomizableObjectInstanceUsage->GetPrivate()->GetAttachedSkeletalMesh();
+		return GetPublic()->CustomizableObjectInstanceUsage->GetPrivate()->GetAttachedSkeletalMesh();
 	}
 
 	return nullptr;
@@ -174,30 +174,30 @@ bool UCustomizableSkeletalComponent::GetSkipSetReferenceSkeletalMesh() const
 
 void UCustomizableSkeletalComponent::SetSkipSetSkeletalMeshOnAttach(bool bSkip)
 {
-	GetPrivate()->bSkipSkipSetSkeletalMeshOnAttach = bSkip;
+	bSkipSkipSetSkeletalMeshOnAttach = bSkip;
 }
 
 
 bool UCustomizableSkeletalComponent::GetSkipSetSkeletalMeshOnAttach() const
 {
-	return GetPrivate()->bSkipSkipSetSkeletalMeshOnAttach;
+	return bSkipSkipSetSkeletalMeshOnAttach;
 }
 
 
 void UCustomizableSkeletalComponent::UpdateSkeletalMeshAsync(bool bNeverSkipUpdate)
 {
-	if (GetPrivate()->CustomizableObjectInstanceUsage)
+	if (CustomizableObjectInstanceUsage)
 	{
-		GetPrivate()->CustomizableObjectInstanceUsage->UpdateSkeletalMeshAsync(bNeverSkipUpdate);
+		CustomizableObjectInstanceUsage->UpdateSkeletalMeshAsync(bNeverSkipUpdate);
 	}
 }
 
 
 void UCustomizableSkeletalComponent::UpdateSkeletalMeshAsyncResult(FInstanceUpdateDelegate Callback, bool bIgnoreCloseDist, bool bForceHighPriority)
 {
-	if (GetPrivate()->CustomizableObjectInstanceUsage)
+	if (CustomizableObjectInstanceUsage)
 	{
-		GetPrivate()->CustomizableObjectInstanceUsage->UpdateSkeletalMeshAsyncResult(Callback, bIgnoreCloseDist, bForceHighPriority);
+		CustomizableObjectInstanceUsage->UpdateSkeletalMeshAsyncResult(Callback, bIgnoreCloseDist, bForceHighPriority);
 	}
 }
 
@@ -219,14 +219,20 @@ const UCustomizableSkeletalComponentPrivate* UCustomizableSkeletalComponent::Get
 #if WITH_EDITOR
 void UCustomizableSkeletalComponentPrivate::EditorUpdateComponent()
 {
-	if (CustomizableObjectInstanceUsage)
+	if (GetPublic()->CustomizableObjectInstanceUsage)
 	{
-		CustomizableObjectInstanceUsage->GetPrivate()->EditorUpdateComponent();
+		GetPublic()->CustomizableObjectInstanceUsage->GetPrivate()->EditorUpdateComponent();
 	}
 }
 #endif
 
 
+bool& UCustomizableSkeletalComponentPrivate::PendingSetSkeletalMesh()
+{
+	return GetPublic()->bPendingSetSkeletalMesh;	
+}
+
+
 UCustomizableSkeletalComponent* UCustomizableSkeletalComponentPrivate::GetPublic()
 {
 	UCustomizableSkeletalComponent* Public = StaticCast<UCustomizableSkeletalComponent*>(GetOuter());
diff --git a/Engine/Plugins/Mutable/Source/CustomizableObject/Public/MuCO/CustomizableObjectInstanceUsage.h b/Engine/Plugins/Mutable/Source/CustomizableObject/Public/MuCO/CustomizableObjectInstanceUsage.h
--- a/Engine/Plugins/Mutable/Source/CustomizableObject/Public/MuCO/CustomizableObjectInstanceUsage.h
+++ b/Engine/Plugins/Mutable/Source/CustomizableObject/Public/MuCO/CustomizableObjectInstanceUsage.h
@@ -26,6 +26,8 @@ DECLARE_DELEGATE(FCustomizableObjectInstanceUsageUpdatedDelegate);
 UCLASS(Blueprintable, BlueprintType, ClassGroup = (CustomizableObject), meta = (BlueprintSpawnableComponent))
 class CUSTOMIZABLEOBJECT_API UCustomizableObjectInstanceUsage : public UObject
 {
+	friend UCustomizableObjectInstanceUsagePrivate;
+	
 public:
 	GENERATED_BODY()
 
@@ -106,6 +108,16 @@ private:
 	UPROPERTY()
 	FName UsedComponentName;
 	
+	// Used to replace the SkeletalMesh of the parent component by the ReferenceSkeletalMesh or the generated SkeletalMesh 
+	bool bUsedPendingSetSkeletalMesh = false;
+
+	// Used to avoid replacing the SkeletalMesh of the parent component by the ReferenceSkeletalMesh if bPendingSetSkeletalMesh is true
+	UPROPERTY()
+	bool bUsedSkipSetReferenceSkeletalMesh = false;
+
+	UPROPERTY()
+	bool bUsedSkipSetSkeletalMeshOnAttach = false;
+	
 	UPROPERTY()
 	TObjectPtr<UCustomizableObjectInstanceUsagePrivate> Private;
 	
diff --git a/Engine/Plugins/Mutable/Source/CustomizableObject/Public/MuCO/CustomizableSkeletalComponent.h b/Engine/Plugins/Mutable/Source/CustomizableObject/Public/MuCO/CustomizableSkeletalComponent.h
--- a/Engine/Plugins/Mutable/Source/CustomizableObject/Public/MuCO/CustomizableSkeletalComponent.h
+++ b/Engine/Plugins/Mutable/Source/CustomizableObject/Public/MuCO/CustomizableSkeletalComponent.h
@@ -16,6 +16,8 @@ DECLARE_DELEGATE(FCustomizableSkeletalComponentUpdatedDelegate);
 UCLASS(Blueprintable, BlueprintType, ClassGroup = (CustomizableObject), meta = (BlueprintSpawnableComponent))
 class CUSTOMIZABLEOBJECT_API UCustomizableSkeletalComponent : public USceneComponent
 {
+	friend UCustomizableSkeletalComponentPrivate;
+
 public:
 	GENERATED_BODY()
 	
@@ -39,6 +41,15 @@ public:
 	FCustomizableSkeletalComponentUpdatedDelegate UpdatedDelegate;
 
 private:
+	// Used to replace the SkeletalMesh of the parent component by the ReferenceSkeletalMesh or the generated SkeletalMesh 
+	bool bPendingSetSkeletalMesh = false;
+
+	UPROPERTY()
+	bool bSkipSkipSetSkeletalMeshOnAttach = false;
+
+	UPROPERTY(Transient)
+	TObjectPtr<UCustomizableObjectInstanceUsage> CustomizableObjectInstanceUsage;
+	
 	UPROPERTY()
 	TObjectPtr<UCustomizableSkeletalComponentPrivate> Private;
 
