[Mutable] Fix projector when the projection mesh comes from a table or a switch.

- Mesh to project would use the base mesh op but we actually need the mesh op.

#rnx

diff --git a/Engine/Plugins/Mutable/Source/MutableTools/Private/MuT/CodeGenerator_Image.cpp b/Engine/Plugins/Mutable/Source/MutableTools/Private/MuT/CodeGenerator_Image.cpp
--- a/Engine/Plugins/Mutable/Source/MutableTools/Private/MuT/CodeGenerator_Image.cpp
+++ b/Engine/Plugins/Mutable/Source/MutableTools/Private/MuT/CodeGenerator_Image.cpp
@@ -1449,9 +1449,10 @@ namespace mu
 			// TODO: 
 			// MeshResult.MeshOp has some modifiers applied already: the ones applied before other operations directly in the mesh constant generation. 
 			// This is not what was happening before the refactor  so use MeshResult.BaseMeshOp. This is another case of ambiguity of order of modifiers 
-			// that whould be fixed with the general ordering design.
-			//Ptr<ASTOp> CurrentMeshToProjectOp = MeshResult.MeshOp;
-			Ptr<ASTOp> CurrentMeshToProjectOp = MeshResult.BaseMeshOp;
+			// that whould be fixed with the general ordering design. 
+			// Actually use the MeshOp, otherwise the projector will only project to the first option if the mesh operation is a switch. 
+			Ptr<ASTOp> CurrentMeshToProjectOp = MeshResult.MeshOp;
+			//Ptr<ASTOp> CurrentMeshToProjectOp = MeshResult.BaseMeshOp;
 
             if (projectorResult.type == PROJECTOR_TYPE::WRAPPING)
             {
