[Mutable] Fix component index pointing to the wrong component data due to nodes being processed in a different order between FirstPass and SecondPass.

#rnx

diff --git a/Engine/Plugins/Mutable/Source/CustomizableObject/Private/MuCO/CustomizableObjectInstance.cpp b/Engine/Plugins/Mutable/Source/CustomizableObject/Private/MuCO/CustomizableObjectInstance.cpp
--- a/Engine/Plugins/Mutable/Source/CustomizableObject/Private/MuCO/CustomizableObjectInstance.cpp
+++ b/Engine/Plugins/Mutable/Source/CustomizableObject/Private/MuCO/CustomizableObjectInstance.cpp
@@ -6592,11 +6592,11 @@ void UCustomizableInstancePrivate::BuildMaterials(const TSharedRef<FUpdateContex
 													if (MutableMipDataProviderFactory)
 													{
 														MutableMipDataProviderFactory->CustomizableObjectInstance = Public;
-														check(LODIndex < 256 && ObjectComponentIndex < 256 && ImageIndex < 256);
+														check(LODIndex < 256 && InstanceComponentIndex < 256 && ImageIndex < 256);
 														MutableMipDataProviderFactory->ImageRef.ImageID = Image.ImageID;
 														MutableMipDataProviderFactory->ImageRef.SurfaceId = Surface.SurfaceId;
 														MutableMipDataProviderFactory->ImageRef.LOD = uint8(Image.BaseLOD);
-														MutableMipDataProviderFactory->ImageRef.Component = uint8(ObjectComponentIndex);
+														MutableMipDataProviderFactory->ImageRef.Component = uint8(InstanceComponentIndex);
 														MutableMipDataProviderFactory->ImageRef.Image = uint8(ImageIndex);
 														MutableMipDataProviderFactory->ImageRef.BaseMip = uint8(Image.BaseMip);
 														MutableMipDataProviderFactory->UpdateContext = UpdateContext;
diff --git a/Engine/Plugins/Mutable/Source/CustomizableObjectEditor/Private/MuCOE/GenerateMutableSource/GenerateMutableSourceComponent.cpp b/Engine/Plugins/Mutable/Source/CustomizableObjectEditor/Private/MuCOE/GenerateMutableSource/GenerateMutableSourceComponent.cpp
--- a/Engine/Plugins/Mutable/Source/CustomizableObjectEditor/Private/MuCOE/GenerateMutableSource/GenerateMutableSourceComponent.cpp
+++ b/Engine/Plugins/Mutable/Source/CustomizableObjectEditor/Private/MuCOE/GenerateMutableSource/GenerateMutableSourceComponent.cpp
@@ -139,7 +139,7 @@ mu::Ptr<mu::NodeComponent> GenerateMutableSourceComponent(const UEdGraphPin* Pin
 		}
 		
 		mu::Ptr<mu::NodeComponentNew> NodeComponentNew = new mu::NodeComponentNew();
-		NodeComponentNew->Id = GenerationContext.ComponentNames.Add(TypedComponentMesh->ComponentName);
+		NodeComponentNew->Id = GenerationContext.ComponentNames.Find(TypedComponentMesh->ComponentName);
 		NodeComponentNew->SetMessageContext(Node);
 		
 		Result = NodeComponentNew;
@@ -235,8 +235,7 @@ mu::Ptr<mu::NodeComponent> GenerateMutableSourceComponent(const UEdGraphPin* Pin
 
 		// Create the component node
 		mu::Ptr<mu::NodeComponentNew> ComponentNode = new mu::NodeComponentNew;
-		ComponentNode->Id = GenerationContext.ComponentNames.Num();
-		GenerationContext.ComponentNames.Add(TypedComponentPassthroughMesh->ComponentName);
+		ComponentNode->Id = GenerationContext.ComponentNames.Find(TypedComponentPassthroughMesh->ComponentName);
 
 		// Create a LOD for each pass-through mesh LOD.
 		const FSkeletalMeshModel* Model = SkeletalMesh->GetImportedModel();
@@ -532,6 +531,8 @@ void FirstPass(UCustomizableObjectNodeComponentMesh& Node, FMutableGraphGenerati
 	// Make sure the Skeleton from the reference mesh is added to the list of referenced Skeletons.
 	GenerationContext.ReferencedSkeletons.Add(RefSkeleton);
 
+	GenerationContext.ComponentNames.Add(Node.ComponentName);
+
 	// Add reference meshes to the participating objects
 	GenerationContext.AddParticipatingObject(*RefSkeletalMesh);
 }
