[mutable] Ensure the image size for compose operations over a switch get the biggest possible size.
#rnx

diff --git a/Engine/Plugins/Mutable/Source/MutableTools/Private/MuT/ASTOpSwitch.cpp b/Engine/Plugins/Mutable/Source/MutableTools/Private/MuT/ASTOpSwitch.cpp
--- a/Engine/Plugins/Mutable/Source/MutableTools/Private/MuT/ASTOpSwitch.cpp
+++ b/Engine/Plugins/Mutable/Source/MutableTools/Private/MuT/ASTOpSwitch.cpp
@@ -251,6 +251,10 @@ namespace mu
 					if (bReturnBestOption)
 					{
 						Candidate.m_format = GetMostGenericFormat(Candidate.m_format, ChildDesc.m_format);
+
+						// Return the biggest size
+						Candidate.m_size[0] = FMath::Max(Candidate.m_size[0], ChildDesc.m_size[0]);
+						Candidate.m_size[1] = FMath::Max(Candidate.m_size[1], ChildDesc.m_size[1]);
 					}
 				}
 			}
