[Mutable] Fix crash when compiling some COs due to Participating Objects pass missmatch.

#rnx

diff --git a/Engine/Plugins/Mutable/Source/CustomizableObjectEditor/Private/MuCOE/GenerateMutableSource/GenerateMutableSource.cpp b/Engine/Plugins/Mutable/Source/CustomizableObjectEditor/Private/MuCOE/GenerateMutableSource/GenerateMutableSource.cpp
--- a/Engine/Plugins/Mutable/Source/CustomizableObjectEditor/Private/MuCOE/GenerateMutableSource/GenerateMutableSource.cpp
+++ b/Engine/Plugins/Mutable/Source/CustomizableObjectEditor/Private/MuCOE/GenerateMutableSource/GenerateMutableSource.cpp
@@ -455,14 +455,11 @@ void FMutableGraphGenerationContext::AddParticipatingObjectChecked(const FName&
 	}
 	else
 	{
-		checkCode
-		(
-			if (!bSkipParticipatingObjectsPass)
-			{
-				FGuid* Result = ParticipatingObjects.Find(PackageName);
-				check(Result && *Result == PackageGuid); // If this check is hit it means that this Participating Object is not being discovered in the Participating Objects pass.
-			}
-		)
+		if (!bSkipParticipatingObjectsPass)
+		{
+			FGuid* Result = ParticipatingObjects.Find(PackageName);
+			ensureAlwaysMsgf(Result, TEXT("Participating Object and Generate Source pass missatch. Asset mismatch: %s"), *PackageName.ToString());
+		}
 	}
 }
 
diff --git a/Engine/Plugins/Mutable/Source/CustomizableObjectEditor/Private/MuCOE/GenerateMutableSource/GenerateMutableSourceTable.cpp b/Engine/Plugins/Mutable/Source/CustomizableObjectEditor/Private/MuCOE/GenerateMutableSource/GenerateMutableSourceTable.cpp
--- a/Engine/Plugins/Mutable/Source/CustomizableObjectEditor/Private/MuCOE/GenerateMutableSource/GenerateMutableSourceTable.cpp
+++ b/Engine/Plugins/Mutable/Source/CustomizableObjectEditor/Private/MuCOE/GenerateMutableSource/GenerateMutableSourceTable.cpp
@@ -497,8 +497,6 @@ bool FillTableColumn(const UCustomizableObjectNodeTable* TableNode, mu::TablePtr
 
 				LogRowGenerationMessage(TableNode, DataTablePtr, GenerationContext, Warning.ToString(), RowName);
 			}
-
-			GenerationContext.AddParticipatingObject(*MaterialInstance);
 			
 			FString EncodedSwitchParameterName = "__MutableMaterialId";
 			if (ColumnName.Contains(EncodedSwitchParameterName))
