[Mutable] Added new PreSetSkeletalMesh delegate.

#rnx

diff --git a/Engine/Plugins/Mutable/Source/CustomizableObject/Private/MuCO/CustomizableObjectSystem.cpp b/Engine/Plugins/Mutable/Source/CustomizableObject/Private/MuCO/CustomizableObjectSystem.cpp
--- a/Engine/Plugins/Mutable/Source/CustomizableObject/Private/MuCO/CustomizableObjectSystem.cpp
+++ b/Engine/Plugins/Mutable/Source/CustomizableObject/Private/MuCO/CustomizableObjectSystem.cpp
@@ -937,6 +937,15 @@ void UpdateSkeletalMesh(const TSharedRef<FUpdateContextPrivate>& Context)
 	UCustomizableObjectInstance* CustomizableObjectInstance = Context->Instance.Get();
 	check(CustomizableObjectInstance);
 
+	for (const TTuple<FName, TObjectPtr<USkeletalMesh>>& Pair : CustomizableObjectInstance->GetPrivate()->SkeletalMeshes)
+	{
+		FPreSetSkeletalMeshParams Params;
+		Params.SkeletalMesh = Pair.Value;
+		
+		CustomizableObjectInstance->PreSetSkeletalMeshDelegate.Broadcast(Params);
+		CustomizableObjectInstance->PreSetSkeletalMeshNativeDelegate.Broadcast(Params);
+	}
+	
 	UCustomizableInstancePrivate* CustomizableObjectInstancePrivateData = CustomizableObjectInstance->GetPrivate();
 	check(CustomizableObjectInstancePrivateData != nullptr);
 	for (TObjectIterator<UCustomizableObjectInstanceUsage> It; It; ++It)
diff --git a/Engine/Plugins/Mutable/Source/CustomizableObject/Public/MuCO/CustomizableObjectInstance.h b/Engine/Plugins/Mutable/Source/CustomizableObject/Public/MuCO/CustomizableObjectInstance.h
--- a/Engine/Plugins/Mutable/Source/CustomizableObject/Public/MuCO/CustomizableObjectInstance.h
+++ b/Engine/Plugins/Mutable/Source/CustomizableObject/Public/MuCO/CustomizableObjectInstance.h
@@ -210,13 +210,24 @@ DECLARE_MULTICAST_DELEGATE_OneParam(FInstanceUpdateNativeDelegate, const FUpdate
 DECLARE_DYNAMIC_MULTICAST_DELEGATE_OneParam(FObjectInstanceUpdatedDelegate, UCustomizableObjectInstance*, Instance);
 DECLARE_MULTICAST_DELEGATE_OneParam(FObjectInstanceUpdatedNativeDelegate, UCustomizableObjectInstance*);
 
+USTRUCT()
+struct FPreSetSkeletalMeshParams
+{
+	GENERATED_BODY()
+
+	UPROPERTY()
+	TObjectPtr<USkeletalMesh> SkeletalMesh;
+};
+
+DECLARE_DYNAMIC_MULTICAST_DELEGATE_OneParam(FPreSetSkeletalMeshDelegate, const FPreSetSkeletalMeshParams&, Params);
+DECLARE_MULTICAST_DELEGATE_OneParam(FPreSetSkeletalMeshNativeDelegate, const FPreSetSkeletalMeshParams&);
+
 DECLARE_DELEGATE_OneParam(FProjectorStateChangedDelegate, FString);
 
 DECLARE_DYNAMIC_DELEGATE_TwoParams(FEachComponentAnimInstanceClassDelegate, FName, SlotIndex, TSubclassOf<UAnimInstance>, AnimInstClass);
 
 DECLARE_DELEGATE_TwoParams(FEachComponentAnimInstanceClassNativeDelegate, FName /*SlotIndex*/, TSubclassOf<UAnimInstance> /*AnimInstClass*/);
 
-
 UCLASS( Blueprintable, BlueprintType, HideCategories=(CustomizableObjectInstance) )
 class CUSTOMIZABLEOBJECT_API UCustomizableObjectInstance : public UObject
 {
@@ -236,6 +247,14 @@ public:
 	/** Broadcast when the Customizable Object Instance is updated. */
 	FObjectInstanceUpdatedNativeDelegate UpdatedNativeDelegate;
 
+	/** Broadcast before setting the generated Skeletal Mesh to the Skeletal Mesh Components.
+	  * Broadcast even if Skeletal Meshes are reused. */
+	FPreSetSkeletalMeshDelegate PreSetSkeletalMeshDelegate;
+
+	/** Broadcast before setting the generated Skeletal Mesh to the Skeletal Mesh Components.
+	  * Broadcast even if Skeletal Meshes are reused. */
+	FPreSetSkeletalMeshNativeDelegate PreSetSkeletalMeshNativeDelegate;
+	
 	// UObject interface.
 #if WITH_EDITOR
 	virtual void PostEditChangeProperty(FPropertyChangedEvent& PropertyChangedEvent) override;
