[mutable] Fixed an ensure happening when compiling a CO without a material set.

#rnx

diff --git a/Engine/Plugins/Mutable/Source/CustomizableObjectEditor/Private/MuCOE/GenerateMutableSource/GenerateMutableSourceSurface.cpp b/Engine/Plugins/Mutable/Source/CustomizableObjectEditor/Private/MuCOE/GenerateMutableSource/GenerateMutableSourceSurface.cpp
--- a/Engine/Plugins/Mutable/Source/CustomizableObjectEditor/Private/MuCOE/GenerateMutableSource/GenerateMutableSourceSurface.cpp
+++ b/Engine/Plugins/Mutable/Source/CustomizableObjectEditor/Private/MuCOE/GenerateMutableSource/GenerateMutableSourceSurface.cpp
@@ -249,6 +249,15 @@ mu::Ptr<mu::NodeSurface> GenerateMutableSourceSurface(const UEdGraphPin * Pin, F
 			}
 		}
 
+		if (!TypedNodeMat->GetMaterial())
+		{
+			const FText Message = LOCTEXT("FailedToGenerateMeshSection", "Could not generate a mesh section because it didn't have a material selected. Please assign one and recompile.");
+			GenerationContext.Log(Message, Node);
+			Result = nullptr;
+
+			return Result;
+		}
+
 		mu::Ptr<mu::NodeSurfaceNew> SurfNode = new mu::NodeSurfaceNew();
 		Result = SurfNode;
 
