[Mutable] Fixed SurfaceMetadata not being serialied in FModelResources.

#rnx

diff --git a/Engine/Plugins/Mutable/Source/CustomizableObject/Internal/MuCO/CustomizableObjectPrivate.h b/Engine/Plugins/Mutable/Source/CustomizableObject/Internal/MuCO/CustomizableObjectPrivate.h
--- a/Engine/Plugins/Mutable/Source/CustomizableObject/Internal/MuCO/CustomizableObjectPrivate.h
+++ b/Engine/Plugins/Mutable/Source/CustomizableObject/Internal/MuCO/CustomizableObjectPrivate.h
@@ -1567,12 +1567,14 @@ public:
 
 		RemoveModifiersHack,
 
+		SurfaceMetadataSerialized,
+
 		// -----<new versions can be added above this line>--------
 		LastCustomizableObjectVersion
 	};
 	
 	static constexpr int32 CurrentSupportedVersion = ECustomizableObjectVersions::LastCustomizableObjectVersion;
 	
-	static constexpr int32 DerivedDataVersion = 0;
+	static constexpr int32 DerivedDataVersion = 1;
 };
 
diff --git a/Engine/Plugins/Mutable/Source/CustomizableObject/Private/MuCO/CustomizableObject.cpp b/Engine/Plugins/Mutable/Source/CustomizableObject/Private/MuCO/CustomizableObject.cpp
--- a/Engine/Plugins/Mutable/Source/CustomizableObject/Private/MuCO/CustomizableObject.cpp
+++ b/Engine/Plugins/Mutable/Source/CustomizableObject/Private/MuCO/CustomizableObject.cpp
@@ -766,6 +766,8 @@ void FModelResources::Serialize(FObjectAndNameAsStringProxyArchive& MemoryWriter
 	MemoryWriter << SkinWeightProfilesInfo;
 
 	MemoryWriter << ImageProperties;
+	MemoryWriter << MeshMetadata;
+	MemoryWriter << SurfaceMetadata;
 	MemoryWriter << ParameterUIDataMap;
 	MemoryWriter << StateUIDataMap;
 
@@ -918,6 +920,8 @@ bool FModelResources::Unserialize(FObjectAndNameAsStringProxyArchive& MemoryRead
 	MemoryReader << SkinWeightProfilesInfo;
 
 	MemoryReader << ImageProperties;
+	MemoryReader << MeshMetadata;
+	MemoryReader << SurfaceMetadata;
 	MemoryReader << ParameterUIDataMap;
 	MemoryReader << StateUIDataMap;
 
