[mutable] Fix handling of instances that add and remove components based on parameters.
- Move Actor ownership to the preview widget.
- Preview actor now only has the components that are actually created, and maintains them.
- Fix misleading comment in GetNumComponents method.
- Remove some unused generation data.
#rnx

diff --git a/Engine/Plugins/Mutable/Source/CustomizableObject/Internal/MuCO/CustomizableObjectInstancePrivate.h b/Engine/Plugins/Mutable/Source/CustomizableObject/Internal/MuCO/CustomizableObjectInstancePrivate.h
--- a/Engine/Plugins/Mutable/Source/CustomizableObject/Internal/MuCO/CustomizableObjectInstancePrivate.h
+++ b/Engine/Plugins/Mutable/Source/CustomizableObject/Internal/MuCO/CustomizableObjectInstancePrivate.h
@@ -199,7 +199,7 @@ class CUSTOMIZABLEOBJECT_API UCustomizableInstancePrivate : public UObject
 public:
 	GENERATED_BODY()
 
-	/** The generated skeletal meshes for this Instance, one for each component in the Object component list which means that some may be null if the component hasn't been generated. */
+	/** The generated skeletal meshes for this Instance. They may be null if the component is empty. */
 	UPROPERTY(Transient, VisibleAnywhere, Category = NoCategory, meta=(DisplayName = "Meshes"))
 	TMap<FName, TObjectPtr<USkeletalMesh>> SkeletalMeshes;
 	
diff --git a/Engine/Plugins/Mutable/Source/CustomizableObject/Internal/MuCO/UnrealPortabilityHelpers.h b/Engine/Plugins/Mutable/Source/CustomizableObject/Internal/MuCO/UnrealPortabilityHelpers.h
--- a/Engine/Plugins/Mutable/Source/CustomizableObject/Internal/MuCO/UnrealPortabilityHelpers.h
+++ b/Engine/Plugins/Mutable/Source/CustomizableObject/Internal/MuCO/UnrealPortabilityHelpers.h
@@ -28,7 +28,6 @@
 #define UE_MUTABLE_ASSETCLASS(X)				X.AssetClassPath
 #define UE_MUTABLE_GETOBJECTPATH(X)				X.GetObjectPathString()
 
-#define UE_MUTABLE_GETSKINNEDASSET(X)			X->GetSkinnedAsset()
 #define UE_MUTABLE_SETSKINNEDASSET(X,Y)			X->SetSkinnedAsset(Y)
 #define UE_MUTABLE_GETSKELETALMESHASSET(X)		X->GetSkeletalMeshAsset()
 
@@ -43,7 +42,6 @@
 #define UE_MUTABLE_ASSETCLASS(X)				X.AssetClass
 #define UE_MUTABLE_GETOBJECTPATH(X)				X.ObjectPath.ToString()
 
-#define UE_MUTABLE_GETSKINNEDASSET(X)			X->SkeletalMesh
 #define UE_MUTABLE_SETSKINNEDASSET(X,Y)			X->SetSkeletalMesh(Y)
 #define UE_MUTABLE_GETSKELETALMESHASSET(X)		X->SkeletalMesh
 
diff --git a/Engine/Plugins/Mutable/Source/CustomizableObject/Private/MuCO/CustomizableObjectInstance.cpp b/Engine/Plugins/Mutable/Source/CustomizableObject/Private/MuCO/CustomizableObjectInstance.cpp
--- a/Engine/Plugins/Mutable/Source/CustomizableObject/Private/MuCO/CustomizableObjectInstance.cpp
+++ b/Engine/Plugins/Mutable/Source/CustomizableObject/Private/MuCO/CustomizableObjectInstance.cpp
@@ -1994,12 +1994,12 @@ bool UCustomizableInstancePrivate::DoComponentsNeedUpdate(UCustomizableObjectIns
 		// Components with mesh must have valid geometry at CurrentMaxLOD
 
 		TObjectPtr<USkeletalMesh>* SkeletalMeshPtr = SkeletalMeshes.Find(ComponentName);
-		const bool bHasSkeletalMesh = SkeletalMeshPtr && *SkeletalMeshPtr;
+		const bool bHadSkeletalMesh = SkeletalMeshPtr && *SkeletalMeshPtr;
 
 		if (Component.LODCount == 0)
 		{
 			// We don't have a mesh in the component, so it has changed if we had one before.
-			OperationData->MeshChanged[InstanceComponentIndex] = bHasSkeletalMesh;
+			OperationData->MeshChanged[InstanceComponentIndex] = bHadSkeletalMesh;
 			continue;
 		}
 
@@ -2017,8 +2017,8 @@ bool UCustomizableInstancePrivate::DoComponentsNeedUpdate(UCustomizableObjectIns
 			}
 		}
 
-		// Update the component if there is a mesh and it shouldn't, or the other way around.
-		OperationData->MeshChanged[InstanceComponentIndex] = (ComponentWithMesh[InstanceComponentIndex] != bHasSkeletalMesh);
+		// If the component wasn't there and now is there, we need to update it.
+		OperationData->MeshChanged[InstanceComponentIndex] = !bHadSkeletalMesh && LOD.Mesh && (LOD.Mesh->GetFaceCount()>0);
 
 		const FCustomizableInstanceComponentData* ComponentData = GetComponentData(ObjectComponentIndex);
 		if (!ComponentData) // Could be nullptr if the component has not been generated.
@@ -2039,6 +2039,33 @@ bool UCustomizableInstancePrivate::DoComponentsNeedUpdate(UCustomizableObjectIns
 		||
 		OperationData->MeshChanged.Find(true) != INDEX_NONE;
 
+	// It also changed if we removed a component that we did have before
+	if (!bChanged)
+	{
+		for (const TPair<FName, TObjectPtr<USkeletalMesh>>& OldMesh : SkeletalMeshes)
+		{
+			bool bFound = false;
+			for (int32 InstanceComponentIndex = 0; InstanceComponentIndex < NumInstanceComponents; ++InstanceComponentIndex)
+			{
+				const FInstanceUpdateData::FComponent& Component = OperationData->InstanceUpdateData.Components[InstanceComponentIndex];
+				const int32 ObjectComponentIndex = Component.Id;
+				const FName ComponentName = OperationData->Instance->GetCustomizableObject()->GetPrivate()->GetModelResources().ComponentNames[ObjectComponentIndex];
+
+				if (ComponentName == OldMesh.Key)
+				{
+					bFound = true;
+					break;
+				}
+			}
+
+			if (!bFound)
+			{
+				bChanged = true;
+				break;
+			}
+		}
+	}
+
 	return !bHasInvalidMesh && bChanged;
 }
 
diff --git a/Engine/Plugins/Mutable/Source/CustomizableObject/Public/MuCO/CustomizableObjectInstance.h b/Engine/Plugins/Mutable/Source/CustomizableObject/Public/MuCO/CustomizableObjectInstance.h
--- a/Engine/Plugins/Mutable/Source/CustomizableObject/Public/MuCO/CustomizableObjectInstance.h
+++ b/Engine/Plugins/Mutable/Source/CustomizableObject/Public/MuCO/CustomizableObjectInstance.h
@@ -639,7 +639,7 @@ public:
 	float GetMinSquareDistToPlayer() const;
 	void SetMinSquareDistToPlayer(float NewValue);
 
-	/** Return the number of components generated in this instance. */
+	/** Return the number of components that are defined in the CO, unrelated to this instance parameters. */
 	UFUNCTION(BlueprintCallable, Category = CustomizableObjectInstance)
 	int32 GetNumComponents() const;
 
diff --git a/Engine/Plugins/Mutable/Source/CustomizableObjectEditor/Private/MuCOE/CustomizableObjectCustomSettingsDetails.cpp b/Engine/Plugins/Mutable/Source/CustomizableObjectEditor/Private/MuCOE/CustomizableObjectCustomSettingsDetails.cpp
--- a/Engine/Plugins/Mutable/Source/CustomizableObjectEditor/Private/MuCOE/CustomizableObjectCustomSettingsDetails.cpp
+++ b/Engine/Plugins/Mutable/Source/CustomizableObjectEditor/Private/MuCOE/CustomizableObjectCustomSettingsDetails.cpp
@@ -67,9 +67,9 @@ void FCustomizableObjectCustomSettingsDetails::CustomizeDetails(IDetailLayoutBui
 					return true;
 				}
 
-				for (TWeakObjectPtr<UDebugSkelMeshComponent>& WeakPreviewMeshComponent : Editor->GetViewport()->GetViewportClient()->GetPreviewMeshComponents())
+				for (TPair<FName,TWeakObjectPtr<UDebugSkelMeshComponent>>& Entry : Editor->GetViewport()->GetViewportClient()->GetPreviewMeshComponents())
 				{
-					UDebugSkelMeshComponent* PreviewMeshComponent = WeakPreviewMeshComponent.Get();
+					UDebugSkelMeshComponent* PreviewMeshComponent = Entry.Value.Get();
 					if (!PreviewMeshComponent)
 					{
 						continue;
diff --git a/Engine/Plugins/Mutable/Source/CustomizableObjectEditor/Private/MuCOE/CustomizableObjectEditor.cpp b/Engine/Plugins/Mutable/Source/CustomizableObjectEditor/Private/MuCOE/CustomizableObjectEditor.cpp
--- a/Engine/Plugins/Mutable/Source/CustomizableObjectEditor/Private/MuCOE/CustomizableObjectEditor.cpp
+++ b/Engine/Plugins/Mutable/Source/CustomizableObjectEditor/Private/MuCOE/CustomizableObjectEditor.cpp
@@ -23,7 +23,6 @@
 #include "MaterialGraph/MaterialGraphNode_Comment.h"
 #include "MuCO/CustomizableObjectInstance.h"
 #include "MuCO/CustomizableObjectSystem.h"
-#include "MuCO/CustomizableSkeletalComponent.h"
 #include "MuCO/UnrealPortabilityHelpers.h"
 #include "MuCO/CustomizableObjectCompilerTypes.h"
 #include "MuCOE/CustomizableObjectCustomSettings.h"
@@ -67,9 +66,7 @@
 #include "Widgets/Layout/SScrollBox.h"
 #include "Widgets/Notifications/SNotificationList.h"
 #include "SMutableObjectViewer.h"
-#include "MuCO/CustomizableSkeletalComponentPrivate.h"
 #include "MuCO/CustomizableObjectCustomVersion.h"
-#include "MuCO/CustomizableSkeletalMeshActor.h"
 #include "MuCOE/CustomizableObjectInstanceFactory.h"
 
 class FAdvancedPreviewScene;
@@ -2109,7 +2106,7 @@ void FCustomizableObjectEditor::FindProperty(const FProperty* Property, const vo
 
 void FCustomizableObjectEditor::OnPostCompile()
 {
-	CreatePreviewActor();
+	Viewport->CreatePreviewActor(PreviewInstance);
 	PreviewInstance->UpdateSkeletalMeshAsync(true, true);
 }
 
@@ -2199,12 +2196,6 @@ void FCustomizableObjectEditor::LogSearchResult(const UObject& Context, const FS
 
 void FCustomizableObjectEditor::OnUpdatePreviewInstance(UCustomizableObjectInstance* Instance)
 {
-	if (PreviewInstance->GetPrivate()->SkeletalMeshStatus != ESkeletalMeshStatus::Success)
-	{
-		Actor->GetRootComponent()->SetVisibility(false, true);
-	}
-	Actor->GetRootComponent()->UpdateBounds();
-	
 	if (TextureAnalyzer.IsValid())
 	{
 		TextureAnalyzer->RefreshTextureAnalyzerTable(PreviewInstance);
@@ -2317,41 +2308,6 @@ UCustomizableObject* FCustomizableObjectEditor::GetAbsoluteCOParent(const UCusto
 }
 
 
-void FCustomizableObjectEditor::CreatePreviewActor()
-{
-	if (Actor)
-	{
-		Actor->Destroy();
-	}
-	
-	Actor.Reset(ViewportClient->GetWorld()->SpawnActor<ASkeletalMeshActor>());
-
-	TArray<TWeakObjectPtr<UDebugSkelMeshComponent>> PreviewSkeletalMeshComponents;
-	
-	for (int32 ObjectComponentIndex = 0; ObjectComponentIndex < CustomizableObject->GetComponentCount(); ++ObjectComponentIndex)
-	{
-		UDebugSkelMeshComponent* DebugComponent = NewObject<UDebugSkelMeshComponent>(Actor.Get(), NAME_None, RF_Transient);
-		DebugComponent->bCastInsetShadow = true; // For better quality shadows in the editor previews, more similar to the in-game ones
-		DebugComponent->bCanHighlightSelectedSections = false;
-		DebugComponent->bComponentUseFixedSkelBounds = true; // First bounds computed would be using physics asset
-		DebugComponent->MarkRenderStateDirty();
-		DebugComponent->AttachToComponent(Actor->GetRootComponent(), FAttachmentTransformRules::KeepRelativeTransform);
-		DebugComponent->RegisterComponent();
-		
-		PreviewSkeletalMeshComponents.Add(DebugComponent);
-		
-		UCustomizableSkeletalComponent* CustomizableComponent = NewObject<UCustomizableSkeletalComponent>(DebugComponent);
-		CustomizableComponent->SetSkipSetReferenceSkeletalMesh(true);
-		CustomizableComponent->CustomizableObjectInstance = PreviewInstance;
-		CustomizableComponent->SetComponentName(CustomizableObject->GetComponentName(ObjectComponentIndex));
-		CustomizableComponent->AttachToComponent(DebugComponent, FAttachmentTransformRules::KeepRelativeTransform);
-		CustomizableComponent->RegisterComponent();
-	}
-	
-	Viewport->SetPreviewActor(Actor.Get(), PreviewInstance, PreviewSkeletalMeshComponents);
-}
-
-
 void RemoveRestrictedChars(FString& String)
 {
 	// Remove restricted chars, according to FPaths::ValidatePath, RestrictedChars = "/?:&\\*\"<>|%#@^ ";
@@ -2380,7 +2336,7 @@ void FCustomizableObjectEditor::OnCustomizableObjectStatusChanged(FCustomizableO
 	{
 		if (CurrentState == FCustomizableObjectStatusTypes::EState::ModelLoaded)
 		{
-			CreatePreviewActor();
+			Viewport->CreatePreviewActor(PreviewInstance);
 			PreviewInstance->UpdateSkeletalMeshAsync(true, true);
 		}
 		else if (CurrentState == FCustomizableObjectStatusTypes::EState::NoModel)
diff --git a/Engine/Plugins/Mutable/Source/CustomizableObjectEditor/Private/MuCOE/CustomizableObjectEditor.h b/Engine/Plugins/Mutable/Source/CustomizableObjectEditor/Private/MuCOE/CustomizableObjectEditor.h
--- a/Engine/Plugins/Mutable/Source/CustomizableObjectEditor/Private/MuCOE/CustomizableObjectEditor.h
+++ b/Engine/Plugins/Mutable/Source/CustomizableObjectEditor/Private/MuCOE/CustomizableObjectEditor.h
@@ -357,9 +357,6 @@ private:
 
 	/** Open the Performance Analyzer tab */
 	void OpenPerformanceAnalyzerTab();
-	
-	/** Creates the actor and necessary components for the preview of the CO instance */
-	void CreatePreviewActor();
 
 	/** Recursively find any property that its name or value contains the given string.
 	  * @param Property Root property.
@@ -396,9 +393,6 @@ private:
 	TObjectPtr<UCustomizableObject> CustomizableObject = nullptr;
 	TObjectPtr<UCustomizableObjectInstance> PreviewInstance = nullptr;
 
-	/** Preview Actor. All preview components are attached to this actor. */
-	TStrongObjectPtr<AActor> Actor;
-
 	/** List of open tool panels; used to ensure only one exists at any one time */
 	TMap< FName, TWeakPtr<class SDockableTab> > SpawnedToolPanels;
 
diff --git a/Engine/Plugins/Mutable/Source/CustomizableObjectEditor/Private/MuCOE/CustomizableObjectEditorViewportClient.cpp b/Engine/Plugins/Mutable/Source/CustomizableObjectEditor/Private/MuCOE/CustomizableObjectEditorViewportClient.cpp
--- a/Engine/Plugins/Mutable/Source/CustomizableObjectEditor/Private/MuCOE/CustomizableObjectEditorViewportClient.cpp
+++ b/Engine/Plugins/Mutable/Source/CustomizableObjectEditor/Private/MuCOE/CustomizableObjectEditorViewportClient.cpp
@@ -31,7 +31,11 @@
 #include "Misc/PackageName.h"
 #include "Misc/ConfigCacheIni.h"
 #include "MuCO/CustomizableObjectInstance.h"
+#include "MuCO/CustomizableObjectInstancePrivate.h"
 #include "MuCO/CustomizableObjectSystem.h"
+#include "MuCO/CustomizableSkeletalMeshActor.h"
+#include "MuCO/CustomizableSkeletalComponent.h"
+#include "MuCO/CustomizableSkeletalComponentPrivate.h"
 #include "MuCOE/CustomizableObjectPreviewScene.h"
 #include "MuCOE/ICustomizableObjectInstanceEditor.h"
 #include "MuCOE/Nodes/CustomizableObjectNodeModifierClipMorph.h"
@@ -259,11 +263,12 @@ void FCustomizableObjectEditorViewportClient::Draw(const FSceneView* View, FPrim
 		{
 			float MaxSphereRadius = 0.f;
 
-			for (TWeakObjectPtr<USkeletalMeshComponent> SkeletalMeshComponent : SkeletalMeshComponents)
+			for (const TPair<FName, TWeakObjectPtr<UDebugSkelMeshComponent>>& Key : SkeletalMeshComponents)
 			{
-				if (SkeletalMeshComponent.IsValid())
+				const UDebugSkelMeshComponent* Component = Key.Value.Get();
+				if (Component)
 				{
-					MaxSphereRadius = FMath::Max(MaxSphereRadius, SkeletalMeshComponent->Bounds.SphereRadius);
+					MaxSphereRadius = FMath::Max(MaxSphereRadius, Component->Bounds.SphereRadius);
 				}
 			}
 
@@ -382,11 +387,12 @@ void FCustomizableObjectEditorViewportClient::Draw(const FSceneView* View, FPrim
 
 	if (bShowBones)
 	{
-		for (const TWeakObjectPtr<UDebugSkelMeshComponent>& SkeletalMeshComponent : SkeletalMeshComponents)
+		for (const TPair<FName, TWeakObjectPtr<UDebugSkelMeshComponent>>& Key : SkeletalMeshComponents)
 		{
-			if (SkeletalMeshComponent.IsValid())
+			const UDebugSkelMeshComponent* Component = Key.Value.Get();
+			if (Component)
 			{
-				DrawMeshBones(SkeletalMeshComponent.Get(), PDI);
+				DrawMeshBones(Component, PDI);
 			}
 		}
 	}
@@ -398,11 +404,12 @@ void FCustomizableObjectEditorViewportClient::Draw(const FSceneView* View, FPrim
 		const bool bPreviousValue = CVar->GetBool();
 		CVar->Set(true);
 		
-		for (const TWeakObjectPtr<UDebugSkelMeshComponent>& SkeletalMeshComponent : SkeletalMeshComponents)
+		for (const TPair<FName, TWeakObjectPtr<UDebugSkelMeshComponent>>& Key : SkeletalMeshComponents)
 		{
-			if (SkeletalMeshComponent.IsValid())
+			UDebugSkelMeshComponent* Component = Key.Value.Get();
+			if (Component)
 			{
-				SkeletalMeshComponent->DebugDrawClothing(PDI);
+				Component->DebugDrawClothing(PDI);
 			}
 		}
 		
@@ -414,14 +421,38 @@ void FCustomizableObjectEditorViewportClient::Draw(const FSceneView* View, FPrim
 void FCustomizableObjectEditorViewportClient::Draw(FViewport* InViewport, FCanvas* Canvas)
 {
 	// Defensive check to avoid unreal crashing inside render if the mesh is degenereated
-	for (TWeakObjectPtr<USkeletalMeshComponent> SkeletalMeshComponent : SkeletalMeshComponents)
+	for (const TPair<FName, TWeakObjectPtr<UDebugSkelMeshComponent>>& Key : SkeletalMeshComponents)
 	{
-		if (SkeletalMeshComponent.IsValid() && UE_MUTABLE_GETSKINNEDASSET(SkeletalMeshComponent) && UE_MUTABLE_GETSKINNEDASSET(SkeletalMeshComponent)->GetLODNum() == 0)
+		UDebugSkelMeshComponent* Component = Key.Value.Get();
+
+		if (Component && Component->GetSkinnedAsset() && Component->GetSkinnedAsset()->GetLODNum() == 0)
+		{
+			Component->SetSkeletalMesh(nullptr);
+		}
+	}
+
+	// Configure the initial orbital position of the camera
+	if (!bIsCameraSetup)
+	{
+		if (!Actor.IsValid())
 		{
-			SkeletalMeshComponent->SetSkeletalMesh(nullptr);
+			return;
 		}
+
+		FVector Center;
+		FVector Extents;
+		Actor.Get()->GetActorBounds(false, Center, Extents, true);
+
+		bIsCameraSetup = Extents.X * Extents.Y * Extents.Z > 0.0;
+
+		static FRotator CustomOrbitRotation(-33.75, -135, 0);
+		FVector CustomOrbitZoom(0, Extents.GetMax() * 2.5 / (75.0 * PI / 360.0), 0);
+
+		SetCameraSetup(Center, CustomOrbitRotation, CustomOrbitZoom, Center, FVector::Zero(), {} /** Not used since orbital is enable just after. */);
+		EnableCameraLock(true);
 	}
 
+
 	FEditorViewportClient::Draw(InViewport, Canvas);
 
 	FSceneViewFamilyContext ViewFamily( FSceneViewFamily::ConstructionValues( InViewport, GetScene(), EngineShowFlags ));
@@ -449,7 +480,7 @@ namespace
 
 void FCustomizableObjectEditorViewportClient::DrawUVs(FViewport* InViewport, FCanvas* InCanvas, int32 InTextYPos)
 {
-	const uint32 ComponentIndex = UVDrawComponentIndex;
+	const FName ComponentName = UVDrawComponentName;
 	const uint32 LODLevel = UVDrawLODIndex; 	// TODO use the overriden LOD level
 	const int32 SectionIndex = UVDrawSectionIndex;
 	const int32 UVChannel = UVDrawUVIndex;
@@ -496,16 +527,17 @@ void FCustomizableObjectEditorViewportClient::DrawUVs(FViewport* InViewport, FCa
 
 	if (SkeletalMeshComponents.Num())
 	{
-		if (SkeletalMeshComponents.IsValidIndex(ComponentIndex))
+		TWeakObjectPtr<UDebugSkelMeshComponent>* ComponentPtr = SkeletalMeshComponents.Find(ComponentName);
+		if (ComponentPtr && ComponentPtr->IsValid())
 		{
-			TWeakObjectPtr<USkeletalMeshComponent> SkeletalMeshComponent = SkeletalMeshComponents[ComponentIndex];
+			TWeakObjectPtr<USkeletalMeshComponent> SkeletalMeshComponent = *ComponentPtr;
 
-			if (!SkeletalMeshComponent.IsValid() || !UE_MUTABLE_GETSKINNEDASSET(SkeletalMeshComponent))
+			if (!SkeletalMeshComponent.IsValid() || !SkeletalMeshComponent->GetSkinnedAsset())
 			{
 				return;
 			}
 
-			const FSkeletalMeshRenderData* MeshRes = UE_MUTABLE_GETSKINNEDASSET(SkeletalMeshComponent)->GetResourceForRendering();
+			const FSkeletalMeshRenderData* MeshRes = SkeletalMeshComponent->GetSkinnedAsset()->GetResourceForRendering();
 			if (!MeshRes->LODRenderData.IsValidIndex(LODLevel))
 			{
 				return;
@@ -696,22 +728,16 @@ void FCustomizableObjectEditorViewportClient::HideGizmoLight()
 }
 
 
-void FCustomizableObjectEditorViewportClient::SetPreviewActor(const TWeakObjectPtr<AActor>& InActor, const TWeakObjectPtr<UCustomizableObjectInstance>& InInstance, const TArray<TWeakObjectPtr<UDebugSkelMeshComponent>>& InSkeletalMeshComponents)
+void FCustomizableObjectEditorViewportClient::CreatePreviewActor(const TWeakObjectPtr<UCustomizableObjectInstance>& InInstance)
 {
-	SkeletalMeshComponents = InSkeletalMeshComponents;
-	Actor = InActor;
-
-	for (const TWeakObjectPtr<UDebugSkelMeshComponent>& SkeletalMeshComponent : SkeletalMeshComponents)
+	if (Actor)
 	{
-		if (SkeletalMeshComponent.IsValid())
-		{
-			SkeletalMeshComponent->bDisableClothSimulation = bDisableClothSimulation;
-			SkeletalMeshComponent->bDrawNormals = bDrawNormals;
-			SkeletalMeshComponent->bDrawTangents = bDrawTangents;
-			SkeletalMeshComponent->bDrawBinormals = bDrawBinormals;
-		}
+		Actor->Destroy();
+		SkeletalMeshComponents.Reset();
 	}
 
+	Actor.Reset(GetWorld()->SpawnActor<ASkeletalMeshActor>());
+
 	bUpdated = false;
 	
 	InInstance->UpdatedNativeDelegate.AddSP(SharedThis(this), &FCustomizableObjectEditorViewportClient::OnInstanceUpdate);
@@ -720,7 +746,7 @@ void FCustomizableObjectEditorViewportClient::SetPreviewActor(const TWeakObjectP
 }
 
 
-TArray<TWeakObjectPtr<UDebugSkelMeshComponent>>& FCustomizableObjectEditorViewportClient::GetPreviewMeshComponents()
+TMap<FName, TWeakObjectPtr<UDebugSkelMeshComponent>>& FCustomizableObjectEditorViewportClient::GetPreviewMeshComponents()
 {
 	return SkeletalMeshComponents;
 }
@@ -728,9 +754,9 @@ TArray<TWeakObjectPtr<UDebugSkelMeshComponent>>& FCustomizableObjectEditorViewpo
 
 void FCustomizableObjectEditorViewportClient::SetPreviewAnimationAsset(UAnimationAsset* AnimAsset)
 {
-	for (TWeakObjectPtr<UDebugSkelMeshComponent>& WeakSkeletalMeshComponent : SkeletalMeshComponents)
+	for (TPair<FName,TWeakObjectPtr<UDebugSkelMeshComponent>>& Entry : SkeletalMeshComponents)
 	{
-		UDebugSkelMeshComponent* SkeletalMeshComponent = WeakSkeletalMeshComponent.Get();
+		UDebugSkelMeshComponent* SkeletalMeshComponent = Entry.Value.Get();
 		if (!SkeletalMeshComponent)
 		{
 			continue;
@@ -760,6 +786,75 @@ void FCustomizableObjectEditorViewportClient::SetPreviewAnimationAsset(UAnimatio
 
 void FCustomizableObjectEditorViewportClient::OnInstanceUpdate(UCustomizableObjectInstance* Instance)
 {
+	if (!Instance)
+	{
+		return;
+	}
+
+	const UCustomizableObject* CO = Instance->GetCustomizableObject();
+	if (!CO)
+	{
+		return;
+	}
+
+	// Remove components that are no longer there
+	TArray<FName> ToRemove;
+	ToRemove.Reserve(SkeletalMeshComponents.Num());
+	for (const TPair<FName,TWeakObjectPtr<UDebugSkelMeshComponent>>& Item : SkeletalMeshComponents)
+	{
+		// TODO: This will not work with different types of components like grooms or panel clothing.
+		bool bInstanceHasComponent = Instance->GetComponentMeshSkeletalMesh(Item.Key) != nullptr;
+
+		if (!bInstanceHasComponent)
+		{ 
+			if (UDebugSkelMeshComponent* Comp = Item.Value.Get())
+			{
+				Comp->DestroyComponent();
+			}
+			ToRemove.Add(Item.Key);
+		}
+	}
+
+	for (const FName Name: ToRemove)
+	{
+		SkeletalMeshComponents.Remove(Name);
+	}
+
+	// Add new components
+	for (int32 ObjectComponentIndex = 0; ObjectComponentIndex < CO->GetComponentCount(); ++ObjectComponentIndex)
+	{
+		FName Name = CO->GetComponentName(ObjectComponentIndex);
+		if (Instance->GetComponentMeshSkeletalMesh(Name))
+		{
+			if (!SkeletalMeshComponents.Contains(Name))
+			{
+				// We need to add it.
+				UDebugSkelMeshComponent* DebugComponent = NewObject<UDebugSkelMeshComponent>(Actor.Get(), Name, RF_Transient);
+				DebugComponent->bCastInsetShadow = true; // For better quality shadows in the editor previews, more similar to the in-game ones
+				DebugComponent->bCanHighlightSelectedSections = false;
+				DebugComponent->bComponentUseFixedSkelBounds = true; // First bounds computed would be using physics asset
+				DebugComponent->MarkRenderStateDirty();
+				DebugComponent->AttachToComponent(Actor->GetRootComponent(), FAttachmentTransformRules::KeepRelativeTransform);
+				DebugComponent->RegisterComponent();
+
+				SkeletalMeshComponents.Add(Name, DebugComponent);
+
+				UCustomizableSkeletalComponent* CustomizableComponent = NewObject<UCustomizableSkeletalComponent>(DebugComponent);
+				CustomizableComponent->SetSkipSetReferenceSkeletalMesh(true);
+				CustomizableComponent->CustomizableObjectInstance = Instance;
+				CustomizableComponent->SetComponentName(Name);
+				CustomizableComponent->AttachToComponent(DebugComponent, FAttachmentTransformRules::KeepRelativeTransform);
+				CustomizableComponent->RegisterComponent();
+
+				DebugComponent->bDisableClothSimulation = bDisableClothSimulation;
+				DebugComponent->bDrawNormals = bDrawNormals;
+				DebugComponent->bDrawTangents = bDrawTangents;
+				DebugComponent->bDrawBinormals = bDrawBinormals;
+			}
+		}
+	}
+
+
 	if (!bUpdated)
 	{
 		bUpdated = true;
@@ -769,25 +864,13 @@ void FCustomizableObjectEditorViewportClient::OnInstanceUpdate(UCustomizableObje
 	
 	Invalidate();
 	
-	// Configure the initial orbital position of the camera
-	if (!bIsCameraSetup)
+	if (Actor)
 	{
-		if (!Actor.IsValid())
+		if (Instance->GetPrivate()->SkeletalMeshStatus != ESkeletalMeshStatus::Success)
 		{
-			return;
+			Actor->GetRootComponent()->SetVisibility(false, true);
 		}
-		
-		FVector Center;
-		FVector Extents;
-		Actor.Get()->GetActorBounds(false, Center, Extents, true);
-
-		bIsCameraSetup = Extents.X * Extents.Y * Extents.Z > 0.0;
-		
-		static FRotator CustomOrbitRotation(-33.75, -135, 0);
-		FVector CustomOrbitZoom(0, Extents.GetMax() * 2.5 / (75.0 * PI / 360.0), 0);
-
-		SetCameraSetup(Center, CustomOrbitRotation, CustomOrbitZoom, Center, FVector::Zero(), {} /** Not used since orbital is enable just after. */ );	
-		EnableCameraLock(true);
+		Actor->GetRootComponent()->UpdateBounds();
 	}
 }
 
@@ -799,9 +882,9 @@ void FCustomizableObjectEditorViewportClient::SetDrawUVOverlay()
 }
 
 
-void FCustomizableObjectEditorViewportClient::SetDrawUV(const int32 ComponentIndex, const int32 LODIndex, const int32 SectionIndex, const int32 UVIndex)
+void FCustomizableObjectEditorViewportClient::SetDrawUV(const FName ComponentName, const int32 LODIndex, const int32 SectionIndex, const int32 UVIndex)
 {
-	UVDrawComponentIndex = ComponentIndex;
+	UVDrawComponentName = ComponentName;
 	UVDrawLODIndex = LODIndex;
 	UVDrawSectionIndex = SectionIndex;
 	UVDrawUVIndex = UVIndex;
@@ -1371,9 +1454,9 @@ bool FCustomizableObjectEditorViewportClient::CanSetWidgetMode(UE::Widget::EWidg
 
 void FCustomizableObjectEditorViewportClient::SetAnimation(UAnimationAsset* Animation)
 {
-	for (TWeakObjectPtr<UDebugSkelMeshComponent>& WeakPreviewMeshComponent : SkeletalMeshComponents)
+	for (TPair<FName, TWeakObjectPtr<UDebugSkelMeshComponent>>& Entry : SkeletalMeshComponents)
 	{
-		UDebugSkelMeshComponent* PreviewMeshComponent = WeakPreviewMeshComponent.Get();
+		UDebugSkelMeshComponent* PreviewMeshComponent = Entry.Value.Get();
 		if (!PreviewMeshComponent)
 		{
 			continue;
@@ -1548,11 +1631,13 @@ void FCustomizableObjectEditorViewportClient::ShowInstanceGeometryInformation(FC
 	int32 ComponentIndex = 0;
 
 	// Show total number of triangles and vertices
-	for (TWeakObjectPtr<USkeletalMeshComponent> SkeletalMeshComponent : SkeletalMeshComponents)
+	for (TPair<FName, TWeakObjectPtr<UDebugSkelMeshComponent>>& Entry : SkeletalMeshComponents)
 	{
-		if (SkeletalMeshComponent.IsValid() && UE_MUTABLE_GETSKINNEDASSET(SkeletalMeshComponent))
+		UDebugSkelMeshComponent* SkeletalMeshComponent = Entry.Value.Get();
+
+		if (SkeletalMeshComponent && SkeletalMeshComponent->GetSkinnedAsset())
 		{
-			const FSkeletalMeshRenderData* MeshRes = UE_MUTABLE_GETSKINNEDASSET(SkeletalMeshComponent)->GetResourceForRendering();
+			const FSkeletalMeshRenderData* MeshRes = SkeletalMeshComponent->GetSkinnedAsset()->GetResourceForRendering();
 			int32 NumTriangles;
 			int32 NumVertices;
 			int32 NumLODLevel = MeshRes->LODRenderData.Num();
@@ -1892,9 +1977,10 @@ void FCustomizableObjectEditorViewportClient::OnEnableClothSimulation()
 {
 	bDisableClothSimulation = !bDisableClothSimulation;
 	
-	for (const TWeakObjectPtr<UDebugSkelMeshComponent>& SkeletalMeshComponent : SkeletalMeshComponents)
+	for (TPair<FName, TWeakObjectPtr<UDebugSkelMeshComponent>>& Entry : SkeletalMeshComponents)
 	{
-		if (SkeletalMeshComponent.IsValid())
+		UDebugSkelMeshComponent* SkeletalMeshComponent = Entry.Value.Get();
+		if (SkeletalMeshComponent)
 		{
 			SkeletalMeshComponent->bDisableClothSimulation = bDisableClothSimulation;
 		}
@@ -1939,12 +2025,13 @@ FText MergeLine(const FText& InText, const FText& InNewLine)
 FText FCustomizableObjectEditorViewportClient::GetMeshInfoText() const
 {
 	FText TextValue;
+	bool bFirst = true;
 
-	for (int32 ComponentIndex = 0; ComponentIndex < SkeletalMeshComponents.Num(); ++ComponentIndex)
+	for (const TPair<FName, TWeakObjectPtr<UDebugSkelMeshComponent>>& Entry : SkeletalMeshComponents)
 	{
-		TWeakObjectPtr<UDebugSkelMeshComponent> PreviewMeshComponent = SkeletalMeshComponents[ComponentIndex];
+		UDebugSkelMeshComponent* PreviewMeshComponent = Entry.Value.Get();
 		
-		if (!PreviewMeshComponent.IsValid())
+		if (!PreviewMeshComponent)
 		{
 			continue;
 		}
@@ -1975,13 +2062,14 @@ FText FCustomizableObjectEditorViewportClient::GetMeshInfoText() const
 			NumTotalTriangles += LODData.RenderSections[SectionIndex].NumTriangles;
 		}
 
-		if (ComponentIndex > 0)
+		if (!bFirst)
 		{
 			TextValue = FText::Format(LOCTEXT("MeshInfoComponentSeparation", "{0}\n"), TextValue);
 		}
+		bFirst = false;
 		
 		TextValue = MergeLine(TextValue, FText::Format(LOCTEXT("MeshInfoFormat", "Component: {0}, LOD: {1}, Bones: {2} (Mapped to Vertices: {3}), Polys: {4}"),
-			FText::AsNumber(ComponentIndex),
+			FText::FromString(Entry.Key.ToString()),
 			FText::AsNumber(LODIndex),
 			FText::AsNumber(NumBonesInUse),
 			FText::AsNumber(NumBonesMappedToVerts),
@@ -2060,9 +2148,10 @@ void FCustomizableObjectEditorViewportClient::ToggleShowNormals()
 {
 	bDrawNormals = !bDrawNormals;
 	
-	for (const TWeakObjectPtr<UDebugSkelMeshComponent>& SkeletalMeshComponent : SkeletalMeshComponents)
+	for (TPair<FName, TWeakObjectPtr<UDebugSkelMeshComponent>>& Entry : SkeletalMeshComponents)
 	{
-		if (SkeletalMeshComponent.IsValid())
+		UDebugSkelMeshComponent* SkeletalMeshComponent = Entry.Value.Get();
+		if (SkeletalMeshComponent)
 		{
 			SkeletalMeshComponent->bDrawNormals = bDrawNormals;
 			SkeletalMeshComponent->MarkRenderStateDirty();
@@ -2083,9 +2172,10 @@ void FCustomizableObjectEditorViewportClient::ToggleShowTangents()
 {
 	bDrawTangents = !bDrawTangents;
 
-	for (const TWeakObjectPtr<UDebugSkelMeshComponent>& SkeletalMeshComponent : SkeletalMeshComponents)
+	for (TPair<FName, TWeakObjectPtr<UDebugSkelMeshComponent>>& Entry : SkeletalMeshComponents)
 	{
-		if (SkeletalMeshComponent.IsValid())
+		UDebugSkelMeshComponent* SkeletalMeshComponent = Entry.Value.Get();
+		if (SkeletalMeshComponent)
 		{
 			SkeletalMeshComponent->bDrawTangents = bDrawTangents;
 			SkeletalMeshComponent->MarkRenderStateDirty();
@@ -2106,9 +2196,10 @@ void FCustomizableObjectEditorViewportClient::ToggleShowBinormals()
 {
 	bDrawBinormals = !bDrawBinormals;
 	
-	for (const TWeakObjectPtr<UDebugSkelMeshComponent>& SkeletalMeshComponent : SkeletalMeshComponents)
+	for (TPair<FName, TWeakObjectPtr<UDebugSkelMeshComponent>>& Entry : SkeletalMeshComponents)
 	{
-		if (SkeletalMeshComponent.IsValid())
+		UDebugSkelMeshComponent* SkeletalMeshComponent = Entry.Value.Get();
+		if (SkeletalMeshComponent)
 		{
 			SkeletalMeshComponent->bDrawBinormals = bDrawBinormals;
 			SkeletalMeshComponent->MarkRenderStateDirty();
@@ -2127,7 +2218,7 @@ bool FCustomizableObjectEditorViewportClient::IsSetShowBinormalsChecked() const
 }
 
 
-void FCustomizableObjectEditorViewportClient::DrawMeshBones(UDebugSkelMeshComponent* MeshComponent, FPrimitiveDrawInterface* PDI)
+void FCustomizableObjectEditorViewportClient::DrawMeshBones(const UDebugSkelMeshComponent* MeshComponent, FPrimitiveDrawInterface* PDI)
 {
 	if (!MeshComponent ||
 		!MeshComponent->GetSkeletalMeshAsset() ||
diff --git a/Engine/Plugins/Mutable/Source/CustomizableObjectEditor/Private/MuCOE/CustomizableObjectInstanceEditor.cpp b/Engine/Plugins/Mutable/Source/CustomizableObjectEditor/Private/MuCOE/CustomizableObjectInstanceEditor.cpp
--- a/Engine/Plugins/Mutable/Source/CustomizableObjectEditor/Private/MuCOE/CustomizableObjectInstanceEditor.cpp
+++ b/Engine/Plugins/Mutable/Source/CustomizableObjectEditor/Private/MuCOE/CustomizableObjectInstanceEditor.cpp
@@ -671,7 +671,7 @@ void FCustomizableObjectInstanceEditor::OnCustomizableObjectStatusChanged(FCusto
 	{
 		if (CurrentState == FCustomizableObjectStatusTypes::EState::ModelLoaded)
 		{
-			CreatePreviewActor();
+			Viewport->CreatePreviewActor(CustomizableObjectInstance);
 			CustomizableObjectInstance->UpdateSkeletalMeshAsync(true, true);
 		}
 		else if (CurrentState == FCustomizableObjectStatusTypes::EState::NoModel)
@@ -693,7 +693,7 @@ void FCustomizableObjectInstanceEditor::OpenTextureAnalyzerTab()
 
 void FCustomizableObjectInstanceEditor::OnPostCompile()
 {
-	CreatePreviewActor();
+	Viewport->CreatePreviewActor(CustomizableObjectInstance);
 	CustomizableObjectInstance->UpdateSkeletalMeshAsync(true, true);
 }
 
@@ -956,48 +956,6 @@ void FCustomizableObjectInstanceEditor::HideGizmoProjectorParameter(const TShare
 }
 
 
-void FCustomizableObjectInstanceEditor::CreatePreviewActor()
-{
-	const UCustomizableObject* Object = CustomizableObjectInstance->GetCustomizableObject();
-	
-	if (!Object)
-	{
-		return;
-	}
-
-	if (Actor)
-	{
-		Viewport->GetViewportClient()->GetWorld()->RemoveActor(Actor.Get(), false);
-	}
-	
-	Actor.Reset(Viewport->GetViewportClient()->GetWorld()->SpawnActor<ASkeletalMeshActor>());
-
-	PreviewSkeletalMeshComponents.Reset();
-	
-	for (int32 ObjectComponentIndex = 0; ObjectComponentIndex < Object->GetComponentCount(); ++ObjectComponentIndex)
-	{
-		UDebugSkelMeshComponent* DebugComponent = NewObject<UDebugSkelMeshComponent>(Actor.Get(), NAME_None, RF_Transient);
-		DebugComponent->bCastInsetShadow = true; // For better quality shadows in the editor previews, more similar to the in-game ones
-		DebugComponent->bCanHighlightSelectedSections = false;
-		DebugComponent->bComponentUseFixedSkelBounds = true; // First bounds computed would be using physics asset
-		DebugComponent->MarkRenderStateDirty();
-		DebugComponent->AttachToComponent(Actor->GetRootComponent(), FAttachmentTransformRules::KeepRelativeTransform);
-		DebugComponent->RegisterComponent();
-		
-		PreviewSkeletalMeshComponents.Add(DebugComponent);
-		
-		UCustomizableSkeletalComponent* CustomizableComponent = NewObject<UCustomizableSkeletalComponent>(DebugComponent);
-		CustomizableComponent->SetSkipSetReferenceSkeletalMesh(true);
-		CustomizableComponent->CustomizableObjectInstance = CustomizableObjectInstance;
-		CustomizableComponent->SetComponentName(Object->GetComponentName(ObjectComponentIndex));
-		CustomizableComponent->AttachToComponent(DebugComponent, FAttachmentTransformRules::KeepRelativeTransform);
-		CustomizableComponent->RegisterComponent();
-	}
-	
-	Viewport->SetPreviewActor(Actor.Get(), CustomizableObjectInstance, PreviewSkeletalMeshComponents);
-}
-
-
 TSharedRef<SDockTab> FCustomizableObjectInstanceEditor::SpawnTab_TextureAnalyzer(const FSpawnTabArgs& Args)
 {
 	check(Args.GetTabId() == TextureAnalyzerTabId);
diff --git a/Engine/Plugins/Mutable/Source/CustomizableObjectEditor/Private/MuCOE/CustomizableObjectInstanceEditor.h b/Engine/Plugins/Mutable/Source/CustomizableObjectEditor/Private/MuCOE/CustomizableObjectInstanceEditor.h
--- a/Engine/Plugins/Mutable/Source/CustomizableObjectEditor/Private/MuCOE/CustomizableObjectInstanceEditor.h
+++ b/Engine/Plugins/Mutable/Source/CustomizableObjectEditor/Private/MuCOE/CustomizableObjectInstanceEditor.h
@@ -276,7 +276,6 @@ public:
 		const TSharedPtr<IDetailsView>& InstanceDetailsView);
 	
 private:
-	void CreatePreviewActor();
 	
 	void OnPostCompile();
 	
diff --git a/Engine/Plugins/Mutable/Source/CustomizableObjectEditor/Private/MuCOE/CustomizableObjectPreviewScene.cpp b/Engine/Plugins/Mutable/Source/CustomizableObjectEditor/Private/MuCOE/CustomizableObjectPreviewScene.cpp
--- a/Engine/Plugins/Mutable/Source/CustomizableObjectEditor/Private/MuCOE/CustomizableObjectPreviewScene.cpp
+++ b/Engine/Plugins/Mutable/Source/CustomizableObjectEditor/Private/MuCOE/CustomizableObjectPreviewScene.cpp
@@ -3,8 +3,6 @@
 #include "MuCOE/CustomizableObjectPreviewScene.h"
 
 
-class UStaticMeshComponent;
-
 FCustomizableObjectPreviewScene::FCustomizableObjectPreviewScene(ConstructionValues CVS, float InFloorOffset)
 	: FAdvancedPreviewScene(CVS, InFloorOffset)
 {
diff --git a/Engine/Plugins/Mutable/Source/CustomizableObjectEditor/Private/MuCOE/GenerateMutableSource/GenerateMutableSource.h b/Engine/Plugins/Mutable/Source/CustomizableObjectEditor/Private/MuCOE/GenerateMutableSource/GenerateMutableSource.h
--- a/Engine/Plugins/Mutable/Source/CustomizableObjectEditor/Private/MuCOE/GenerateMutableSource/GenerateMutableSource.h
+++ b/Engine/Plugins/Mutable/Source/CustomizableObjectEditor/Private/MuCOE/GenerateMutableSource/GenerateMutableSource.h
@@ -854,10 +854,6 @@ public:
 	// For a given material node (the key is node package path + node uid + image index in node) stores images generated for the same node at a higher quality LOD to reuse that image node
 	TMap<FString, FGroupProjectorImageInfo> GroupProjectorLODCache;
 
-	// Data structures used for clipping feature by assigning a Customizable Object
-	// Map with pairs (Unreal Mutable material node Guid, array with its corresponding Mutable surface node) to add tags from clipping nodes
-	TMap<class UCustomizableObjectNodeMaterial*, TArray<mu::Ptr<mu::NodeSurfaceNew>>> MapMaterialNodeToMutableSurfaceNodeArray;
-
 	// Data used for MorphTarget reconstruction.
 	TMap<uint32, FRealTimeMorphMeshData> RealTimeMorphTargetPerMeshData;
 
diff --git a/Engine/Plugins/Mutable/Source/CustomizableObjectEditor/Private/MuCOE/GenerateMutableSource/GenerateMutableSourceSurface.cpp b/Engine/Plugins/Mutable/Source/CustomizableObjectEditor/Private/MuCOE/GenerateMutableSource/GenerateMutableSourceSurface.cpp
--- a/Engine/Plugins/Mutable/Source/CustomizableObjectEditor/Private/MuCOE/GenerateMutableSource/GenerateMutableSourceSurface.cpp
+++ b/Engine/Plugins/Mutable/Source/CustomizableObjectEditor/Private/MuCOE/GenerateMutableSource/GenerateMutableSourceSurface.cpp
@@ -958,18 +958,6 @@ mu::Ptr<mu::NodeSurface> GenerateMutableSourceSurface(const UEdGraphPin * Pin, F
 			SurfNode->Tags.AddUnique( TypedNodeMat->GetInternalTag() );
 		}
 
-		TArray<mu::Ptr<mu::NodeSurfaceNew>>* ArraySurfaceNodePtr = GenerationContext.MapMaterialNodeToMutableSurfaceNodeArray.Find(TypedNodeMat->GetMaterialNode());
-		if (ArraySurfaceNodePtr == nullptr)
-		{
-			TArray<mu::Ptr<mu::NodeSurfaceNew>> ArraySurfaceNode;
-			ArraySurfaceNode.Add(SurfNode);
-			ArraySurfaceNodePtr = &GenerationContext.MapMaterialNodeToMutableSurfaceNodeArray.Add(TypedNodeMat->GetMaterialNode(), ArraySurfaceNode);
-		}
-		else
-		{
-			ArraySurfaceNodePtr->AddUnique(SurfNode);
-		}
-
 		// If an alternate resolution for a particular state is present, clone the surface node, add the image resizing and inject the surface variation node
 		if (TextureNameToProjectionResFactor.Num() > 0 && !AlternateResStateName.IsEmpty())
 		{
@@ -1017,8 +1005,6 @@ mu::Ptr<mu::NodeSurface> GenerateMutableSourceSurface(const UEdGraphPin * Pin, F
 				}
 			}
 
-			ArraySurfaceNodePtr->AddUnique(SurfNode2);
-
 			mu::Ptr<mu::NodeSurfaceVariation> SurfaceVariation = new mu::NodeSurfaceVariation;
 			SurfaceVariation->Type = mu::NodeSurfaceVariation::VariationType::State;
 			SurfaceVariation->Variations.SetNum(1);
diff --git a/Engine/Plugins/Mutable/Source/CustomizableObjectEditor/Private/MuCOE/SCustomizableObjectEditorViewport.cpp b/Engine/Plugins/Mutable/Source/CustomizableObjectEditor/Private/MuCOE/SCustomizableObjectEditorViewport.cpp
--- a/Engine/Plugins/Mutable/Source/CustomizableObjectEditor/Private/MuCOE/SCustomizableObjectEditorViewport.cpp
+++ b/Engine/Plugins/Mutable/Source/CustomizableObjectEditor/Private/MuCOE/SCustomizableObjectEditorViewport.cpp
@@ -308,19 +308,22 @@ void SCustomizableObjectEditorViewportTabBody::Tick(const FGeometry& AllottedGeo
 	// Needed?
 	PreviewScenePtr->GetWorld()->Tick(LEVELTICK_All, InDeltaTime);
 
+
+	const TMap<FName, TWeakObjectPtr<UDebugSkelMeshComponent>>& PreviewSkeletalMeshComponents = LevelViewportClient->GetPreviewMeshComponents();
+
 	// Update the material list. Not ideal to do it every tick, but tracking changes on materials in the current instance is not easy right now.
 	if (PreviewSkeletalMeshComponents.Num()>0)
 	{
 		MaterialNames.Empty();
 
-		for (TWeakObjectPtr<UDebugSkelMeshComponent> PreviewSkeletalMeshComponent : PreviewSkeletalMeshComponents)
+		for (TPair<FName,TWeakObjectPtr<UDebugSkelMeshComponent>> Entry : PreviewSkeletalMeshComponents)
 		{
-			if (!PreviewSkeletalMeshComponent.IsValid())
+			if (!Entry.Value.IsValid())
 			{
 				continue;
 			}
 			
-			const TArray<UMaterialInterface*> Materials = PreviewSkeletalMeshComponent->GetMaterials();
+			const TArray<UMaterialInterface*> Materials = Entry.Value->GetMaterials();
 			for (UMaterialInterface* m : Materials)
 			{
 				if (m)
@@ -442,10 +445,9 @@ void SCustomizableObjectEditorViewportTabBody::HideGizmoLight() const
 }
 
 
-void SCustomizableObjectEditorViewportTabBody::SetPreviewActor(const TWeakObjectPtr<AActor>& InActor, const TWeakObjectPtr<UCustomizableObjectInstance>& InInstance, const TArray<TWeakObjectPtr<UDebugSkelMeshComponent>>& InSkeletalMeshComponents)
+void SCustomizableObjectEditorViewportTabBody::CreatePreviewActor(const TWeakObjectPtr<UCustomizableObjectInstance>& InInstance)
 {
-	PreviewSkeletalMeshComponents = InSkeletalMeshComponents;
-	LevelViewportClient->SetPreviewActor(InActor, InInstance, InSkeletalMeshComponents);
+	LevelViewportClient->CreatePreviewActor(InInstance);
 	
 	LODSelection = 0;
 }
@@ -689,7 +691,12 @@ void SCustomizableObjectEditorViewportTabBody::BindCommands()
 
 void SCustomizableObjectEditorViewportTabBody::OnTakeHighResScreenshot()
 {
-	CustomizableObjectHighresScreenshot = SCustomizableObjectHighresScreenshot::OpenDialog(ViewportWidget->GetSceneViewport(), LevelViewportClient, PreviewSkeletalMeshComponents[0].Get(), PreviewScenePtr);
+	// TODO: Fix for multicomponents
+	TMap<FName, TWeakObjectPtr<UDebugSkelMeshComponent>>& Comps = LevelViewportClient->GetPreviewMeshComponents();
+	if (!Comps.IsEmpty())
+	{
+		CustomizableObjectHighresScreenshot = SCustomizableObjectHighresScreenshot::OpenDialog(ViewportWidget->GetSceneViewport(), LevelViewportClient, Comps.begin().Value().Get(), PreviewScenePtr);
+	}
 }
 
 
@@ -714,7 +721,7 @@ void SCustomizableObjectEditorViewportTabBody::SetDrawDefaultUVMaterial()
 	if (!SelectedUVSection ||
 		!SelectedUVChannel)
 	{
-		LevelViewportClient->SetDrawUV(-1, -1, -1, -1);
+		LevelViewportClient->SetDrawUV(NAME_None, -1, -1, -1);
 	}
 	else
 	{
@@ -725,7 +732,7 @@ void SCustomizableObjectEditorViewportTabBody::SetDrawDefaultUVMaterial()
 		const int32 UVIndex = UVChannelOptionString.IndexOfByKey(SelectedUVChannel);
 		check(UVIndex != INDEX_NONE);
 	
-		LevelViewportClient->SetDrawUV(Section.ComponentIndex, Section.LODIndex, Section.SectionIndex, UVIndex);
+		LevelViewportClient->SetDrawUV(Section.ComponentName, Section.LODIndex, Section.SectionIndex, UVIndex);
 	}
 }
 
@@ -846,14 +853,15 @@ void SCustomizableObjectEditorViewportTabBody::GenerateUVSectionOptions()
 	UVSectionOptionString.Empty();
 	UVSectionOption.Empty();
 	
-	for (int32 ComponentIndex = 0; ComponentIndex < PreviewSkeletalMeshComponents.Num(); ++ComponentIndex)
+	TMap<FName, TWeakObjectPtr<UDebugSkelMeshComponent>>& Comps = LevelViewportClient->GetPreviewMeshComponents();
+	for (TPair<FName, TWeakObjectPtr<UDebugSkelMeshComponent>>& Entry : Comps)
 	{
-		const TWeakObjectPtr<UDebugSkelMeshComponent>& PreviewSkeletalMeshComponent = PreviewSkeletalMeshComponents[ComponentIndex];
+		const TWeakObjectPtr<UDebugSkelMeshComponent>& PreviewSkeletalMeshComponent = Entry.Value;
 		
-		if (PreviewSkeletalMeshComponent.IsValid() && UE_MUTABLE_GETSKINNEDASSET(PreviewSkeletalMeshComponent) != nullptr && UE_MUTABLE_GETSKINNEDASSET(PreviewSkeletalMeshComponent)->GetResourceForRendering() != nullptr)
+		if (PreviewSkeletalMeshComponent.IsValid() && PreviewSkeletalMeshComponent->GetSkinnedAsset() != nullptr && PreviewSkeletalMeshComponent->GetSkinnedAsset()->GetResourceForRendering() != nullptr)
 		{
 			const TArray<UMaterialInterface*> Materials = PreviewSkeletalMeshComponent->GetMaterials();
-			const FSkeletalMeshRenderData* MeshRes = UE_MUTABLE_GETSKINNEDASSET(PreviewSkeletalMeshComponent)->GetResourceForRendering();
+			const FSkeletalMeshRenderData* MeshRes = PreviewSkeletalMeshComponent->GetSkinnedAsset()->GetResourceForRendering();
 			for (int32 LODIndex = 0; LODIndex < MeshRes->LODRenderData.Num(); ++LODIndex)
 			{
 				for (int32 SectionIndex = 0; SectionIndex < MeshRes->LODRenderData[LODIndex].RenderSections.Num(); ++SectionIndex)
@@ -876,7 +884,7 @@ void SCustomizableObjectEditorViewportTabBody::GenerateUVSectionOptions()
 					UVSectionOptionString.Add(MakeShared<FString>(BaseMaterialName));
 
 					FSection SectionOption;
-					SectionOption.ComponentIndex = ComponentIndex;
+					SectionOption.ComponentName = Entry.Key;
 					SectionOption.SectionIndex = SectionIndex;
 					SectionOption.LODIndex = LODIndex;
 
@@ -924,11 +932,11 @@ void SCustomizableObjectEditorViewportTabBody::OnSectionChanged(TSharedPtr<FStri
 		const int32 UVIndex = UVChannelOptionString.IndexOfByKey(SelectedUVChannel);
 		check(UVIndex != INDEX_NONE);
 	
-		LevelViewportClient->SetDrawUV(Section.ComponentIndex, Section.LODIndex, Section.SectionIndex, UVIndex);
+		LevelViewportClient->SetDrawUV(Section.ComponentName, Section.LODIndex, Section.SectionIndex, UVIndex);
 	}
 	else
 	{
-		LevelViewportClient->SetDrawUV(-1, -1, -1, -1);
+		LevelViewportClient->SetDrawUV(NAME_None, -1, -1, -1);
 	}
 }
 
@@ -956,12 +964,13 @@ void SCustomizableObjectEditorViewportTabBody::GenerateUVChannelOptions()
 	check(Index != INDEX_NONE);
 	const FSection& Section = UVSectionOption[Index];
 
-	const TWeakObjectPtr<UDebugSkelMeshComponent>& PreviewSkeletalMeshComponent = PreviewSkeletalMeshComponents[Section.ComponentIndex];
-
-	if (PreviewSkeletalMeshComponent.IsValid() && UE_MUTABLE_GETSKINNEDASSET(PreviewSkeletalMeshComponent) != nullptr
-		&& UE_MUTABLE_GETSKINNEDASSET(PreviewSkeletalMeshComponent)->GetResourceForRendering() != nullptr)
+	TMap<FName, TWeakObjectPtr<UDebugSkelMeshComponent>>& Comps = LevelViewportClient->GetPreviewMeshComponents();
+	const TWeakObjectPtr<UDebugSkelMeshComponent>* PreviewSkeletalMeshComponent = Comps.Find(Section.ComponentName);
+	
+	if (PreviewSkeletalMeshComponent && PreviewSkeletalMeshComponent->IsValid() && (*PreviewSkeletalMeshComponent)->GetSkinnedAsset() != nullptr
+		&& (*PreviewSkeletalMeshComponent)->GetSkinnedAsset()->GetResourceForRendering() != nullptr)
 	{
-		const FSkeletalMeshRenderData* MeshRes = UE_MUTABLE_GETSKINNEDASSET(PreviewSkeletalMeshComponent)->GetResourceForRendering();
+		const FSkeletalMeshRenderData* MeshRes = (*PreviewSkeletalMeshComponent)->GetSkinnedAsset()->GetResourceForRendering();
 		
 		const int32 UVChannels = MeshRes->LODRenderData[Section.LODIndex].GetNumTexCoords();
 		for (int32 UVChan = 0; UVChan < UVChannels; ++UVChan)
@@ -1001,11 +1010,11 @@ void SCustomizableObjectEditorViewportTabBody::OnUVChannelChanged(TSharedPtr<FSt
 		const int32 UVIndex = UVChannelOptionString.IndexOfByKey(SelectedUVChannel);
 		check(UVIndex != INDEX_NONE);
 	
-		LevelViewportClient->SetDrawUV(Section.ComponentIndex, Section.LODIndex, Section.SectionIndex, UVIndex);
+		LevelViewportClient->SetDrawUV(Section.ComponentName, Section.LODIndex, Section.SectionIndex, UVIndex);
 	}
 	else
 	{
-		LevelViewportClient->SetDrawUV(-1, -1, -1, -1);
+		LevelViewportClient->SetDrawUV(NAME_None, -1, -1, -1);
 	}
 }
 
@@ -1035,11 +1044,12 @@ int32 SCustomizableObjectEditorViewportTabBody::GetLODModelCount() const
 {
 	int32 LODModelCount = 0;
 
-	for (const TWeakObjectPtr<UDebugSkelMeshComponent>& PreviewComponent : PreviewSkeletalMeshComponents)
+	TMap<FName, TWeakObjectPtr<UDebugSkelMeshComponent>>& Comps = LevelViewportClient->GetPreviewMeshComponents();
+	for (const TPair<FName,TWeakObjectPtr<UDebugSkelMeshComponent>>& Entry : Comps)
 	{
-		if (PreviewComponent.IsValid() && UE_MUTABLE_GETSKINNEDASSET(PreviewComponent))
+		if (Entry.Value.IsValid() && Entry.Value->GetSkinnedAsset())
 		{
-			const TIndirectArray<FSkeletalMeshLODRenderData>& LODModels = UE_MUTABLE_GETSKINNEDASSET(PreviewComponent)->GetResourceForRendering()->LODRenderData;
+			const TIndirectArray<FSkeletalMeshLODRenderData>& LODModels = Entry.Value->GetSkinnedAsset()->GetResourceForRendering()->LODRenderData;
 			LODModelCount = FMath::Max(LODModelCount, LODModels.Num());
 		}
 	}
@@ -1094,11 +1104,12 @@ void SCustomizableObjectEditorViewportTabBody::OnSetLODModel(int32 LODSelectionT
 {
 	LODSelection = LODSelectionType;
 
-	for (const TWeakObjectPtr<UDebugSkelMeshComponent>& PreviewComponent : PreviewSkeletalMeshComponents)
+	TMap<FName, TWeakObjectPtr<UDebugSkelMeshComponent>>& Comps = LevelViewportClient->GetPreviewMeshComponents();
+	for (const TPair<FName, TWeakObjectPtr<UDebugSkelMeshComponent>>& Entry : Comps)
 	{
-		if (PreviewComponent.IsValid())
+		if (Entry.Value.IsValid())
 		{
-			PreviewComponent->SetForcedLOD(LODSelection);
+			Entry.Value.Get()->SetForcedLOD(LODSelection);
 		}
 	}
 }
diff --git a/Engine/Plugins/Mutable/Source/CustomizableObjectEditor/Private/MuCOE/SCustomizableObjectEditorViewport.h b/Engine/Plugins/Mutable/Source/CustomizableObjectEditor/Private/MuCOE/SCustomizableObjectEditorViewport.h
--- a/Engine/Plugins/Mutable/Source/CustomizableObjectEditor/Private/MuCOE/SCustomizableObjectEditorViewport.h
+++ b/Engine/Plugins/Mutable/Source/CustomizableObjectEditor/Private/MuCOE/SCustomizableObjectEditorViewport.h
@@ -152,7 +152,7 @@ public:
 	/** Callback to show / hide the state and runtime parameter test in the viewport */
 	TSharedRef<SWidget> ShowStateTestData();
 	
-	void SetPreviewActor(const TWeakObjectPtr<AActor>& InActor, const TWeakObjectPtr<UCustomizableObjectInstance>& InInstance, const TArray<TWeakObjectPtr<UDebugSkelMeshComponent>>& InSkeletalMeshComponents);
+	void CreatePreviewActor(const TWeakObjectPtr<UCustomizableObjectInstance>& InInstance);
 	
 	/** Function to get the number of LOD models associated with the preview skeletal mesh*/
 	int32 GetLODModelCount() const;
@@ -228,9 +228,6 @@ private:
 	
 	bool IsPlaybackSpeedSelected(int32 PlaybackSpeedMode);
 
-	// Components for the preview mesh.
-	TArray<TWeakObjectPtr<UDebugSkelMeshComponent>> PreviewSkeletalMeshComponents;
-
 	// The scene for this viewport.
 	TSharedPtr<FCustomizableObjectPreviewScene> PreviewScenePtr;
 
@@ -250,7 +247,7 @@ private:
 	/** Section identifier. */
 	struct FSection
 	{
-		int32 ComponentIndex = -1;
+		FName ComponentName = NAME_None;
 		int32 LODIndex = -1;
 		int32 SectionIndex = -1;
 	};
diff --git a/Engine/Plugins/Mutable/Source/CustomizableObjectEditor/Private/MuCOE/SMutableScrubPanel.cpp b/Engine/Plugins/Mutable/Source/CustomizableObjectEditor/Private/MuCOE/SMutableScrubPanel.cpp
--- a/Engine/Plugins/Mutable/Source/CustomizableObjectEditor/Private/MuCOE/SMutableScrubPanel.cpp
+++ b/Engine/Plugins/Mutable/Source/CustomizableObjectEditor/Private/MuCOE/SMutableScrubPanel.cpp
@@ -1,4 +1,4 @@
-﻿// Copyright Epic Games, Inc. All Rights Reserved.
+// Copyright Epic Games, Inc. All Rights Reserved.
 
 #include "MuCOE/SMutableScrubPanel.h"
 
@@ -67,9 +67,9 @@ void SMutableScrubPanel::Construct(const FArguments& InArgs, const TSharedRef<FC
 
 float SMutableScrubPanel::GetViewMinInput() const
 {
-	for (TWeakObjectPtr<UDebugSkelMeshComponent>& WeakPreviewMeshComponent : GetPreviewScene()->GetPreviewMeshComponents())
+	for (TPair<FName,TWeakObjectPtr<UDebugSkelMeshComponent>>& Entry : GetPreviewScene()->GetPreviewMeshComponents())
 	{
-		UDebugSkelMeshComponent* PreviewMeshComponent = WeakPreviewMeshComponent.Get();
+		UDebugSkelMeshComponent* PreviewMeshComponent = Entry.Value.Get();
 		if (!PreviewMeshComponent)
 		{
 			continue;
@@ -91,9 +91,9 @@ float SMutableScrubPanel::GetViewMinInput() const
 
 float SMutableScrubPanel::GetViewMaxInput() const
 { 
-	for (TWeakObjectPtr<UDebugSkelMeshComponent>& WeakPreviewMeshComponent : GetPreviewScene()->GetPreviewMeshComponents())
+	for (TPair<FName, TWeakObjectPtr<UDebugSkelMeshComponent>>& Entry : GetPreviewScene()->GetPreviewMeshComponents())
 	{
-		UDebugSkelMeshComponent* PreviewMeshComponent = WeakPreviewMeshComponent.Get();
+		UDebugSkelMeshComponent* PreviewMeshComponent = Entry.Value.Get();
 		if (!PreviewMeshComponent)
 		{
 			continue;
@@ -115,9 +115,9 @@ float SMutableScrubPanel::GetViewMaxInput() const
 
 FReply SMutableScrubPanel::OnClick_Forward_Step()
 {
-	for (TWeakObjectPtr<UDebugSkelMeshComponent>& WeakPreviewMeshComponent : GetPreviewScene()->GetPreviewMeshComponents())
+	for (TPair<FName, TWeakObjectPtr<UDebugSkelMeshComponent>>& Entry : GetPreviewScene()->GetPreviewMeshComponents())
 	{
-		UDebugSkelMeshComponent* PreviewMeshComponent = WeakPreviewMeshComponent.Get();
+		UDebugSkelMeshComponent* PreviewMeshComponent = Entry.Value.Get();
 		if (!PreviewMeshComponent)
 		{
 			continue;
@@ -153,9 +153,9 @@ FReply SMutableScrubPanel::OnClick_Forward_Step()
 
 FReply SMutableScrubPanel::OnClick_Forward_End()
 {
-	for (TWeakObjectPtr<UDebugSkelMeshComponent>& WeakPreviewMeshComponent : GetPreviewScene()->GetPreviewMeshComponents())
+	for (TPair<FName, TWeakObjectPtr<UDebugSkelMeshComponent>>& Entry : GetPreviewScene()->GetPreviewMeshComponents())
 	{
-		UDebugSkelMeshComponent* PreviewMeshComponent = WeakPreviewMeshComponent.Get();
+		UDebugSkelMeshComponent* PreviewMeshComponent = Entry.Value.Get();
 		if (!PreviewMeshComponent)
         {
         	continue;
@@ -175,9 +175,9 @@ FReply SMutableScrubPanel::OnClick_Forward_End()
 
 FReply SMutableScrubPanel::OnClick_Backward_Step()
 {
-	for (TWeakObjectPtr<UDebugSkelMeshComponent>& WeakPreviewMeshComponent : GetPreviewScene()->GetPreviewMeshComponents())
+	for (TPair<FName, TWeakObjectPtr<UDebugSkelMeshComponent>>& Entry : GetPreviewScene()->GetPreviewMeshComponents())
 	{
-		UDebugSkelMeshComponent* PreviewMeshComponent = WeakPreviewMeshComponent.Get();
+		UDebugSkelMeshComponent* PreviewMeshComponent = Entry.Value.Get();
 		if (!PreviewMeshComponent)
 		{
 			continue;
@@ -204,9 +204,9 @@ FReply SMutableScrubPanel::OnClick_Backward_Step()
 
 FReply SMutableScrubPanel::OnClick_Backward_End()
 {
-	for (TWeakObjectPtr<UDebugSkelMeshComponent>& WeakPreviewMeshComponent : GetPreviewScene()->GetPreviewMeshComponents())
+	for (TPair<FName, TWeakObjectPtr<UDebugSkelMeshComponent>>& Entry : GetPreviewScene()->GetPreviewMeshComponents())
 	{
-		UDebugSkelMeshComponent* PreviewMeshComponent = WeakPreviewMeshComponent.Get();
+		UDebugSkelMeshComponent* PreviewMeshComponent = Entry.Value.Get();
 		if (!PreviewMeshComponent)
 		{
 			continue;
@@ -226,9 +226,9 @@ FReply SMutableScrubPanel::OnClick_Backward_End()
 
 FReply SMutableScrubPanel::OnClick_Forward()
 {
-	for (TWeakObjectPtr<UDebugSkelMeshComponent>& WeakPreviewMeshComponent : GetPreviewScene()->GetPreviewMeshComponents())
+	for (TPair<FName, TWeakObjectPtr<UDebugSkelMeshComponent>>& Entry : GetPreviewScene()->GetPreviewMeshComponents())
 	{
-		UDebugSkelMeshComponent* PreviewMeshComponent = WeakPreviewMeshComponent.Get();
+		UDebugSkelMeshComponent* PreviewMeshComponent = Entry.Value.Get();
 		if (!PreviewMeshComponent)
 		{
 			continue;
@@ -285,9 +285,9 @@ FReply SMutableScrubPanel::OnClick_Forward()
 
 FReply SMutableScrubPanel::OnClick_Backward()
 {
-	for (TWeakObjectPtr<UDebugSkelMeshComponent>& WeakPreviewMeshComponent : GetPreviewScene()->GetPreviewMeshComponents())
+	for (TPair<FName, TWeakObjectPtr<UDebugSkelMeshComponent>>& Entry : GetPreviewScene()->GetPreviewMeshComponents())
 	{
-		UDebugSkelMeshComponent* PreviewMeshComponent = WeakPreviewMeshComponent.Get();
+		UDebugSkelMeshComponent* PreviewMeshComponent = Entry.Value.Get();
 		if (!PreviewMeshComponent)
 		{
 			continue;
@@ -327,9 +327,9 @@ FReply SMutableScrubPanel::OnClick_Backward()
 
 FReply SMutableScrubPanel::OnClick_ToggleLoop()
 {
-	for (TWeakObjectPtr<UDebugSkelMeshComponent>& WeakPreviewMeshComponent : GetPreviewScene()->GetPreviewMeshComponents())
+	for (TPair<FName, TWeakObjectPtr<UDebugSkelMeshComponent>>& Entry : GetPreviewScene()->GetPreviewMeshComponents())
 	{
-		UDebugSkelMeshComponent* PreviewMeshComponent = WeakPreviewMeshComponent.Get();
+		UDebugSkelMeshComponent* PreviewMeshComponent = Entry.Value.Get();
 		if (!PreviewMeshComponent)
 		{
 			continue;
@@ -349,9 +349,9 @@ FReply SMutableScrubPanel::OnClick_ToggleLoop()
 
 bool SMutableScrubPanel::IsLoopStatusOn() const
 {
-	for (TWeakObjectPtr<UDebugSkelMeshComponent>& WeakPreviewMeshComponent : GetPreviewScene()->GetPreviewMeshComponents())
+	for (TPair<FName, TWeakObjectPtr<UDebugSkelMeshComponent>>& Entry : GetPreviewScene()->GetPreviewMeshComponents())
 	{
-		UDebugSkelMeshComponent* PreviewMeshComponent = WeakPreviewMeshComponent.Get();
+		UDebugSkelMeshComponent* PreviewMeshComponent = Entry.Value.Get();
 		if (!PreviewMeshComponent)
 		{
 			continue;
@@ -369,9 +369,9 @@ bool SMutableScrubPanel::IsLoopStatusOn() const
 
 EPlaybackMode::Type SMutableScrubPanel::GetPlaybackMode() const
 {
-	for (TWeakObjectPtr<UDebugSkelMeshComponent>& WeakPreviewMeshComponent : GetPreviewScene()->GetPreviewMeshComponents())
+	for (TPair<FName, TWeakObjectPtr<UDebugSkelMeshComponent>>& Entry : GetPreviewScene()->GetPreviewMeshComponents())
 	{
-		UDebugSkelMeshComponent* PreviewMeshComponent = WeakPreviewMeshComponent.Get();
+		UDebugSkelMeshComponent* PreviewMeshComponent = Entry.Value.Get();
 		if (!PreviewMeshComponent)
 		{
 			continue;
@@ -397,9 +397,9 @@ EPlaybackMode::Type SMutableScrubPanel::GetPlaybackMode() const
 
 bool SMutableScrubPanel::IsRealtimeStreamingMode() const
 {
-	for (TWeakObjectPtr<UDebugSkelMeshComponent>& WeakPreviewMeshComponent : GetPreviewScene()->GetPreviewMeshComponents())
+	for (TPair<FName, TWeakObjectPtr<UDebugSkelMeshComponent>>& Entry : GetPreviewScene()->GetPreviewMeshComponents())
 	{
-		UDebugSkelMeshComponent* PreviewMeshComponent = WeakPreviewMeshComponent.Get();
+		UDebugSkelMeshComponent* PreviewMeshComponent = Entry.Value.Get();
 		if (!PreviewMeshComponent)
 		{
 			continue;
@@ -417,9 +417,9 @@ bool SMutableScrubPanel::IsRealtimeStreamingMode() const
 
 void SMutableScrubPanel::OnValueChanged(float NewValue)
 {
-	for (TWeakObjectPtr<UDebugSkelMeshComponent>& WeakPreviewMeshComponent : GetPreviewScene()->GetPreviewMeshComponents())
+	for (TPair<FName, TWeakObjectPtr<UDebugSkelMeshComponent>>& Entry : GetPreviewScene()->GetPreviewMeshComponents())
 	{
-		UDebugSkelMeshComponent* PreviewMeshComponent = WeakPreviewMeshComponent.Get();
+		UDebugSkelMeshComponent* PreviewMeshComponent = Entry.Value.Get();
 		if (!PreviewMeshComponent)
 		{
 			continue;
@@ -446,9 +446,9 @@ void SMutableScrubPanel::OnBeginSliderMovement()
 {
 	bSliderBeingDragged = true;
 
-	for (TWeakObjectPtr<UDebugSkelMeshComponent>& WeakPreviewMeshComponent : GetPreviewScene()->GetPreviewMeshComponents())
+	for (TPair<FName, TWeakObjectPtr<UDebugSkelMeshComponent>>& Entry : GetPreviewScene()->GetPreviewMeshComponents())
 	{
-		UDebugSkelMeshComponent* PreviewMeshComponent = WeakPreviewMeshComponent.Get();
+		UDebugSkelMeshComponent* PreviewMeshComponent = Entry.Value.Get();
 		if (!PreviewMeshComponent)
 		{
 			continue;
@@ -470,9 +470,9 @@ void SMutableScrubPanel::OnEndSliderMovement(float NewValue)
 
 uint32 SMutableScrubPanel::GetNumberOfKeys() const
 {
-	for (TWeakObjectPtr<UDebugSkelMeshComponent>& WeakPreviewMeshComponent : GetPreviewScene()->GetPreviewMeshComponents())
+	for (TPair<FName, TWeakObjectPtr<UDebugSkelMeshComponent>>& Entry : GetPreviewScene()->GetPreviewMeshComponents())
 	{
-		UDebugSkelMeshComponent* PreviewMeshComponent = WeakPreviewMeshComponent.Get();
+		UDebugSkelMeshComponent* PreviewMeshComponent = Entry.Value.Get();
 		if (!PreviewMeshComponent)
 		{
 			continue;
@@ -508,9 +508,9 @@ uint32 SMutableScrubPanel::GetNumberOfKeys() const
 
 float SMutableScrubPanel::GetSequenceLength() const
 {
-	for (TWeakObjectPtr<UDebugSkelMeshComponent>& WeakPreviewMeshComponent : GetPreviewScene()->GetPreviewMeshComponents())
+	for (TPair<FName, TWeakObjectPtr<UDebugSkelMeshComponent>>& Entry : GetPreviewScene()->GetPreviewMeshComponents())
 	{
-		UDebugSkelMeshComponent* PreviewMeshComponent = WeakPreviewMeshComponent.Get();
+		UDebugSkelMeshComponent* PreviewMeshComponent = Entry.Value.Get();
 		if (!PreviewMeshComponent)
 		{
 			continue;
@@ -546,9 +546,9 @@ UAnimSingleNodeInstance* SMutableScrubPanel::GetPreviewInstance(UDebugSkelMeshCo
 
 float SMutableScrubPanel::GetScrubValue() const
 {
-	for (TWeakObjectPtr<UDebugSkelMeshComponent>& WeakPreviewMeshComponent : GetPreviewScene()->GetPreviewMeshComponents())
+	for (TPair<FName, TWeakObjectPtr<UDebugSkelMeshComponent>>& Entry : GetPreviewScene()->GetPreviewMeshComponents())
 	{
-		UDebugSkelMeshComponent* PreviewMeshComponent = WeakPreviewMeshComponent.Get();
+		UDebugSkelMeshComponent* PreviewMeshComponent = Entry.Value.Get();
 		if (!PreviewMeshComponent)
 		{
 			continue;
@@ -571,9 +571,9 @@ bool SMutableScrubPanel::GetAnimBlueprintDebugData(UAnimInstance*& Instance, FAn
 {
 	Instance = nullptr;
 
-	for (TWeakObjectPtr<UDebugSkelMeshComponent>& WeakPreviewMeshComponent : GetPreviewScene()->GetPreviewMeshComponents())
+	for (TPair<FName, TWeakObjectPtr<UDebugSkelMeshComponent>>& Entry : GetPreviewScene()->GetPreviewMeshComponents())
 	{
-		UDebugSkelMeshComponent* PreviewMeshComponent = WeakPreviewMeshComponent.Get();
+		UDebugSkelMeshComponent* PreviewMeshComponent = Entry.Value.Get();
 		if (!PreviewMeshComponent)
 		{
 			continue;
@@ -615,9 +615,9 @@ TSharedRef<FCustomizableObjectEditorViewportClient> SMutableScrubPanel::GetPrevi
 
 void SMutableScrubPanel::OnReZeroAnimSequence(int32 FrameIndex)
 {
-	for (TWeakObjectPtr<UDebugSkelMeshComponent>& WeakPreviewMeshComponent : GetPreviewScene()->GetPreviewMeshComponents())
+	for (TPair<FName, TWeakObjectPtr<UDebugSkelMeshComponent>>& Entry : GetPreviewScene()->GetPreviewMeshComponents())
 	{
-		UDebugSkelMeshComponent* PreviewSkelComp = WeakPreviewMeshComponent.Get();
+		UDebugSkelMeshComponent* PreviewSkelComp = Entry.Value.Get();
 		if (!PreviewSkelComp)
 		{
 			continue;
@@ -688,9 +688,9 @@ void SMutableScrubPanel::OnReZeroAnimSequence(int32 FrameIndex)
 
 bool SMutableScrubPanel::GetDisplayDrag() const
 {
-	for (TWeakObjectPtr<UDebugSkelMeshComponent>& WeakPreviewMeshComponent : GetPreviewScene()->GetPreviewMeshComponents())
+	for (TPair<FName, TWeakObjectPtr<UDebugSkelMeshComponent>>& Entry : GetPreviewScene()->GetPreviewMeshComponents())
 	{
-		UDebugSkelMeshComponent* PreviewMeshComponent = WeakPreviewMeshComponent.Get();
+		UDebugSkelMeshComponent* PreviewMeshComponent = Entry.Value.Get();
 		if (!PreviewMeshComponent)
 		{
 			continue;
diff --git a/Engine/Plugins/Mutable/Source/CustomizableObjectEditor/Public/MuCOE/CustomizableObjectEditorViewportClient.h b/Engine/Plugins/Mutable/Source/CustomizableObjectEditor/Public/MuCOE/CustomizableObjectEditorViewportClient.h
--- a/Engine/Plugins/Mutable/Source/CustomizableObjectEditor/Public/MuCOE/CustomizableObjectEditorViewportClient.h
+++ b/Engine/Plugins/Mutable/Source/CustomizableObjectEditor/Public/MuCOE/CustomizableObjectEditorViewportClient.h
@@ -112,9 +112,9 @@ public:
 	virtual void SetViewportType(ELevelViewportType InViewportType) override;
 	// End of FEditorViewportClient
 	
-	void SetPreviewActor(const TWeakObjectPtr<AActor>& InActor, const TWeakObjectPtr<UCustomizableObjectInstance>& InInstance, const TArray<TWeakObjectPtr<UDebugSkelMeshComponent>>& InSkeletalMeshComponents);
+	void CreatePreviewActor(const TWeakObjectPtr<UCustomizableObjectInstance>& InInstance);
 
-	TArray<TWeakObjectPtr<UDebugSkelMeshComponent>>& GetPreviewMeshComponents();
+	TMap<FName, TWeakObjectPtr<UDebugSkelMeshComponent>>& GetPreviewMeshComponents();
 
 	void SetPreviewAnimationAsset(UAnimationAsset* AnimAsset);
 
@@ -152,7 +152,7 @@ public:
 	bool IsSetDrawUVOverlayChecked() const;
 
 	/** Specify which UV to draw. -1 values will not draw anything. */
-	void SetDrawUV(const int32 ComponentIndex, const int32 LODIndex, const int32 SectionIndex, const int32 UVIndex);
+	void SetDrawUV(const FName ComponentName, const int32 LODIndex, const int32 SectionIndex, const int32 UVIndex);
 
 	/** Callback for toggling the grid show flag. */
 	void UpdateShowGridFromButton();
@@ -316,7 +316,7 @@ public:
 
 private:
 	/** Draws Mesh Bones in foreground (From: FAnimationViewportClient) */
-	void DrawMeshBones(UDebugSkelMeshComponent* MeshComponent, FPrimitiveDrawInterface* PDI);
+	void DrawMeshBones(const UDebugSkelMeshComponent* MeshComponent, FPrimitiveDrawInterface* PDI);
 
 	void SetWidgetType(EWidgetType Type);
 
@@ -325,8 +325,10 @@ private:
 
 	void OnInstanceUpdate(UCustomizableObjectInstance* Instance);
 
-	TWeakObjectPtr<AActor> Actor;
-	TArray<TWeakObjectPtr<UDebugSkelMeshComponent>> SkeletalMeshComponents;
+
+	/** Preview Actor. All preview components are attached to this actor. */
+	TStrongObjectPtr<AActor> Actor;
+	TMap<FName, TWeakObjectPtr<UDebugSkelMeshComponent>> SkeletalMeshComponents;
 
 	/** True if the widget is being dragged. */
 	bool bManipulating = false;
@@ -339,7 +341,7 @@ private:
 	bool bDrawSky;
 
 	// UV to draw
-	int32 UVDrawComponentIndex = 0;
+	FName UVDrawComponentName = NAME_None;
 	int32 UVDrawSectionIndex = 0;
 	int32 UVDrawLODIndex = 0;
 	int32 UVDrawUVIndex = 0;
