[mutable] Fix wrong image size when mixing very big textures and small textures in a table column.
#rnx

diff --git a/Engine/Plugins/Mutable/Source/CustomizableObjectEditor/Private/MuCOE/GenerateMutableSource/GenerateMutableSourceImage.cpp b/Engine/Plugins/Mutable/Source/CustomizableObjectEditor/Private/MuCOE/GenerateMutableSource/GenerateMutableSourceImage.cpp
--- a/Engine/Plugins/Mutable/Source/CustomizableObjectEditor/Private/MuCOE/GenerateMutableSource/GenerateMutableSourceImage.cpp
+++ b/Engine/Plugins/Mutable/Source/CustomizableObjectEditor/Private/MuCOE/GenerateMutableSource/GenerateMutableSourceImage.cpp
@@ -1109,7 +1109,10 @@ mu::NodeImagePtr GenerateMutableSourceImage(const UEdGraphPin* Pin, FMutableGrap
 							ImageTableNode->bNoneOption = TypedNodeTable->bAddNoneOption;
 							ImageTableNode->DefaultRowName = TypedNodeTable->DefaultRowName.ToString();
 
-							if (UTexture2D* DefaultTexture2D = Cast<UTexture2D>(DefaultTexture))
+							bool bIsPassthrough = (Pin->PinType.PinCategory == Schema->PC_PassThroughImage);
+
+							UTexture2D* DefaultTexture2D = Cast<UTexture2D>(DefaultTexture);
+							if (!bIsPassthrough && DefaultTexture2D)
 							{
 								mu::FImageDesc ImageDesc = GenerateImageDescriptor(DefaultTexture2D);
 
diff --git a/Engine/Plugins/Mutable/Source/MutableRuntime/Private/MuR/CodeRunnerTasked.cpp b/Engine/Plugins/Mutable/Source/MutableRuntime/Private/MuR/CodeRunnerTasked.cpp
--- a/Engine/Plugins/Mutable/Source/MutableRuntime/Private/MuR/CodeRunnerTasked.cpp
+++ b/Engine/Plugins/Mutable/Source/MutableRuntime/Private/MuR/CodeRunnerTasked.cpp
@@ -1810,6 +1810,13 @@ namespace mu
 			Lods = Image::GetMipmapCount(DestSize[0], DestSize[1]);
 		}
 
+		if (Base->IsReference())
+		{
+			// We are trying to resize an external reference. This shouldn't happen, but be deffensive.
+			Runner->StoreImage(Op, Base);
+			return false;
+		}
+
 		Result = Runner->CreateImage(DestSize[0], DestSize[1], Lods, Base->GetFormat(), EInitializationType::NotInitialized);
 
 		ImagePixelFormatFunc = Runner->m_pSystem->ImagePixelFormatOverride;
diff --git a/Engine/Plugins/Mutable/Source/MutableTools/Private/MuT/CodeGenerator_Image.cpp b/Engine/Plugins/Mutable/Source/MutableTools/Private/MuT/CodeGenerator_Image.cpp
--- a/Engine/Plugins/Mutable/Source/MutableTools/Private/MuT/CodeGenerator_Image.cpp
+++ b/Engine/Plugins/Mutable/Source/MutableTools/Private/MuT/CodeGenerator_Image.cpp
@@ -1927,12 +1927,13 @@ namespace mu
 
 					if (MaxTextureSize > 0 && (MaxTextureSize < pImage->GetSizeX() || MaxTextureSize < pImage->GetSizeY()))
 					{
+						// Use a relative resize, because at this point we may be generating a layout block and not the full image
 						float Factor = FMath::Min(MaxTextureSize / (float)(pImage->GetSizeX()), MaxTextureSize / (float)(pImage->GetSizeY()));
 						Ptr<ASTOpFixed> op = new ASTOpFixed();
-						op->op.type = OP_TYPE::IM_RESIZE;
-						op->op.args.ImageResize.size[0] = (uint16)pImage->GetSizeX() * Factor;
-						op->op.args.ImageResize.size[1] = (uint16)pImage->GetSizeY() * Factor;
-						op->SetChild(op->op.args.ImageResize.source, ImageOp);
+						op->op.type = OP_TYPE::IM_RESIZEREL;
+						op->op.args.ImageResizeRel.factor[0] = Factor;
+						op->op.args.ImageResizeRel.factor[1] = Factor;
+						op->SetChild(op->op.args.ImageResizeRel.source, ImageOp);
 						ImageOp = op;
 					}
 				}
