[mutable] Fix crash opening instance editor.
#rnx

diff --git a/Engine/Plugins/Mutable/Source/CustomizableObjectEditor/Private/MuCOE/CustomizableObjectInstanceEditor.cpp b/Engine/Plugins/Mutable/Source/CustomizableObjectEditor/Private/MuCOE/CustomizableObjectInstanceEditor.cpp
--- a/Engine/Plugins/Mutable/Source/CustomizableObjectEditor/Private/MuCOE/CustomizableObjectInstanceEditor.cpp
+++ b/Engine/Plugins/Mutable/Source/CustomizableObjectEditor/Private/MuCOE/CustomizableObjectInstanceEditor.cpp
@@ -615,6 +615,11 @@ void FCustomizableObjectInstanceEditor::OnInstancePropertySelectionChanged(FProp
 
 void FCustomizableObjectInstanceEditor::OnUpdatePreviewInstance(UCustomizableObjectInstance* Instance)
 {
+	if (!Actor)
+	{
+		return;
+	}
+
 	check(CustomizableObjectInstance)
 	
 	if (Instance->GetPrivate()->SkeletalMeshStatus != ESkeletalMeshStatus::Success)
