[mutable] Allow Mesh Section nodes to connect their output to multiple other nodes.
- Make sure a single USkeletalMesh is created when the mesh cache is enabled.
- Optimize code to pull conditionals out of of mesh operations if possible, to ensure more meshes have the same root operation.
#rnx

diff --git a/Engine/Plugins/Mutable/Source/CustomizableObject/Internal/MuCO/CustomizableObjectPrivate.h b/Engine/Plugins/Mutable/Source/CustomizableObject/Internal/MuCO/CustomizableObjectPrivate.h
--- a/Engine/Plugins/Mutable/Source/CustomizableObject/Internal/MuCO/CustomizableObjectPrivate.h
+++ b/Engine/Plugins/Mutable/Source/CustomizableObject/Internal/MuCO/CustomizableObjectPrivate.h
@@ -1569,6 +1569,8 @@ public:
 
 		SurfaceMetadataSerialized,
 
+		FixesForMeshSectionMultipleOutputs,
+
 		// -----<new versions can be added above this line>--------
 		LastCustomizableObjectVersion
 	};
diff --git a/Engine/Plugins/Mutable/Source/CustomizableObjectEditor/Private/MuCOE/EdGraphSchema_CustomizableObject.cpp b/Engine/Plugins/Mutable/Source/CustomizableObjectEditor/Private/MuCOE/EdGraphSchema_CustomizableObject.cpp
--- a/Engine/Plugins/Mutable/Source/CustomizableObjectEditor/Private/MuCOE/EdGraphSchema_CustomizableObject.cpp
+++ b/Engine/Plugins/Mutable/Source/CustomizableObjectEditor/Private/MuCOE/EdGraphSchema_CustomizableObject.cpp
@@ -1329,16 +1329,6 @@ bool UEdGraphSchema_CustomizableObject::TryCreateConnection(UEdGraphPin* PinA, U
 	{
 		return bResult;
 	}
-
-	if (UCustomizableObjectNode* Node = Cast<UCustomizableObjectNode>(InputPin->GetOwningNode()))
-	{
-		Node->BreakExistingConnectionsPostConnection(InputPin, OutputPin);
-	}
-
-	if (UCustomizableObjectNode* Node = Cast<UCustomizableObjectNode>(OutputPin->GetOwningNode()))
-	{
-		Node->BreakExistingConnectionsPostConnection(InputPin, OutputPin);
-	}
 	
 #if WITH_EDITOR
 	if (bResult)
diff --git a/Engine/Plugins/Mutable/Source/CustomizableObjectEditor/Private/MuCOE/Nodes/CustomizableObjectNodeMaterial.cpp b/Engine/Plugins/Mutable/Source/CustomizableObjectEditor/Private/MuCOE/Nodes/CustomizableObjectNodeMaterial.cpp
--- a/Engine/Plugins/Mutable/Source/CustomizableObjectEditor/Private/MuCOE/Nodes/CustomizableObjectNodeMaterial.cpp
+++ b/Engine/Plugins/Mutable/Source/CustomizableObjectEditor/Private/MuCOE/Nodes/CustomizableObjectNodeMaterial.cpp
@@ -1480,33 +1480,4 @@ void UCustomizableObjectNodeMaterial::MeshPostEditChangeProperty(FPostEditChange
 }
 
 
-void UCustomizableObjectNodeMaterial::BreakExistingConnectionsPostConnection(UEdGraphPin* InputPin, UEdGraphPin* OutputPin)
-{
-	if (OutputPin->GetOwningNode() == this) // Case: NodeMaterial (OutputPin) --> (InputPin) Other Node
-	{
-		// Allow multiple NodeCopyMaterial nodes to be connected
-		if (InputPin->GetOwningNode()->IsA(UCustomizableObjectNodeCopyMaterial::StaticClass()))
-		{
-			return;
-		}
-
-		// Remove the previous NodeMaterial
-		UEdGraphPin* RemovePin = nullptr;
-
-		for (UEdGraphPin* Pin : OutputPin->LinkedTo)
-		{
-			if (Pin != InputPin && !Pin->GetOwningNode()->IsA(UCustomizableObjectNodeCopyMaterial::StaticClass())) {
-				RemovePin = Pin;
-				break;
-			}
-		}
-
-		if (RemovePin)
-		{
-			GetSchema()->BreakSinglePinLink(OutputPin, RemovePin); // Can not be called inside the range for loop.
-		}
-	}
-}
-
-
 #undef LOCTEXT_NAMESPACE
diff --git a/Engine/Plugins/Mutable/Source/CustomizableObjectEditor/Private/MuCOE/SMutableCodeViewer.cpp b/Engine/Plugins/Mutable/Source/CustomizableObjectEditor/Private/MuCOE/SMutableCodeViewer.cpp
--- a/Engine/Plugins/Mutable/Source/CustomizableObjectEditor/Private/MuCOE/SMutableCodeViewer.cpp
+++ b/Engine/Plugins/Mutable/Source/CustomizableObjectEditor/Private/MuCOE/SMutableCodeViewer.cpp
@@ -1401,6 +1401,22 @@ void SMutableCodeViewer::GenerateElementRecursive(const int32& InStateIndex, mu:
 	const mu::OP_TYPE ParentOperationType = InProgram.GetOpType(InParentAddress);
 	switch (ParentOperationType)
 	{
+	case mu::OP_TYPE::IM_CONDITIONAL:
+	case mu::OP_TYPE::LA_CONDITIONAL:
+	case mu::OP_TYPE::ME_CONDITIONAL:
+	case mu::OP_TYPE::CO_CONDITIONAL:
+	case mu::OP_TYPE::SC_CONDITIONAL:
+	case mu::OP_TYPE::NU_CONDITIONAL:
+	case mu::OP_TYPE::IN_CONDITIONAL:
+	case mu::OP_TYPE::ED_CONDITIONAL:
+	{
+		mu::OP::ConditionalArgs Args = InProgram.GetOpArgs<mu::OP::ConditionalArgs>(InParentAddress);
+		AddOpFunc(Args.condition, TEXT("cond "));
+		AddOpFunc(Args.yes, TEXT("true "));
+		AddOpFunc(Args.no, TEXT("false "));
+		break;
+	}
+
 	case mu::OP_TYPE::IM_SWITCH:
 	case mu::OP_TYPE::LA_SWITCH:
 	case mu::OP_TYPE::ME_SWITCH:
diff --git a/Engine/Plugins/Mutable/Source/CustomizableObjectEditor/Private/MuCOE/SMutableGraphViewer.cpp b/Engine/Plugins/Mutable/Source/CustomizableObjectEditor/Private/MuCOE/SMutableGraphViewer.cpp
--- a/Engine/Plugins/Mutable/Source/CustomizableObjectEditor/Private/MuCOE/SMutableGraphViewer.cpp
+++ b/Engine/Plugins/Mutable/Source/CustomizableObjectEditor/Private/MuCOE/SMutableGraphViewer.cpp
@@ -23,6 +23,7 @@
 #include "MuT/NodeColourSampleImage.h"
 #include "MuT/NodeColourSwitch.h"
 #include "MuT/NodeComponentEdit.h"
+#include "MuT/NodeComponentSwitch.h"
 #include "MuT/NodeImageFormat.h"
 #include "MuT/NodeImageInterpolate.h"
 #include "MuT/NodeImageInvert.h"
@@ -99,6 +100,137 @@ namespace MutableGraphTreeViewColumns
 	static const FName Name("Name");
 };
 
+static const char* MutableNodeNames[] =
+{
+	"None",
+
+	"Node",
+
+	"Mesh",
+	"MeshConstant",
+	"MeshInterpolate",
+	"MeshTable",
+	"MeshFormat",
+	"MeshTangents",
+	"MeshMorph",
+	"MeshMakeMorph",
+	"MeshSwitch",
+	"MeshFragment",
+	"MeshTransform",
+	"MeshClipMorphPlane",
+	"MeshClipWithMesh",
+	"MeshApplyPose",
+	"MeshVariation",
+	"MeshGeometryOperation",
+	"MeshReshape",
+	"MeshClipDeform",
+
+	"Image",
+	"ImageConstant",
+	"ImageInterpolate",
+	"ImageSaturate",
+	"ImageTable",
+	"ImageSwizzle",
+	"ImageColorMap",
+	"ImageGradient",
+	"ImageBinarise",
+	"ImageLuminance",
+	"ImageLayer",
+	"ImageLayerColour",
+	"ImageResize",
+	"ImagePlainColour",
+	"ImageProject",
+	"ImageMipmap",
+	"ImageSwitch",
+	"ImageConditional",
+	"ImageFormat",
+	"ImageParameter",
+	"ImageMultiLayer",
+	"ImageInvert",
+	"ImageVariation",
+	"ImageNormalComposite",
+	"ImageTransform",
+
+	"Bool",
+	"BoolConstant",
+	"BoolParameter",
+	"BoolNot",
+	"BoolAnd",
+
+	"Color",
+	"ColorConstant",
+	"ColorParameter",
+	"ColorSampleImage",
+	"ColorTable",
+	"ColorImageSize",
+	"ColorFromScalars",
+	"ColorArithmeticOperation",
+	"ColorSwitch",
+	"ColorVariation",
+
+	"Scalar",
+	"ScalarConstant",
+	"ScalarParameter",
+	"ScalarEnumParameter",
+	"ScalarCurve",
+	"ScalarSwitch",
+	"ScalarArithmeticOperation",
+	"ScalarVariation",
+	"ScalarTable",
+
+	"String",
+	"StringConstant",
+	"StringParameter",
+
+	"Projector",
+	"ProjectorConstant",
+	"ProjectorParameter",
+
+	"Range",
+	"RangeFromScalar",
+
+	"Layout",
+
+	"PatchImage",
+	"PatchMesh",
+
+	"Surface",
+	"SurfaceNew",
+	"SurfaceSwitch",
+	"SurfaceVariation",
+
+	"LOD",
+
+	"Component",
+	"ComponentNew",
+	"ComponentEdit",
+	"ComponentSwitch",
+	"ComponentVariation",
+
+	"Object",
+	"ObjectNew",
+	"ObjectGroup",
+
+	"Modifier",
+	"ModifierMeshClipMorphPlane",
+	"ModifierMeshClipWithMesh",
+	"ModifierMeshClipDeform",
+	"ModifierMeshClipWithUVMask",
+	"ModifierSurfaceEdit",
+	"ModifierTransformInMesh",
+
+	"ExtensionData",
+	"ExtensionDataConstant",
+	"ExtensionDataSwitch",
+	"ExtensionDataVariation",
+
+	"Matrix",
+	"MatrixConstant",
+	"MatrixParameter",
+};
+
+static_assert(sizeof(MutableNodeNames) == SIZE_T(mu::Node::EType::Count) * sizeof(char*));
+
 
 class SMutableGraphTreeRow : public STableRow<TSharedPtr<FMutableGraphTreeElement>>
 {
@@ -113,14 +245,16 @@ public:
 		{
 			mu::Node::EType MutableType = RowItem->MutableNode->GetType()->Type;
 			
+			FString NodeName = MutableNodeNames[int32(MutableType)];
+
 			const FString LabelString = RowItem->Prefix.IsEmpty() 
-				? FString::Printf(TEXT("%d"), int32(MutableType) )
-				: FString::Printf( TEXT("%s : %d"), *RowItem->Prefix, int32(MutableType) );
+				? FString::Printf( TEXT("%s"), *NodeName)
+				: FString::Printf( TEXT("%s : %s"), *RowItem->Prefix, *NodeName);
 
 			MainLabel = FText::FromString(LabelString);
 			if (RowItem->DuplicatedOf)
 			{
-				MainLabel = FText::FromString( FString::Printf(TEXT("%d (Duplicated)"), int32(MutableType)));
+				MainLabel = FText::FromString( FString::Printf(TEXT("%s (Duplicated)"), *NodeName));
 			}
 		}
 		else
@@ -449,6 +583,16 @@ void SMutableGraphViewer::GetChildrenForInfo(TSharedPtr<FMutableGraphTreeElement
 		{
 			AddChildFunc(ComponentVar->LODs[LODIndex].get(), FString::Printf(TEXT("LOD [%d]"), LODIndex));
 		}
+		}
+
+	else if (ParentNode->GetType() == mu::NodeComponentSwitch::GetStaticType())
+	{
+		mu::NodeComponentSwitch* ComponentSwitch = StaticCast<mu::NodeComponentSwitch*>(ParentNode);
+		AddChildFunc(ComponentSwitch->Parameter.get(), TEXT("PARAM"));
+		for (int32 OptionIndex = 0; OptionIndex < ComponentSwitch->Options.Num(); OptionIndex++)
+		{
+			AddChildFunc(ComponentSwitch->Options[OptionIndex].get(), FString::Printf(TEXT("OPTION [%d]"), OptionIndex));
+		}
 	}
 
 	else if (ParentNode->GetType() == mu::NodeMeshConstant::GetStaticType())
diff --git a/Engine/Plugins/Mutable/Source/CustomizableObjectEditor/Public/MuCOE/Nodes/CustomizableObjectNode.h b/Engine/Plugins/Mutable/Source/CustomizableObjectEditor/Public/MuCOE/Nodes/CustomizableObjectNode.h
--- a/Engine/Plugins/Mutable/Source/CustomizableObjectEditor/Public/MuCOE/Nodes/CustomizableObjectNode.h
+++ b/Engine/Plugins/Mutable/Source/CustomizableObjectEditor/Public/MuCOE/Nodes/CustomizableObjectNode.h
@@ -223,13 +223,6 @@ public:
 	 */
 	virtual bool ShouldBreakExistingConnections(const UEdGraphPin* InputPin, const UEdGraphPin* OutputPin) const;
 
-	/** Called when a new connection is perform. Used to break an specific existing connectoin. Empty default implementation. This method can be overridden.
-	 *
-	 * @param	InputPin	New connection input pin.
-	 * @param	OutputPin	New connection output pin.
-	 */
-	virtual void BreakExistingConnectionsPostConnection(UEdGraphPin* InputPin, UEdGraphPin* OutputPin) {};
-
 	void PostEditChangeProperty(FPropertyChangedEvent& PropertyChangedEvent) override;
 
 	/** Custom post duplicate function. Called at the beginning of duplication, before the nodes have their Guid updated. */
diff --git a/Engine/Plugins/Mutable/Source/CustomizableObjectEditor/Public/MuCOE/Nodes/CustomizableObjectNodeMaterial.h b/Engine/Plugins/Mutable/Source/CustomizableObjectEditor/Public/MuCOE/Nodes/CustomizableObjectNodeMaterial.h
--- a/Engine/Plugins/Mutable/Source/CustomizableObjectEditor/Public/MuCOE/Nodes/CustomizableObjectNodeMaterial.h
+++ b/Engine/Plugins/Mutable/Source/CustomizableObjectEditor/Public/MuCOE/Nodes/CustomizableObjectNodeMaterial.h
@@ -128,7 +128,6 @@ public:
 	virtual bool ProvidesCustomPinRelevancyTest() const override { return true; }
 	virtual bool IsPinRelevant(const UEdGraphPin* Pin) const override;
 	virtual bool CustomRemovePin(UEdGraphPin& Pin) override;
-	virtual void BreakExistingConnectionsPostConnection(UEdGraphPin* InputPin, UEdGraphPin* OutputPin) override;
 	virtual bool IsNodeOutDatedAndNeedsRefresh() override;
 	virtual FString GetRefreshMessage() const override;
 	virtual TSharedPtr<IDetailsView> CustomizePinDetails(const UEdGraphPin& Pin) const override;
diff --git a/Engine/Plugins/Mutable/Source/MutableRuntime/Private/MuR/CodeRunner.cpp b/Engine/Plugins/Mutable/Source/MutableRuntime/Private/MuR/CodeRunner.cpp
--- a/Engine/Plugins/Mutable/Source/MutableRuntime/Private/MuR/CodeRunner.cpp
+++ b/Engine/Plugins/Mutable/Source/MutableRuntime/Private/MuR/CodeRunner.cpp
@@ -617,7 +617,7 @@ namespace mu
                         pResult->GetPrivate()->Components[NewComponentIndex] = pComp->GetPrivate()->Components[0];
 
 						// Id
-                    	pResult->GetPrivate()->Components[NewComponentIndex].Id = args.id;
+                    	pResult->GetPrivate()->Components[NewComponentIndex].Id = args.ExternalId;
                     }
                 }
                 StoreInstance( item, pResult );
diff --git a/Engine/Plugins/Mutable/Source/MutableTools/Internal/MuT/AST.h b/Engine/Plugins/Mutable/Source/MutableTools/Internal/MuT/AST.h
--- a/Engine/Plugins/Mutable/Source/MutableTools/Internal/MuT/AST.h
+++ b/Engine/Plugins/Mutable/Source/MutableTools/Internal/MuT/AST.h
@@ -60,6 +60,13 @@ namespace mu
 {
 	class ASTOp;
 	class ASTOpFixed;
+	class ASTOpImageCrop;
+	class ASTOpImagePixelFormat;
+	class ASTOpMeshOptimizeSkinning;
+	class ASTOpMeshFormat;
+	class ASTOpImageSwizzle;
+	class ASTOpImageMipmap;
+	class ASTOpMeshExtractLayoutBlocks;
 	struct FProxyFileContext;
 
 	template<typename T>
@@ -324,56 +331,55 @@ namespace mu
 	};
 
 
+	/** */
 	class Sink_ImageCropAST
 	{
 	public:
 
-		Ptr<ASTOp> Apply(const class ASTOpImageCrop* root);
+		Ptr<ASTOp> Apply(const ASTOpImageCrop* root);
 
 	protected:
 
-		const class ASTOpImageCrop* m_root = nullptr;
+		const ASTOpImageCrop* m_root = nullptr;
 		Ptr<ASTOp> m_initialSource;
 		//! For each operation we sink, the map from old instructions to new instructions.
 		TMap<FSinkerOldToNewKey, Ptr<ASTOp>> OldToNew;
 
-		Ptr<ASTOp> Visit(Ptr<ASTOp> at, const class ASTOpImageCrop* currentCropOp);
+		Ptr<ASTOp> Visit(Ptr<ASTOp> at, const ASTOpImageCrop* currentCropOp);
 	};
 
-	//---------------------------------------------------------------------------------------------
-	//!
-	//---------------------------------------------------------------------------------------------
+
+	/** */
 	class Sink_ImagePixelFormatAST
 	{
 	public:
 
-		Ptr<ASTOp> Apply(const class ASTOpImagePixelFormat* root);
+		Ptr<ASTOp> Apply(const ASTOpImagePixelFormat* root);
 
 	protected:
 
-		const class ASTOpImagePixelFormat* m_root = nullptr;
+		const ASTOpImagePixelFormat* m_root = nullptr;
 		Ptr<ASTOp> m_initialSource;
 		TMap<FSinkerOldToNewKey, Ptr<ASTOp>> OldToNew;
 
-		Ptr<ASTOp> Visit(Ptr<ASTOp> at, const class ASTOpImagePixelFormat* currentFormatOp);
+		Ptr<ASTOp> Visit(Ptr<ASTOp> at, const ASTOpImagePixelFormat* currentFormatOp);
 	};
 
-	//---------------------------------------------------------------------------------------------
-	//!
-	//---------------------------------------------------------------------------------------------
+
+	/** */
 	class Sink_ImageSwizzleAST
 	{
 	public:
 
-		Ptr<ASTOp> Apply(const class ASTOpImageSwizzle* Root);
+		Ptr<ASTOp> Apply(const ASTOpImageSwizzle* Root);
 
 	protected:
 
-		const class ASTOpImageSwizzle* Root = nullptr;
+		const ASTOpImageSwizzle* Root = nullptr;
 		Ptr<ASTOp> InitialSource;
 		TMap<FSinkerOldToNewKey, Ptr<ASTOp>> OldToNew;
 
-		Ptr<ASTOp> Visit(Ptr<ASTOp> at, const class ASTOpImageSwizzle* CurrentSwizzleOp);
+		Ptr<ASTOp> Visit(Ptr<ASTOp> at, const ASTOpImageSwizzle* CurrentSwizzleOp);
 	};
 
 
@@ -382,16 +388,32 @@ namespace mu
 	{
 	public:
 
-		// \TODO This is recursive and may cause stack overflows in big models.
-		Ptr<ASTOp> Apply(const class ASTOpMeshFormat* Root);
+		Ptr<ASTOp> Apply(const ASTOpMeshFormat* Root);
 
 	protected:
 
-		const class ASTOpMeshFormat* m_root = nullptr;
+		const ASTOpMeshFormat* m_root = nullptr;
 		Ptr<ASTOp> m_initialSource;
 		TMap<FSinkerOldToNewKey, Ptr<ASTOp>> OldToNew;
 
-		Ptr<ASTOp> Visit(const Ptr<ASTOp>& at, const class ASTOpMeshFormat* currentFormatOp);
+		Ptr<ASTOp> Visit(const Ptr<ASTOp>& Op, const ASTOpMeshFormat* CurrentSinkOp);
+	};
+
+
+	/** */
+	class Sink_MeshOptimizeSkinningAST
+	{
+	public:
+
+		Ptr<ASTOp> Apply(const ASTOpMeshOptimizeSkinning* InRoot);
+
+	protected:
+
+		const ASTOpMeshOptimizeSkinning* Root = nullptr;
+		Ptr<ASTOp> InitialSource;
+		TMap<FSinkerOldToNewKey, Ptr<ASTOp>> OldToNew;
+
+		Ptr<ASTOp> Visit(const Ptr<ASTOp>& Op, const ASTOpMeshOptimizeSkinning* CurrentSinkOp);
 	};
 
 
@@ -400,37 +422,34 @@ namespace mu
 	{
 	public:
 
-		// \TODO This is recursive and may cause stack overflows in big models.
-		Ptr<ASTOp> Apply(const class ASTOpMeshExtractLayoutBlocks* Root);
+		Ptr<ASTOp> Apply(const ASTOpMeshExtractLayoutBlocks* Root);
 
 	protected:
 
-		const class ASTOpMeshExtractLayoutBlocks* m_root = nullptr;
+		const ASTOpMeshExtractLayoutBlocks* m_root = nullptr;
 		Ptr<ASTOp> m_initialSource;
 		TMap<FSinkerOldToNewKey, Ptr<ASTOp>> OldToNew;
 
-		Ptr<ASTOp> Visit(const Ptr<ASTOp>& at, const class ASTOpMeshExtractLayoutBlocks* currentFormatOp);
+		Ptr<ASTOp> Visit(const Ptr<ASTOp>& at, const ASTOpMeshExtractLayoutBlocks* currentFormatOp);
 	};
 
 
-	//---------------------------------------------------------------------------------------------
-	//!
-	//---------------------------------------------------------------------------------------------
+	/** */
 	class Sink_ImageMipmapAST
 	{
 	public:
 
-		Ptr<ASTOp> Apply(const class ASTOpImageMipmap* Root);
+		Ptr<ASTOp> Apply(const ASTOpImageMipmap* Root);
 
 	protected:
 
-		const class ASTOpImageMipmap* m_root = nullptr;
+		const ASTOpImageMipmap* m_root = nullptr;
 		Ptr<ASTOp> m_initialSource;
 
 		//! For each operation we sink, the map from old instructions to new instructions.
 		TMap<FSinkerOldToNewKey, Ptr<ASTOp>> OldToNew;
 
-		Ptr<ASTOp> Visit(Ptr<ASTOp> at, const class ASTOpImageMipmap* currentMipmapOp);
+		Ptr<ASTOp> Visit(Ptr<ASTOp> at, const ASTOpImageMipmap* currentMipmapOp);
 
 	};
 
@@ -442,6 +461,7 @@ namespace mu
 		Sink_ImageMipmapAST ImageMipmapSinker;
 		Sink_MeshFormatAST MeshFormatSinker;
 		Sink_MeshExtractLayoutBlocksAST MeshExtractLayoutBlocksSinker;
+		Sink_MeshOptimizeSkinningAST MeshOptimizeSkinningSinker;
 	};
 
 
@@ -1125,7 +1145,7 @@ namespace mu
 
         ~ASTOpFixed() override;
 
-        OP_TYPE GetOpType() const override { return op.type; }
+		virtual OP_TYPE GetOpType() const override { return op.type; }
 
 		virtual Ptr<ASTOp> Clone( MapChildFuncRef mapChild ) const override;
 		virtual void ForEachChild( const TFunctionRef<void(ASTChild&)> f ) override;
diff --git a/Engine/Plugins/Mutable/Source/MutableTools/Private/MuT/ASTOpInstanceAdd.cpp b/Engine/Plugins/Mutable/Source/MutableTools/Private/MuT/ASTOpInstanceAdd.cpp
--- a/Engine/Plugins/Mutable/Source/MutableTools/Private/MuT/ASTOpInstanceAdd.cpp
+++ b/Engine/Plugins/Mutable/Source/MutableTools/Private/MuT/ASTOpInstanceAdd.cpp
@@ -2,19 +2,18 @@
 
 #include "MuT/ASTOpInstanceAdd.h"
 
-#include "Misc/AssertionMacros.h"
+#include "MuT/ASTOpConditional.h"
+#include "MuT/CodeOptimiser.h"
 #include "MuR/ModelPrivate.h"
 #include "MuR/ParametersPrivate.h"
 #include "MuR/RefCounted.h"
 #include "MuR/Types.h"
-#include "MuT/CodeOptimiser.h"
+#include "Misc/AssertionMacros.h"
 
 
 namespace mu
 {
 
-
-	//-------------------------------------------------------------------------------------------------
 	ASTOpInstanceAdd::ASTOpInstanceAdd()
 		: instance(this)
 		, value(this)
@@ -22,7 +21,6 @@ namespace mu
 	}
 
 
-	//-------------------------------------------------------------------------------------------------
 	ASTOpInstanceAdd::~ASTOpInstanceAdd()
 	{
 		// Explicit call needed to avoid recursive destruction
@@ -30,7 +28,6 @@ namespace mu
 	}
 
 
-	//-------------------------------------------------------------------------------------------------
 	bool ASTOpInstanceAdd::IsEqual(const ASTOp& otherUntyped) const
 	{
 		if (otherUntyped.GetOpType()==GetOpType())
@@ -48,7 +45,6 @@ namespace mu
 	}
 
 
-	//-------------------------------------------------------------------------------------------------
 	mu::Ptr<ASTOp> ASTOpInstanceAdd::Clone(MapChildFuncRef mapChild) const
 	{
 		Ptr<ASTOpInstanceAdd> n = new ASTOpInstanceAdd();
@@ -63,7 +59,6 @@ namespace mu
 	}
 
 
-	//-------------------------------------------------------------------------------------------------
 	uint64 ASTOpInstanceAdd::Hash() const
 	{
 		uint64 res = std::hash<size_t>()(size_t(type));
@@ -73,7 +68,6 @@ namespace mu
 	}
 
 
-	//-------------------------------------------------------------------------------------------------
 	void ASTOpInstanceAdd::Assert()
 	{
 		switch (type)
@@ -97,7 +91,6 @@ namespace mu
 	}
 
 
-	//-------------------------------------------------------------------------------------------------
 	void ASTOpInstanceAdd::ForEachChild(const TFunctionRef<void(ASTChild&)> f)
 	{
 		f(instance);
@@ -105,7 +98,6 @@ namespace mu
 	}
 
 
-	//-------------------------------------------------------------------------------------------------
 	void ASTOpInstanceAdd::Link(FProgram& program, FLinkerOptions*)
 	{
 		// Already linked?
@@ -166,4 +158,86 @@ namespace mu
 
 	}
 
+
+	Ptr<ASTOp> ASTOpInstanceAdd::OptimiseSink(const FModelOptimizationOptions& options, FOptimizeSinkContext& context) const
+	{
+		Ptr<ASTOp> at;
+
+		switch (type)
+		{
+
+		case OP_TYPE::IN_ADDMESH:
+		{
+			Ptr<ASTOp> ValueAt = value.child();
+			if (!ValueAt)
+			{
+				break;
+			}
+
+			OP_TYPE ValueType = ValueAt->GetOpType();
+			switch (ValueType)
+			{
+
+			// We want to move the conditional up the op graph because this way the mesh root operation
+			// can easily match other mesh root operations in the program. This is important because this operation
+			// address is used for caching, and avoiding duplicate mesh work with multi-components.
+			case OP_TYPE::ME_CONDITIONAL:
+			{
+				const ASTOpConditional* TypedValue = static_cast<const ASTOpConditional*>(ValueAt.get());
+
+				bool bCanSink = !instance.child();
+
+				const ASTOpConditional* TypedSource = nullptr;
+				if (instance.child() && instance.child()->GetOpType() == OP_TYPE::IN_CONDITIONAL)
+				{
+					TypedSource = static_cast<const ASTOpConditional*>(instance.child().get());
+					if (TypedSource->condition == TypedValue->condition)
+					{
+						bCanSink = true;
+					}
+				}
+
+				if (bCanSink)
+				{
+					Ptr<ASTOpConditional> NewOp = mu::Clone<ASTOpConditional>(ValueAt);
+					NewOp->type = OP_TYPE::IN_CONDITIONAL;
+
+					Ptr<ASTOpInstanceAdd> TrueOp = mu::Clone<ASTOpInstanceAdd>(this);
+					TrueOp->instance = TypedSource ? TypedSource->yes.child() : nullptr; 
+					TrueOp->value = TypedValue->yes.child();
+					NewOp->yes = TrueOp;
+
+					Ptr<ASTOpInstanceAdd> FalseOp = mu::Clone<ASTOpInstanceAdd>(this);
+					FalseOp->instance = TypedSource ? TypedSource->no.child() : nullptr; 
+					FalseOp->value = TypedValue->no.child();
+					NewOp->no = FalseOp;
+
+					at = NewOp;
+				}
+				break;
+			}
+
+			// I don't think we have switches at this point but if we did, it would be interesting to optimize
+			// for the same reasons 
+			case OP_TYPE::ME_SWITCH:
+			{
+				// TODO
+				break;
+			}
+
+			default:
+				break;
+			}
+
+			break;
+		}
+
+		default:
+			break;
+		}
+
+		return at;
+	}
+
+
 }
diff --git a/Engine/Plugins/Mutable/Source/MutableTools/Private/MuT/ASTOpInstanceAdd.h b/Engine/Plugins/Mutable/Source/MutableTools/Private/MuT/ASTOpInstanceAdd.h
--- a/Engine/Plugins/Mutable/Source/MutableTools/Private/MuT/ASTOpInstanceAdd.h
+++ b/Engine/Plugins/Mutable/Source/MutableTools/Private/MuT/ASTOpInstanceAdd.h
@@ -36,14 +36,15 @@ struct FProgram;
 		ASTOpInstanceAdd(const ASTOpInstanceAdd&) = delete;
 		~ASTOpInstanceAdd();
 
-		OP_TYPE GetOpType() const override { return type; }
-
-		void ForEachChild(const TFunctionRef<void(ASTChild&)>) override;
-		bool IsEqual(const ASTOp& otherUntyped) const override;
-		Ptr<ASTOp> Clone(MapChildFuncRef mapChild) const override;
-		uint64 Hash() const override;
-		void Assert() override;
-		void Link(FProgram& program, FLinkerOptions* Options) override;
+		virtual OP_TYPE GetOpType() const override { return type; }
+
+		virtual void ForEachChild(const TFunctionRef<void(ASTChild&)>) override;
+		virtual bool IsEqual(const ASTOp& otherUntyped) const override;
+		virtual Ptr<ASTOp> Clone(MapChildFuncRef mapChild) const override;
+		virtual uint64 Hash() const override;
+		virtual void Assert() override;
+		virtual void Link(FProgram& program, FLinkerOptions* Options) override;
+		virtual Ptr<ASTOp> OptimiseSink(const FModelOptimizationOptions&, FOptimizeSinkContext&) const override;
 	};
 
 
diff --git a/Engine/Plugins/Mutable/Source/MutableTools/Private/MuT/ASTOpMeshOptimizeSkinning.cpp b/Engine/Plugins/Mutable/Source/MutableTools/Private/MuT/ASTOpMeshOptimizeSkinning.cpp
--- a/Engine/Plugins/Mutable/Source/MutableTools/Private/MuT/ASTOpMeshOptimizeSkinning.cpp
+++ b/Engine/Plugins/Mutable/Source/MutableTools/Private/MuT/ASTOpMeshOptimizeSkinning.cpp
@@ -2,26 +2,25 @@
 
 #include "MuT/ASTOpMeshOptimizeSkinning.h"
 
-#include "HAL/PlatformMath.h"
+#include "MuT/ASTOpConditional.h"
+#include "MuT/ASTOpMeshMorph.h"
+#include "MuT/ASTOpMeshAddTags.h"
+#include "MuT/ASTOpSwitch.h"
 #include "MuR/ModelPrivate.h"
 #include "MuR/RefCounted.h"
 #include "MuR/Types.h"
+#include "HAL/PlatformMath.h"
 
 
 namespace mu
 {
 
-
-	//-------------------------------------------------------------------------------------------------
-	//-------------------------------------------------------------------------------------------------
-	//-------------------------------------------------------------------------------------------------
 	ASTOpMeshOptimizeSkinning::ASTOpMeshOptimizeSkinning()
-		: source(this)
+		: Source(this)
 	{
 	}
 
 
-	//-------------------------------------------------------------------------------------------------
 	ASTOpMeshOptimizeSkinning::~ASTOpMeshOptimizeSkinning()
 	{
 		// Explicit call needed to avoid recursive destruction
@@ -29,70 +28,172 @@ namespace mu
 	}
 
 
-	//-------------------------------------------------------------------------------------------------
-	bool ASTOpMeshOptimizeSkinning::IsEqual(const ASTOp& otherUntyped) const
+	bool ASTOpMeshOptimizeSkinning::IsEqual(const ASTOp& OtherUntyped) const
 	{
-		if (otherUntyped.GetOpType() == GetOpType())
+		if (OtherUntyped.GetOpType() == GetOpType())
 		{
-			const ASTOpMeshOptimizeSkinning* other = static_cast<const ASTOpMeshOptimizeSkinning*>(&otherUntyped);
-			return source == other->source;
+			const ASTOpMeshOptimizeSkinning* Other = static_cast<const ASTOpMeshOptimizeSkinning*>(&OtherUntyped);
+			return Source == Other->Source;
 		}
 		return false;
 	}
 
 
-	//-------------------------------------------------------------------------------------------------
 	uint64 ASTOpMeshOptimizeSkinning::Hash() const
 	{
-		uint64 res = std::hash<void*>()(source.child().get());
-		return res;
+		uint64 Res = std::hash<void*>()(Source.child().get());
+		return Res;
 	}
 
 
-	//-------------------------------------------------------------------------------------------------
-	mu::Ptr<ASTOp> ASTOpMeshOptimizeSkinning::Clone(MapChildFuncRef mapChild) const
+	mu::Ptr<ASTOp> ASTOpMeshOptimizeSkinning::Clone(MapChildFuncRef MapChild) const
 	{
 		Ptr<ASTOpMeshOptimizeSkinning> n = new ASTOpMeshOptimizeSkinning();
-		n->source = mapChild(source.child());
+		n->Source = MapChild(Source.child());
 		return n;
 	}
 
 
-	//-------------------------------------------------------------------------------------------------
-	void ASTOpMeshOptimizeSkinning::ForEachChild(const TFunctionRef<void(ASTChild&)> f)
+	void ASTOpMeshOptimizeSkinning::ForEachChild(const TFunctionRef<void(ASTChild&)> Func)
 	{
-		f(source);
+		Func(Source);
 	}
 
 
-	//-------------------------------------------------------------------------------------------------
-	void ASTOpMeshOptimizeSkinning::Link(FProgram& program, FLinkerOptions*)
+	void ASTOpMeshOptimizeSkinning::Link(FProgram& Program, FLinkerOptions*)
 	{
 		// Already linked?
 		if (!linkedAddress)
 		{
-			OP::MeshOptimizeSkinningArgs args;
-			FMemory::Memset(&args, 0, sizeof(args));
+			OP::MeshOptimizeSkinningArgs Args = {};
 
-			if (source) args.source = source->linkedAddress;
+			if (Source)
+			{
+				Args.source = Source->linkedAddress;
+			}
 
-			linkedAddress = (OP::ADDRESS)program.m_opAddress.Num();
-			program.m_opAddress.Add((uint32_t)program.m_byteCode.Num());
-			AppendCode(program.m_byteCode, OP_TYPE::ME_OPTIMIZESKINNING);
-			AppendCode(program.m_byteCode, args);
+			linkedAddress = OP::ADDRESS(Program.m_opAddress.Num());
+			Program.m_opAddress.Add((uint32)Program.m_byteCode.Num());
+			AppendCode(Program.m_byteCode, OP_TYPE::ME_OPTIMIZESKINNING);
+			AppendCode(Program.m_byteCode, Args);
 		}
 	}
 
 
-	//-------------------------------------------------------------------------------------------------
+	mu::Ptr<ASTOp> ASTOpMeshOptimizeSkinning::OptimiseSink(const FModelOptimizationOptions&, FOptimizeSinkContext& Context) const
+	{
+		return Context.MeshOptimizeSkinningSinker.Apply(this);
+	}
+
+
 	FSourceDataDescriptor ASTOpMeshOptimizeSkinning::GetSourceDataDescriptor(FGetSourceDataDescriptorContext* Context) const
 	{
-		if (source)
+		if (Source)
 		{
-			return source->GetSourceDataDescriptor(Context);
+			return Source->GetSourceDataDescriptor(Context);
 		}
 
 		return {};
 	}
 
+
+
+
+	//---------------------------------------------------------------------------------------------
+	//---------------------------------------------------------------------------------------------
+	//---------------------------------------------------------------------------------------------
+	mu::Ptr<ASTOp> Sink_MeshOptimizeSkinningAST::Apply(const ASTOpMeshOptimizeSkinning* InRoot)
+	{
+		Root = InRoot;
+
+		OldToNew.Reset();
+
+		InitialSource = Root->Source.child();
+		mu::Ptr<ASTOp> NewSource = Visit(InitialSource, Root);
+
+		// If there is any change, it is the new root.
+		if (NewSource != InitialSource)
+		{
+			return NewSource;
+		}
+
+		return nullptr;
+	}
+
+
+	mu::Ptr<ASTOp> Sink_MeshOptimizeSkinningAST::Visit(const mu::Ptr<ASTOp>& at, const ASTOpMeshOptimizeSkinning* CurrentSinkOp)
+	{
+		if (!at)
+		{
+			return nullptr;
+		}
+
+		// Move the operation down the conditionals and harmless operations to try to optimize it into a constant if possible
+		// and also try to move the conditionals to theroot so that may be moved out of the mesh operation graph entirely.
+
+		// Already visited?
+		const Ptr<ASTOp>* Cached = OldToNew.Find({ at, CurrentSinkOp });
+		if (Cached)
+		{
+			return *Cached;
+		}
+
+		mu::Ptr<ASTOp> NewAt = at;
+		switch (at->GetOpType())
+		{
+
+		case OP_TYPE::ME_APPLYLAYOUT:
+		{
+			Ptr<ASTOpFixed> NewOp = mu::Clone<ASTOpFixed>(at);
+			NewOp->SetChild(NewOp->op.args.MeshApplyLayout.mesh, Visit(NewOp->children[NewOp->op.args.MeshApplyLayout.mesh].child(), CurrentSinkOp));
+			NewAt = NewOp;
+			break;
+		}
+
+		case OP_TYPE::ME_ADDTAGS:
+		{
+			Ptr<ASTOpMeshAddTags> NewOp = mu::Clone<ASTOpMeshAddTags>(at);
+			NewOp->Source = Visit(NewOp->Source.child(), CurrentSinkOp);
+			NewAt = NewOp;
+			break;
+		}
+
+		case OP_TYPE::ME_CONDITIONAL:
+		{
+			Ptr<ASTOpConditional> NewOp = mu::Clone<ASTOpConditional>(at);
+			NewOp->yes = Visit(NewOp->yes.child(), CurrentSinkOp);
+			NewOp->no = Visit(NewOp->no.child(), CurrentSinkOp);
+			NewAt = NewOp;
+			break;
+		}
+
+		case OP_TYPE::ME_SWITCH:
+		{
+			Ptr<ASTOpSwitch> NewOp = mu::Clone<ASTOpSwitch>(at);
+			NewOp->def = Visit(NewOp->def.child(), CurrentSinkOp);
+			for (auto& c : NewOp->cases)
+			{
+				c.branch = Visit(c.branch.child(), CurrentSinkOp);
+			}
+			NewAt = NewOp;
+			break;
+		}
+
+		// If we reach here it means the operation type has not been optimized.
+		default:
+			if (at != InitialSource)
+			{
+				mu::Ptr<ASTOpMeshOptimizeSkinning> NewOp = mu::Clone<ASTOpMeshOptimizeSkinning>(CurrentSinkOp);
+				NewOp->Source = at;
+				NewAt = NewOp;
+			}
+			break;
+
+		}
+
+		OldToNew.Add({ at, CurrentSinkOp }, NewAt);
+
+		return NewAt;
+	}
+
 }
diff --git a/Engine/Plugins/Mutable/Source/MutableTools/Private/MuT/ASTOpMeshOptimizeSkinning.h b/Engine/Plugins/Mutable/Source/MutableTools/Private/MuT/ASTOpMeshOptimizeSkinning.h
--- a/Engine/Plugins/Mutable/Source/MutableTools/Private/MuT/ASTOpMeshOptimizeSkinning.h
+++ b/Engine/Plugins/Mutable/Source/MutableTools/Private/MuT/ASTOpMeshOptimizeSkinning.h
@@ -18,7 +18,7 @@ namespace mu
 	public:
 
 		//! Mesh that will have the skinning optimized.
-		ASTChild source;
+		ASTChild Source;
 
 	public:
 
@@ -27,12 +27,13 @@ namespace mu
 		~ASTOpMeshOptimizeSkinning() override;
 
 		// ASTOp interface
-		OP_TYPE GetOpType() const override { return OP_TYPE::ME_OPTIMIZESKINNING; }
-		uint64 Hash() const override;
-		void ForEachChild(const TFunctionRef<void(ASTChild&)>) override;
-		bool IsEqual(const ASTOp& otherUntyped) const override;
-		Ptr<ASTOp> Clone(MapChildFuncRef mapChild) const override;
-		void Link(FProgram& program, FLinkerOptions*) override;
+		virtual OP_TYPE GetOpType() const override { return OP_TYPE::ME_OPTIMIZESKINNING; }
+		virtual uint64 Hash() const override;
+		virtual void ForEachChild(const TFunctionRef<void(ASTChild&)>) override;
+		virtual bool IsEqual(const ASTOp&) const override;
+		virtual Ptr<ASTOp> Clone(MapChildFuncRef) const override;
+		virtual void Link(FProgram&, FLinkerOptions*) override;
+		virtual Ptr<ASTOp> OptimiseSink(const FModelOptimizationOptions&, FOptimizeSinkContext&) const override;
 		virtual FSourceDataDescriptor GetSourceDataDescriptor(FGetSourceDataDescriptorContext*) const override;
 	};
 
diff --git a/Engine/Plugins/Mutable/Source/MutableTools/Private/MuT/CodeGenerator.cpp b/Engine/Plugins/Mutable/Source/MutableTools/Private/MuT/CodeGenerator.cpp
--- a/Engine/Plugins/Mutable/Source/MutableTools/Private/MuT/CodeGenerator.cpp
+++ b/Engine/Plugins/Mutable/Source/MutableTools/Private/MuT/CodeGenerator.cpp
@@ -595,7 +595,7 @@ namespace mu
 		InstanceOp->type = OP_TYPE::IN_ADDCOMPONENT;
 		InstanceOp->instance = Options.BaseInstance;
 		InstanceOp->value = LODsOp;
-		InstanceOp->id = node.Id;
+		InstanceOp->ExternalId = node.Id;
 
 		Result.op = InstanceOp;
 
@@ -856,7 +856,9 @@ namespace mu
 
 		// Do we need to generate the mesh? Or was it already generated for state conditions 
 		// accepting the current state?
-		FirstPassGenerator::FSurface* TargetSurface = nullptr;
+		TArray<FirstPassGenerator::FSurface*> TargetSurfaces;
+		TargetSurfaces.Reserve(FirstPass.Surfaces.Num());
+
 		for (FirstPassGenerator::FSurface& Surface : FirstPass.Surfaces)
 		{
 			if (Surface.Node != SurfaceNode)
@@ -883,11 +885,11 @@ namespace mu
 			else
 			{
 				// Not already generated, we will generate this
-				TargetSurface = &Surface;
+				TargetSurfaces.Add( &Surface );
 			}
 		}
 
-        if (!TargetSurface)
+        if (TargetSurfaces.IsEmpty())
         {
             return;
         }
@@ -1112,7 +1114,10 @@ namespace mu
 			MeshResults.LayoutOps = MoveTemp(SurfaceLayoutOps); 
 
             // Store in the surface for later use.
-            TargetSurface->ResultMeshOp = LastMeshOp;
+			for (FirstPassGenerator::FSurface* TargetSurface : TargetSurfaces)
+			{
+				TargetSurface->ResultMeshOp = LastMeshOp;
+			}
         }
 
         // Create the expression for each texture, if we are not reusing the surface from another LOD.
@@ -1548,7 +1553,11 @@ namespace mu
 		}
 
         SurfaceResult.surfaceOp = LastSurfOp;
-        TargetSurface->ResultSurfaceOp = LastSurfOp;
+
+		for (FirstPassGenerator::FSurface* TargetSurface : TargetSurfaces)
+		{
+			TargetSurface->ResultSurfaceOp = LastSurfOp;
+		}
 
 		// If we are going to share this surface properties, remember it.
 		if (bIsBaseForSharedSurface)
@@ -1571,13 +1580,12 @@ namespace mu
         Ptr<ASTOp> lastMeshOp;
         FString lastMeshName;
 
-        // This generates a different ID for each surface. It can be used to match it to the
-        // mesh surface, or for debugging. It cannot be 0 because it is a special case for the
-        // merge operation.
+        // This generates a different ID for each surface in the LOD and the component. It can be used to match it to the
+        // mesh surface. It cannot be 0 because it is a special case for the merge operation.
         int32 surfaceID=1;
 
         // Look for all surfaces that belong to this component
-		for (int32 i = 0; i<FirstPass.Surfaces.Num(); ++i, ++surfaceID)
+		for (int32 i = 0; i<FirstPass.Surfaces.Num(); ++i)
 		{
 			const FirstPassGenerator::FSurface& its = FirstPass.Surfaces[i];
 			if (its.Component==Options.Component
@@ -1600,20 +1608,19 @@ namespace mu
                     }
                 }
 
-                Ptr<ASTOpInstanceAdd> sop = new ASTOpInstanceAdd();
-                sop->type = OP_TYPE::IN_ADDSURFACE;
-                sop->name = its.Node->Name;
-                sop->instance = lastCompOp;
+                Ptr<ASTOpInstanceAdd> SurfaceOp = new ASTOpInstanceAdd();
+				SurfaceOp->type = OP_TYPE::IN_ADDSURFACE;
+				SurfaceOp->name = its.Node->Name;
+				SurfaceOp->instance = lastCompOp;
 
 				FSurfaceGenerationOptions SurfaceOptions(Options);
 				FSurfaceGenerationResult surfaceGenerationResult;
                 GenerateSurface( surfaceGenerationResult, SurfaceOptions, its.Node );
-                sop->value = surfaceGenerationResult.surfaceOp;
+				SurfaceOp->value = surfaceGenerationResult.surfaceOp;
 
-                sop->id = surfaceID;
-                sop->ExternalId = its.Node->ExternalId;
-                sop->SharedSurfaceId = its.Node->SharedSurfaceId;
-                Ptr<ASTOp> surfaceAt = sop;
+				SurfaceOp->id = surfaceID;
+				SurfaceOp->ExternalId = its.Node->ExternalId;
+				SurfaceOp->SharedSurfaceId = its.Node->SharedSurfaceId;
 
                 Ptr<ASTOp> SurfaceConditionOp = its.FinalCondition;
 
@@ -1621,7 +1628,7 @@ namespace mu
                     Ptr<ASTOpConditional> op = new ASTOpConditional();
                     op->type = OP_TYPE::IN_CONDITIONAL;
                     op->no = lastCompOp;
-                    op->yes = surfaceAt;
+                    op->yes = SurfaceOp;
                     op->condition = SurfaceConditionOp;
                     lastCompOp = op;
                 }
@@ -1654,13 +1661,15 @@ namespace mu
                 {
                     lastMeshOp = mergeAd;
                 }
+
+				++surfaceID;
 			}
 		}
 
 		// Add op to optimize the skinning of the resulting mesh
 		{
 			Ptr<ASTOpMeshOptimizeSkinning> mop = new ASTOpMeshOptimizeSkinning();
-			mop->source = lastMeshOp;
+			mop->Source = lastMeshOp;
 			lastMeshOp = mop;
 		}
 
diff --git a/Engine/Plugins/Mutable/Source/MutableTools/Private/MuT/CodeGenerator_FirstPass.h b/Engine/Plugins/Mutable/Source/MutableTools/Private/MuT/CodeGenerator_FirstPass.h
--- a/Engine/Plugins/Mutable/Source/MutableTools/Private/MuT/CodeGenerator_FirstPass.h
+++ b/Engine/Plugins/Mutable/Source/MutableTools/Private/MuT/CodeGenerator_FirstPass.h
@@ -136,6 +136,8 @@ namespace mu
             // This is filled in the final code generation pass
             Ptr<ASTOp> ResultSurfaceOp;
             Ptr<ASTOp> ResultMeshOp;
+
+			bool operator==(const FSurface&) const = default;
         };
 		TArray<FSurface> Surfaces;
 
diff --git a/Engine/Plugins/Mutable/Source/MutableTools/Private/MuT/CodeOptimiser_Logic.cpp b/Engine/Plugins/Mutable/Source/MutableTools/Private/MuT/CodeOptimiser_Logic.cpp
--- a/Engine/Plugins/Mutable/Source/MutableTools/Private/MuT/CodeOptimiser_Logic.cpp
+++ b/Engine/Plugins/Mutable/Source/MutableTools/Private/MuT/CodeOptimiser_Logic.cpp
@@ -1059,11 +1059,11 @@ namespace mu
                                     }
                                     else
                                     {
-                                        // Other parents may not impose the same condition that allows
-                                        // the optimisation.
+                                        // Other parents may not impose the same condition that allows the optimisation.
                                         Ptr<ASTOpInstanceAdd> newAddSurface = mu::Clone<ASTOpInstanceAdd>(addSurface);
                                         newAddSurface->instance = bottomConditional->no.child();
                                         topConditional->yes = newAddSurface;
+										addSurface = newAddSurface.get();
                                     }
 
                                     modified = true;
