[mutable] Replace some index types (object component, instance component) by opaque types to prevent misusing them.
- Remove some unused code related to legacy logging.
- Remove internal use of deprecated functions.
#rnx

diff --git a/Engine/Plugins/Mutable/Source/CustomizableObject/Internal/MuCO/CustomizableObjectInstanceDescriptor.h b/Engine/Plugins/Mutable/Source/CustomizableObject/Internal/MuCO/CustomizableObjectInstanceDescriptor.h
--- a/Engine/Plugins/Mutable/Source/CustomizableObject/Internal/MuCO/CustomizableObjectInstanceDescriptor.h
+++ b/Engine/Plugins/Mutable/Source/CustomizableObject/Internal/MuCO/CustomizableObjectInstanceDescriptor.h
@@ -3,6 +3,7 @@
 #pragma once
 
 #include "MuCO/MultilayerProjector.h"
+#include "MuCO/CustomizableObject.h"
 #include "MuCO/CustomizableObjectParameterTypeDefinitions.h"
 #include "Math/RandomStream.h"
 
@@ -81,6 +82,7 @@ struct CUSTOMIZABLEOBJECT_API FCustomizableObjectInstanceDescriptor
 	void SetRequestedLODLevels(const TArray<uint16>& InRequestedLODLevels);
 
 	const TArray<uint16>& GetRequestedLODLevels() const;
+	uint16 GetRequestedLODs(FCustomizableObjectComponentIndex) const;
 
 	// ------------------------------------------------------------
 	// Parameters
diff --git a/Engine/Plugins/Mutable/Source/CustomizableObject/Internal/MuCO/CustomizableObjectInstancePrivate.h b/Engine/Plugins/Mutable/Source/CustomizableObject/Internal/MuCO/CustomizableObjectInstancePrivate.h
--- a/Engine/Plugins/Mutable/Source/CustomizableObject/Internal/MuCO/CustomizableObjectInstancePrivate.h
+++ b/Engine/Plugins/Mutable/Source/CustomizableObject/Internal/MuCO/CustomizableObjectInstancePrivate.h
@@ -8,6 +8,7 @@
 #include "Rendering/SkeletalMeshModel.h"
 #include "MuCO/CustomizableObjectInstance.h"
 #include "MuCO/CustomizableObjectPrivate.h"
+#include "MuCO/CustomizableObjectSystemPrivate.h"
 #include "MuR/Instance.h"
 #include "GameplayTagContainer.h"
 #include "MuCO/DescriptorHash.h"
@@ -190,7 +191,7 @@ struct FExtensionInstanceData
 
 #if WITH_EDITOR
 DECLARE_MULTICAST_DELEGATE_OneParam(FObjectInstanceTransactedDelegate, const FTransactionObjectEvent&);
-#endif //WIT_EDITOR
+#endif //WITH_EDITOR
 
 
 UCLASS()
@@ -243,10 +244,6 @@ public:
 
 	FCustomizableInstanceComponentData* GetComponentData(const FName& ComponentName);
 
-	// DEPRECATED
-	FCustomizableInstanceComponentData* GetComponentData(int32 ComponentIndex);
-	const FCustomizableInstanceComponentData* GetComponentData(int32 ComponentIndex) const;
-
 	ECOInstanceFlags GetCOInstanceFlags() const { return InstanceFlagsPrivate; }
 	void SetCOInstanceFlags(ECOInstanceFlags FlagsToSet) { InstanceFlagsPrivate = (ECOInstanceFlags)(InstanceFlagsPrivate | FlagsToSet); }
 	void ClearCOInstanceFlags(ECOInstanceFlags FlagsToClear) { InstanceFlagsPrivate = (ECOInstanceFlags)(InstanceFlagsPrivate & ~FlagsToClear); }
@@ -295,20 +292,24 @@ public:
 
 private:
 
-	void InitSkeletalMeshData(const TSharedRef<FUpdateContextPrivate>& OperationData, USkeletalMesh* SkeletalMesh, const FMutableRefSkeletalMeshData& RefSkeletalMeshData, const UCustomizableObject& CustomizableObject, int32 ObjectComponentIndex);
+	void InitSkeletalMeshData(const TSharedRef<FUpdateContextPrivate>& OperationData, USkeletalMesh* SkeletalMesh, const FMutableRefSkeletalMeshData& RefSkeletalMeshData, const UCustomizableObject& CustomizableObject, FCustomizableObjectComponentIndex ObjectComponentIndex);
 
-	bool BuildSkeletonData(const TSharedRef<FUpdateContextPrivate>& OperationData, USkeletalMesh& SkeletalMesh, const FMutableRefSkeletalMeshData& RefSkeletalMeshData, UCustomizableObject& CustomizableObject, int32 InstanceComponentIndex);
+	bool BuildSkeletonData(const TSharedRef<FUpdateContextPrivate>& OperationData, USkeletalMesh& SkeletalMesh, const FMutableRefSkeletalMeshData& RefSkeletalMeshData, UCustomizableObject& CustomizableObject, FCustomizableObjectInstanceComponentIndex InstanceComponentIndex);
 	void BuildMeshSockets(const TSharedRef<FUpdateContextPrivate>& OperationData, USkeletalMesh* SkeletalMesh, const FModelResources& ModelResources, const FMutableRefSkeletalMeshData& RefSkeletalMeshData, mu::Ptr<const mu::Mesh> MutableMesh);
-	void BuildOrCopyElementData(const TSharedRef<FUpdateContextPrivate>& OperationData, USkeletalMesh* SkeletalMesh, UCustomizableObjectInstance* CustomizableObjectInstance, int32 InstanceComponentIndex);
-	void BuildOrCopyMorphTargetsData(const TSharedRef<FUpdateContextPrivate>& OperationData, USkeletalMesh* SkeletalMesh, const USkeletalMesh* SrcSkeletalMesh, UCustomizableObjectInstance* CustomizableObjectInstance, int32 InstanceComponentIndex);
-	bool BuildOrCopyRenderData(const TSharedRef<FUpdateContextPrivate>& OperationData, USkeletalMesh* SkeletalMesh, const USkeletalMesh* SrcSkeletalMesh, UCustomizableObjectInstance* CustomizableObjectInstance, int32 InstanceComponentIndex);
-	void BuildOrCopyClothingData(const TSharedRef<FUpdateContextPrivate>& OperationData, USkeletalMesh* SkeletalMesh, const USkeletalMesh* SrcSkeletalMesh, UCustomizableObjectInstance* CustomizableObjectInstance, int32 InstanceComponentIndex);
-	
+	void BuildOrCopyElementData(const TSharedRef<FUpdateContextPrivate>& OperationData, USkeletalMesh* SkeletalMesh, UCustomizableObjectInstance* CustomizableObjectInstance, FCustomizableObjectInstanceComponentIndex InstanceComponentIndex);
+	void BuildOrCopyMorphTargetsData(const TSharedRef<FUpdateContextPrivate>& OperationData, USkeletalMesh* SkeletalMesh, const USkeletalMesh* SrcSkeletalMesh, UCustomizableObjectInstance* CustomizableObjectInstance, FCustomizableObjectInstanceComponentIndex InstanceComponentIndex);
+	bool BuildOrCopyRenderData(const TSharedRef<FUpdateContextPrivate>& OperationData, USkeletalMesh* SkeletalMesh, const USkeletalMesh* SrcSkeletalMesh, UCustomizableObjectInstance* CustomizableObjectInstance, FCustomizableObjectInstanceComponentIndex InstanceComponentIndex);
+	void BuildOrCopyClothingData(const TSharedRef<FUpdateContextPrivate>& OperationData, USkeletalMesh* SkeletalMesh, const USkeletalMesh* SrcSkeletalMesh, UCustomizableObjectInstance* CustomizableObjectInstance, FCustomizableObjectInstanceComponentIndex InstanceComponentIndex);
+
+	// 
+	FCustomizableInstanceComponentData* GetComponentData(FCustomizableObjectComponentIndex ObjectComponentIndex);
+	const FCustomizableInstanceComponentData* GetComponentData(FCustomizableObjectComponentIndex ObjectComponentIndex) const;
+
 	//
-	USkeleton* MergeSkeletons(UCustomizableObject& CustomizableObject, const FMutableRefSkeletalMeshData& RefSkeletalMeshData, int32 ObjectComponentIndex, bool& bOutCreatedNewSkeleton);
+	USkeleton* MergeSkeletons(UCustomizableObject& CustomizableObject, const FMutableRefSkeletalMeshData& RefSkeletalMeshData, FCustomizableObjectComponentIndex ObjectComponentIndex, bool& bOutCreatedNewSkeleton);
 
 	//
-	UPhysicsAsset* GetOrBuildMainPhysicsAsset(const TSharedRef<FUpdateContextPrivate>& OperationData, TObjectPtr<class UPhysicsAsset> TamplateAsset, const mu::PhysicsBody* PhysicsBody, bool bDisableCollisionBetweenAssets, int32 InstanceComponentIndex);
+	UPhysicsAsset* GetOrBuildMainPhysicsAsset(const TSharedRef<FUpdateContextPrivate>& OperationData, TObjectPtr<class UPhysicsAsset> TamplateAsset, const mu::PhysicsBody* PhysicsBody, bool bDisableCollisionBetweenAssets, FCustomizableObjectInstanceComponentIndex InstanceComponentIndex);
 	
 	// Create a transient texture and add it to the TextureTrackerArray
 	UTexture2D* CreateTexture();
@@ -317,7 +318,8 @@ private:
 
 	bool DoComponentsNeedUpdate(UCustomizableObjectInstance* CustomizableObjectInstance, const TSharedRef<FUpdateContextPrivate>& OperationData, bool& bOutEmptyMesh);
 
-	void SetLastMeshId(int32 ObjectComponentIndex, int32 LODIndex, mu::FResourceID MeshId);
+	// TODO: Why is this a method?
+	void SetLastMeshId(FCustomizableObjectComponentIndex ObjectComponentIndex, int32 LODIndex, mu::FResourceID MeshId);
 
 public:
 	bool LoadParametersFromProfile(int32 ProfileIndex);
diff --git a/Engine/Plugins/Mutable/Source/CustomizableObject/Internal/MuCO/CustomizableObjectPrivate.h b/Engine/Plugins/Mutable/Source/CustomizableObject/Internal/MuCO/CustomizableObjectPrivate.h
--- a/Engine/Plugins/Mutable/Source/CustomizableObject/Internal/MuCO/CustomizableObjectPrivate.h
+++ b/Engine/Plugins/Mutable/Source/CustomizableObject/Internal/MuCO/CustomizableObjectPrivate.h
@@ -1032,7 +1032,7 @@ struct FModelResources
 	 * the Reference Skeletal Meshes to avoid having them loaded at all times. This includes data like LOD distances,
 	 * LOD render data settings, Mesh sockets, Bounding volumes, etc.
 	 * 
-	 * Index with Component index
+	 * Index with CustomizableObject Component index
 	 */
 	UPROPERTY()
 	TArray<FMutableRefSkeletalMeshData> ReferenceSkeletalMeshesData;
@@ -1171,7 +1171,7 @@ struct FModelResources
 
 	/** Name of all possible components. Index is the ObjectComponentIndex. */
 	UPROPERTY()
-	TArray<FName> ComponentNames;
+	TArray<FName> ComponentNamesPerObjectComponent;
 
 	UPROPERTY()
 	FString ReleaseVersion;
diff --git a/Engine/Plugins/Mutable/Source/CustomizableObject/Internal/MuCO/CustomizableObjectSystemPrivate.h b/Engine/Plugins/Mutable/Source/CustomizableObject/Internal/MuCO/CustomizableObjectSystemPrivate.h
--- a/Engine/Plugins/Mutable/Source/CustomizableObject/Internal/MuCO/CustomizableObjectSystemPrivate.h
+++ b/Engine/Plugins/Mutable/Source/CustomizableObject/Internal/MuCO/CustomizableObjectSystemPrivate.h
@@ -48,6 +48,37 @@ struct FTexturePlatformData;
 void FinishUpdateGlobal(const TSharedRef<FUpdateContextPrivate>& Context);
 
 
+/** Strongly typed index for the index of a component in a UCustomizableObjectInstance. */
+class FCustomizableObjectInstanceComponentIndex
+{
+private:
+
+	int32 Index = 0;
+
+public:
+
+	explicit FCustomizableObjectInstanceComponentIndex(int32 InIndex = 0)
+		: Index(InIndex)
+	{
+	}
+
+	inline void Invalidate()
+	{
+		Index = INDEX_NONE;
+	}
+
+	inline bool IsValid() const
+	{
+		return Index != INDEX_NONE;
+	}
+
+	inline int32 GetValue() const
+	{
+		return Index;
+	}
+};
+
+
 class FMutableUpdateCandidate
 {
 public:
@@ -386,7 +417,7 @@ struct FInstanceUpdateData
 
 	struct FComponent
 	{
-		uint16 Id = 0;
+		FCustomizableObjectComponentIndex Id = FCustomizableObjectComponentIndex(0);
 		
 		/** Range in the LODs array */
 		uint16 FirstLOD = 0;
@@ -423,13 +454,15 @@ struct FInstanceUpdateData
 
 	struct FRealTimeMorphsComponentData
 	{
-		int32 ObjectComponentIndex = INDEX_NONE;
+		FCustomizableObjectComponentIndex ObjectComponentIndex = FCustomizableObjectComponentIndex(INDEX_NONE);
 
 		TArray<FName> RealTimeMorphTargetNames;
 		TArray<TArray<FMorphTargetLODModel>> RealTimeMorphsLODData; 
 	};
 
+	// Index in this array is not related to any other index.
 	TArray<FRealTimeMorphsComponentData> RealTimeMorphTargets;
+	
 
 	struct FClothingMeshData
 	{
@@ -450,9 +483,8 @@ struct FInstanceUpdateData
 		TMap<mu::FBoneName, TPair<FName, uint16>> BoneInfoMap;
 	};
 
-	// Access by object component index
-	// \TODO: somewhat wasteful, since instances may have less components
-	TArray<FSkeletonData> Skeletons;
+	// Access by instance component index
+	TArray<FSkeletonData> SkeletonsPerInstanceComponent;
 
 	struct FNamedExtensionData
 	{
@@ -471,7 +503,7 @@ struct FInstanceUpdateData
 		Scalars.Empty();
 		Vectors.Empty();
 		RealTimeMorphTargets.Empty();
-		Skeletons.Empty();
+		SkeletonsPerInstanceComponent.Empty();
 		ExtendedInputPins.Empty();
 	}
 };
@@ -494,6 +526,7 @@ public:
 
 	/** Return an array of LODs per object component. */
 	const TArray<uint16>& GetRequestedLODs() const;
+	uint16 GetRequestedLODs(FCustomizableObjectComponentIndex) const;
 
 	void SetRequestedLODs(TArray<uint16>& RequestedLODs);
 
@@ -512,6 +545,11 @@ public:
 	 *It is weak because we don't want to lock it in case it becomes irrelevant in the game while operations are pending and it needs to be destroyed. */
 	TWeakObjectPtr<UCustomizableObjectInstance> Instance;
 
+	/** Return the object component index associated with a component in this instance. */
+	FCustomizableObjectComponentIndex GetObjectComponentIndex(FCustomizableObjectInstanceComponentIndex InstanceComponentIndex);
+
+	FInstanceUpdateData::FRealTimeMorphsComponentData* GetComponentMorphTargets(FCustomizableObjectComponentIndex ObjectComponentIndex);
+
 private:
 	/** Descriptor which the update will be performed on. */
 	FCustomizableObjectInstanceDescriptor CapturedDescriptor;
@@ -531,6 +569,8 @@ public:
 	FInstanceUpdateData InstanceUpdateData;
 	TArray<int32> RelevantParametersInProgress;
 
+	const FInstanceUpdateData::FComponent* GetComponentUpdateData(FCustomizableObjectInstanceComponentIndex InstanceComponentIndex);
+
 	TArray<FString> LowPriorityTextures;
 
 	/** This option comes from the operation request */
@@ -560,10 +600,37 @@ public:
 	// TODO: Redundant to store it here because it is implicit in the size of some arrays or in the MutableInstance while it is valid.
 	uint8 NumInstanceComponents = 0;
 
+private:
+
 	// Index with ObjectComponent index
 	// \TODO: Wasteful: it only needs to be as big as NumInstanceComponents
 	TArray<uint8> NumLODsAvailablePerComponent;
 
+public:
+
+	inline void InitNumLODsAvailablePerComponent(int32 Size)
+	{
+		NumLODsAvailablePerComponent.SetNumZeroed(Size);
+	}
+
+	inline uint8 GetLODsAvailable(FCustomizableObjectComponentIndex Index) const
+	{
+		if (NumLODsAvailablePerComponent.IsValidIndex(Index.GetValue()))
+		{
+			return NumLODsAvailablePerComponent[Index.GetValue()];
+		}
+		return 0;
+	}
+
+	inline void SetNumLODsAvailable(FCustomizableObjectComponentIndex Index, uint8 LODCount)
+	{
+		if (NumLODsAvailablePerComponent.IsValidIndex(Index.GetValue()))
+		{
+			NumLODsAvailablePerComponent[Index.GetValue()] = LODCount;
+		}
+	}
+
+
 	uint8 FirstLODAvailable = 0;
 	uint8 FirstResidentLOD = 0;
 
@@ -573,13 +640,21 @@ public:
 
 	mu::FImageOperator::FImagePixelFormatFunc PixelFormatOverride;
 
+private:
+
 	/** Mutable Meshes required for each component. Outermost index is the object component index, inner index is the LOD. */
 	TArray<TArray<mu::FResourceID>> MeshDescriptors;
 
+public:
+
+	void InitMeshDescriptors(int32 Size);
+	const TArray<TArray<mu::FResourceID>>& GetMeshDescriptors() const;
+	TArray<mu::FResourceID>* GetMeshDescriptors(FCustomizableObjectComponentIndex Index);
+
 	/** Used to know if the updated instances' meshes are different from the previous ones. 
 	  * The index of the array is the instance component's index.
 	  * @return true if the mesh is new or new to this instance (e.g. mesh cached by another instance). */
-	TArray<bool> MeshChanged;
+	TArray<bool> MeshChangedPerInstanceComponent;
 	
 	bool UpdateStarted = false;
 	bool bLevelBegunPlay = false;
diff --git a/Engine/Plugins/Mutable/Source/CustomizableObject/Private/MuCO/CustomizableInstanceLODManagement.cpp b/Engine/Plugins/Mutable/Source/CustomizableObject/Private/MuCO/CustomizableInstanceLODManagement.cpp
--- a/Engine/Plugins/Mutable/Source/CustomizableObject/Private/MuCO/CustomizableInstanceLODManagement.cpp
+++ b/Engine/Plugins/Mutable/Source/CustomizableObject/Private/MuCO/CustomizableInstanceLODManagement.cpp
@@ -233,6 +233,10 @@ void UCustomizableInstanceLODManagement::UpdateInstanceDistsAndLODs(FMutableInst
 		}
 
 		TArray<uint16> RequestedLODs;
+		for (int32 i = NumGeneratedInstancesLimitLODs + NumGeneratedInstancesLimitLOD1 + NumGeneratedInstancesLimitLOD2; i < SortedInstances.Num(); ++i)
+		{
+			SortedInstances[i]->SetIsDiscardedBecauseOfTooManyInstances(true);
+		}
 
 		if (SortedInstances.Num() > NumGeneratedInstancesLimitLODs)
 		{
@@ -248,14 +252,16 @@ void UCustomizableInstanceLODManagement::UpdateInstanceDistsAndLODs(FMutableInst
 
 			for (int32 i = 0; i < NumGeneratedInstancesLimitLODs && i < SortedInstances.Num(); ++i)
 			{
+				const UCustomizableObject* CO = SortedInstances[i]->GetCustomizableObject();
+				int32 NumObjectComponents = CO ? CO->GetComponentCount() : 0;
+				RequestedLODs.Init(MAX_uint16, NumObjectComponents);
+
 				if (!bAlreadyReachedFixedLOD && SortedInstances[i]->GetMinSquareDistToPlayer() < DistanceForFixedLODSquared)
 				{
-					RequestedLODs.Init(MAX_uint16, SortedInstances[i]->GetNumComponents());
 					SortedInstances[i]->SetRequestedLODs(0, 0, RequestedLODs, InOutRequestedUpdates);
 				}
 				else
 				{
-					RequestedLODs.Init(MAX_uint16, SortedInstances[i]->GetNumComponents());
 					SortedInstances[i]->SetRequestedLODs(2, 0, RequestedLODs, InOutRequestedUpdates);
 					bAlreadyReachedFixedLOD = true;
 				}
@@ -263,14 +269,16 @@ void UCustomizableInstanceLODManagement::UpdateInstanceDistsAndLODs(FMutableInst
 
 			for (int32 i = NumGeneratedInstancesLimitLODs; i < NumGeneratedInstancesLimitLODs + NumGeneratedInstancesLimitLOD1 && i < SortedInstances.Num(); ++i)
 			{
+				const UCustomizableObject* CO = SortedInstances[i]->GetCustomizableObject();
+				int32 NumObjectComponents = CO ? CO->GetComponentCount() : 0;
+				RequestedLODs.Init(MAX_uint16, NumObjectComponents);
+
 				if (!bAlreadyReachedFixedLOD && SortedInstances[i]->GetMinSquareDistToPlayer() < DistanceForFixedLODSquared)
 				{
-					RequestedLODs.Init(MAX_uint16, SortedInstances[i]->GetNumComponents());
 					SortedInstances[i]->SetRequestedLODs(1, 0, RequestedLODs, InOutRequestedUpdates);
 				}
 				else
 				{
-					RequestedLODs.Init(MAX_uint16, SortedInstances[i]->GetNumComponents());
 					SortedInstances[i]->SetRequestedLODs(2, 0, RequestedLODs, InOutRequestedUpdates);
 					bAlreadyReachedFixedLOD = true;
 				}
@@ -278,7 +286,9 @@ void UCustomizableInstanceLODManagement::UpdateInstanceDistsAndLODs(FMutableInst
 
 			for (int32 i = NumGeneratedInstancesLimitLODs + NumGeneratedInstancesLimitLOD1; i < NumGeneratedInstancesLimitLODs + NumGeneratedInstancesLimitLOD1 + NumGeneratedInstancesLimitLOD2 && i < SortedInstances.Num(); ++i)
 			{
-				RequestedLODs.Init(MAX_uint16, SortedInstances[i]->GetNumComponents());
+				const UCustomizableObject* CO = SortedInstances[i]->GetCustomizableObject();
+				int32 NumObjectComponents = CO ? CO->GetComponentCount() : 0;
+				RequestedLODs.Init(MAX_uint16, NumObjectComponents);
 				SortedInstances[i]->SetRequestedLODs(2, 0, RequestedLODs, InOutRequestedUpdates);
 			}
 
@@ -292,7 +302,9 @@ void UCustomizableInstanceLODManagement::UpdateInstanceDistsAndLODs(FMutableInst
 			// No limit surpassed, set all instances to have all LODs, there will be an UpdateSkeletalMesh only if there's a change in LOD state
 			for (int32 i = 0; i < SortedInstances.Num(); ++i)
 			{
-				RequestedLODs.Init(MAX_uint16, SortedInstances[i]->GetNumComponents());
+				const UCustomizableObject* CO = SortedInstances[i]->GetCustomizableObject();
+				int32 NumObjectComponents = CO ? CO->GetComponentCount() : 0;
+				RequestedLODs.Init(MAX_uint16, NumObjectComponents);
 				SortedInstances[i]->SetRequestedLODs(2, 0, RequestedLODs, InOutRequestedUpdates);
 			}
 		}
@@ -370,7 +382,7 @@ void UCustomizableInstanceLODManagement::UpdateInstanceDistsAndLODs(FMutableInst
 					COI->SetIsBeingUsedByComponentInPlay(true);
 
 					const FName& ComponentName = CustomizableObjectInstanceUsage->GetComponentName();
-					const int32 ObjectComponentIndex = COI->GetCustomizableObject()->GetPrivate()->GetModelResources().ComponentNames.IndexOfByKey(ComponentName);
+					const int32 ObjectComponentIndex = COI->GetCustomizableObject()->GetPrivate()->GetModelResources().ComponentNamesPerObjectComponent.IndexOfByKey(ComponentName);
 
 #if WITH_EDITOR
 					// If the instance is generated but the component doesn't have a mesh, set it.
diff --git a/Engine/Plugins/Mutable/Source/CustomizableObject/Private/MuCO/CustomizableObject.cpp b/Engine/Plugins/Mutable/Source/CustomizableObject/Private/MuCO/CustomizableObject.cpp
--- a/Engine/Plugins/Mutable/Source/CustomizableObject/Private/MuCO/CustomizableObject.cpp
+++ b/Engine/Plugins/Mutable/Source/CustomizableObject/Private/MuCO/CustomizableObject.cpp
@@ -782,7 +782,7 @@ void FModelResources::Serialize(FObjectAndNameAsStringProxyArchive& MemoryWriter
 	MemoryWriter << NumLODsToStream;
 	MemoryWriter << FirstLODAvailable;
 
-	MemoryWriter << ComponentNames;
+	MemoryWriter << ComponentNamesPerObjectComponent;
 	MemoryWriter << ReleaseVersion;
 
 	// Editor Only data
@@ -936,7 +936,7 @@ bool FModelResources::Unserialize(FObjectAndNameAsStringProxyArchive& MemoryRead
 	MemoryReader << NumLODsToStream;
 	MemoryReader << FirstLODAvailable;
 
-	MemoryReader << ComponentNames;
+	MemoryReader << ComponentNamesPerObjectComponent;
 	MemoryReader << ReleaseVersion;
 
 	// Editor Only data
@@ -1294,7 +1294,7 @@ USkeletalMesh* UCustomizableObject::GetComponentMeshReferenceSkeletalMesh(const
 	}
 #else
 	const FModelResources& ModelResources = Private->GetModelResources();
-	int32 ObjectComponentIndex = ModelResources.ComponentNames.IndexOfByKey(ComponentName);
+	int32 ObjectComponentIndex = ModelResources.ComponentNamesPerObjectComponent.IndexOfByKey(ComponentName);
 	if (ModelResources.ReferenceSkeletalMeshesData.IsValidIndex(ObjectComponentIndex))
 	{
 		// Can be nullptr if RefSkeletalMeshes are not loaded yet.
@@ -1463,21 +1463,21 @@ int32 UCustomizableObject::GetComponentCount() const
 {
 	if (IsCompiled())
 	{
-		return GetPrivate()->GetModelResources().ComponentNames.Num();
+		return GetPrivate()->GetModelResources().ComponentNamesPerObjectComponent.Num();
 	}
 
 	return 0;
 }
 
 
-FName UCustomizableObject::GetComponentName(int32 ObjectComponentIndex) const
+FName UCustomizableObject::GetComponentName(FCustomizableObjectComponentIndex ObjectComponentIndex) const
 {
 	if (IsCompiled())
 	{
-		const TArray<FName>& ComponentNames = GetPrivate()->GetModelResources().ComponentNames;
-		if (ComponentNames.IsValidIndex(ObjectComponentIndex))
+		const TArray<FName>& ComponentNames = GetPrivate()->GetModelResources().ComponentNamesPerObjectComponent;
+		if (ComponentNames.IsValidIndex(ObjectComponentIndex.GetValue()))
 		{
-			return ComponentNames[ObjectComponentIndex];
+			return ComponentNames[ObjectComponentIndex.GetValue()];
 		}
 	}
 
diff --git a/Engine/Plugins/Mutable/Source/CustomizableObject/Private/MuCO/CustomizableObjectInstance.cpp b/Engine/Plugins/Mutable/Source/CustomizableObject/Private/MuCO/CustomizableObjectInstance.cpp
--- a/Engine/Plugins/Mutable/Source/CustomizableObject/Private/MuCO/CustomizableObjectInstance.cpp
+++ b/Engine/Plugins/Mutable/Source/CustomizableObject/Private/MuCO/CustomizableObjectInstance.cpp
@@ -201,7 +201,7 @@ UTexture2D* UCustomizableInstancePrivate::CreateTexture()
 }
 
 
-void UCustomizableInstancePrivate::SetLastMeshId(int32 ObjectComponentIndex, int32 LODIndex, mu::FResourceID MeshId)
+void UCustomizableInstancePrivate::SetLastMeshId(FCustomizableObjectComponentIndex ObjectComponentIndex, int32 LODIndex, mu::FResourceID MeshId)
 {
 	FCustomizableInstanceComponentData* ComponentData = GetComponentData(ObjectComponentIndex);
 	if (ComponentData && ComponentData->LastMeshIdPerLOD.IsValidIndex(LODIndex))
@@ -264,30 +264,30 @@ FCustomizableInstanceComponentData* UCustomizableInstancePrivate::GetComponentDa
 		return nullptr;
 	}
 
-	const int32 ComponentIndex = Object->GetPrivate()->GetModelResources().ComponentNames.IndexOfByKey(ComponentName);
-	if (ComponentIndex == INDEX_NONE)
+	const int32 ObjectComponentIndex = Object->GetPrivate()->GetModelResources().ComponentNamesPerObjectComponent.IndexOfByKey(ComponentName);
+	if (ObjectComponentIndex == INDEX_NONE)
 	{
 		return nullptr;
 	}
 
-	if (!ComponentsData.IsValidIndex(ComponentIndex))
+	if (!ComponentsData.IsValidIndex(ObjectComponentIndex))
 	{
 		return nullptr;	
 	}
 
-	return &ComponentsData[ComponentIndex];
+	return &ComponentsData[ObjectComponentIndex];
 }
 
 
-FCustomizableInstanceComponentData* UCustomizableInstancePrivate::GetComponentData(int32 ComponentIndex)
+FCustomizableInstanceComponentData* UCustomizableInstancePrivate::GetComponentData(FCustomizableObjectComponentIndex ObjectComponentIndex)
 {
-	return ComponentsData.IsValidIndex(ComponentIndex) ? &ComponentsData[ComponentIndex] : nullptr;
+	return ComponentsData.IsValidIndex(ObjectComponentIndex.GetValue()) ? &ComponentsData[ObjectComponentIndex.GetValue()] : nullptr;
 }
 
 
-const FCustomizableInstanceComponentData* UCustomizableInstancePrivate::GetComponentData(int32 ComponentIndex) const
+const FCustomizableInstanceComponentData* UCustomizableInstancePrivate::GetComponentData(FCustomizableObjectComponentIndex ObjectComponentIndex) const
 {
-	return ComponentsData.IsValidIndex(ComponentIndex) ? &ComponentsData[ComponentIndex] : nullptr;
+	return ComponentsData.IsValidIndex(ObjectComponentIndex.GetValue()) ? &ComponentsData[ObjectComponentIndex.GetValue()] : nullptr;
 }
 
 
@@ -1248,7 +1248,7 @@ bool AreSkeletonsCompatible(const TArray<TObjectPtr<USkeleton>>& InSkeletons)
 #endif
 
 
-USkeleton* UCustomizableInstancePrivate::MergeSkeletons(UCustomizableObject& CustomizableObject, const FMutableRefSkeletalMeshData& RefSkeletalMeshData, int32 ObjectComponentIndex, bool& bOutCreatedNewSkeleton)
+USkeleton* UCustomizableInstancePrivate::MergeSkeletons(UCustomizableObject& CustomizableObject, const FMutableRefSkeletalMeshData& RefSkeletalMeshData, FCustomizableObjectComponentIndex ObjectComponentIndex, bool& bOutCreatedNewSkeleton)
 {
 	MUTABLE_CPUPROFILER_SCOPE(BuildSkeletonData_MergeSkeletons);
 	bOutCreatedNewSkeleton = false;
@@ -1476,7 +1476,7 @@ namespace
 		const TSharedRef<FUpdateContextPrivate>& OperationData,
 		TObjectPtr<UPhysicsAsset> TemplateAsset,
 		const mu::PhysicsBody* MutablePhysics,
-		int32 InstanceComponentIndex)
+		FCustomizableObjectInstanceComponentIndex InstanceComponentIndex)
 	{
 		check(TemplateAsset);
 		TObjectPtr<UPhysicsAsset> Result = NewObject<UPhysicsAsset>();
@@ -1491,7 +1491,7 @@ namespace
 
 		Result->bNotForDedicatedServer = TemplateAsset->bNotForDedicatedServer;
 
-		const TMap<mu::FBoneName, TPair<FName, uint16>>& BoneInfoMap = OperationData->InstanceUpdateData.Skeletons[InstanceComponentIndex].BoneInfoMap;
+		const TMap<mu::FBoneName, TPair<FName, uint16>>& BoneInfoMap = OperationData->InstanceUpdateData.SkeletonsPerInstanceComponent[InstanceComponentIndex.GetValue()].BoneInfoMap;
 		TMap<FName, int32> BonesInUse;
 
 		const int32 MutablePhysicsBodyCount = MutablePhysics->GetBodyCount();
@@ -1604,7 +1604,7 @@ UPhysicsAsset* UCustomizableInstancePrivate::GetOrBuildMainPhysicsAsset(
 	TObjectPtr<UPhysicsAsset> TemplateAsset,
 	const mu::PhysicsBody* MutablePhysics,
 	bool bDisableCollisionsBetweenDifferentAssets,
-	int32 InstanceComponentIndex)
+	FCustomizableObjectInstanceComponentIndex InstanceComponentIndex)
 {
 
 	MUTABLE_CPUPROFILER_SCOPE(MergePhysicsAssets);
@@ -1613,9 +1613,13 @@ UPhysicsAsset* UCustomizableInstancePrivate::GetOrBuildMainPhysicsAsset(
 
 	UPhysicsAsset* Result = nullptr;
 
-	const FInstanceUpdateData::FComponent& Component = OperationData->InstanceUpdateData.Components[InstanceComponentIndex];
-	int32 ObjectComponentIndex = Component.Id;
+	const FInstanceUpdateData::FComponent* Component = OperationData->GetComponentUpdateData(InstanceComponentIndex);
+	if (!Component)
+	{
+		return nullptr;
+	}
 
+	FCustomizableObjectComponentIndex ObjectComponentIndex = Component->Id;
 	FCustomizableInstanceComponentData* ComponentData = GetComponentData(ObjectComponentIndex);
 	check(ComponentData);
 
@@ -1660,7 +1664,7 @@ UPhysicsAsset* UCustomizableInstancePrivate::GetOrBuildMainPhysicsAsset(
 
 	Result->bNotForDedicatedServer = TemplateAsset->bNotForDedicatedServer;
 
-	const TMap<mu::FBoneName, TPair<FName, uint16>>& BoneInfoMap = OperationData->InstanceUpdateData.Skeletons[InstanceComponentIndex].BoneInfoMap;
+	const TMap<mu::FBoneName, TPair<FName, uint16>>& BoneInfoMap = OperationData->InstanceUpdateData.SkeletonsPerInstanceComponent[InstanceComponentIndex.GetValue()].BoneInfoMap;
 	TMap<FName, int32> BonesInUse;
 
 	const int32 MutablePhysicsBodyCount = MutablePhysics->GetBodyCount();
@@ -1970,22 +1974,22 @@ bool UCustomizableInstancePrivate::DoComponentsNeedUpdate(UCustomizableObjectIns
 	}
 
 	// Find which components need an update
-	OperationData->MeshChanged.Init(false, NumInstanceComponents);
+	OperationData->MeshChangedPerInstanceComponent.Init(false, NumInstanceComponents);
 
 	for (int32 InstanceComponentIndex = 0; InstanceComponentIndex < NumInstanceComponents; ++InstanceComponentIndex)
 	{
 		const FInstanceUpdateData::FComponent& Component = OperationData->InstanceUpdateData.Components[InstanceComponentIndex];
-		const int32 ObjectComponentIndex = Component.Id;
-		const FName ComponentName = OperationData->Instance->GetCustomizableObject()->GetPrivate()->GetModelResources().ComponentNames[ObjectComponentIndex];
+		const FCustomizableObjectComponentIndex ObjectComponentIndex = Component.Id;
+		const FName ComponentName = OperationData->Instance->GetCustomizableObject()->GetPrivate()->GetModelResources().ComponentNamesPerObjectComponent[ObjectComponentIndex.GetValue()];
 
 		if (OperationData->bUseMeshCache)
 		{
-			USkeletalMesh* CachedMesh = CustomizableObject->GetPrivate()->MeshCache.Get(OperationData->MeshDescriptors[ObjectComponentIndex]);
+			USkeletalMesh* CachedMesh = CustomizableObject->GetPrivate()->MeshCache.Get(*OperationData->GetMeshDescriptors(ObjectComponentIndex));
 			if (CachedMesh)
 			{
 				TObjectPtr<USkeletalMesh>* SkeletalMesh = SkeletalMeshes.Find(ComponentName);
 				const bool bMeshNeedsUpdate = !SkeletalMesh || (*SkeletalMesh != CachedMesh);
-				OperationData->MeshChanged[InstanceComponentIndex] = bMeshNeedsUpdate;
+				OperationData->MeshChangedPerInstanceComponent[InstanceComponentIndex] = bMeshNeedsUpdate;
 				ComponentWithMesh[InstanceComponentIndex] = true;
 				continue;
 			}
@@ -1999,7 +2003,7 @@ bool UCustomizableInstancePrivate::DoComponentsNeedUpdate(UCustomizableObjectIns
 		if (Component.LODCount == 0)
 		{
 			// We don't have a mesh in the component, so it has changed if we had one before.
-			OperationData->MeshChanged[InstanceComponentIndex] = bHadSkeletalMesh;
+			OperationData->MeshChangedPerInstanceComponent[InstanceComponentIndex] = bHadSkeletalMesh;
 			continue;
 		}
 
@@ -2007,18 +2011,18 @@ bool UCustomizableInstancePrivate::DoComponentsNeedUpdate(UCustomizableObjectIns
 		bool bIsReferenced = LOD.Mesh && LOD.Mesh->IsReference();
 		if (!bIsReferenced)
 		{
-			if (ComponentWithMesh[InstanceComponentIndex] && MeshIDs[InstanceComponentIndex * MAX_MESH_LOD_COUNT + OperationData->NumLODsAvailablePerComponent[ObjectComponentIndex] - 1] == MAX_uint64)
+			if (ComponentWithMesh[InstanceComponentIndex] && MeshIDs[InstanceComponentIndex * MAX_MESH_LOD_COUNT + OperationData->GetLODsAvailable(ObjectComponentIndex) - 1] == MAX_uint64)
 			{
 				UE_LOG(LogMutable, Error, TEXT("Failed to generate SkeletalMesh for CO Instance [%s]. CO [%s] is missing geometry for LOD [%d] Object Component [%d]."),
 					*Public->GetName(), *CustomizableObject->GetName(),
-					OperationData->NumLODsAvailablePerComponent[ObjectComponentIndex] - 1, ObjectComponentIndex);
+					OperationData->GetLODsAvailable(ObjectComponentIndex) - 1, ObjectComponentIndex.GetValue());
 				bHasInvalidMesh = true;
 				continue;
 			}
 		}
 
 		// If the component wasn't there and now is there, we need to update it.
-		OperationData->MeshChanged[InstanceComponentIndex] = !bHadSkeletalMesh && LOD.Mesh && (LOD.Mesh->GetFaceCount()>0);
+		OperationData->MeshChangedPerInstanceComponent[InstanceComponentIndex] = !bHadSkeletalMesh && LOD.Mesh && (LOD.Mesh->GetFaceCount()>0);
 
 		const FCustomizableInstanceComponentData* ComponentData = GetComponentData(ObjectComponentIndex);
 		if (!ComponentData) // Could be nullptr if the component has not been generated.
@@ -2028,16 +2032,16 @@ bool UCustomizableInstancePrivate::DoComponentsNeedUpdate(UCustomizableObjectIns
 
 		// Update if MeshIDs are different
 		const int32 ComponentOffset = InstanceComponentIndex * MAX_MESH_LOD_COUNT;
-		for (int32 MeshIndex = 0; !OperationData->MeshChanged[InstanceComponentIndex] && MeshIndex < MAX_MESH_LOD_COUNT; ++MeshIndex)
+		for (int32 MeshIndex = 0; !OperationData->MeshChangedPerInstanceComponent[InstanceComponentIndex] && MeshIndex < MAX_MESH_LOD_COUNT; ++MeshIndex)
 		{
-			OperationData->MeshChanged[InstanceComponentIndex] = MeshIDs[ComponentOffset + MeshIndex] != ComponentData->LastMeshIdPerLOD[MeshIndex];
+			OperationData->MeshChangedPerInstanceComponent[InstanceComponentIndex] = MeshIDs[ComponentOffset + MeshIndex] != ComponentData->LastMeshIdPerLOD[MeshIndex];
 		}
 	}
 
 	bool bChanged =
-		OperationData->MeshChanged.Num() != SkeletalMeshes.Num()
+		OperationData->MeshChangedPerInstanceComponent.Num() != SkeletalMeshes.Num()
 		||
-		OperationData->MeshChanged.Find(true) != INDEX_NONE;
+		OperationData->MeshChangedPerInstanceComponent.Find(true) != INDEX_NONE;
 
 	// It also changed if we removed a component that we did have before
 	if (!bChanged)
@@ -2048,8 +2052,8 @@ bool UCustomizableInstancePrivate::DoComponentsNeedUpdate(UCustomizableObjectIns
 			for (int32 InstanceComponentIndex = 0; InstanceComponentIndex < NumInstanceComponents; ++InstanceComponentIndex)
 			{
 				const FInstanceUpdateData::FComponent& Component = OperationData->InstanceUpdateData.Components[InstanceComponentIndex];
-				const int32 ObjectComponentIndex = Component.Id;
-				const FName ComponentName = OperationData->Instance->GetCustomizableObject()->GetPrivate()->GetModelResources().ComponentNames[ObjectComponentIndex];
+				const FCustomizableObjectComponentIndex ObjectComponentIndex = Component.Id;
+				const FName ComponentName = OperationData->Instance->GetCustomizableObject()->GetPrivate()->GetModelResources().ComponentNamesPerObjectComponent[ObjectComponentIndex.GetValue()];
 
 				if (ComponentName == OldMesh.Key)
 				{
@@ -2230,13 +2234,19 @@ bool UCustomizableInstancePrivate::UpdateSkeletalMesh_PostBeginUpdate0(UCustomiz
 	SkeletalMeshes.Reset();
 	
 	const int32 NumInstanceComponents = OperationData->InstanceUpdateData.Components.Num();
-	for (int32 InstanceComponentIndex = 0; InstanceComponentIndex < NumInstanceComponents; ++InstanceComponentIndex)
+	for (int32 InstanceComponentIndexValue = 0; InstanceComponentIndexValue < NumInstanceComponents; ++InstanceComponentIndexValue)
 	{
-		const FInstanceUpdateData::FComponent& Component = OperationData->InstanceUpdateData.Components[InstanceComponentIndex];
-		const int32 ObjectComponentIndex = Component.Id;
-		const FName ComponentName = OperationData->Instance->GetCustomizableObject()->GetPrivate()->GetModelResources().ComponentNames[ObjectComponentIndex];
+		FCustomizableObjectInstanceComponentIndex InstanceComponentIndex(InstanceComponentIndexValue);
+		const FInstanceUpdateData::FComponent* Component = OperationData->GetComponentUpdateData( InstanceComponentIndex );
+		if (!Component)
+		{
+			continue;
+		}
+		const FCustomizableObjectComponentIndex ObjectComponentIndex = Component->Id;
+		const FName ComponentName = OperationData->Instance->GetCustomizableObject()->GetPrivate()->GetModelResources().ComponentNamesPerObjectComponent[ObjectComponentIndex.GetValue()];
 
-		if (!ComponentsData.IsValidIndex(ObjectComponentIndex))
+		FCustomizableInstanceComponentData* ComponentData = GetComponentData(ComponentName);
+		if (!ComponentData)
 		{
 			ensure(false);
 			
@@ -2245,7 +2255,7 @@ bool UCustomizableInstancePrivate::UpdateSkeletalMesh_PostBeginUpdate0(UCustomiz
 		}
 		
 		// If the component doesn't need an update copy the previously generated mesh.
-		if (!OperationData->MeshChanged[InstanceComponentIndex])
+		if (!OperationData->MeshChangedPerInstanceComponent[InstanceComponentIndex.GetValue()])
 		{
 			if (TObjectPtr<USkeletalMesh>* Result = OldSkeletalMeshes.Find(ComponentName))
 			{
@@ -2257,21 +2267,27 @@ bool UCustomizableInstancePrivate::UpdateSkeletalMesh_PostBeginUpdate0(UCustomiz
 		
 		if (OperationData->bUseMeshCache)
 		{
-			if (USkeletalMesh* CachedMesh = CustomizableObject->GetPrivate()->MeshCache.Get(OperationData->MeshDescriptors[ObjectComponentIndex]))
+			TArray<mu::FResourceID>* MeshDescriptors = OperationData->GetMeshDescriptors(ObjectComponentIndex);
+			if (!MeshDescriptors)
 			{
-				check(OperationData->MeshDescriptors[ObjectComponentIndex].Num() == MAX_MESH_LOD_COUNT);
-				ComponentsData[ObjectComponentIndex].LastMeshIdPerLOD = OperationData->MeshDescriptors[ObjectComponentIndex];
+				continue;
+			}
+
+			if (USkeletalMesh* CachedMesh = CustomizableObject->GetPrivate()->MeshCache.Get(*MeshDescriptors))
+			{
+				check(MeshDescriptors->Num() == MAX_MESH_LOD_COUNT);
+				ComponentData->LastMeshIdPerLOD = *MeshDescriptors;
 				SkeletalMeshes.Add(ComponentName, CachedMesh);
 				continue;
 			}
 		}
 
 		// Reset last mesh IDs.
-		ComponentsData[ObjectComponentIndex].LastMeshIdPerLOD.Init(MAX_uint64, MAX_MESH_LOD_COUNT);
+		ComponentData->LastMeshIdPerLOD.Init(MAX_uint64, MAX_MESH_LOD_COUNT);
 
 		// We need the first valid mesh. get it from the component, considering that some LOSs may have been skipped.
 		mu::Ptr<const mu::Mesh> ComponentMesh;
-		for (int32 FirstValidLODIndex = Component.FirstLOD; FirstValidLODIndex < OperationData->InstanceUpdateData.LODs.Num() && !ComponentMesh; ++FirstValidLODIndex)
+		for (int32 FirstValidLODIndex = Component->FirstLOD; FirstValidLODIndex < OperationData->InstanceUpdateData.LODs.Num() && !ComponentMesh; ++FirstValidLODIndex)
 		{
 			ComponentMesh = OperationData->InstanceUpdateData.LODs[FirstValidLODIndex].Mesh;
 		}
@@ -2297,7 +2313,7 @@ bool UCustomizableInstancePrivate::UpdateSkeletalMesh_PostBeginUpdate0(UCustomiz
         	continue;
         }
 
-		if (!ModelResources.ReferenceSkeletalMeshesData.IsValidIndex(ObjectComponentIndex))
+		if (!ModelResources.ReferenceSkeletalMeshesData.IsValidIndex(ObjectComponentIndex.GetValue()))
 		{
 			InvalidateGeneratedData();
 			return false;
@@ -2310,7 +2326,7 @@ bool UCustomizableInstancePrivate::UpdateSkeletalMesh_PostBeginUpdate0(UCustomiz
 
 		if (OperationData->bStreamMeshLODs)
 		{
-			SkeletalMesh =UCustomizableObjectSkeletalMesh::CreateSkeletalMesh(OperationData, *Public, *CustomizableObject, InstanceComponentIndex);
+			SkeletalMesh = UCustomizableObjectSkeletalMesh::CreateSkeletalMesh(OperationData, *Public, *CustomizableObject, Component->FirstLOD, Component->LODCount);
 		}
 		else
 		{
@@ -2321,13 +2337,13 @@ bool UCustomizableInstancePrivate::UpdateSkeletalMesh_PostBeginUpdate0(UCustomiz
 		check(SkeletalMesh);		
 		SkeletalMeshes.Add(ComponentName, SkeletalMesh);
 		
-		const FMutableRefSkeletalMeshData& RefSkeletalMeshData = ModelResources.ReferenceSkeletalMeshesData[ObjectComponentIndex];
+		const FMutableRefSkeletalMeshData& RefSkeletalMeshData = ModelResources.ReferenceSkeletalMeshesData[ObjectComponentIndex.GetValue()];
 
 		// Set up the default information any mesh from this component will have (LODArrayInfos, RenderData, Mesh settings, etc). 
 		InitSkeletalMeshData(OperationData, SkeletalMesh, RefSkeletalMeshData, *CustomizableObject, ObjectComponentIndex);
 		
 		// Construct a new skeleton, fix up ActiveBones and Bonemap arrays and recompute the RefInvMatrices
-		const bool bBuildSkeletonDataSuccess = BuildSkeletonData(OperationData, *SkeletalMesh, RefSkeletalMeshData, *CustomizableObject, InstanceComponentIndex);
+		const bool bBuildSkeletonDataSuccess = BuildSkeletonData(OperationData, *SkeletalMesh, RefSkeletalMeshData, *CustomizableObject, FCustomizableObjectInstanceComponentIndex(InstanceComponentIndex) );
 		if (!bBuildSkeletonDataSuccess)
 		{
 			InvalidateGeneratedData();
@@ -2339,7 +2355,7 @@ bool UCustomizableInstancePrivate::UpdateSkeletalMesh_PostBeginUpdate0(UCustomiz
 		{
 			constexpr bool bDisallowCollisionBetweenAssets = true;
 			UPhysicsAsset* PhysicsAssetResult = GetOrBuildMainPhysicsAsset(
-				OperationData, RefSkeletalMeshData.PhysicsAsset, MutablePhysics.get(), bDisallowCollisionBetweenAssets, InstanceComponentIndex);
+				OperationData, RefSkeletalMeshData.PhysicsAsset, MutablePhysics.get(), bDisallowCollisionBetweenAssets, FCustomizableObjectInstanceComponentIndex(InstanceComponentIndex) );
 
 			SkeletalMesh->SetPhysicsAsset(PhysicsAssetResult);
 
@@ -2395,7 +2411,7 @@ bool UCustomizableInstancePrivate::UpdateSkeletalMesh_PostBeginUpdate0(UCustomiz
 		for (const TPair<const UCustomizableObjectExtension*, TArray<FInputPinDataContainer>>& Pair : ExtensionToExtensionData)
 		{
 			PRAGMA_DISABLE_DEPRECATION_WARNINGS
-			Pair.Key->OnSkeletalMeshCreated(Pair.Value, ObjectComponentIndex, SkeletalMesh);
+			Pair.Key->OnSkeletalMeshCreated(Pair.Value, ObjectComponentIndex.GetValue(), SkeletalMesh);
 			PRAGMA_ENABLE_DEPRECATION_WARNINGS
 
 			Pair.Key->OnSkeletalMeshCreated(Pair.Value, ComponentName, SkeletalMesh);
@@ -2427,8 +2443,11 @@ bool UCustomizableInstancePrivate::UpdateSkeletalMesh_PostBeginUpdate0(UCustomiz
 
 		if (OperationData->bUseMeshCache)
 		{
-			const TArray<mu::FResourceID>& MeshId = OperationData->MeshDescriptors[ObjectComponentIndex];
-			CustomizableObject->GetPrivate()->MeshCache.Add(MeshId, SkeletalMesh);
+			const TArray<mu::FResourceID>* MeshId = OperationData->GetMeshDescriptors(ObjectComponentIndex);
+			if (MeshId)
+			{
+				CustomizableObject->GetPrivate()->MeshCache.Add(*MeshId, SkeletalMesh);
+			}
 		}
 	}
 
@@ -3047,7 +3066,7 @@ void UCustomizableInstancePrivate::SetDefaultSkeletalMesh(bool bSetEmptyMesh) co
 #endif
 
 		const FName& ComponentName = CustomizableObjectInstanceUsage->GetComponentName();
-		const int32 ObjectComponentIndex = Instance->GetCustomizableObject()->GetPrivate()->GetModelResources().ComponentNames.IndexOfByKey(ComponentName);
+		const int32 ObjectComponentIndex = Instance->GetCustomizableObject()->GetPrivate()->GetModelResources().ComponentNamesPerObjectComponent.IndexOfByKey(ComponentName);
 		if (!ModelResources.ReferenceSkeletalMeshesData.IsValidIndex(ObjectComponentIndex))
 		{
 			continue;
@@ -3509,7 +3528,7 @@ void ConvertImage(UTexture2D* Texture, mu::Ptr<const mu::Image> MutableImage, co
 }
 
 
-void UCustomizableInstancePrivate::InitSkeletalMeshData(const TSharedRef<FUpdateContextPrivate>& OperationData, USkeletalMesh* SkeletalMesh, const FMutableRefSkeletalMeshData& RefSkeletalMeshData, const UCustomizableObject& CustomizableObject, int32 ObjectComponentIndex)
+void UCustomizableInstancePrivate::InitSkeletalMeshData(const TSharedRef<FUpdateContextPrivate>& OperationData, USkeletalMesh* SkeletalMesh, const FMutableRefSkeletalMeshData& RefSkeletalMeshData, const UCustomizableObject& CustomizableObject, FCustomizableObjectComponentIndex ObjectComponentIndex)
 {
 	MUTABLE_CPUPROFILER_SCOPE(UCustomizableInstancePrivate::InitSkeletalMesh);
 
@@ -3547,20 +3566,21 @@ void UCustomizableInstancePrivate::InitSkeletalMeshData(const TSharedRef<FUpdate
 			}
 		}
 
-		const FCustomizableInstanceComponentData& ComponentData = ComponentsData[ObjectComponentIndex];
-		for (TObjectPtr<UAssetUserData> AssetUserData : ComponentData.AssetUserDataArray)
+		const FCustomizableInstanceComponentData* ComponentData = GetComponentData(ObjectComponentIndex);
+		check(ComponentData);
+		for (TObjectPtr<UAssetUserData> AssetUserData : ComponentData->AssetUserDataArray)
 		{
 			SkeletalMesh->AddAssetUserData(AssetUserData);
 		}
 
 		//Custom Asset User Data
 		if (OperationData->Instance->GetAnimationGameplayTags().Num() ||
-			ComponentData.AnimSlotToBP.Num())
+			ComponentData->AnimSlotToBP.Num())
 		{
 			UCustomizableObjectInstanceUserData* InstanceData = NewObject<UCustomizableObjectInstanceUserData>(SkeletalMesh, NAME_None, RF_Public | RF_Transactional);
 			InstanceData->AnimationGameplayTag = OperationData->Instance->GetAnimationGameplayTags();
 			
-			for (const TTuple<FName, TSoftClassPtr<UAnimInstance>>& AnimSlot : ComponentData.AnimSlotToBP)
+			for (const TTuple<FName, TSoftClassPtr<UAnimInstance>>& AnimSlot : ComponentData->AnimSlotToBP)
 			{
 				FCustomizableObjectAnimationSlot AnimationSlot;
 				AnimationSlot.Name = AnimSlot.Key;
@@ -3579,12 +3599,13 @@ void UCustomizableInstancePrivate::InitSkeletalMeshData(const TSharedRef<FUpdate
 		SkeletalMesh->AllocateResourceForRendering();
 
 		FSkeletalMeshRenderData* RenderData = SkeletalMesh->GetResourceForRendering();
-		RenderData->NumInlinedLODs = OperationData->NumLODsAvailablePerComponent[ObjectComponentIndex] - OperationData->FirstResidentLOD;
-		RenderData->NumNonOptionalLODs = OperationData->NumLODsAvailablePerComponent[ObjectComponentIndex] - OperationData->FirstLODAvailable;
+		int32 NumLODsAvailablePerComponent = OperationData->GetLODsAvailable(ObjectComponentIndex);
+		RenderData->NumInlinedLODs = NumLODsAvailablePerComponent - OperationData->FirstResidentLOD;
+		RenderData->NumNonOptionalLODs = NumLODsAvailablePerComponent - OperationData->FirstLODAvailable;
 		RenderData->CurrentFirstLODIdx = OperationData->FirstResidentLOD;
 		RenderData->LODBiasModifier = 0;
 
-		for (int32 LODIndex = 0; LODIndex < OperationData->NumLODsAvailablePerComponent[ObjectComponentIndex]; ++LODIndex)
+		for (int32 LODIndex = 0; LODIndex < NumLODsAvailablePerComponent; ++LODIndex)
 		{
 			RenderData->LODRenderData.Add(new FSkeletalMeshLODRenderData());
 			
@@ -3656,11 +3677,11 @@ void UCustomizableInstancePrivate::InitSkeletalMeshData(const TSharedRef<FUpdate
 }
 
 
-bool UCustomizableInstancePrivate::BuildSkeletonData(const TSharedRef<FUpdateContextPrivate>& OperationData, USkeletalMesh& SkeletalMesh, const FMutableRefSkeletalMeshData& RefSkeletalMeshData, UCustomizableObject& CustomizableObject, int32 InstanceComponentIndex)
+bool UCustomizableInstancePrivate::BuildSkeletonData(const TSharedRef<FUpdateContextPrivate>& OperationData, USkeletalMesh& SkeletalMesh, const FMutableRefSkeletalMeshData& RefSkeletalMeshData, UCustomizableObject& CustomizableObject, FCustomizableObjectInstanceComponentIndex InstanceComponentIndex)
 {
 	MUTABLE_CPUPROFILER_SCOPE(UCustomizableInstancePrivate::BuildSkeletonData);
 
-	int32 ObjectComponentIndex = OperationData->InstanceUpdateData.Components[InstanceComponentIndex].Id;
+	FCustomizableObjectComponentIndex ObjectComponentIndex = OperationData->GetObjectComponentIndex(InstanceComponentIndex);
 
 	bool bCreatedNewSkeleton = false;
 	const TObjectPtr<USkeleton> Skeleton = MergeSkeletons(CustomizableObject, RefSkeletalMeshData, ObjectComponentIndex, bCreatedNewSkeleton);
@@ -3677,8 +3698,8 @@ bool UCustomizableInstancePrivate::BuildSkeletonData(const TSharedRef<FUpdateCon
 	const TArray<FMeshBoneInfo>& RawRefBoneInfo = ReferenceSkeleton.GetRawRefBoneInfo();
 	const int32 RawRefBoneCount = ReferenceSkeleton.GetRawBoneNum();
 
-	const TArray<FInstanceUpdateData::FBone>& BonePose = OperationData->InstanceUpdateData.Skeletons[InstanceComponentIndex].BonePose;
-	TMap<mu::FBoneName, TPair<FName,uint16>>& BoneInfoMap = OperationData->InstanceUpdateData.Skeletons[InstanceComponentIndex].BoneInfoMap;
+	const TArray<FInstanceUpdateData::FBone>& BonePose = OperationData->InstanceUpdateData.SkeletonsPerInstanceComponent[InstanceComponentIndex.GetValue()].BonePose;
+	TMap<mu::FBoneName, TPair<FName,uint16>>& BoneInfoMap = OperationData->InstanceUpdateData.SkeletonsPerInstanceComponent[InstanceComponentIndex.GetValue()].BoneInfoMap;
 
 	{
 		MUTABLE_CPUPROFILER_SCOPE(BuildSkeletonData_BuildBoneInfoMap);
@@ -3904,15 +3925,20 @@ void UCustomizableInstancePrivate::BuildMeshSockets(const TSharedRef<FUpdateCont
 }
 
 
-void UCustomizableInstancePrivate::BuildOrCopyElementData(const TSharedRef<FUpdateContextPrivate>& OperationData, USkeletalMesh* SkeletalMesh, UCustomizableObjectInstance* CustomizableObjectInstance, int32 InstanceComponentIndex)
+void UCustomizableInstancePrivate::BuildOrCopyElementData(const TSharedRef<FUpdateContextPrivate>& OperationData, USkeletalMesh* SkeletalMesh, UCustomizableObjectInstance* CustomizableObjectInstance, 
+	FCustomizableObjectInstanceComponentIndex InstanceComponentIndex)
 {
 	MUTABLE_CPUPROFILER_SCOPE(UCustomizableInstancePrivate::BuildOrCopyElementData);
 
-	const FInstanceUpdateData::FComponent& Component = OperationData->InstanceUpdateData.Components[InstanceComponentIndex];
+	const FInstanceUpdateData::FComponent* Component = OperationData->GetComponentUpdateData(InstanceComponentIndex);
+	if (!Component)
+	{
+		return;
+	}
 
-	for (int32 LODIndex = FirstLODAvailable; LODIndex < Component.LODCount; ++LODIndex)
+	for (int32 LODIndex = FirstLODAvailable; LODIndex < Component->LODCount; ++LODIndex)
 	{
-		const FInstanceUpdateData::FLOD& LOD = OperationData->InstanceUpdateData.LODs[Component.FirstLOD+LODIndex];
+		const FInstanceUpdateData::FLOD& LOD = OperationData->InstanceUpdateData.LODs[Component->FirstLOD+LODIndex];
 
 		if (!LOD.SurfaceCount)
 		{
@@ -3931,7 +3957,7 @@ void UCustomizableInstancePrivate::BuildOrCopyElementData(const TSharedRef<FUpda
 	}
 }
 
-void UCustomizableInstancePrivate::BuildOrCopyMorphTargetsData(const TSharedRef<FUpdateContextPrivate>& OperationData, USkeletalMesh* SkeletalMesh, const USkeletalMesh* LastUpdateSkeletalMesh, UCustomizableObjectInstance* CustomizableObjectInstance, int32 InstanceComponentIndex)
+void UCustomizableInstancePrivate::BuildOrCopyMorphTargetsData(const TSharedRef<FUpdateContextPrivate>& OperationData, USkeletalMesh* SkeletalMesh, const USkeletalMesh* LastUpdateSkeletalMesh, UCustomizableObjectInstance* CustomizableObjectInstance, FCustomizableObjectInstanceComponentIndex InstanceComponentIndex)
 {
 	MUTABLE_CPUPROFILER_SCOPE(UCustomizableInstancePrivate::BuildOrCopyMorphTargetsData);
 
@@ -3946,18 +3972,8 @@ void UCustomizableInstancePrivate::BuildOrCopyMorphTargetsData(const TSharedRef<
 		return;
 	}
 	
-	FInstanceUpdateData& UpdateData = OperationData->InstanceUpdateData;
-
-	FInstanceUpdateData::FComponent& Component = UpdateData.Components[InstanceComponentIndex];
-	int32 ObjectComponentIndex = Component.Id;
-
-	FInstanceUpdateData::FRealTimeMorphsComponentData* ComponentMorphTargets = 
-		UpdateData.RealTimeMorphTargets.FindByPredicate(
-				[ObjectComponentIndex](const FInstanceUpdateData::FRealTimeMorphsComponentData& Elem)
-				{ 
-					return Elem.ObjectComponentIndex == ObjectComponentIndex;
-				});
-
+	FCustomizableObjectComponentIndex ObjectComponentIndex = OperationData->GetObjectComponentIndex(InstanceComponentIndex);
+	FInstanceUpdateData::FRealTimeMorphsComponentData* ComponentMorphTargets = OperationData->GetComponentMorphTargets(ObjectComponentIndex);
 	if (!ComponentMorphTargets)
 	{
 		return;
@@ -3982,7 +3998,7 @@ void UCustomizableInstancePrivate::BuildOrCopyMorphTargetsData(const TSharedRef<
 	}
 
 	// Copy MorphTargets from the FirstGeneratedLOD to the LODs below
-	const int32 FirstGeneratedLOD = FMath::Max((int32)OperationData->GetRequestedLODs()[ObjectComponentIndex], OperationData->GetMinLOD());
+	const int32 FirstGeneratedLOD = FMath::Max((int32)OperationData->GetRequestedLODs(ObjectComponentIndex), OperationData->GetMinLOD());
 	for (int32 LODIndex = OperationData->FirstLODAvailable; LODIndex < FirstGeneratedLOD; ++LODIndex)
 	{
 		MUTABLE_CPUPROFILER_SCOPE(CopyMorphTargetsData);
@@ -4012,7 +4028,8 @@ namespace
 
 }
 
-void UCustomizableInstancePrivate::BuildOrCopyClothingData(const TSharedRef<FUpdateContextPrivate>& OperationData, USkeletalMesh * SkeletalMesh, const USkeletalMesh* LastUpdateSkeletalMesh, UCustomizableObjectInstance* CustomizableObjectInstance, int32 InstanceComponentIndex)
+void UCustomizableInstancePrivate::BuildOrCopyClothingData(const TSharedRef<FUpdateContextPrivate>& OperationData, USkeletalMesh * SkeletalMesh, const USkeletalMesh* LastUpdateSkeletalMesh, 
+	UCustomizableObjectInstance* CustomizableObjectInstance, FCustomizableObjectInstanceComponentIndex InstanceComponentIndex)
 {
 	MUTABLE_CPUPROFILER_SCOPE(UCustomizableInstancePrivate::BuildOrCopyClothingData);
 
@@ -4053,15 +4070,18 @@ void UCustomizableInstancePrivate::BuildOrCopyClothingData(const TSharedRef<FUpd
 
 	const int32 LODCount = OperationData->InstanceUpdateData.LODs.Num();
 
-	FInstanceUpdateData::FComponent& Component = OperationData->InstanceUpdateData.Components[InstanceComponentIndex];
-	int32 ObjectComponentIndex = Component.Id;
+	const FInstanceUpdateData::FComponent* Component = OperationData->GetComponentUpdateData(InstanceComponentIndex);
+	if (!Component)
+	{
+		return;
+	}
 
 	{
 		MUTABLE_CPUPROFILER_SCOPE(DiscoverSectionsWithCloth);
 
-		for (int32 LODIndex = OperationData->FirstLODAvailable; LODIndex < Component.LODCount; ++LODIndex)
+		for (int32 LODIndex = OperationData->FirstLODAvailable; LODIndex < Component->LODCount; ++LODIndex)
 		{
-			const FInstanceUpdateData::FLOD& LOD = OperationData->InstanceUpdateData.LODs[Component.FirstLOD+LODIndex];
+			const FInstanceUpdateData::FLOD& LOD = OperationData->InstanceUpdateData.LODs[Component->FirstLOD+LODIndex];
 
 			if (!LOD.bGenerated)
 			{
@@ -4196,7 +4216,7 @@ void UCustomizableInstancePrivate::BuildOrCopyClothingData(const TSharedRef<FUpd
 			// Only initilialize once, multiple sections with cloth could point to the same cloth asset.
 			if (!DstAssetData.LodMap.Num())
 			{
-				DstAssetData.LodMap.Init(INDEX_NONE, OperationData->NumLODsAvailablePerComponent[ObjectComponentIndex]);
+				DstAssetData.LodMap.Init(INDEX_NONE, OperationData->GetLODsAvailable(Component->Id));
 
 				DstAssetData.UsedBoneNames = SrcAssetData.UsedBoneNames;
 				DstAssetData.UsedBoneIndices = SrcAssetData.UsedBoneIndices;
@@ -4784,9 +4804,9 @@ void UCustomizableInstancePrivate::BuildOrCopyClothingData(const TSharedRef<FUpd
 	{
 		int32 NumSectionsWithClothProcessed = 0;
 
-		for (int32 LODIndex = OperationData->GetMinLOD(); LODIndex < Component.LODCount; ++LODIndex)
+		for (int32 LODIndex = OperationData->GetMinLOD(); LODIndex < Component->LODCount; ++LODIndex)
 		{
-			const FInstanceUpdateData::FLOD& LOD = OperationData->InstanceUpdateData.LODs[Component.FirstLOD+LODIndex];
+			const FInstanceUpdateData::FLOD& LOD = OperationData->InstanceUpdateData.LODs[Component->FirstLOD+LODIndex];
 
 			if (mu::MeshPtrConst MutableMesh = LOD.Mesh)
 			{
@@ -4833,7 +4853,7 @@ void UCustomizableInstancePrivate::BuildOrCopyClothingData(const TSharedRef<FUpd
 	{
 		MUTABLE_CPUPROFILER_SCOPE(InitClothRenderData)
 		// Based on FSkeletalMeshLODModel::GetClothMappingData().
-		for (int32 LODIndex = OperationData->GetMinLOD(); LODIndex < OperationData->NumLODsAvailablePerComponent[ObjectComponentIndex]; ++LODIndex)
+		for (int32 LODIndex = OperationData->GetMinLOD(); LODIndex < OperationData->GetLODsAvailable(Component->Id); ++LODIndex)
 		{
 			FSkeletalMeshLODRenderData& LODModel = RenderResource->LODRenderData[LODIndex];
 	
@@ -4969,7 +4989,8 @@ void UCustomizableInstancePrivate::BuildOrCopyClothingData(const TSharedRef<FUpd
 	SkeletalMesh->SetHasActiveClothingAssets( static_cast<bool>( SectionsWithCloth.Num() ) );
 }
 
-bool UCustomizableInstancePrivate::BuildOrCopyRenderData(const TSharedRef<FUpdateContextPrivate>& OperationData, USkeletalMesh* SkeletalMesh, const USkeletalMesh* LastUpdateSkeletalMesh, UCustomizableObjectInstance* Public, int32 InstanceComponentIndex)
+bool UCustomizableInstancePrivate::BuildOrCopyRenderData(const TSharedRef<FUpdateContextPrivate>& OperationData, USkeletalMesh* SkeletalMesh, const USkeletalMesh* LastUpdateSkeletalMesh, 
+	UCustomizableObjectInstance* Public, FCustomizableObjectInstanceComponentIndex InstanceComponentIndex)
 {
 	MUTABLE_CPUPROFILER_SCOPE(UCustomizableInstancePrivate::BuildOrCopyRenderData);
 
@@ -4981,19 +5002,22 @@ bool UCustomizableInstancePrivate::BuildOrCopyRenderData(const TSharedRef<FUpdat
 	// It must be not null as it's checked in the calling function
 	check(CustomizableObject);
 
-	FInstanceUpdateData::FComponent& Component = OperationData->InstanceUpdateData.Components[InstanceComponentIndex];
-	int32 ObjectComponentIndex = Component.Id;
+	const FInstanceUpdateData::FComponent* Component = OperationData->GetComponentUpdateData(InstanceComponentIndex);
+	if (!Component)
+	{
+		return false;
+	}
 
 	const FModelResources& ModelResources = CustomizableObject->GetPrivate()->GetModelResources();
 
-	const int32 FirstGeneratedLOD = FMath::Max((int32)OperationData->GetRequestedLODs()[ObjectComponentIndex], OperationData->GetMinLOD());
+	const int32 FirstGeneratedLOD = FMath::Max((int32)OperationData->GetRequestedLODs(Component->Id), OperationData->GetMinLOD());
 
 
-	for (int32 LODIndex = FirstGeneratedLOD; LODIndex < Component.LODCount; ++LODIndex)
+	for (int32 LODIndex = FirstGeneratedLOD; LODIndex < Component->LODCount; ++LODIndex)
 	{
 		MUTABLE_CPUPROFILER_SCOPE(BuildRenderData);
 
-		const FInstanceUpdateData::FLOD& LOD = OperationData->InstanceUpdateData.LODs[Component.FirstLOD+LODIndex];
+		const FInstanceUpdateData::FLOD& LOD = OperationData->InstanceUpdateData.LODs[Component->FirstLOD+LODIndex];
 
 		// There could be components without a mesh in LODs
 		if (!LOD.bGenerated || !LOD.Mesh || LOD.SurfaceCount == 0)
@@ -5001,19 +5025,19 @@ bool UCustomizableInstancePrivate::BuildOrCopyRenderData(const TSharedRef<FUpdat
 			UE_LOG(LogMutable, Warning, TEXT("Building instance: generated mesh [%s] has LOD [%d] of object component index [%d] with no mesh.")
 				, *SkeletalMesh->GetName()
 				, LODIndex
-				, ObjectComponentIndex);
+				, Component->Id.GetValue());
 
 			// End with failure
 			return false;
 		}
 
-		TRACE_CPUPROFILER_EVENT_SCOPE_TEXT(*FString::Printf(TEXT("BuildRenderData: Component index %d, LOD %d"), ObjectComponentIndex, LODIndex));
+		TRACE_CPUPROFILER_EVENT_SCOPE_TEXT(*FString::Printf(TEXT("BuildRenderData: Component index %d, LOD %d"), Component->Id.GetValue(), LODIndex));
 
-		SetLastMeshId(ObjectComponentIndex, LODIndex, LOD.MeshID);
+		SetLastMeshId(Component->Id, LODIndex, LOD.MeshID);
 
 		FSkeletalMeshLODRenderData& LODResource = RenderData->LODRenderData[LODIndex];
 
-		const TMap<mu::FBoneName, TPair<FName, uint16>>& BoneInfoMap = OperationData->InstanceUpdateData.Skeletons[InstanceComponentIndex].BoneInfoMap;
+		const TMap<mu::FBoneName, TPair<FName, uint16>>& BoneInfoMap = OperationData->InstanceUpdateData.SkeletonsPerInstanceComponent[InstanceComponentIndex.GetValue()].BoneInfoMap;
 		
 		// Set active and required bones
 		{
@@ -5039,14 +5063,14 @@ bool UCustomizableInstancePrivate::BuildOrCopyRenderData(const TSharedRef<FUpdat
 		for (int32 MeshSectionIndex = 0; MeshSectionIndex < MeshNumSurfaces; ++MeshSectionIndex)
 		{
 			uint32 MeshSurfaceId = LOD.Mesh->GetSurfaceId(MeshSectionIndex);
-			int32 InstanceSurfaceIndex = OperationData->MutableInstance->FindSurfaceById(InstanceComponentIndex, LODIndex, MeshSurfaceId);
+			int32 InstanceSurfaceIndex = OperationData->MutableInstance->FindSurfaceById(InstanceComponentIndex.GetValue(), LODIndex, MeshSurfaceId);
 			
 			if (InstanceSurfaceIndex < 0)
 			{
 				continue;
 			}
 			
-			uint32 SurfaceMetadataId = OperationData->MutableInstance->GetSurfaceCustomId(InstanceComponentIndex, LODIndex, InstanceSurfaceIndex);
+			uint32 SurfaceMetadataId = OperationData->MutableInstance->GetSurfaceCustomId(InstanceComponentIndex.GetValue(), LODIndex, InstanceSurfaceIndex);
 			
 			uint32 UsedSurfaceMetadataId = 0;
 			if (SurfaceMetadataId != 0)
@@ -5270,9 +5294,9 @@ UE::Tasks::FTask UCustomizableInstancePrivate::LoadAdditionalAssetsAndData(
 	// Load Skeletons required by the SubMeshes of the newly generated Mesh, will be merged later
 	for (int32 InstanceComponentIndex = 0; InstanceComponentIndex < OperationData->NumInstanceComponents; ++InstanceComponentIndex)
 	{
-		int32 ObjectComponentIndex = OperationData->InstanceUpdateData.Components[InstanceComponentIndex].Id;
+		FCustomizableObjectComponentIndex ObjectComponentIndex = OperationData->GetObjectComponentIndex(FCustomizableObjectInstanceComponentIndex(InstanceComponentIndex));
 
-		const FInstanceUpdateData::FSkeletonData& SkeletonData = OperationData->InstanceUpdateData.Skeletons[InstanceComponentIndex];
+		const FInstanceUpdateData::FSkeletonData& SkeletonData = OperationData->InstanceUpdateData.SkeletonsPerInstanceComponent[InstanceComponentIndex];
 		
 		FCustomizableInstanceComponentData* ComponentData = GetComponentData(ObjectComponentIndex);
 		if (!ComponentData)
@@ -5318,9 +5342,9 @@ UE::Tasks::FTask UCustomizableInstancePrivate::LoadAdditionalAssetsAndData(
 	// Load assets coming from SubMeshes of the newly generated Mesh
 	if (OperationData->InstanceUpdateData.LODs.Num())
 	{
-		for (int32 ComponentIndex = 0; ComponentIndex < OperationData->InstanceUpdateData.Components.Num(); ++ComponentIndex)
+		for (int32 InstanceComponentIndex = 0; InstanceComponentIndex < OperationData->InstanceUpdateData.Components.Num(); ++InstanceComponentIndex)
 		{
-			const FInstanceUpdateData::FComponent& Component = Components[ComponentIndex];
+			const FInstanceUpdateData::FComponent& Component = Components[InstanceComponentIndex];
 
 			for (int32 LODIndex = OperationData->FirstLODAvailable; LODIndex < Component.LODCount; ++LODIndex)
 			{
@@ -5332,7 +5356,7 @@ UE::Tasks::FTask UCustomizableInstancePrivate::LoadAdditionalAssetsAndData(
 					continue;
 				}
 
-				FCustomizableInstanceComponentData* ComponentData = GetComponentData(ComponentIndex);
+				FCustomizableInstanceComponentData* ComponentData = GetComponentData(Component.Id);
 
 				const TArray<uint64>& StreamedResources = MutableMesh->GetStreamedResources();
 
@@ -5884,9 +5908,16 @@ FCustomizableObjectInstanceDescriptor& UCustomizableInstancePrivate::GetDescript
 }
 
 
-const TArray<TObjectPtr<UMaterialInterface>>* UCustomizableObjectInstance::GetOverrideMaterials(int32 ComponentIndex) const
+const TArray<TObjectPtr<UMaterialInterface>>* UCustomizableObjectInstance::GetOverrideMaterials(FCustomizableObjectComponentIndex ComponentIndex) const
 {
-	FCustomizableInstanceComponentData* ComponentData = PrivateData->GetComponentData(ComponentIndex);
+	UCustomizableObject* Object = GetCustomizableObject();
+	if (!Object)
+	{
+		return nullptr;
+	}
+
+	FName ComponentName = Object->GetPrivate()->GetModelResources().ComponentNamesPerObjectComponent[ComponentIndex.GetValue()];
+	FCustomizableInstanceComponentData* ComponentData = PrivateData->GetComponentData(ComponentName);
 	return ComponentData ? &ComponentData->OverrideMaterials : nullptr;
 }
 
@@ -6200,8 +6231,8 @@ void UCustomizableInstancePrivate::BuildMaterials(const TSharedRef<FUpdateContex
 	{
 		const FInstanceUpdateData::FComponent& Component = OperationData->InstanceUpdateData.Components[InstanceComponentIndex];
 
-		const int32 ObjectComponentIndex = Component.Id;
-		const FName& ComponentName = GetPublic()->GetCustomizableObject()->GetPrivate()->GetModelResources().ComponentNames[ObjectComponentIndex];
+		const FCustomizableObjectComponentIndex ObjectComponentIndex = Component.Id;
+		const FName& ComponentName = GetPublic()->GetCustomizableObject()->GetPrivate()->GetModelResources().ComponentNamesPerObjectComponent[ObjectComponentIndex.GetValue()];
 
 		TObjectPtr<USkeletalMesh>* Result = SkeletalMeshes.Find(ComponentName);
 		TObjectPtr<USkeletalMesh> SkeletalMesh = Result ? *Result : nullptr;
@@ -6210,7 +6241,7 @@ void UCustomizableInstancePrivate::BuildMaterials(const TSharedRef<FUpdateContex
 			continue;
 		}
 
-		const bool bReuseMaterials = !OperationData->MeshChanged[InstanceComponentIndex];
+		const bool bReuseMaterials = !OperationData->MeshChangedPerInstanceComponent[InstanceComponentIndex];
 
 		// If the mesh is not transient, it means it's pass-through so it should use material overrides and not be modified in any way
 		const bool bIsTransientMesh = static_cast<bool>(SkeletalMesh->HasAllFlags(EObjectFlags::RF_Transient));
@@ -6218,7 +6249,11 @@ void UCustomizableInstancePrivate::BuildMaterials(const TSharedRef<FUpdateContex
 		// It is not safe to replace the materials of a SkeletalMesh whose resources are initialized. Use overrides instead.
 		const bool bUseOverrideMaterialsOnly = !bIsTransientMesh || (OperationData->bUseMeshCache && SkeletalMesh->GetResourceForRendering()->IsInitialized());
 
-		ComponentsData[ObjectComponentIndex].OverrideMaterials.Reset();
+		FCustomizableInstanceComponentData* ComponentData = GetComponentData(ObjectComponentIndex);
+		if (ComponentData)
+		{
+			ComponentData->OverrideMaterials.Reset();
+		}
 
 		TArray<FSkeletalMaterial> Materials;
 
@@ -6230,7 +6265,7 @@ void UCustomizableInstancePrivate::BuildMaterials(const TSharedRef<FUpdateContex
 
 		MUTABLE_CPUPROFILER_SCOPE(BuildMaterials_LODLoop);
 
-		const int32 FirstGeneratedLOD = FMath::Max((int32)OperationData->GetRequestedLODs()[ObjectComponentIndex], OperationData->GetMinLOD());
+		const int32 FirstGeneratedLOD = FMath::Max((int32)OperationData->GetRequestedLODs(ObjectComponentIndex), OperationData->GetMinLOD());
 
 		for (int32 LODIndex = FirstGeneratedLOD; LODIndex < Component.LODCount; LODIndex++)
 		{
@@ -6248,9 +6283,9 @@ void UCustomizableInstancePrivate::BuildMaterials(const TSharedRef<FUpdateContex
 
 			// Pass-through components will not have a reference mesh.
 			const FMutableRefSkeletalMeshData* RefSkeletalMeshData = nullptr;
-			if (ModelResources.ReferenceSkeletalMeshesData.IsValidIndex(ObjectComponentIndex))
+			if (ModelResources.ReferenceSkeletalMeshesData.IsValidIndex(ObjectComponentIndex.GetValue()))
 			{
-				RefSkeletalMeshData = &ModelResources.ReferenceSkeletalMeshesData[ObjectComponentIndex];
+				RefSkeletalMeshData = &ModelResources.ReferenceSkeletalMeshesData[ObjectComponentIndex.GetValue()];
 			}
 
 			for (int32 SurfaceIndex = 0; SurfaceIndex < LOD.SurfaceCount; ++SurfaceIndex)
@@ -6441,7 +6476,7 @@ void UCustomizableInstancePrivate::BuildMaterials(const TSharedRef<FUpdateContex
 						UTexture* PassThroughTexture = nullptr; // Texture not generated by mutable
 
 						// \TODO: Change this key to a struct.
-						FString TextureReuseCacheRef = bReuseTextures ? FString::Printf(TEXT("%d-%d-%d-%d"), Image.BaseLOD, ObjectComponentIndex, Surface.SurfaceId, ImageIndex) : FString();
+						FString TextureReuseCacheRef = bReuseTextures ? FString::Printf(TEXT("%d-%d-%d-%d"), Image.BaseLOD, ObjectComponentIndex.GetValue(), Surface.SurfaceId, ImageIndex) : FString();
 
 						// If the mutable image is null, it must be in the cache
 						FMutableImageCacheKey ImageCacheKey = { Image.ImageID, OperationData->MipsToSkip };
@@ -6563,7 +6598,7 @@ void UCustomizableInstancePrivate::BuildMaterials(const TSharedRef<FUpdateContex
 										UE_LOG(LogMutable, Error, TEXT("Required image [%s] with ID [%lld] was not generated in the mutable thread, and it is not cached. LOD [%d]. Object Component [%d]"),
 											*Props.TextureParameterName,
 											Image.ImageID,
-											LODIndex, ObjectComponentIndex);
+											LODIndex, ObjectComponentIndex.GetValue());
 										continue;
 									}
 
@@ -6799,7 +6834,10 @@ void UCustomizableInstancePrivate::BuildMaterials(const TSharedRef<FUpdateContex
 
 				// Assign the material to the slot, and add it to the  OverrideMaterials
 				MaterialSlot.MaterialInterface = MaterialInterface;
-				ComponentsData[ObjectComponentIndex].OverrideMaterials.Add(MaterialInterface);
+				if (ComponentData)
+				{
+					ComponentData->OverrideMaterials.Add(MaterialInterface);
+				}
 			}
 		}
 
@@ -6866,12 +6904,12 @@ void UCustomizableInstancePrivate::BuildMaterials(const TSharedRef<FUpdateContex
 #endif
 
 			const FName& ComponentName = CustomizableObjectInstanceUsage->GetComponentName();
-			const int32 ObjectComponentIndex = OperationData->Instance->GetCustomizableObject()->GetPrivate()->GetModelResources().ComponentNames.IndexOfByKey(ComponentName);
+			const int32 ObjectComponentIndex = OperationData->Instance->GetCustomizableObject()->GetPrivate()->GetModelResources().ComponentNamesPerObjectComponent.IndexOfByKey(ComponentName);
 			
 			int32 InstanceComponentIndex = -1;
 			for (int32 CurrentInstanceIndex=0; CurrentInstanceIndex<OperationData->InstanceUpdateData.Components.Num(); ++CurrentInstanceIndex)
 			{
-				if (OperationData->InstanceUpdateData.Components[CurrentInstanceIndex].Id == ObjectComponentIndex)
+				if (OperationData->InstanceUpdateData.Components[CurrentInstanceIndex].Id.GetValue() == ObjectComponentIndex)
 				{
 					InstanceComponentIndex = CurrentInstanceIndex;
 					break;
@@ -7021,7 +7059,7 @@ FAutoConsoleVariableRef CVarMutableIgnoreMinMaxLOD(
 #endif
 
 
-void UCustomizableObjectInstance::SetRequestedLODs(int32 InMinLOD, int32 , const TArray<uint16>& InRequestedLODsPerComponent, 
+void UCustomizableObjectInstance::SetRequestedLODs(int32 InMinLOD, int32 , const TArray<uint16>& InRequestedLODsPerObjectComponent, 
 	                                               FMutableInstanceUpdateMap& InOutRequestedUpdates)
 {
 	check(PrivateData);
@@ -7078,34 +7116,34 @@ void UCustomizableObjectInstance::SetRequestedLODs(int32 InMinLOD, int32 , const
 	{
 		const uint16 FirstNonStreamedLODIndex = FMath::Clamp(PrivateData->FirstResidentLOD, 0, MaxLODIdx);
 
-		const int32 ComponentCount = GetNumComponents();
-		if (ComponentCount != MutableUpdateCandidate.RequestedLODLevels.Num())
+		const int32 ObjectComponentCount = CustomizableObject->GetComponentCount();
+		if (ObjectComponentCount != MutableUpdateCandidate.RequestedLODLevels.Num())
 		{
-			MutableUpdateCandidate.RequestedLODLevels.Init(FirstNonStreamedLODIndex, ComponentCount);
+			MutableUpdateCandidate.RequestedLODLevels.Init(FirstNonStreamedLODIndex, ObjectComponentCount);
 		}
 
 		const TArray<uint16>& GeneratedLODsPerComponent = GetPrivate()->CommittedDescriptorHash.RequestedLODsPerComponent;
 
-		const bool bIgnoreGeneratedLODs = GeneratedLODsPerComponent.Num() != ComponentCount;
+		const bool bIgnoreGeneratedLODs = GeneratedLODsPerComponent.Num() != ObjectComponentCount;
 
-		if (bMinLODChanged || bIgnoreGeneratedLODs || Descriptor.GetRequestedLODLevels() != InRequestedLODsPerComponent)
+		if (bMinLODChanged || bIgnoreGeneratedLODs || Descriptor.GetRequestedLODLevels() != InRequestedLODsPerObjectComponent)
 		{
-			check(InRequestedLODsPerComponent.Num() == ComponentCount);
+			check(InRequestedLODsPerObjectComponent.Num() == ObjectComponentCount);
 
 			bUpdateRequestedLODs = bIgnoreGeneratedLODs;
-			for (int32 ComponentIndex = 0; ComponentIndex < ComponentCount; ++ComponentIndex)
+			for (int32 ObjectComponentIndex = 0; ObjectComponentIndex < ObjectComponentCount; ++ObjectComponentIndex)
 			{
-				uint16 PredictedLOD = FMath::Min(InRequestedLODsPerComponent[ComponentIndex], FirstNonStreamedLODIndex);
+				uint16 PredictedLOD = FMath::Min(InRequestedLODsPerObjectComponent[ObjectComponentIndex], FirstNonStreamedLODIndex);
 
 				if (!bIgnoreGeneratedLODs)
 				{
-					PredictedLOD = FMath::Min(PredictedLOD, GeneratedLODsPerComponent[ComponentIndex]);
+					PredictedLOD = FMath::Min(PredictedLOD, GeneratedLODsPerComponent[ObjectComponentIndex]);
 				}
 
-				bUpdateRequestedLODs |= (PredictedLOD != MutableUpdateCandidate.RequestedLODLevels[ComponentIndex]);
+				bUpdateRequestedLODs |= (PredictedLOD != MutableUpdateCandidate.RequestedLODLevels[ObjectComponentIndex]);
 
 				// Save new RequestedLODs
-				MutableUpdateCandidate.RequestedLODLevels[ComponentIndex] = PredictedLOD;
+				MutableUpdateCandidate.RequestedLODLevels[ObjectComponentIndex] = PredictedLOD;
 			}
 		}
 	}
@@ -7174,13 +7212,32 @@ bool UCustomizableObjectInstance::HasAnyParameters() const
 }
 
 
-TSubclassOf<UAnimInstance> UCustomizableObjectInstance::GetAnimBP(int32 ComponentIndex, const FName& SlotName) const
+TSet<FName> UCustomizableObjectInstance::GetComponents() const
+{
+	TSet<FName> GeneratedComponents;
+
+	// For now, the instances don't really hold a direct array of generated components FNames. 
+	// They can be identified with the ones having a valid SkeletalMesh in the SkeletalMeshes array, but this will 
+	// not longer work when we have components that don't have a SkeletalMesh, like grooms, or panel clothing. (TODO)
+	for ( const TPair<FName, TObjectPtr<USkeletalMesh>>& Entry : GetPrivate()->SkeletalMeshes )
+	{
+		if (Entry.Value)
+		{
+			GeneratedComponents.Add(Entry.Key);
+		}
+	}
+
+	return GeneratedComponents;
+}
+
+
+TSubclassOf<UAnimInstance> UCustomizableObjectInstance::GetAnimBP(FName ComponentName, const FName& SlotName) const
 {
-	FCustomizableInstanceComponentData* ComponentData =	GetPrivate()->GetComponentData(ComponentIndex);
+	FCustomizableInstanceComponentData* ComponentData =	GetPrivate()->GetComponentData(ComponentName);
 	
 	if (!ComponentData)
 	{
-		FString ErrorMsg = FString::Printf(TEXT("Tried to access an invalid component index [%d] in a Mutable Instance."), ComponentIndex);
+		FString ErrorMsg = FString::Printf(TEXT("Tried to access an invalid component index [%s] in a Mutable Instance."), *ComponentName.ToString());
 		UE_LOG(LogMutable, Error, TEXT("%s"), *ErrorMsg);
 #if WITH_EDITOR
 		FMessageLogModule& MessageLogModule = FModuleManager::LoadModuleChecked<FMessageLogModule>("MessageLog");
@@ -7203,63 +7260,103 @@ const FGameplayTagContainer& UCustomizableObjectInstance::GetAnimationGameplayTa
 	return GetPrivate()->AnimBPGameplayTags;
 }
 
-void UCustomizableObjectInstance::ForEachAnimInstance(int32 ComponentIndex, FEachComponentAnimInstanceClassDelegate Delegate) const
+namespace MutablePrivate
 {
-	// allow us to log out both bad states with one pass
-	bool bAnyErrors = false;
-
-	if (!Delegate.IsBound())
+	template<typename DELEGATE>
+	void InternalForEachAnimInstance(UCustomizableInstancePrivate* Private, FName ComponentName, DELEGATE Delegate)
 	{
-		FString ErrorMsg = FString::Printf(TEXT("Attempting to iterate over AnimInstances with an unbound delegate - does nothing!"), ComponentIndex);
-		UE_LOG(LogMutable, Warning, TEXT("%s"), *ErrorMsg);
+		// allow us to log out both bad states with one pass
+		bool bAnyErrors = false;
+
+		if (!Delegate.IsBound())
+		{
+			FString ErrorMsg = FString::Printf(TEXT("Attempting to iterate over AnimInstances with an unbound delegate for component [%s]."), *ComponentName.ToString());
+			UE_LOG(LogMutable, Warning, TEXT("%s"), *ErrorMsg);
 #if WITH_EDITOR
-		FMessageLogModule& MessageLogModule = FModuleManager::LoadModuleChecked<FMessageLogModule>("MessageLog");
-		MessageLogModule.RegisterLogListing(FName("Mutable"), FText::FromString(FString("Mutable")));
-		FMessageLog MessageLog("Mutable");
+			FMessageLogModule& MessageLogModule = FModuleManager::LoadModuleChecked<FMessageLogModule>("MessageLog");
+			MessageLogModule.RegisterLogListing(FName("Mutable"), FText::FromString(FString("Mutable")));
+			FMessageLog MessageLog("Mutable");
 
-		MessageLog.Notify(FText::FromString(ErrorMsg), EMessageSeverity::Warning, true);
+			MessageLog.Notify(FText::FromString(ErrorMsg), EMessageSeverity::Warning, true);
 #endif
-		bAnyErrors = true;
-	}
+			bAnyErrors = true;
+		}
 
-	const FCustomizableInstanceComponentData* ComponentData =	GetPrivate()->GetComponentData(ComponentIndex);
-	
-	if (!ComponentData)
-	{
-		FString ErrorMsg = FString::Printf(TEXT("Tried to access an invalid component index [%d] in a Mutable Instance."), ComponentIndex);
-		UE_LOG(LogMutable, Error, TEXT("%s"), *ErrorMsg);
+		const FCustomizableInstanceComponentData* ComponentData = Private->GetComponentData(ComponentName);
+
+		if (!ComponentData)
+		{
+			FString ErrorMsg = FString::Printf(TEXT("Tried to access an invalid component [%s] in a Mutable Instance."), *ComponentName.ToString());
+			UE_LOG(LogMutable, Error, TEXT("%s"), *ErrorMsg);
 #if WITH_EDITOR
-		FMessageLogModule& MessageLogModule = FModuleManager::LoadModuleChecked<FMessageLogModule>("MessageLog");
-		MessageLogModule.RegisterLogListing(FName("Mutable"), FText::FromString(FString("Mutable")));
-		FMessageLog MessageLog("Mutable");
+			FMessageLogModule& MessageLogModule = FModuleManager::LoadModuleChecked<FMessageLogModule>("MessageLog");
+			MessageLogModule.RegisterLogListing(FName("Mutable"), FText::FromString(FString("Mutable")));
+			FMessageLog MessageLog("Mutable");
 
-		MessageLog.Notify(FText::FromString(ErrorMsg), EMessageSeverity::Error, true);
+			MessageLog.Notify(FText::FromString(ErrorMsg), EMessageSeverity::Error, true);
 #endif
 
-		bAnyErrors = true;
-	}
-
-	if (bAnyErrors)
-	{
-		return;
-	}
+			bAnyErrors = true;
+		}
 
-	for (const TPair<FName, TSoftClassPtr<UAnimInstance>>& MapElem : ComponentData->AnimSlotToBP)
-	{
-		const FName& Index = MapElem.Key;
-		const TSoftClassPtr<UAnimInstance>& AnimBP = MapElem.Value;
+		if (bAnyErrors)
+		{
+			return;
+		}
 
-		// if this _can_ resolve to a real AnimBP
-		if (!AnimBP.IsNull())
+		for (const TPair<FName, TSoftClassPtr<UAnimInstance>>& MapElem : ComponentData->AnimSlotToBP)
 		{
-			// force a load right now - we don't know whether we would have loaded already - this could be called in editor
-			const TSubclassOf<UAnimInstance> LiveAnimBP = AnimBP.LoadSynchronous();
-			if (LiveAnimBP)
+			const FName& Index = MapElem.Key;
+			const TSoftClassPtr<UAnimInstance>& AnimBP = MapElem.Value;
+
+			// if this _can_ resolve to a real AnimBP
+			if (!AnimBP.IsNull())
 			{
-				Delegate.Execute(Index, LiveAnimBP);
+				// force a load right now - we don't know whether we would have loaded already - this could be called in editor
+				const TSubclassOf<UAnimInstance> LiveAnimBP = AnimBP.LoadSynchronous();
+				if (LiveAnimBP)
+				{
+					Delegate.Execute(Index, LiveAnimBP);
+				}
 			}
 		}
 	}
+
+}
+
+
+void UCustomizableObjectInstance::ForEachComponentAnimInstance(FName ComponentName, FEachComponentAnimInstanceClassDelegate Delegate) const
+{
+	MutablePrivate::InternalForEachAnimInstance<>(GetPrivate(), ComponentName, Delegate);
+}
+
+
+void UCustomizableObjectInstance::ForEachComponentAnimInstance(FName ComponentName, FEachComponentAnimInstanceClassNativeDelegate Delegate) const
+{
+	MutablePrivate::InternalForEachAnimInstance<>(GetPrivate(), ComponentName, Delegate);
+}
+
+// Deprecated
+void UCustomizableObjectInstance::ForEachAnimInstance(int32 ObjectComponentIndex, FEachComponentAnimInstanceClassDelegate Delegate) const
+{
+	UCustomizableObject* CO = GetCustomizableObject();
+	if (CO)
+	{
+		FName ComponentName = CO->GetComponentName(FCustomizableObjectComponentIndex(ObjectComponentIndex));
+		MutablePrivate::InternalForEachAnimInstance<>(GetPrivate(), ComponentName, Delegate);
+	}
+}
+
+
+// Deprecated
+void UCustomizableObjectInstance::ForEachAnimInstance(int32 ObjectComponentIndex, FEachComponentAnimInstanceClassNativeDelegate Delegate) const
+{
+	UCustomizableObject* CO = GetCustomizableObject();
+	if (CO)
+	{
+		FName ComponentName = CO->GetComponentName(FCustomizableObjectComponentIndex(ObjectComponentIndex));
+		MutablePrivate::InternalForEachAnimInstance<>(GetPrivate(), ComponentName, Delegate);
+	}
 }
 
 
@@ -7352,65 +7449,6 @@ const TArray<FAnimInstanceOverridePhysicsAsset>* UCustomizableInstancePrivate::G
 	return &Found->AnimInstancePropertyIndexAndPhysicsAssets;
 }
 
-void UCustomizableObjectInstance::ForEachAnimInstance(int32 ComponentIndex, FEachComponentAnimInstanceClassNativeDelegate Delegate) const
-{
-	// allow us to log out both bad states with one pass
-	bool bAnyErrors = false;
-
-	if (!Delegate.IsBound())
-	{
-		FString ErrorMsg = FString::Printf(TEXT("Attempting to iterate over AnimInstances with an unbound delegate - does nothing!"), ComponentIndex);
-		UE_LOG(LogMutable, Warning, TEXT("%s"), *ErrorMsg);
-#if WITH_EDITOR
-		FMessageLogModule& MessageLogModule = FModuleManager::LoadModuleChecked<FMessageLogModule>("MessageLog");
-		MessageLogModule.RegisterLogListing(FName("Mutable"), FText::FromString(FString("Mutable")));
-		FMessageLog MessageLog("Mutable");
-
-		MessageLog.Notify(FText::FromString(ErrorMsg), EMessageSeverity::Warning, true);
-#endif
-		bAnyErrors = true;
-	}
-
-	const FCustomizableInstanceComponentData* ComponentData =	GetPrivate()->GetComponentData(ComponentIndex);
-	
-	if (!ComponentData)
-	{
-		FString ErrorMsg = FString::Printf(TEXT("Tried to access an invalid component index [%d] in a Mutable Instance."), ComponentIndex);
-		UE_LOG(LogMutable, Error, TEXT("%s"), *ErrorMsg);
-#if WITH_EDITOR
-		FMessageLogModule& MessageLogModule = FModuleManager::LoadModuleChecked<FMessageLogModule>("MessageLog");
-		MessageLogModule.RegisterLogListing(FName("Mutable"), FText::FromString(FString("Mutable")));
-		FMessageLog MessageLog("Mutable");
-
-		MessageLog.Notify(FText::FromString(ErrorMsg), EMessageSeverity::Error, true);
-#endif
-
-		bAnyErrors = true;
-	}
-
-	if (bAnyErrors)
-	{
-		return;
-	}
-
-	for (const TPair<FName, TSoftClassPtr<UAnimInstance>>& MapElem : ComponentData->AnimSlotToBP)
-	{
-		const FName& Index = MapElem.Key;
-		const TSoftClassPtr<UAnimInstance>& AnimBP = MapElem.Value;
-
-		// if this _can_ resolve to a real AnimBP
-		if (!AnimBP.IsNull())
-		{
-			// force a load right now - we don't know whether we would have loaded already - this could be called in editor
-			const TSubclassOf<UAnimInstance> LiveAnimBP = AnimBP.LoadSynchronous();
-			if (LiveAnimBP)
-			{
-				Delegate.Execute(Index, LiveAnimBP);
-			}
-		}
-	}
-}
-
 
 FInstancedStruct UCustomizableObjectInstance::GetExtensionInstanceData(const UCustomizableObjectExtension* Extension) const
 {
diff --git a/Engine/Plugins/Mutable/Source/CustomizableObject/Private/MuCO/CustomizableObjectInstanceDescriptor.cpp b/Engine/Plugins/Mutable/Source/CustomizableObject/Private/MuCO/CustomizableObjectInstanceDescriptor.cpp
--- a/Engine/Plugins/Mutable/Source/CustomizableObject/Private/MuCO/CustomizableObjectInstanceDescriptor.cpp
+++ b/Engine/Plugins/Mutable/Source/CustomizableObject/Private/MuCO/CustomizableObjectInstanceDescriptor.cpp
@@ -1200,6 +1200,21 @@ const TArray<uint16>& FCustomizableObjectInstanceDescriptor::GetRequestedLODLeve
 }
 
 
+uint16 FCustomizableObjectInstanceDescriptor::GetRequestedLODs(FCustomizableObjectComponentIndex Index) const
+{
+	uint16 Result = 0;
+	if (RequestedLODLevels.IsValidIndex(Index.GetValue()))
+	{
+		Result = RequestedLODLevels[Index.GetValue()];
+	}
+	else
+	{
+		ensure(false);
+	}
+	return Result;
+}
+
+
 const TArray<FCustomizableObjectBoolParameterValue>& FCustomizableObjectInstanceDescriptor::GetBoolParameters() const
 {
 	return BoolParameters;
diff --git a/Engine/Plugins/Mutable/Source/CustomizableObject/Private/MuCO/CustomizableObjectInstanceUsage.cpp b/Engine/Plugins/Mutable/Source/CustomizableObject/Private/MuCO/CustomizableObjectInstanceUsage.cpp
--- a/Engine/Plugins/Mutable/Source/CustomizableObject/Private/MuCO/CustomizableObjectInstanceUsage.cpp
+++ b/Engine/Plugins/Mutable/Source/CustomizableObject/Private/MuCO/CustomizableObjectInstanceUsage.cpp
@@ -92,7 +92,7 @@ void UCustomizableObjectInstanceUsage::SetComponentIndex(int32 ComponentIndex)
 	}
 	else
 	{
-		UsedComponentIndex = ComponentIndex;
+		UsedComponentIndex_DEPRECATED = FCustomizableObjectComponentIndex(ComponentIndex);
 	}
 }
 
@@ -105,7 +105,7 @@ int32 UCustomizableObjectInstanceUsage::GetComponentIndex() const
 	}
 	else
 	{
-		return UsedComponentIndex;
+		return UsedComponentIndex_DEPRECATED.GetValue();
 	}
 }
 
@@ -118,7 +118,7 @@ void UCustomizableObjectInstanceUsage::SetComponentName(const FName& Name)
 	}
 	else
 	{
-		UsedComponentIndex = INDEX_NONE;
+		UsedComponentIndex_DEPRECATED.Invalidate();
 		UsedComponentName = Name;
 	}
 }
@@ -132,13 +132,13 @@ FName UCustomizableObjectInstanceUsage::GetComponentName() const
 	}
 	else
 	{
-		if (UsedComponentIndex == INDEX_NONE)
+		if (!UsedComponentIndex_DEPRECATED.IsValid())
 		{
 			return UsedComponentName;	
 		}
 		else
 		{
-			return FName(FString::FromInt(UsedComponentIndex));
+			return FName(FString::FromInt(UsedComponentIndex_DEPRECATED.GetValue()));
 		}
 	}
 }
diff --git a/Engine/Plugins/Mutable/Source/CustomizableObject/Private/MuCO/CustomizableObjectSkeletalMesh.cpp b/Engine/Plugins/Mutable/Source/CustomizableObject/Private/MuCO/CustomizableObjectSkeletalMesh.cpp
--- a/Engine/Plugins/Mutable/Source/CustomizableObject/Private/MuCO/CustomizableObjectSkeletalMesh.cpp
+++ b/Engine/Plugins/Mutable/Source/CustomizableObject/Private/MuCO/CustomizableObjectSkeletalMesh.cpp
@@ -7,6 +7,7 @@
 
 #include "MuCO/CustomizableObject.h"
 #include "MuCO/CustomizableObjectInstance.h"
+#include "MuCO/CustomizableObjectInstancePrivate.h"
 #include "MuCO/CustomizableObjectMeshUpdate.h"
 
 FName GenerateUniqueNameFromCOInstance(const UCustomizableObjectInstance& Instance)
@@ -20,7 +21,7 @@ FName GenerateUniqueNameFromCOInstance(const UCustomizableObjectInstance& Instan
 
 
 UCustomizableObjectSkeletalMesh* UCustomizableObjectSkeletalMesh::CreateSkeletalMesh(const TSharedRef<FUpdateContextPrivate>& OperationData, 
-	const UCustomizableObjectInstance& Instance, UCustomizableObject& CustomizableObject, const int32 InstanceComponentIndex)
+	const UCustomizableObjectInstance& Instance, UCustomizableObject& CustomizableObject, const int32 FirstLOD, const int32 LODCount)
 {
 	FName SkeletalMeshName = GenerateUniqueNameFromCOInstance(Instance);
 	UCustomizableObjectSkeletalMesh* OutSkeletalMesh = NewObject<UCustomizableObjectSkeletalMesh>(GetTransientPackage(), SkeletalMeshName, RF_Transient);
@@ -29,7 +30,6 @@ UCustomizableObjectSkeletalMesh* UCustomizableObjectSkeletalMesh::CreateSkeletal
 	OutSkeletalMesh->CustomizableObjectPathName = GetNameSafe(&CustomizableObject);
 	OutSkeletalMesh->InstancePathName = Instance.GetName();
 
-
 	// Init properties
 	OutSkeletalMesh->Model = CustomizableObject.GetPrivate()->GetModel();
 
@@ -38,11 +38,9 @@ UCustomizableObjectSkeletalMesh* UCustomizableObjectSkeletalMesh::CreateSkeletal
 	
 	OutSkeletalMesh->MeshIDs.Init(MAX_uint64, MAX_MESH_LOD_COUNT);
 	
-	const FInstanceUpdateData::FComponent& Component = OperationData->InstanceUpdateData.Components[InstanceComponentIndex];
-
-	for (int32 LODIndex = OperationData->FirstLODAvailable; LODIndex < Component.LODCount; ++LODIndex)
+	for (int32 LODIndex = OperationData->FirstLODAvailable; LODIndex < LODCount; ++LODIndex)
 	{
-		const FInstanceUpdateData::FLOD& LOD = OperationData->InstanceUpdateData.LODs[Component.FirstLOD + LODIndex];
+		const FInstanceUpdateData::FLOD& LOD = OperationData->InstanceUpdateData.LODs[FirstLOD + LODIndex];
 		OutSkeletalMesh->MeshIDs[LODIndex] = LOD.MeshID;
 	}
 
diff --git a/Engine/Plugins/Mutable/Source/CustomizableObject/Private/MuCO/CustomizableObjectSystem.cpp b/Engine/Plugins/Mutable/Source/CustomizableObject/Private/MuCO/CustomizableObjectSystem.cpp
--- a/Engine/Plugins/Mutable/Source/CustomizableObject/Private/MuCO/CustomizableObjectSystem.cpp
+++ b/Engine/Plugins/Mutable/Source/CustomizableObject/Private/MuCO/CustomizableObjectSystem.cpp
@@ -18,7 +18,6 @@
 #include "MuCO/CustomizableObjectInstanceUsage.h"
 #include "MuCO/ICustomizableObjectModule.h"
 #include "MuCO/LogBenchmarkUtil.h"
-#include "MuCO/LogInformationUtil.h"
 #include "MuCO/UnrealExtensionDataStreamer.h"
 #include "MuCO/UnrealMutableImageProvider.h"
 #include "MuCO/UnrealMutableModelDiskStreamer.h"
@@ -209,6 +208,12 @@ const TArray<uint16>& FUpdateContextPrivate::GetRequestedLODs() const
 }
 
 
+uint16 FUpdateContextPrivate::GetRequestedLODs(FCustomizableObjectComponentIndex Index) const
+{
+	return CapturedDescriptor.GetRequestedLODs(Index);
+}
+
+
 void FUpdateContextPrivate::SetRequestedLODs(TArray<uint16>& RequestedLODs)
 {
 	CapturedDescriptor.SetRequestedLODLevels(RequestedLODs);
@@ -234,6 +239,63 @@ const FCustomizableObjectInstanceDescriptor&& FUpdateContextPrivate::MoveCommitt
 }
 
 
+FCustomizableObjectComponentIndex FUpdateContextPrivate::GetObjectComponentIndex(FCustomizableObjectInstanceComponentIndex InstanceComponentIndex)
+{
+	FCustomizableObjectComponentIndex ObjectComponentIndex;
+	if (InstanceUpdateData.Components.IsValidIndex(InstanceComponentIndex.GetValue()))
+	{
+		ObjectComponentIndex = InstanceUpdateData.Components[InstanceComponentIndex.GetValue()].Id;
+	}
+	else
+	{
+		ObjectComponentIndex.Invalidate();
+	}
+	return ObjectComponentIndex;
+}
+
+
+FInstanceUpdateData::FRealTimeMorphsComponentData* FUpdateContextPrivate::GetComponentMorphTargets(FCustomizableObjectComponentIndex ObjectComponentIndex)
+{
+	return InstanceUpdateData.RealTimeMorphTargets.FindByPredicate(
+		[ObjectComponentIndex](const FInstanceUpdateData::FRealTimeMorphsComponentData& Elem)
+		{
+			return Elem.ObjectComponentIndex == ObjectComponentIndex;
+		});
+}
+
+
+const FInstanceUpdateData::FComponent* FUpdateContextPrivate::GetComponentUpdateData(FCustomizableObjectInstanceComponentIndex InstanceComponentIndex)
+{
+	if (InstanceUpdateData.Components.IsValidIndex(InstanceComponentIndex.GetValue()))
+	{
+		return &InstanceUpdateData.Components[InstanceComponentIndex.GetValue()];
+	}
+	return nullptr;
+}
+
+
+void FUpdateContextPrivate::InitMeshDescriptors(int32 Size)
+{
+	MeshDescriptors.SetNum(Size);
+}
+
+
+const TArray<TArray<mu::FResourceID>>& FUpdateContextPrivate::GetMeshDescriptors() const
+{
+	return MeshDescriptors;
+}
+
+
+TArray<mu::FResourceID>* FUpdateContextPrivate::GetMeshDescriptors(FCustomizableObjectComponentIndex Index)
+{
+	if (MeshDescriptors.IsValidIndex(Index.GetValue()))
+	{
+		return &MeshDescriptors[Index.GetValue()];
+	}
+	return nullptr;
+}
+
+
 FMutablePendingInstanceUpdate::FMutablePendingInstanceUpdate(const TSharedRef<FUpdateContextPrivate>& InContext) :
 	Context(InContext)
 {
@@ -428,75 +490,6 @@ FString UCustomizableObjectSystem::GetPluginVersion() const
 }
 
 
-void UCustomizableObjectSystem::LogShowData(bool bFullInfo, bool ShowMaterialInfo) const
-{
-	LogInformationUtil::ResetCounters();
-
-	TArray<UCustomizableObjectInstance*> ArrayData;
-
-	for (TObjectIterator<UCustomizableObjectInstanceUsage> It; It; ++It)
-	{
-		const UCustomizableObjectInstanceUsage* CustomizableObjectInstanceUsage = *It;
-
-#if WITH_EDITOR
-		if (IsValid(CustomizableObjectInstanceUsage) && CustomizableObjectInstanceUsage->GetPrivate()->IsNetMode(NM_DedicatedServer))
-		{
-			continue;
-		}
-#endif
-
-		if (IsValid(CustomizableObjectInstanceUsage) && CustomizableObjectInstanceUsage->GetCustomizableObjectInstance()
-			&& CustomizableObjectInstanceUsage->GetAttachParent())
-		{
-			const AActor* ParentActor = CustomizableObjectInstanceUsage->GetAttachParent()->GetAttachmentRootActor();
-
-			if (ParentActor != nullptr)
-			{
-				ArrayData.AddUnique(CustomizableObjectInstanceUsage->GetCustomizableObjectInstance());
-			}
-		}
-	}
-
-	ArrayData.Sort([](UCustomizableObjectInstance& A, UCustomizableObjectInstance& B)
-	{
-		check(A.GetPrivate() != nullptr);
-		check(B.GetPrivate() != nullptr);
-		return (A.GetPrivate()->LastMinSquareDistFromComponentToPlayer < B.GetPrivate()->LastMinSquareDistFromComponentToPlayer);
-	});
-
-	int32 i;
-	const int32 Max = ArrayData.Num();
-
-	if (bFullInfo)
-	{
-		for (i = 0; i < Max; ++i)
-		{
-			LogInformationUtil::LogShowInstanceDataFull(ArrayData[i], ShowMaterialInfo);
-		}
-	}
-	else
-	{
-		FString LogData = "\n\n";
-		for (i = 0; i < Max; ++i)
-		{
-			LogInformationUtil::LogShowInstanceData(ArrayData[i], LogData);
-		}
-		UE_LOG(LogMutable, Log, TEXT("%s"), *LogData);
-
-		UWorld* World = GWorld;
-
-		if (World)
-		{
-			APlayerController* PlayerController = World->GetFirstPlayerController();
-			if (PlayerController)
-			{
-				PlayerController->ClientMessage(LogData);
-			}
-		}
-	}
-}
-
-
 UCustomizableObjectSystemPrivate* UCustomizableObjectSystem::GetPrivate()
 {
 	check(Private);
@@ -1699,20 +1692,27 @@ namespace impl
 			return;
 		}
 
-		Operation->NumLODsAvailablePerComponent.SetNumZeroed(Operation->NumObjectComponents);
+		Operation->InitNumLODsAvailablePerComponent(Operation->NumObjectComponents);
 
-		for (int32 ComponentIndex = 0; ComponentIndex < Operation->NumInstanceComponents; ++ComponentIndex)
+		int32 LODsInFirstValidComponent = -1;
+		for (int32 InstanceComponentIndex = 0; InstanceComponentIndex < Operation->NumInstanceComponents; ++InstanceComponentIndex)
 		{
-			Operation->NumLODsAvailablePerComponent[Operation->MutableInstance->GetComponentId(ComponentIndex)] = Operation->MutableInstance->GetLODCount(ComponentIndex);
+			int32 NumLODs = Operation->MutableInstance->GetLODCount(InstanceComponentIndex);
+			FCustomizableObjectComponentIndex Index(Operation->MutableInstance->GetComponentId(InstanceComponentIndex));
+			Operation->SetNumLODsAvailable(Index, NumLODs);
+			if (LODsInFirstValidComponent <= 0)
+			{
+				LODsInFirstValidComponent = NumLODs;
+			}
 		}
 		
 		int32 CurrentMinLOD = Operation->bStreamMeshLODs ? 0 : Operation->GetMinLOD();
-		CurrentMinLOD = FMath::Clamp(CurrentMinLOD, Operation->FirstLODAvailable, Operation->NumLODsAvailablePerComponent[0] - 1); 
+		CurrentMinLOD = FMath::Clamp(CurrentMinLOD, Operation->FirstLODAvailable, LODsInFirstValidComponent - 1);
 		Operation->SetMinLOD(CurrentMinLOD);
 
 		if (Operation->bStreamMeshLODs)
 		{
-			Operation->FirstResidentLOD = FMath::Clamp(Operation->FirstResidentLOD, Operation->FirstLODAvailable, Operation->NumLODsAvailablePerComponent[0] - 1); 
+			Operation->FirstResidentLOD = FMath::Clamp(Operation->FirstResidentLOD, Operation->FirstLODAvailable, LODsInFirstValidComponent - 1);
 		}
 		else 
 		{
@@ -1732,7 +1732,7 @@ namespace impl
 			else
 			{
 				// Clamp value to the valid LOD range.
-				RequestedLODs[ComponentIndex] = FMath::Min(RequestedLODs[ComponentIndex], (uint16)(Operation->NumLODsAvailablePerComponent[0] - 1));
+				RequestedLODs[ComponentIndex] = FMath::Min(RequestedLODs[ComponentIndex], (uint16)(LODsInFirstValidComponent - 1));
 			}
 		}
 
@@ -1782,11 +1782,12 @@ namespace impl
 	{
 		MUTABLE_CPUPROFILER_SCOPE(Subtask_Mutable_PrepareSkeletonData);
 
-		OperationData->InstanceUpdateData.Skeletons.SetNum(OperationData->NumObjectComponents);
+		int32 NumInstanceComponents = OperationData->InstanceUpdateData.Components.Num();
+		OperationData->InstanceUpdateData.SkeletonsPerInstanceComponent.SetNum(NumInstanceComponents);
 
-		for (int32 ComponentIndex = 0; ComponentIndex<OperationData->InstanceUpdateData.Components.Num(); ++ComponentIndex)
+		for (int32 InstanceComponentIndex = 0; InstanceComponentIndex< NumInstanceComponents; ++InstanceComponentIndex)
 		{
-			const FInstanceUpdateData::FComponent& Component = OperationData->InstanceUpdateData.Components[ComponentIndex];
+			const FInstanceUpdateData::FComponent& Component = OperationData->InstanceUpdateData.Components[InstanceComponentIndex];
 
 			for (int32 LODIndex = 0; LODIndex < Component.LODCount; ++LODIndex)
 			{
@@ -1798,7 +1799,7 @@ namespace impl
 					continue;
 				}
 
-				FInstanceUpdateData::FSkeletonData& SkeletonData = OperationData->InstanceUpdateData.Skeletons[ComponentIndex];
+				FInstanceUpdateData::FSkeletonData& SkeletonData = OperationData->InstanceUpdateData.SkeletonsPerInstanceComponent[InstanceComponentIndex];
 
 				// Add SkeletonIds 
 				const int32 SkeletonIDsCount = Mesh->GetSkeletonIDsCount();
@@ -1849,29 +1850,15 @@ namespace impl
 		
 		int32 NumNotFoundLoadedMorphsResources = 0;
 		
-		for (FInstanceUpdateData::FComponent& Component : OperationData->InstanceUpdateData.Components)
+		const int32 NumInstanceComponents = OperationData->InstanceUpdateData.Components.Num();
+		check(OperationData->InstanceUpdateData.RealTimeMorphTargets.IsEmpty());
+		OperationData->InstanceUpdateData.RealTimeMorphTargets.Reserve(NumInstanceComponents);
+		for ( int32 InstanceComponentIndex = 0; InstanceComponentIndex< NumInstanceComponents; ++InstanceComponentIndex)
 		{
-			if (!OperationData->InstanceUpdateData.RealTimeMorphTargets.IsValidIndex(Component.Id))
-			{
-				OperationData->InstanceUpdateData.RealTimeMorphTargets.SetNum(Component.Id + 1);
-			}
-
-			FInstanceUpdateData::FRealTimeMorphsComponentData& ComponentMorphTargetsData = 
-					OperationData->InstanceUpdateData.RealTimeMorphTargets[Component.Id];
-
+			FInstanceUpdateData::FComponent& Component = OperationData->InstanceUpdateData.Components[InstanceComponentIndex];
+			FInstanceUpdateData::FRealTimeMorphsComponentData& ComponentMorphTargetsData = OperationData->InstanceUpdateData.RealTimeMorphTargets.Emplace_GetRef();
 			ComponentMorphTargetsData.ObjectComponentIndex = Component.Id;
-		}
-
-		const int32 NumComponents = OperationData->InstanceUpdateData.RealTimeMorphTargets.Num();
-		for (int32 ComponentIndex = 0; ComponentIndex < NumComponents; ++ComponentIndex)
-		{
-			FInstanceUpdateData::FRealTimeMorphsComponentData& ComponentMorphTargetsData = 
-					OperationData->InstanceUpdateData.RealTimeMorphTargets[ComponentIndex];
-			
-			if (ComponentMorphTargetsData.ObjectComponentIndex == INDEX_NONE)
-			{
-				continue;
-			}
+			check(ComponentMorphTargetsData.ObjectComponentIndex.IsValid());
 
 			struct FMorphTargetMeshData
 			{
@@ -1904,10 +1891,12 @@ namespace impl
 			TArray<TArray<FMorphTargetLODModel>>& MorphsData = ComponentMorphTargetsData.RealTimeMorphsLODData;
 			const int32 NumMorphs = MorphTargetNames.Num();
 
+			const int32 NumLODsInThisComponent = OperationData->GetLODsAvailable(ComponentMorphTargetsData.ObjectComponentIndex);
+
 			MorphsData.SetNum(MorphTargetNames.Num());
 			for (int32 MorphIndex = 0; MorphIndex < NumMorphs; ++MorphIndex)
 			{
-				MorphsData[MorphIndex].SetNum(OperationData->NumLODsAvailablePerComponent[ComponentIndex]);
+				MorphsData[MorphIndex].SetNum(NumLODsInThisComponent);
 			}
 
 			TArray<int32> SectionMorphTargetVerticesCount;
@@ -1915,10 +1904,9 @@ namespace impl
 
 			int32 NumInvalidVertexMorphNamesFound = 0;
 			
-			const int32 FirstGeneratedLOD = FMath::Max((int32)OperationData->GetRequestedLODs()[ComponentIndex], OperationData->GetMinLOD());
-			for (int32 LODIndex = FirstGeneratedLOD; LODIndex < OperationData->NumLODsAvailablePerComponent[ComponentIndex]; ++LODIndex)
+			const int32 FirstGeneratedLOD = FMath::Max((int32)OperationData->GetRequestedLODs(ComponentMorphTargetsData.ObjectComponentIndex), OperationData->GetMinLOD());
+			for (int32 LODIndex = FirstGeneratedLOD; LODIndex < NumLODsInThisComponent; ++LODIndex)
 			{
-				const FInstanceUpdateData::FComponent& Component = UpdateData.Components[ComponentIndex];
 				const FInstanceUpdateData::FLOD& LOD = UpdateData.LODs[Component.FirstLOD+LODIndex];
 				check(LOD.bGenerated);
 
@@ -2498,8 +2486,6 @@ namespace impl
 				return;
 			}
 
-			const TArray<uint16>& RequestedLODs = OperationData->GetRequestedLODs();
-
 			const TSharedRef<TArray<FGetMeshData>> GetMeshesData = MakeShared<TArray<FGetMeshData>>();
 		
 			OperationData->InstanceUpdateData.Components.SetNum(OperationData->NumInstanceComponents);
@@ -2507,18 +2493,17 @@ namespace impl
 			{
 				FInstanceUpdateData::FComponent& Component = OperationData->InstanceUpdateData.Components[InstanceComponentIndex];
 				Component.FirstLOD = OperationData->InstanceUpdateData.LODs.Num();
-				Component.Id = Instance->GetComponentId(InstanceComponentIndex);
+				Component.Id = FCustomizableObjectComponentIndex(Instance->GetComponentId(InstanceComponentIndex));
 
-				int32 ObjectComponentIndex = Component.Id;
+				FCustomizableObjectComponentIndex ObjectComponentIndex = Component.Id;
 
-				if (!OperationData->NumLODsAvailablePerComponent.IsValidIndex(ObjectComponentIndex))
+				Component.LODCount = OperationData->GetLODsAvailable(ObjectComponentIndex);
+				if (!Component.LODCount)
 				{
 					// It happens in degenerated cases with empty components.
 					continue;
 				}
 
-				Component.LODCount = OperationData->NumLODsAvailablePerComponent[ObjectComponentIndex];
-
 				for (int32 MutableLODIndex = 0; MutableLODIndex < Component.LODCount; ++MutableLODIndex)
 				{
 					// If the LOD is not generated we still add an empty one to keep indexes aligned.
@@ -2532,7 +2517,7 @@ namespace impl
 
 					FInstanceUpdateData::FLOD& LOD = OperationData->InstanceUpdateData.LODs.Last();
 
-					const bool bGenerateLOD = RequestedLODs.IsValidIndex(ObjectComponentIndex) ? RequestedLODs[ObjectComponentIndex] <= MutableLODIndex : true;
+					const bool bGenerateLOD = OperationData->GetRequestedLODs(ObjectComponentIndex) <= MutableLODIndex;
 
 					// Mesh
 					{
@@ -3198,7 +3183,7 @@ namespace impl
 		UCustomizableObject* CustomizableObject = Operation->Instance->GetCustomizableObject();
 		UCustomizableObjectPrivate* CustomizableObjectPrivate = CustomizableObject->GetPrivate();
 		
-		for (const TArray<mu::FResourceID>& MeshId : Operation->MeshDescriptors)
+		for (const TArray<mu::FResourceID>& MeshId : Operation->GetMeshDescriptors())
 		{
 			if (USkeletalMesh* CachedMesh = CustomizableObject->GetPrivate()->MeshCache.Get(MeshId))
 			{
@@ -3237,22 +3222,25 @@ namespace impl
 		CreateMutableInstance(Operation);
 		FixLODs(Operation);
 
-		Operation->MeshDescriptors.SetNum(Operation->NumObjectComponents);
+		Operation->InitMeshDescriptors(Operation->NumObjectComponents);
 
-		const TArray<uint16>& RequestedLODs = Operation->GetRequestedLODs();
 		for (int32 InstanceComponentIndex = 0; InstanceComponentIndex < Operation->NumInstanceComponents; ++InstanceComponentIndex)
 		{
-			int32 ObjectComponentIndex = Operation->MutableInstance->GetComponentId(InstanceComponentIndex);
-			TArray<mu::FResourceID>& MeshId = Operation->MeshDescriptors[ObjectComponentIndex];
-			MeshId.Init(MAX_uint64, MAX_MESH_LOD_COUNT);
+			FCustomizableObjectComponentIndex ObjectComponentIndex(Operation->MutableInstance->GetComponentId(InstanceComponentIndex));
+			TArray<mu::FResourceID>* MeshId = Operation->GetMeshDescriptors(ObjectComponentIndex);
+			if (!MeshId)
+			{
+				continue;
+			}
 
+			MeshId->Init(MAX_uint64, MAX_MESH_LOD_COUNT);
 			
-			for (int32 LODIndex = Operation->GetMinLOD(); LODIndex < Operation->NumLODsAvailablePerComponent[ObjectComponentIndex]; ++LODIndex)
+			for (int32 LODIndex = Operation->GetMinLOD(); LODIndex < Operation->GetLODsAvailable(ObjectComponentIndex); ++LODIndex)
 			{
-				const bool bGenerateLOD = RequestedLODs.IsValidIndex(ObjectComponentIndex) ? RequestedLODs[ObjectComponentIndex] <= LODIndex : true;
+				const bool bGenerateLOD = Operation->GetRequestedLODs(ObjectComponentIndex) <= LODIndex;
 				if (bGenerateLOD)
 				{
-					MeshId[LODIndex] = Operation->MutableInstance->GetMeshId(InstanceComponentIndex, LODIndex);
+					(*MeshId)[LODIndex] = Operation->MutableInstance->GetMeshId(InstanceComponentIndex, LODIndex);
 				}
 			}
 		}
diff --git a/Engine/Plugins/Mutable/Source/CustomizableObject/Private/MuCO/CustomizableSkeletalMeshActor.cpp b/Engine/Plugins/Mutable/Source/CustomizableObject/Private/MuCO/CustomizableSkeletalMeshActor.cpp
--- a/Engine/Plugins/Mutable/Source/CustomizableObject/Private/MuCO/CustomizableSkeletalMeshActor.cpp
+++ b/Engine/Plugins/Mutable/Source/CustomizableObject/Private/MuCO/CustomizableSkeletalMeshActor.cpp
@@ -113,33 +113,33 @@ void ACustomizableSkeletalMeshActor::SwitchComponentsMaterials(UCustomizableObje
 			return;
 		}
 
-		for (int32 CompIndex = 0; CompIndex < SkeletalMeshComponents.Num(); ++CompIndex)
+		for (int32 ObjectComponentIndex = 0; ObjectComponentIndex < SkeletalMeshComponents.Num(); ++ObjectComponentIndex)
 		{
-			int32 NumMaterials = SkeletalMeshComponents[CompIndex]->GetNumMaterials();
+			int32 NumMaterials = SkeletalMeshComponents[ObjectComponentIndex]->GetNumMaterials();
 
 			if (bDebugMaterialEnabled)
 			{
 				for (int32 MatIndex = 0; MatIndex < NumMaterials; ++MatIndex)
 				{
-					SkeletalMeshComponents[CompIndex]->SetMaterial(MatIndex, DebugMaterial);
+					SkeletalMeshComponents[ObjectComponentIndex]->SetMaterial(MatIndex, DebugMaterial);
 				}
 			}
 			else // Remove debugmaterial
 			{
 				// check if original materials already overriden
-				const TArray<TObjectPtr<UMaterialInterface>>* OverrideMaterials = COInstance->GetOverrideMaterials(CompIndex);
+				const TArray<TObjectPtr<UMaterialInterface>>* OverrideMaterials = COInstance->GetOverrideMaterials(FCustomizableObjectComponentIndex(ObjectComponentIndex));
 				const bool bUseOverrideMaterials = COInstance->GetCustomizableObject()->bEnableMeshCache && CVarEnableMeshCache.GetValueOnAnyThread();
 
 				if (bUseOverrideMaterials && OverrideMaterials->Num() > 0)
 				{
 					for (int32 MatIndex = 0; MatIndex < OverrideMaterials->Num(); ++MatIndex)
 					{
-						SkeletalMeshComponents[CompIndex]->SetMaterial(MatIndex, (*OverrideMaterials)[MatIndex]);
+						SkeletalMeshComponents[ObjectComponentIndex]->SetMaterial(MatIndex, (*OverrideMaterials)[MatIndex]);
 					}
 				}
 				else
 				{
-					SkeletalMeshComponents[CompIndex]->EmptyOverrideMaterials();
+					SkeletalMeshComponents[ObjectComponentIndex]->EmptyOverrideMaterials();
 				}
 
 				bRemoveDebugMaterial = false;
diff --git a/Engine/Plugins/Mutable/Source/CustomizableObject/Private/MuCO/LogBenchmarkUtil.cpp b/Engine/Plugins/Mutable/Source/CustomizableObject/Private/MuCO/LogBenchmarkUtil.cpp
--- a/Engine/Plugins/Mutable/Source/CustomizableObject/Private/MuCO/LogBenchmarkUtil.cpp
+++ b/Engine/Plugins/Mutable/Source/CustomizableObject/Private/MuCO/LogBenchmarkUtil.cpp
@@ -104,12 +104,17 @@ void FLogBenchmarkUtil::GetInstancesStats(int32& OutNumInstances, int32& OutNumB
 		{
 			++OutNumBuiltInstances;
 		}
-		
-		for (int32 ComponentIndex = 0; ComponentIndex < Instance->GetNumComponents(); ++ComponentIndex)
+
+		UCustomizableObject* CO = Instance->GetCustomizableObject();
+		if (CO)
 		{
-			if (Instance->GetSkeletalMesh(ComponentIndex))
+			TSet<FName> InstanceComponents = Instance->GetComponents();
+			for (FName ComponentName: InstanceComponents)
 			{
-				++OutNumAllocatedSkeletalMeshes;
+				if (Instance->GetComponentMeshSkeletalMesh(ComponentName))
+				{
+					++OutNumAllocatedSkeletalMeshes;
+				}
 			}
 		}
 	}
@@ -268,8 +273,8 @@ void FLogBenchmarkUtil::UpdateStats()
 			GEngine->AddOnScreenDebugMessage(MsgIndex++, .0f, FColor::Magenta, TEXT("Meshes: "));
 
 			GEngine->AddOnScreenDebugMessage(MsgIndex++, .0f, FColor::Cyan,
-				TEXT("Player Pawn Mutable Mesh/Animation info for component ") + FString::Printf(TEXT("%d"),
-					CustomizableObjectInstanceUsage->GetComponentIndex()));
+				FString::Printf(TEXT("Player Pawn Mutable Mesh/Animation info for component %s"),
+					*CustomizableObjectInstanceUsage->GetComponentName().ToString()));
 		}
 	}
 
@@ -316,7 +321,7 @@ void FLogBenchmarkUtil::FinishUpdateMesh(const TSharedRef<FUpdateContextPrivate>
 	uint32 TriangleCount = 0;	
 	for (int32 ComponentIndex = 0; ComponentIndex < Object->GetComponentCount(); ComponentIndex++)
 	{
-		const FName ComponentName = Object->GetComponentName(ComponentIndex);
+		const FName ComponentName = Object->GetComponentName(FCustomizableObjectComponentIndex(ComponentIndex));
 		
 		// Process the generated components (not null)
 		if (const USkeletalMesh* InstanceSkeletalMesh = Instance->GetComponentMeshSkeletalMesh(ComponentName))
diff --git a/Engine/Plugins/Mutable/Source/CustomizableObject/Private/MuCO/LogInformationUtil.cpp b/Engine/Plugins/Mutable/Source/CustomizableObject/Private/MuCO/LogInformationUtil.cpp
deleted file mode 100644
diff --git a/Engine/Plugins/Mutable/Source/CustomizableObject/Private/MuCO/LogInformationUtil.h b/Engine/Plugins/Mutable/Source/CustomizableObject/Private/MuCO/LogInformationUtil.h
deleted file mode 100644
diff --git a/Engine/Plugins/Mutable/Source/CustomizableObject/Public/MuCO/CustomizableObject.h b/Engine/Plugins/Mutable/Source/CustomizableObject/Public/MuCO/CustomizableObject.h
--- a/Engine/Plugins/Mutable/Source/CustomizableObject/Public/MuCO/CustomizableObject.h
+++ b/Engine/Plugins/Mutable/Source/CustomizableObject/Public/MuCO/CustomizableObject.h
@@ -238,6 +238,42 @@ struct FMutableLODSettings
 };
 
 
+/** Strongly typed index for the index of a component in a UCustomizableObject. */
+USTRUCT(BlueprintType)
+struct FCustomizableObjectComponentIndex
+{
+	GENERATED_BODY()
+
+private:
+
+	int32 Index = 0;
+
+public:
+
+	explicit inline FCustomizableObjectComponentIndex(int32 InIndex=0)
+		: Index(InIndex)
+	{
+	}
+
+	inline void Invalidate()
+	{
+		Index = INDEX_NONE;
+	}
+
+	inline bool IsValid() const
+	{
+		return Index != INDEX_NONE;
+	}
+
+	inline int32 GetValue() const
+	{
+		return Index;
+	}
+
+	inline bool operator==(const FCustomizableObjectComponentIndex&) const = default;
+};
+
+
 UCLASS( BlueprintType, config=Engine )
 class CUSTOMIZABLEOBJECT_API UCustomizableObject : public UObject
 {
@@ -372,7 +408,7 @@ public:
 	/** Return the name of the component.
 	 *  @return NAME_None if the component does not exist. */
 	UFUNCTION(BlueprintCallable, Category = CustomizableObject)
-	FName GetComponentName(int32 ObjectComponentIndex) const;
+	FName GetComponentName(FCustomizableObjectComponentIndex ComponentIndex) const;
 	
 	/** Get the number of parameters available in instances of this object. */
 	UFUNCTION(BlueprintCallable, Category = CustomizableObject)
diff --git a/Engine/Plugins/Mutable/Source/CustomizableObject/Public/MuCO/CustomizableObjectInstance.h b/Engine/Plugins/Mutable/Source/CustomizableObject/Public/MuCO/CustomizableObjectInstance.h
--- a/Engine/Plugins/Mutable/Source/CustomizableObject/Public/MuCO/CustomizableObjectInstance.h
+++ b/Engine/Plugins/Mutable/Source/CustomizableObject/Public/MuCO/CustomizableObjectInstance.h
@@ -32,7 +32,6 @@ struct FTexturePlatformData;
 struct FMutableModelImageProperties;
 
 
-
 /**
  * Represents what kind of saving procedure was performed to save the package
  */
@@ -295,7 +294,7 @@ public:
 	UFUNCTION(BlueprintCallable, Category = CustomizableObjectInstance)
 	void SetCurrentState(const FString& StateName);
 
-	// DEPRECATED
+	// DEPRECATED: Use GetComponentMeshSkeletalMesh 
 	UFUNCTION(BlueprintCallable, Category = CustomizableObjectInstance)
 	USkeletalMesh* GetSkeletalMesh(int32 ComponentIndex = 0) const;
 
@@ -586,10 +585,17 @@ public:
 	void MultilayerProjectorUpdateLayer(const FName& ProjectorParamName, int32 Index, const FMultilayerProjectorLayer& Layer);
 	
 	// ------------------------------------------------------------
-	
+
+	/** Return the list of names of components generated for this instance. 
+	* This only has values when the instaance has been completely generated. 
+	*/
+	TSet<FName> GetComponents() const;
+
+	// ------------------------------------------------------------
+
 	/** Returns the animation BP for the parameter component and slot, gathered from all the meshes that compose this instance. */
 	UFUNCTION(BlueprintCallable, Category = CustomizableObjectInstance)
-	TSubclassOf<UAnimInstance> GetAnimBP(int32 ComponentIndex, const FName& Slot) const;
+	TSubclassOf<UAnimInstance> GetAnimBP(FName ComponentName, const FName& Slot) const;
 
 	/** Return the list of tags for this instance. */
 	UFUNCTION(BlueprintCallable, Category = CustomizableObjectInstance)
@@ -597,9 +603,16 @@ public:
 	
 	/** Execute a delegate for each animation instance involved in this customizable object instance. */
 	UFUNCTION(BlueprintCallable, Category = CustomizableObjectInstance)
-	void ForEachAnimInstance(int32 ComponentIndex, FEachComponentAnimInstanceClassDelegate Delegate) const;
+	void ForEachComponentAnimInstance(FName ComponentName, FEachComponentAnimInstanceClassDelegate Delegate) const;
+
+	void ForEachComponentAnimInstance(FName ComponentName, FEachComponentAnimInstanceClassNativeDelegate Delegate) const;
+
+	/** DEPRECATED: Use ForEachComponentAnimInstance. */
+	UFUNCTION(BlueprintCallable, Category = CustomizableObjectInstance)
+	void ForEachAnimInstance(int32 ObjectComponentIndex, FEachComponentAnimInstanceClassDelegate Delegate) const;
 
-	void ForEachAnimInstance(int32 ComponentIndex, FEachComponentAnimInstanceClassNativeDelegate Delegate) const;
+	/** DEPRECATED: Use ForEachComponentAnimInstance. */
+	void ForEachAnimInstance(int32 ObjectComponentIndex, FEachComponentAnimInstanceClassNativeDelegate Delegate) const;
 
 	/** Check if the given UAnimInstance class requires to be fixed up. */
 	bool AnimInstanceNeedsFixup(TSubclassOf<UAnimInstance> AnimInstance) const;
@@ -628,7 +641,7 @@ public:
 	/** If enabled, low-priority textures will generate resident mipmaps too. */
 	void SetForceGenerateResidentMips(bool bForceGenerateResidentMips);
 
-	const TArray<TObjectPtr<UMaterialInterface>>* GetOverrideMaterials(int32 ComponentIndex) const;
+	const TArray<TObjectPtr<UMaterialInterface>>* GetOverrideMaterials(FCustomizableObjectComponentIndex ComponentIndex) const;
 	
 	// The following methods should only be used in an LOD management class
 	void SetIsBeingUsedByComponentInPlay(bool bIsUsedByComponent);
@@ -639,7 +652,10 @@ public:
 	float GetMinSquareDistToPlayer() const;
 	void SetMinSquareDistToPlayer(float NewValue);
 
-	/** Return the number of components that are defined in the CO, unrelated to this instance parameters. */
+	/** Return the total number of components that can be generated with the CustomizableObject (unrelated to this instance parameters). 
+	* DEPRECATED: Get the CO with GetCustomizableObect and get the components there with GetComponentCount.
+	* If the actual components of this instance are required, use GetComponents to get a list of names.
+	*/
 	UFUNCTION(BlueprintCallable, Category = CustomizableObjectInstance)
 	int32 GetNumComponents() const;
 
@@ -655,7 +671,7 @@ public:
 	  * @param InMaxLOD - MaxLOD to generate - DEPRECATED.
 	  * @param InRequestedLODsPerComponent - Array with bitmasks of requested LODs per component with range from [0 .. CO->GetComponentCount()].
 	  * @param InOutRequestedUpdates - Map from Instance to Update data that stores a request for the Instance to be updated, which will be either processed or discarded by priority (to be rerequested the next tick) */
-	void SetRequestedLODs(int32 InMinLOD, int32 InMaxLOD, const TArray<uint16>& InRequestedLODsPerComponent, FMutableInstanceUpdateMap& InOutRequestedUpdates);
+	void SetRequestedLODs(int32 InMinLOD, int32 InMaxLOD, const TArray<uint16>& InRequestedLODsPerObjectComponent, FMutableInstanceUpdateMap& InOutRequestedUpdates);
 
 	const TArray<uint16>& GetRequestedLODsPerComponent() const;
 
diff --git a/Engine/Plugins/Mutable/Source/CustomizableObject/Public/MuCO/CustomizableObjectInstanceUsage.h b/Engine/Plugins/Mutable/Source/CustomizableObject/Public/MuCO/CustomizableObjectInstanceUsage.h
--- a/Engine/Plugins/Mutable/Source/CustomizableObject/Public/MuCO/CustomizableObjectInstanceUsage.h
+++ b/Engine/Plugins/Mutable/Source/CustomizableObject/Public/MuCO/CustomizableObjectInstanceUsage.h
@@ -3,6 +3,7 @@
 #pragma once
 
 #include "MuCO/CustomizableObjectInstance.h"
+#include "MuCO/CustomizableObject.h"
 #include "MuCO/CustomizableObjectInstanceUsagePrivate.h"
 #include "Tickable.h"
 #include "Engine/EngineBaseTypes.h"
@@ -40,10 +41,12 @@ public:
 	UFUNCTION(BlueprintCallable, Category = CustomizableObjectInstanceUsage)
 	UCustomizableObjectInstance* GetCustomizableObjectInstance() const;
 	
+	// DEPRECATED: use SetComponentName
 	void SetComponentIndex(int32 ObjectComponentIndex);
 
+	// DEPRECATED: use GetComponentName
 	int32 GetComponentIndex() const;
-	
+
 	UFUNCTION(BlueprintCallable, Category = CustomizableObjectInstanceUsage)
 	void SetComponentName(const FName& Name);
 
@@ -102,9 +105,9 @@ private:
 
 	// If no CustomizableSkeletalComponent is associated, this Index will be used
 	UPROPERTY()
-	int32 UsedComponentIndex;
+	FCustomizableObjectComponentIndex UsedComponentIndex_DEPRECATED;
 
-	/** Only used if the ComponentIndex is INDEX_NONE. */
+	/** Only used if the UsedComponentIndex is INDEX_NONE. */
 	UPROPERTY()
 	FName UsedComponentName;
 	
diff --git a/Engine/Plugins/Mutable/Source/CustomizableObject/Public/MuCO/CustomizableObjectSkeletalMesh.h b/Engine/Plugins/Mutable/Source/CustomizableObject/Public/MuCO/CustomizableObjectSkeletalMesh.h
--- a/Engine/Plugins/Mutable/Source/CustomizableObject/Public/MuCO/CustomizableObjectSkeletalMesh.h
+++ b/Engine/Plugins/Mutable/Source/CustomizableObject/Public/MuCO/CustomizableObjectSkeletalMesh.h
@@ -2,12 +2,11 @@
 
 #pragma once
 
-#include "CoreMinimal.h"
-#include "Engine/SkeletalMesh.h"
-
 #include "MuR/Ptr.h"
 #include "MuR/Model.h"
 #include "MuR/Parameters.h"
+#include "Engine/SkeletalMesh.h"
+#include "CoreMinimal.h"
 
 #include "CustomizableObjectSkeletalMesh.generated.h"
 
@@ -38,7 +37,8 @@ class CUSTOMIZABLEOBJECT_API UCustomizableObjectSkeletalMesh : public USkeletalM
 public: // Public Methods
 
 	static UCustomizableObjectSkeletalMesh* CreateSkeletalMesh(const TSharedRef<FUpdateContextPrivate>& InOperationData, 
-		const UCustomizableObjectInstance& InInstance, UCustomizableObject& CustomizableObject, const int32 InstanceComponentIndex);
+		const UCustomizableObjectInstance& InInstance, UCustomizableObject& CustomizableObject, 
+		const int32 FirstLOD, const int32 LODCount);
 
 private: // Private Methods
 
diff --git a/Engine/Plugins/Mutable/Source/CustomizableObject/Public/MuCO/CustomizableObjectSystem.h b/Engine/Plugins/Mutable/Source/CustomizableObject/Public/MuCO/CustomizableObjectSystem.h
--- a/Engine/Plugins/Mutable/Source/CustomizableObject/Public/MuCO/CustomizableObjectSystem.h
+++ b/Engine/Plugins/Mutable/Source/CustomizableObject/Public/MuCO/CustomizableObjectSystem.h
@@ -275,9 +275,6 @@ public:
 	// being used and it's not compiled.  Callers can add additional information to the error log.
 	void AddUncompiledCOWarning(const UCustomizableObject& InObject, FString const* OptionalLogInfo = nullptr);
 
-	// Show data about all UCustomizableObjectInstance existing elements
-	void LogShowData(bool bFullInfo, bool ShowMaterialInfo) const;
-
 	// Give access to the internal object data.
 	UCustomizableObjectSystemPrivate* GetPrivate();
 	const UCustomizableObjectSystemPrivate* GetPrivate() const;
diff --git a/Engine/Plugins/Mutable/Source/CustomizableObject/Public/MuCO/CustomizableSkeletalComponent.h b/Engine/Plugins/Mutable/Source/CustomizableObject/Public/MuCO/CustomizableSkeletalComponent.h
--- a/Engine/Plugins/Mutable/Source/CustomizableObject/Public/MuCO/CustomizableSkeletalComponent.h
+++ b/Engine/Plugins/Mutable/Source/CustomizableObject/Public/MuCO/CustomizableSkeletalComponent.h
@@ -5,6 +5,7 @@
 #include "Components/SceneComponent.h"
 
 #include "MuCO/CustomizableObjectInstance.h"
+#include "MuCO/CustomizableObject.h"
 
 #include "CustomizableSkeletalComponent.generated.h"
 
@@ -24,9 +25,11 @@ public:
 	UPROPERTY(BlueprintReadWrite, EditAnywhere, Category = CustomizableSkeletalComponent)
 	TObjectPtr<UCustomizableObjectInstance> CustomizableObjectInstance;
 
-	/** This component index refers to the object list of components */
+	/** This component index refers to the object list of components 
+	* DEPRECATED: Use FNames instead with Get/SetComponentName
+	*/
 	UPROPERTY(BlueprintReadWrite, EditAnywhere, Category = CustomizableSkeletalComponent)
-	int32 ComponentIndex = 0;
+	int32 ComponentIndex;
 
 private:
 	/** Only used if the ComponentIndex is INDEX_NONE.
diff --git a/Engine/Plugins/Mutable/Source/CustomizableObject/Public/MuCO/CustomizableSkeletalMeshActor.h b/Engine/Plugins/Mutable/Source/CustomizableObject/Public/MuCO/CustomizableSkeletalMeshActor.h
--- a/Engine/Plugins/Mutable/Source/CustomizableObject/Public/MuCO/CustomizableSkeletalMeshActor.h
+++ b/Engine/Plugins/Mutable/Source/CustomizableObject/Public/MuCO/CustomizableSkeletalMeshActor.h
@@ -17,9 +17,9 @@ class CUSTOMIZABLEOBJECT_API ACustomizableSkeletalMeshActor : public ASkeletalMe
 	GENERATED_UCLASS_BODY()
 public:
 	
-	class UCustomizableSkeletalComponent* GetCustomizableSkeletalComponent(int32 Index = 0) const { return CustomizableSkeletalComponents[Index]; }
+	class UCustomizableSkeletalComponent* GetCustomizableSkeletalComponent(int32 ObjectComponentIndex = 0) const { return CustomizableSkeletalComponents[ObjectComponentIndex]; }
 	
-	class USkeletalMeshComponent* GetSkeletalMeshComponentAt(int32 Index = 0) const { return SkeletalMeshComponents[Index]; }
+	class USkeletalMeshComponent* GetSkeletalMeshComponentAt(int32 ObjectComponentIndex = 0) const { return SkeletalMeshComponents[ObjectComponentIndex]; }
 
 	void AttachNewComponent();
 
@@ -44,9 +44,11 @@ private:
 
 private:
 	
+	/** TODO: There are as many as components in the CO (not in the COI). */
 	UPROPERTY(Category = CustomizableSkeletalMesh, VisibleAnywhere, BlueprintReadOnly, meta = (AllowPrivateAccess = "true"))
 	TArray< TObjectPtr<class UCustomizableSkeletalComponent> > CustomizableSkeletalComponents;
 
+	/** TODO: There are as many as components in the CO (not in the COI). */
 	UPROPERTY(Category = CustomizableSkeletalMesh, VisibleAnywhere, BlueprintReadOnly, meta = (AllowPrivateAccess = "true"))
 	TArray< TObjectPtr<class USkeletalMeshComponent>> SkeletalMeshComponents;
 
diff --git a/Engine/Plugins/Mutable/Source/CustomizableObjectEditor/Private/MuCOE/CustomizableObjectBenchmarkingUtils.cpp b/Engine/Plugins/Mutable/Source/CustomizableObjectEditor/Private/MuCOE/CustomizableObjectBenchmarkingUtils.cpp
--- a/Engine/Plugins/Mutable/Source/CustomizableObjectEditor/Private/MuCOE/CustomizableObjectBenchmarkingUtils.cpp
+++ b/Engine/Plugins/Mutable/Source/CustomizableObjectEditor/Private/MuCOE/CustomizableObjectBenchmarkingUtils.cpp
@@ -33,7 +33,7 @@ bool CustomizableObjectBenchmarkingUtils::GenerateDeterministicSetOfInstances(TO
 		{
 			// Force generation of all LODS
 			TArray<uint16> RequestedLodLevels{};
-			RequestedLodLevels.Init(0, GeneratedInstance->GetNumComponents());
+			RequestedLodLevels.Init(0, TargetCustomizableObject->GetComponentCount());
 			GeneratedInstance->GetPrivate()->GetDescriptor().SetRequestedLODLevels(RequestedLodLevels);
 				
 			// Randomize instance values
diff --git a/Engine/Plugins/Mutable/Source/CustomizableObjectEditor/Private/MuCOE/CustomizableObjectCompiler.cpp b/Engine/Plugins/Mutable/Source/CustomizableObjectEditor/Private/MuCOE/CustomizableObjectCompiler.cpp
--- a/Engine/Plugins/Mutable/Source/CustomizableObjectEditor/Private/MuCOE/CustomizableObjectCompiler.cpp
+++ b/Engine/Plugins/Mutable/Source/CustomizableObjectEditor/Private/MuCOE/CustomizableObjectCompiler.cpp
@@ -987,7 +987,7 @@ void FCustomizableObjectCompiler::CompileInternal(bool bAsync)
 		ModelResources.CustomizableObjectPathMap = MoveTemp(GenerationContext.CustomizableObjectPathMap);
 #endif
 
-		ModelResources.ComponentNames = MoveTemp(GenerationContext.ComponentNames);
+		ModelResources.ComponentNamesPerObjectComponent = MoveTemp(GenerationContext.ComponentNames);
 
 		if (ICustomizableObjectVersionBridgeInterface* VersionBridge = Cast<ICustomizableObjectVersionBridgeInterface>(GraphTraversal::GetRootObject(CurrentObject)->VersionBridge))
 		{
diff --git a/Engine/Plugins/Mutable/Source/CustomizableObjectEditor/Private/MuCOE/CustomizableObjectEditor.h b/Engine/Plugins/Mutable/Source/CustomizableObjectEditor/Private/MuCOE/CustomizableObjectEditor.h
--- a/Engine/Plugins/Mutable/Source/CustomizableObjectEditor/Private/MuCOE/CustomizableObjectEditor.h
+++ b/Engine/Plugins/Mutable/Source/CustomizableObjectEditor/Private/MuCOE/CustomizableObjectEditor.h
@@ -143,7 +143,7 @@ public:
 	
 	/** The Component of the texture */
 	UPROPERTY(VisibleAnywhere, AssetRegistrySearchable, Category = "Stats", meta = (DisplayName = "Component", ColumnWidth = "70", NoResetToDefault))
-	int32 Component;
+	FName ComponentName;
 
 };
 
diff --git a/Engine/Plugins/Mutable/Source/CustomizableObjectEditor/Private/MuCOE/CustomizableObjectEditorViewportClient.cpp b/Engine/Plugins/Mutable/Source/CustomizableObjectEditor/Private/MuCOE/CustomizableObjectEditorViewportClient.cpp
--- a/Engine/Plugins/Mutable/Source/CustomizableObjectEditor/Private/MuCOE/CustomizableObjectEditorViewportClient.cpp
+++ b/Engine/Plugins/Mutable/Source/CustomizableObjectEditor/Private/MuCOE/CustomizableObjectEditorViewportClient.cpp
@@ -823,7 +823,7 @@ void FCustomizableObjectEditorViewportClient::OnInstanceUpdate(UCustomizableObje
 	// Add new components
 	for (int32 ObjectComponentIndex = 0; ObjectComponentIndex < CO->GetComponentCount(); ++ObjectComponentIndex)
 	{
-		FName Name = CO->GetComponentName(ObjectComponentIndex);
+		FName Name = CO->GetComponentName(FCustomizableObjectComponentIndex(ObjectComponentIndex));
 		if (Instance->GetComponentMeshSkeletalMesh(Name))
 		{
 			if (!SkeletalMeshComponents.Contains(Name))
diff --git a/Engine/Plugins/Mutable/Source/CustomizableObjectEditor/Private/MuCOE/CustomizableObjectInstanceBakingUtils.cpp b/Engine/Plugins/Mutable/Source/CustomizableObjectEditor/Private/MuCOE/CustomizableObjectInstanceBakingUtils.cpp
--- a/Engine/Plugins/Mutable/Source/CustomizableObjectEditor/Private/MuCOE/CustomizableObjectInstanceBakingUtils.cpp
+++ b/Engine/Plugins/Mutable/Source/CustomizableObjectEditor/Private/MuCOE/CustomizableObjectInstanceBakingUtils.cpp
@@ -434,10 +434,11 @@ bool BakeCustomizableObjectInstance(
 	}
 	
 	
-	const int32 NumComponents = InInstance.GetNumComponents();
-	for (int32 ComponentIndex = 0; ComponentIndex < NumComponents; ++ComponentIndex)
+	const int32 NumObjectComponents = InstanceCO->GetComponentCount();
+	for (int32 ObjectComponentIndex = 0; ObjectComponentIndex < NumObjectComponents; ++ObjectComponentIndex)
 	{
-		USkeletalMesh* Mesh = InInstance.GetSkeletalMesh(ComponentIndex);
+		FName ComponentName = InstanceCO->GetComponentName(FCustomizableObjectComponentIndex(ObjectComponentIndex));
+		USkeletalMesh* Mesh = InInstance.GetComponentMeshSkeletalMesh(ComponentName);
 
 		if (!Mesh)
 		{
@@ -445,9 +446,9 @@ bool BakeCustomizableObjectInstance(
 		}
 
 		FString ObjectName = FileName;
-		if (NumComponents > 1)
+		if (NumObjectComponents > 1)
 		{
-			ObjectName = FileName + "_Component_" + FString::FromInt(ComponentIndex);
+			ObjectName = FileName + "_Component_" + ComponentName.ToString();
 		}
 
 		TMap<UObject*, UObject*> ReplacementMap;
@@ -797,7 +798,7 @@ bool BakeCustomizableObjectInstance(
 
 			// Add Instance Info in a custom AssetUserData
 			{
-				const FCustomizableInstanceComponentData* ComponentData = InInstance.GetPrivate()->GetComponentData(ComponentIndex);
+				const FCustomizableInstanceComponentData* ComponentData = InInstance.GetPrivate()->GetComponentData(ComponentName);
 				check(ComponentData);
 				
 				if (InInstance.GetAnimationGameplayTags().Num() ||
@@ -833,9 +834,9 @@ bool BakeCustomizableObjectInstance(
 
 			// Copy LODSettings from the Reference Skeletal Mesh
 			{
-				if (InstanceCO && InstanceCO->GetPrivate()->GetModelResources().ReferenceSkeletalMeshesData.IsValidIndex(ComponentIndex))
+				if (InstanceCO && InstanceCO->GetPrivate()->GetModelResources().ReferenceSkeletalMeshesData.IsValidIndex(ObjectComponentIndex))
 				{
-					USkeletalMeshLODSettings* LODSettings = InstanceCO->GetPrivate()->GetModelResources().ReferenceSkeletalMeshesData[ComponentIndex].SkeletalMeshLODSettings;
+					USkeletalMeshLODSettings* LODSettings = InstanceCO->GetPrivate()->GetModelResources().ReferenceSkeletalMeshesData[ObjectComponentIndex].SkeletalMeshLODSettings;
 					SkeletalMesh->SetLODSettings(LODSettings);
 				}
 			}
diff --git a/Engine/Plugins/Mutable/Source/CustomizableObjectEditor/Private/MuCOE/CustomizableObjectInstanceFactory.cpp b/Engine/Plugins/Mutable/Source/CustomizableObjectEditor/Private/MuCOE/CustomizableObjectInstanceFactory.cpp
--- a/Engine/Plugins/Mutable/Source/CustomizableObjectEditor/Private/MuCOE/CustomizableObjectInstanceFactory.cpp
+++ b/Engine/Plugins/Mutable/Source/CustomizableObjectEditor/Private/MuCOE/CustomizableObjectInstanceFactory.cpp
@@ -49,9 +49,10 @@ void UCustomizableObjectInstanceFactory::PostSpawnActor(UObject* Asset, AActor*
 		return;
 	}
 	
-	for (int32 ComponentIndex = 0; ComponentIndex < Instance->GetNumComponents(); ++ComponentIndex)
+
+	for (int32 ObjectComponentIndex = 0; ObjectComponentIndex < Object->GetComponentCount(); ++ObjectComponentIndex)
 	{
-		const FName& ComponentName = Object->GetComponentName(ComponentIndex);
+		const FName& ComponentName = Object->GetComponentName(FCustomizableObjectComponentIndex(ObjectComponentIndex));
 		
 		USkeletalMesh* SkeletalMesh = Instance->GetComponentMeshSkeletalMesh(ComponentName);
 		if (!SkeletalMesh)
@@ -59,22 +60,22 @@ void UCustomizableObjectInstanceFactory::PostSpawnActor(UObject* Asset, AActor*
 			SkeletalMesh = Object->GetComponentMeshReferenceSkeletalMesh(ComponentName);
 		}
 		
-		if (ComponentIndex > 0 && ComponentIndex > NewCSMActor->GetNumComponents() - 1)
+		if (ObjectComponentIndex > 0 && ObjectComponentIndex > NewCSMActor->GetNumComponents() - 1)
 		{
 			NewCSMActor->AttachNewComponent();
 		}
 
-		if (USkeletalMeshComponent* SkeletalMeshComp = NewCSMActor->GetSkeletalMeshComponentAt(ComponentIndex))
+		if (USkeletalMeshComponent* SkeletalMeshComp = NewCSMActor->GetSkeletalMeshComponentAt(ObjectComponentIndex))
 		{
 			SkeletalMeshComp->UnregisterComponent();
 			UE_MUTABLE_SETSKINNEDASSET(SkeletalMeshComp, SkeletalMesh);
 
-			if (ComponentIndex == 0 && NewCSMActor->GetWorld()->IsGameWorld())
+			if (ObjectComponentIndex == 0 && NewCSMActor->GetWorld()->IsGameWorld())
 			{
 				NewCSMActor->ReplicatedMesh = SkeletalMesh;
 			}
 
-			if (UCustomizableSkeletalComponent* CustomSkeletalComp = NewCSMActor->GetCustomizableSkeletalComponent(ComponentIndex))
+			if (UCustomizableSkeletalComponent* CustomSkeletalComp = NewCSMActor->GetCustomizableSkeletalComponent(ObjectComponentIndex))
 			{
 				CustomSkeletalComp->UnregisterComponent();
 				CustomSkeletalComp->CustomizableObjectInstance = Instance;
@@ -120,12 +121,19 @@ bool UCustomizableObjectInstanceFactory::CanCreateActorFrom(const FAssetData& As
     {
         if (UCustomizableObjectInstance* CustomizableObjectInstance = Cast<UCustomizableObjectInstance>(AssetData.GetAsset()))
         {
-            if (USkeletalMesh* SkeletalMesh = CustomizableObjectInstance->GetSkeletalMesh())
+			TSet<FName> ComponentNames = CustomizableObjectInstance->GetComponents();
+
+			// If there is a valid component with skeletal mesh
+			// TODO: This is not accurate in the future, if we have components that have grooms or clothing instead.
+            for (FName ComponentName : ComponentNames)
             {
-                return true;
+				if (USkeletalMesh* SkeletalMesh = CustomizableObjectInstance->GetComponentMeshSkeletalMesh(ComponentName))
+				{
+					return true;
+				}
             }
-            else
-            {
+
+			{
                 if (UCustomizableObject* CustomizableObject = CustomizableObjectInstance->GetCustomizableObject())
                 {
                     return true;
diff --git a/Engine/Plugins/Mutable/Source/CustomizableObjectEditor/Private/MuCOE/SCustomizableObjectEditorTextureAnalyzer.cpp b/Engine/Plugins/Mutable/Source/CustomizableObjectEditor/Private/MuCOE/SCustomizableObjectEditorTextureAnalyzer.cpp
--- a/Engine/Plugins/Mutable/Source/CustomizableObjectEditor/Private/MuCOE/SCustomizableObjectEditorTextureAnalyzer.cpp
+++ b/Engine/Plugins/Mutable/Source/CustomizableObjectEditor/Private/MuCOE/SCustomizableObjectEditorTextureAnalyzer.cpp
@@ -133,9 +133,9 @@ void SCustomizableObjecEditorTextureAnalyzer::FillTextureAnalyzerTable(UCustomiz
 		const UCustomizableObject* CustomizableObject = PreviewInstance->GetCustomizableObject();
 		check(CustomizableObject)
 
-		for (int32 ComponentIndex = 0; ComponentIndex < PreviewInstance->GetNumComponents(); ++ComponentIndex)
+		for (int32 ObjectComponentIndex = 0; ObjectComponentIndex < CustomizableObject->GetComponentCount(); ++ObjectComponentIndex)
 		{
-			const FName ComponentName = CustomizableObject->GetComponentName(ComponentIndex);
+			const FName ComponentName = CustomizableObject->GetComponentName(FCustomizableObjectComponentIndex(ObjectComponentIndex));
 			
 			if (const USkeletalMesh* SkeletalMesh = PreviewInstance->GetComponentMeshSkeletalMesh(ComponentName))
 			{
@@ -178,7 +178,7 @@ void SCustomizableObjecEditorTextureAnalyzer::FillTextureAnalyzerTable(UCustomiz
 								Entry->MaterialName = Material->GetName();
 								Entry->ParentMaterial = Material->Parent;
 								Entry->MaterialParameterName = Material->Parent ? Material->Parent->GetName() : FString();
-								Entry->Component = ComponentIndex;
+								Entry->ComponentName = ComponentName;
 
 								//Total size
 								TotalSize += textureSize;
diff --git a/Engine/Plugins/Mutable/Source/CustomizableObjectEditor/Private/MuCOE/Thumbnails/CustomizableObjectInstanceThumbnailRenderer.cpp b/Engine/Plugins/Mutable/Source/CustomizableObjectEditor/Private/MuCOE/Thumbnails/CustomizableObjectInstanceThumbnailRenderer.cpp
--- a/Engine/Plugins/Mutable/Source/CustomizableObjectEditor/Private/MuCOE/Thumbnails/CustomizableObjectInstanceThumbnailRenderer.cpp
+++ b/Engine/Plugins/Mutable/Source/CustomizableObjectEditor/Private/MuCOE/Thumbnails/CustomizableObjectInstanceThumbnailRenderer.cpp
@@ -72,7 +72,7 @@ void UCustomizableObjectInstanceThumbnailRenderer::Draw(UObject* Object, int32 X
 		}
 
 		// The compilation might fail because the graph is not valid or incomplete, so it has to be checked again
-		if (CustomizableObject->GetRefSkeletalMesh() && CustomizableObject->IsCompiled()) // TODO MultiComp
+		if (CustomizableObject->IsCompiled()) // TODO MultiComp
 		{
 			CustomizableObjectInstance->SetObject(CustomizableObject);
 			
diff --git a/Engine/Plugins/Mutable/Source/CustomizableObjectEditor/Private/MuCOE/Thumbnails/CustomizableObjectThumbnailRenderer.cpp b/Engine/Plugins/Mutable/Source/CustomizableObjectEditor/Private/MuCOE/Thumbnails/CustomizableObjectThumbnailRenderer.cpp
--- a/Engine/Plugins/Mutable/Source/CustomizableObjectEditor/Private/MuCOE/Thumbnails/CustomizableObjectThumbnailRenderer.cpp
+++ b/Engine/Plugins/Mutable/Source/CustomizableObjectEditor/Private/MuCOE/Thumbnails/CustomizableObjectThumbnailRenderer.cpp
@@ -61,7 +61,7 @@ void UCustomizableObjectThumbnailRenderer::Draw(UObject* Object, int32 X, int32
 		}
 
 		// The compilation might fail because the graph is not valid or incomplete, so it has to be checked again
-		if (CustomizableObject->GetRefSkeletalMesh() && CustomizableObject->IsCompiled()) // TODO MultiComp
+		if (CustomizableObject->IsCompiled()) // TODO MultiComp
 		{
 			CustomizableObjectInstance = CustomizableObject->CreateInstance();
 
diff --git a/Engine/Plugins/Mutable/Source/MutableValidation/Private/MuV/CustomizableObjectInstanceUpdateUtility.cpp b/Engine/Plugins/Mutable/Source/MutableValidation/Private/MuV/CustomizableObjectInstanceUpdateUtility.cpp
--- a/Engine/Plugins/Mutable/Source/MutableValidation/Private/MuV/CustomizableObjectInstanceUpdateUtility.cpp
+++ b/Engine/Plugins/Mutable/Source/MutableValidation/Private/MuV/CustomizableObjectInstanceUpdateUtility.cpp
@@ -106,12 +106,16 @@ void FCustomizableObjectInstanceUpdateUtility::OnInstanceUpdateResult(const FUpd
 		UE_LOG(LogMutable, Display, TEXT("Instance %s rquesting streaming all MIPs."), *InstanceName);
 
 		check(ComponentsBeingUpdated.IsEmpty());
-		for (int32 Index = 0; Index < Instance->GetNumComponents(); ++Index)
+
+		UCustomizableObject* CO = Instance->GetCustomizableObject();
+
+		TSet<FName> InstanceComponents = Instance->GetComponents();
+		for (FName ComponentName : InstanceComponents)
 		{
 			TStrongObjectPtr<USkeletalMeshComponent> SkeletalComponent = TStrongObjectPtr(NewObject<USkeletalMeshComponent>());
-			SkeletalComponent->SetSkeletalMesh(Instance->GetSkeletalMesh(Index));
-			
-			ComponentsBeingUpdated.Add(SkeletalComponent);            			
+			SkeletalComponent->SetSkeletalMesh(Instance->GetComponentMeshSkeletalMesh(ComponentName));
+
+			ComponentsBeingUpdated.Add(SkeletalComponent);
 		}
 
 		// Request the streaming in of all the components affected by the update
