[Mutable] Fix GetSkeletalMeshSectionFor crash.

Crash due to the Skeletal Mesh changing without updating nodes were it was used (Section or LOD dissapeared).

#rnx

diff --git a/Engine/Plugins/Mutable/Source/CustomizableObjectEditor/Private/MuCOE/Nodes/CustomizableObjectNodeSkeletalMesh.cpp b/Engine/Plugins/Mutable/Source/CustomizableObjectEditor/Private/MuCOE/Nodes/CustomizableObjectNodeSkeletalMesh.cpp
--- a/Engine/Plugins/Mutable/Source/CustomizableObjectEditor/Private/MuCOE/Nodes/CustomizableObjectNodeSkeletalMesh.cpp
+++ b/Engine/Plugins/Mutable/Source/CustomizableObjectEditor/Private/MuCOE/Nodes/CustomizableObjectNodeSkeletalMesh.cpp
@@ -765,21 +765,20 @@ const FSkelMeshSection* UCustomizableObjectNodeSkeletalMesh::GetSkeletalMeshSect
 	{
 		return nullptr;
 	}
+	
 	const FSkeletalMeshModel* ImportedModel = SkeletalMesh->GetImportedModel();
-
 	if (!ImportedModel)
 	{
 		return nullptr;
 	}
 
-	if (LODIndex < 0 && LODIndex >= ImportedModel->LODModels.Num())
+	if (!ImportedModel->LODModels.IsValidIndex(LODIndex))
 	{
 		return nullptr;
 	}
 
 	const FSkeletalMeshLODModel& LODModel = ImportedModel->LODModels[LODIndex];
-
-	if (SectionIndex < 0 && SectionIndex >= LODModel.Sections.Num())
+	if (!LODModel.Sections.IsValidIndex(SectionIndex))
 	{
 		return nullptr;
 	}
