[Mutable] Changed Mutable Tasks to Background priority.

#rnx

diff --git a/Engine/Plugins/Mutable/Source/CustomizableObject/Private/MuCO/CustomizableObjectMeshUpdate.cpp b/Engine/Plugins/Mutable/Source/CustomizableObject/Private/MuCO/CustomizableObjectMeshUpdate.cpp
--- a/Engine/Plugins/Mutable/Source/CustomizableObject/Private/MuCO/CustomizableObjectMeshUpdate.cpp
+++ b/Engine/Plugins/Mutable/Source/CustomizableObject/Private/MuCO/CustomizableObjectMeshUpdate.cpp
@@ -147,7 +147,8 @@ namespace impl
 
 				Task_Mutable_UpdateMesh_Loop(OperationData, Task, InstanceID, LODIndex + 1);
 			},
-			GetMeshTask));
+			GetMeshTask,
+			LowLevelTasks::ETaskPriority::Inherit));
 	}
 
 	void Task_Mutable_UpdateMesh(const TSharedPtr<FMutableMeshOperationData> OperationData, TRefCountPtr<FCustomizableObjectMeshStreamIn>& Task)
diff --git a/Engine/Plugins/Mutable/Source/CustomizableObject/Private/MuCO/CustomizableObjectSystem.cpp b/Engine/Plugins/Mutable/Source/CustomizableObject/Private/MuCO/CustomizableObjectSystem.cpp
--- a/Engine/Plugins/Mutable/Source/CustomizableObject/Private/MuCO/CustomizableObjectSystem.cpp
+++ b/Engine/Plugins/Mutable/Source/CustomizableObject/Private/MuCO/CustomizableObjectSystem.cpp
@@ -2179,7 +2179,8 @@ namespace impl
 		{
 			Task_Mutable_GetMeshes_GetImage_Post(OperationData, StartTime, GetImagesData, GetImageIndex, GetImageTask);
 		},
-		GetImageTask));
+		GetImageTask,
+		LowLevelTasks::ETaskPriority::Inherit));
 	}
 
 
@@ -2440,7 +2441,8 @@ namespace impl
 		{
 			Task_MutableGetMeshes_GetMesh_Post(OperationData, StartTime, GetMeshesData, GetMeshIndex, GetMeshTask);
 		},
-		GetMeshTask));
+		GetMeshTask,
+		LowLevelTasks::ETaskPriority::Inherit));
 	}
 
 
@@ -2579,7 +2581,8 @@ namespace impl
 		{
 			Task_Mutable_GetImages_GetImage(OperationData, StartTime, ImagesInThisInstance, ImageIndex, GetImageDescTask);
 		},
-		GetImageDescTask));
+		GetImageDescTask,
+		LowLevelTasks::ETaskPriority::Inherit));
 	}
 
 
@@ -2722,7 +2725,8 @@ namespace impl
 			{
 				Task_Mutable_GetImages_GetImage_Post(OperationData, StartTime, ImagesInThisInstance, ImageIndex, GetImageTask, MipSizeX, MipSizeY, FullLODCount, MipsToSkip);
 			},
-			GetImageTask));
+			GetImageTask,
+			LowLevelTasks::ETaskPriority::Inherit));
 		}
 	}
 
