[Mutable] Made USkeletalMesh::ReleaseResources optional.

In some scenarios was causing GameThread spikes. Disabled by default.

#rnx

diff --git a/Engine/Plugins/Mutable/Source/CustomizableObject/Private/MuCO/CustomizableObjectInstance.cpp b/Engine/Plugins/Mutable/Source/CustomizableObject/Private/MuCO/CustomizableObjectInstance.cpp
--- a/Engine/Plugins/Mutable/Source/CustomizableObject/Private/MuCO/CustomizableObjectInstance.cpp
+++ b/Engine/Plugins/Mutable/Source/CustomizableObject/Private/MuCO/CustomizableObjectInstance.cpp
@@ -3011,6 +3011,8 @@ bool UCustomizableInstancePrivate::IsSelectedParameterProfileDirty() const
 
 void UCustomizableInstancePrivate::DiscardResources()
 {
+	check(IsInGameThread());
+
 	UCustomizableObjectInstance* Instance = Cast<UCustomizableObjectInstance>(GetOuter());
 	if (!Instance)
 	{
@@ -3019,13 +3021,16 @@ void UCustomizableInstancePrivate::DiscardResources()
 
 	if (SkeletalMeshStatus == ESkeletalMeshStatus::Success)
 	{
-		for (const TTuple<FName, TObjectPtr<USkeletalMesh>>& Tuple : SkeletalMeshes)
+		if (CVarEnableReleaseMeshResources.GetValueOnGameThread())
 		{
-			USkeletalMesh* SkeletalMesh = Tuple.Get<1>();
-			
-			if (SkeletalMesh->IsValidLowLevel())
+			for (const TTuple<FName, TObjectPtr<USkeletalMesh>>& Tuple : SkeletalMeshes)
 			{
-				SkeletalMesh->ReleaseResources();
+				USkeletalMesh* SkeletalMesh = Tuple.Get<1>();
+			
+				if (SkeletalMesh->IsValidLowLevel())
+				{
+					SkeletalMesh->ReleaseResources();
+				}
 			}
 		}
 		
diff --git a/Engine/Plugins/Mutable/Source/CustomizableObject/Private/MuCO/CustomizableObjectSystem.cpp b/Engine/Plugins/Mutable/Source/CustomizableObject/Private/MuCO/CustomizableObjectSystem.cpp
--- a/Engine/Plugins/Mutable/Source/CustomizableObject/Private/MuCO/CustomizableObjectSystem.cpp
+++ b/Engine/Plugins/Mutable/Source/CustomizableObject/Private/MuCO/CustomizableObjectSystem.cpp
@@ -130,6 +130,11 @@ static FAutoConsoleVariableRef CVarMutableEnableLODManagmentInEditor(
 	ECVF_Default);
 #endif
 
+TAutoConsoleVariable<bool> CVarEnableReleaseMeshResources(
+	TEXT("mutable.EnableReleaseMeshResources"),
+	true,
+	TEXT("Allow releasing resources when discarding instances."));
+
 int32 UCustomizableObjectSystemPrivate::SkeletalMeshMinLodQualityLevel = -1;
 
 
diff --git a/Engine/Plugins/Mutable/Source/CustomizableObject/Public/MuCO/CustomizableObjectSystem.h b/Engine/Plugins/Mutable/Source/CustomizableObject/Public/MuCO/CustomizableObjectSystem.h
--- a/Engine/Plugins/Mutable/Source/CustomizableObject/Public/MuCO/CustomizableObjectSystem.h
+++ b/Engine/Plugins/Mutable/Source/CustomizableObject/Public/MuCO/CustomizableObjectSystem.h
@@ -41,6 +41,8 @@ extern TAutoConsoleVariable<bool> CVarEnableRealTimeMorphTargets;
 
 extern TAutoConsoleVariable<bool> CVarRollbackFixModelDiskStreamerDataRace;
 
+extern TAutoConsoleVariable<bool> CVarEnableReleaseMeshResources;
+
 
 #if WITH_EDITOR
 
