[Mutable] Minor scopes improvements.

- Fixed incorrect scope name.
- Renamed UpdateImage scope.

#rnx

diff --git a/Engine/Plugins/Mutable/Source/CustomizableObject/Private/MuCO/CustomizableObjectMipDataProvider.cpp b/Engine/Plugins/Mutable/Source/CustomizableObject/Private/MuCO/CustomizableObjectMipDataProvider.cpp
--- a/Engine/Plugins/Mutable/Source/CustomizableObject/Private/MuCO/CustomizableObjectMipDataProvider.cpp
+++ b/Engine/Plugins/Mutable/Source/CustomizableObject/Private/MuCO/CustomizableObjectMipDataProvider.cpp
@@ -13,7 +13,7 @@
 #include UE_INLINE_GENERATED_CPP_BY_NAME(CustomizableObjectMipDataProvider)
 
 
-#define UE_MUTABLE_MIPDATA_PROVIDER_UPDATE_IMAGE_REGION		TEXT("MipDataProvider_Mutable_UpdateImage")
+#define UE_MUTABLE_MIPDATA_PROVIDER_UPDATE_IMAGE_REGION		TEXT("Task_Mutable_UpdateImage")
 
 UMutableTextureMipDataProviderFactory::UMutableTextureMipDataProviderFactory(const FObjectInitializer& ObjectInitializer)
 	: Super(ObjectInitializer)
@@ -133,6 +133,8 @@ namespace Impl
 {
 	void Task_Mutable_UpdateImage(TSharedPtr<FMutableImageOperationData> OperationData)
 	{
+		MUTABLE_CPUPROFILER_SCOPE(Task_Mutable_UpdateImage)
+		
 		const double StartTime = FPlatformTime::Seconds();
 		
 		if (FLogBenchmarkUtil::IsBenchmarkingReportingEnabled())
diff --git a/Engine/Plugins/Mutable/Source/MutableRuntime/Private/MuR/System.cpp b/Engine/Plugins/Mutable/Source/MutableRuntime/Private/MuR/System.cpp
--- a/Engine/Plugins/Mutable/Source/MutableRuntime/Private/MuR/System.cpp
+++ b/Engine/Plugins/Mutable/Source/MutableRuntime/Private/MuR/System.cpp
@@ -566,7 +566,7 @@ namespace mu
 	UE::Tasks::TTask<Ptr<const Mesh>> System::GetMesh(Instance::ID instanceID, FResourceID MeshId)
     {
 		LLM_SCOPE_BYNAME(TEXT("MutableRuntime"));
-		MUTABLE_CPUPROFILER_SCOPE(SystemGetImage);
+		MUTABLE_CPUPROFILER_SCOPE(SystemGetMesh);
 
 		Ptr<const Mesh> ResultMesh;
 
