[mutable] Fix crash with the use of passthrough components.
- Make sure the component name is added.
- Make sure some LODs exist.
- Initialize internal skeleton array with the instance components count and not the object component count.
- Add some defensive early outs in case of other internal problems.
#rnx

diff --git a/Engine/Plugins/Mutable/Source/CustomizableObject/Private/MuCO/CustomizableObjectInstance.cpp b/Engine/Plugins/Mutable/Source/CustomizableObject/Private/MuCO/CustomizableObjectInstance.cpp
--- a/Engine/Plugins/Mutable/Source/CustomizableObject/Private/MuCO/CustomizableObjectInstance.cpp
+++ b/Engine/Plugins/Mutable/Source/CustomizableObject/Private/MuCO/CustomizableObjectInstance.cpp
@@ -1980,7 +1980,12 @@ bool UCustomizableInstancePrivate::DoComponentsNeedUpdate(UCustomizableObjectIns
 	{
 		const FInstanceUpdateData::FComponent& Component = OperationData->InstanceUpdateData.Components[InstanceComponentIndex];
 		const FCustomizableObjectComponentIndex ObjectComponentIndex = Component.Id;
-		const FName ComponentName = OperationData->Instance->GetCustomizableObject()->GetPrivate()->GetModelResources().ComponentNamesPerObjectComponent[ObjectComponentIndex.GetValue()];
+		if (!ObjectComponentIndex.IsValid())
+		{
+			continue;
+		}
+
+		const FName ComponentName = CustomizableObject->GetPrivate()->GetModelResources().ComponentNamesPerObjectComponent[ObjectComponentIndex.GetValue()];
 
 		if (OperationData->bUseMeshCache)
 		{
@@ -2243,7 +2248,11 @@ bool UCustomizableInstancePrivate::UpdateSkeletalMesh_PostBeginUpdate0(UCustomiz
 			continue;
 		}
 		const FCustomizableObjectComponentIndex ObjectComponentIndex = Component->Id;
-		const FName ComponentName = OperationData->Instance->GetCustomizableObject()->GetPrivate()->GetModelResources().ComponentNamesPerObjectComponent[ObjectComponentIndex.GetValue()];
+		if (!CustomizableObject->GetPrivate()->GetModelResources().ComponentNamesPerObjectComponent.IsValidIndex(ObjectComponentIndex.GetValue()))
+		{
+			continue;
+		}
+		const FName ComponentName = CustomizableObject->GetPrivate()->GetModelResources().ComponentNamesPerObjectComponent[ObjectComponentIndex.GetValue()];
 
 		FCustomizableInstanceComponentData* ComponentData = GetComponentData(ComponentName);
 		if (!ComponentData)
@@ -2253,7 +2262,7 @@ bool UCustomizableInstancePrivate::UpdateSkeletalMesh_PostBeginUpdate0(UCustomiz
 			InvalidateGeneratedData();
 			return false;
 		}
-		
+
 		// If the component doesn't need an update copy the previously generated mesh.
 		if (!OperationData->MeshChangedPerInstanceComponent[InstanceComponentIndex.GetValue()])
 		{
@@ -5300,13 +5309,16 @@ UE::Tasks::FTask UCustomizableInstancePrivate::LoadAdditionalAssetsAndData(
 	for (int32 InstanceComponentIndex = 0; InstanceComponentIndex < OperationData->NumInstanceComponents; ++InstanceComponentIndex)
 	{
 		FCustomizableObjectComponentIndex ObjectComponentIndex = OperationData->GetObjectComponentIndex(FCustomizableObjectInstanceComponentIndex(InstanceComponentIndex));
+		if (!ObjectComponentIndex.IsValid())
+		{
+			continue;
+		}
 
 		const FInstanceUpdateData::FSkeletonData& SkeletonData = OperationData->InstanceUpdateData.SkeletonsPerInstanceComponent[InstanceComponentIndex];
 		
 		FCustomizableInstanceComponentData* ComponentData = GetComponentData(ObjectComponentIndex);
 		if (!ComponentData)
 		{
-			check(false);
 			continue;
 		}
 
@@ -6237,7 +6249,12 @@ void UCustomizableInstancePrivate::BuildMaterials(const TSharedRef<FUpdateContex
 		const FInstanceUpdateData::FComponent& Component = OperationData->InstanceUpdateData.Components[InstanceComponentIndex];
 
 		const FCustomizableObjectComponentIndex ObjectComponentIndex = Component.Id;
-		const FName& ComponentName = GetPublic()->GetCustomizableObject()->GetPrivate()->GetModelResources().ComponentNamesPerObjectComponent[ObjectComponentIndex.GetValue()];
+
+		if (!CustomizableObject->GetPrivate()->GetModelResources().ComponentNamesPerObjectComponent.IsValidIndex(ObjectComponentIndex.GetValue()))
+		{
+			continue;
+		}
+		const FName& ComponentName = CustomizableObject->GetPrivate()->GetModelResources().ComponentNamesPerObjectComponent[ObjectComponentIndex.GetValue()];
 
 		TObjectPtr<USkeletalMesh>* Result = SkeletalMeshes.Find(ComponentName);
 		TObjectPtr<USkeletalMesh> SkeletalMesh = Result ? *Result : nullptr;
@@ -6383,8 +6400,16 @@ void UCustomizableInstancePrivate::BuildMaterials(const TSharedRef<FUpdateContex
 
 				if (!bUseOverrideMaterialsOnly)
 				{
-					const int32 LODMaterialIndex = SkeletalMesh->GetLODInfo(LODIndex)->LODMaterialMap.Add(MaterialSlotIndex);
-					SkeletalMesh->GetResourceForRendering()->LODRenderData[LODIndex].RenderSections[SurfaceIndex].MaterialIndex = LODMaterialIndex;
+					if (SkeletalMesh->GetResourceForRendering()->LODRenderData.IsValidIndex(LODIndex) &&
+						SkeletalMesh->GetResourceForRendering()->LODRenderData[LODIndex].RenderSections.IsValidIndex(SurfaceIndex))
+					{
+						const int32 LODMaterialIndex = SkeletalMesh->GetLODInfo(LODIndex)->LODMaterialMap.Add(MaterialSlotIndex);
+						SkeletalMesh->GetResourceForRendering()->LODRenderData[LODIndex].RenderSections[SurfaceIndex].MaterialIndex = LODMaterialIndex;
+					}
+					else
+					{
+						ensure(false);
+					}
 				}
 
 				FMutableMaterialPlaceholder MutableMaterialPlaceholder;
diff --git a/Engine/Plugins/Mutable/Source/CustomizableObject/Private/MuCO/CustomizableObjectSystem.cpp b/Engine/Plugins/Mutable/Source/CustomizableObject/Private/MuCO/CustomizableObjectSystem.cpp
--- a/Engine/Plugins/Mutable/Source/CustomizableObject/Private/MuCO/CustomizableObjectSystem.cpp
+++ b/Engine/Plugins/Mutable/Source/CustomizableObject/Private/MuCO/CustomizableObjectSystem.cpp
@@ -2500,7 +2500,15 @@ namespace impl
 			{
 				FInstanceUpdateData::FComponent& Component = OperationData->InstanceUpdateData.Components[InstanceComponentIndex];
 				Component.FirstLOD = OperationData->InstanceUpdateData.LODs.Num();
-				Component.Id = FCustomizableObjectComponentIndex(Instance->GetComponentId(InstanceComponentIndex));
+				uint64 MutableComponentId = Instance->GetComponentId(InstanceComponentIndex);
+				if (MutableComponentId < 65535)
+				{
+					Component.Id = FCustomizableObjectComponentIndex(MutableComponentId);
+				}
+				else
+				{
+					Component.Id.Invalidate();
+				}
 
 				FCustomizableObjectComponentIndex ObjectComponentIndex = Component.Id;
 
diff --git a/Engine/Plugins/Mutable/Source/CustomizableObjectEditor/Private/MuCOE/GenerateMutableSource/GenerateMutableSourceComponent.cpp b/Engine/Plugins/Mutable/Source/CustomizableObjectEditor/Private/MuCOE/GenerateMutableSource/GenerateMutableSourceComponent.cpp
--- a/Engine/Plugins/Mutable/Source/CustomizableObjectEditor/Private/MuCOE/GenerateMutableSource/GenerateMutableSourceComponent.cpp
+++ b/Engine/Plugins/Mutable/Source/CustomizableObjectEditor/Private/MuCOE/GenerateMutableSource/GenerateMutableSourceComponent.cpp
@@ -235,7 +235,20 @@ mu::Ptr<mu::NodeComponent> GenerateMutableSourceComponent(const UEdGraphPin* Pin
 
 		// Create the component node
 		mu::Ptr<mu::NodeComponentNew> ComponentNode = new mu::NodeComponentNew;
-		ComponentNode->Id = GenerationContext.ComponentNames.Find(TypedComponentPassthroughMesh->ComponentName);
+		if (GenerationContext.ComponentNames.Contains(TypedComponentPassthroughMesh->ComponentName))
+		{
+			FString Msg = FString::Printf(TEXT("More than one component with the same name [%s] have been found. This is not supported."), *TypedComponentPassthroughMesh->ComponentName.ToString());
+			GenerationContext.Log(FText::FromString(Msg), TypedComponentPassthroughMesh, EMessageSeverity::Warning);
+		}
+		ComponentNode->Id = GenerationContext.ComponentNames.Add(TypedComponentPassthroughMesh->ComponentName);
+
+		// While we don't handle the LODs per component, make sure we have some LODs for the passthrough mesh to appear.
+		if (GenerationContext.NumLODsInRoot == 0)
+		{
+			GenerationContext.NumLODsInRoot = SkeletalMesh->GetLODNum();
+			GenerationContext.NumMaxLODsToStream = FMath::Clamp(GenerationContext.NumMaxLODsToStream, 0, GenerationContext.NumLODsInRoot - 1);
+			GenerationContext.FirstLODAvailable = FMath::Clamp(GenerationContext.FirstLODAvailable, 0, GenerationContext.NumLODsInRoot - 1);
+		}
 
 		// Create a LOD for each pass-through mesh LOD.
 		const FSkeletalMeshModel* Model = SkeletalMesh->GetImportedModel();
