[mutable] Cleanup MutableRuntime API moving to use more unreal types and naming standards.
#rnx

diff --git a/Engine/Plugins/Mutable/Source/CustomizableObject/Internal/MuCO/MutableProjectorTypeUtils.h b/Engine/Plugins/Mutable/Source/CustomizableObject/Internal/MuCO/MutableProjectorTypeUtils.h
--- a/Engine/Plugins/Mutable/Source/CustomizableObject/Internal/MuCO/MutableProjectorTypeUtils.h
+++ b/Engine/Plugins/Mutable/Source/CustomizableObject/Internal/MuCO/MutableProjectorTypeUtils.h
@@ -12,13 +12,13 @@ namespace ProjectorUtils
 	 * @return The equivalent projector type to the one provided to be used with the CO and COI systems.
 	 */
 	CUSTOMIZABLEOBJECT_API
-	ECustomizableObjectProjectorType GetEquivalentProjectorType (mu::PROJECTOR_TYPE ProjectorType);
+	ECustomizableObjectProjectorType GetEquivalentProjectorType (mu::EProjectorType ProjectorType);
 
 	/** Get the equivalent mutable projector type to the CO mutable projector type.
 	 * @param ProjectorType The mutable CO projector type to get the mutable's internal equivalent projector type.
 	 * @return the projector type provided but as a mutable core projector type.
 	 */
 	CUSTOMIZABLEOBJECT_API
-	mu::PROJECTOR_TYPE GetEquivalentProjectorType (ECustomizableObjectProjectorType ProjectorType);
+	mu::EProjectorType GetEquivalentProjectorType (ECustomizableObjectProjectorType ProjectorType);
 }
 
diff --git a/Engine/Plugins/Mutable/Source/CustomizableObject/Internal/MuCO/UnrealMutableModelDiskStreamer.h b/Engine/Plugins/Mutable/Source/CustomizableObject/Internal/MuCO/UnrealMutableModelDiskStreamer.h
--- a/Engine/Plugins/Mutable/Source/CustomizableObject/Internal/MuCO/UnrealMutableModelDiskStreamer.h
+++ b/Engine/Plugins/Mutable/Source/CustomizableObject/Internal/MuCO/UnrealMutableModelDiskStreamer.h
@@ -60,9 +60,9 @@ public:
 	void EndStreaming();
 
 	// mu::ModelReader interface
-	OPERATION_ID BeginReadBlock(const mu::Model*, uint32 key0, void* pBuffer, uint64 size, TFunction<void(bool bSuccess)>* CompletionCallback) override;
-	bool IsReadCompleted(OPERATION_ID) override;
-	bool EndRead(OPERATION_ID) override;
+	FOperationID BeginReadBlock(const mu::Model*, uint32 key0, void* pBuffer, uint64 size, TFunction<void(bool bSuccess)>* CompletionCallback) override;
+	bool IsReadCompleted(FOperationID) override;
+	bool EndRead(FOperationID) override;
 
 protected:
 
@@ -83,7 +83,7 @@ protected:
 		TWeakPtr<const mu::Model> Model;
 
 		FString BulkFilePrefix;
-		TMap<OPERATION_ID, FReadRequest> CurrentReadRequests;
+		TMap<FOperationID, FReadRequest> CurrentReadRequests;
 		TMap<uint32, TSharedPtr<IAsyncReadFileHandle>> ReadFileHandles;
 
 		TSharedPtr<FModelStreamableBulkData> ModelStreamableBulkData;
@@ -101,7 +101,7 @@ protected:
 	FCriticalSection FileHandlesCritical;
 
 	/** This is used to generate unique ids for read requests. */
-	OPERATION_ID LastOperationID = 0;
+	FOperationID LastOperationID = 0;
 };
 
 
diff --git a/Engine/Plugins/Mutable/Source/CustomizableObject/Private/MuCO/CustomizableObject.cpp b/Engine/Plugins/Mutable/Source/CustomizableObject/Private/MuCO/CustomizableObject.cpp
--- a/Engine/Plugins/Mutable/Source/CustomizableObject/Private/MuCO/CustomizableObject.cpp
+++ b/Engine/Plugins/Mutable/Source/CustomizableObject/Private/MuCO/CustomizableObject.cpp
@@ -1561,7 +1561,7 @@ void UCustomizableObjectPrivate::UpdateParameterPropertiesFromModel(const TShare
 		const int32 NumParameters = MutableParameters->GetCount();
 
 		TArray<int32> TypedParametersCount;
-		TypedParametersCount.SetNum(static_cast<int32>(mu::PARAMETER_TYPE::T_COUNT));
+		TypedParametersCount.SetNum(static_cast<int32>(mu::EParameterType::Count));
 
 		ParameterProperties.Reset(NumParameters);
 		ParameterPropertiesLookupTable.Empty(NumParameters);
@@ -1572,16 +1572,16 @@ void UCustomizableObjectPrivate::UpdateParameterPropertiesFromModel(const TShare
 			Data.Name = MutableParameters->GetName(Index);
 			Data.Type = EMutableParameterType::None;
 
-			mu::PARAMETER_TYPE ParameterType = MutableParameters->GetType(Index);
+			mu::EParameterType ParameterType = MutableParameters->GetType(Index);
 			switch (ParameterType)
 			{
-			case mu::PARAMETER_TYPE::T_BOOL:
+			case mu::EParameterType::Bool:
 			{
 				Data.Type = EMutableParameterType::Bool;
 				break;
 			}
 
-			case mu::PARAMETER_TYPE::T_INT:
+			case mu::EParameterType::Int:
 			{
 				Data.Type = EMutableParameterType::Int;
 
@@ -1596,31 +1596,31 @@ void UCustomizableObjectPrivate::UpdateParameterPropertiesFromModel(const TShare
 				break;
 			}
 
-			case mu::PARAMETER_TYPE::T_FLOAT:
+			case mu::EParameterType::Float:
 			{
 				Data.Type = EMutableParameterType::Float;
 				break;
 			}
 
-			case mu::PARAMETER_TYPE::T_COLOUR:
+			case mu::EParameterType::Color:
 			{
 				Data.Type = EMutableParameterType::Color;
 				break;
 			}
 
-			case mu::PARAMETER_TYPE::T_PROJECTOR:
+			case mu::EParameterType::Projector:
 			{
 				Data.Type = EMutableParameterType::Projector;
 				break;
 			}
 
-			case mu::PARAMETER_TYPE::T_MATRIX:
+			case mu::EParameterType::Matrix:
 			{
 				Data.Type = EMutableParameterType::Transform;
 				break;
 			}
 				
-			case mu::PARAMETER_TYPE::T_IMAGE:
+			case mu::EParameterType::Image:
 			{
 				Data.Type = EMutableParameterType::Texture;
 				break;
@@ -1921,8 +1921,8 @@ FLinearColor UCustomizableObject::GetColorParameterDefaultValue(const FString& I
 		return FCustomizableObjectVectorParameterValue::DEFAULT_PARAMETER_VALUE;
 	}
 	
-	FLinearColor Value;
-	Model->GetColourDefaultValue(ParameterIndex, &Value.R, &Value.G, &Value.B, &Value.A);
+	FVector4f Value;
+	Model->GetColourDefaultValue(ParameterIndex, Value);
 
 	return Value;
 }
@@ -1982,7 +1982,7 @@ FCustomizableObjectProjector UCustomizableObject::GetProjectorParameterDefaultVa
 	}
 
 	FCustomizableObjectProjector Value;
-	mu::PROJECTOR_TYPE Type;
+	mu::EProjectorType Type;
 	Model->GetProjectorDefaultValue(ParameterIndex, &Type, &Value.Position, &Value.Direction, &Value.Up, &Value.Scale, &Value.Angle);
 	Value.ProjectionType = ProjectorUtils::GetEquivalentProjectorType(Type);
 	
diff --git a/Engine/Plugins/Mutable/Source/CustomizableObject/Private/MuCO/CustomizableObjectInstance.cpp b/Engine/Plugins/Mutable/Source/CustomizableObject/Private/MuCO/CustomizableObjectInstance.cpp
--- a/Engine/Plugins/Mutable/Source/CustomizableObject/Private/MuCO/CustomizableObjectInstance.cpp
+++ b/Engine/Plugins/Mutable/Source/CustomizableObject/Private/MuCO/CustomizableObjectInstance.cpp
@@ -3287,29 +3287,29 @@ FTexturePlatformData* MutableCreateImagePlatformData(mu::Ptr<const mu::Image> Mu
 	EPixelFormat PlatformFormat = PF_Unknown;
 	switch (MutableFormat)
 	{
-	case mu::EImageFormat::IF_RGB_UBYTE:
+	case mu::EImageFormat::RGB_UByte:
 		// performance penalty. can happen in states that remove compression.
 		PlatformFormat = PF_R8G8B8A8;	
 		UE_LOG(LogMutable, Display, TEXT("Building instance: a texture was generated in a format not supported by the hardware (RGB), this results in an additional conversion, so a performance penalty."));
 		break; 
 
-	case mu::EImageFormat::IF_BGRA_UBYTE:			
+	case mu::EImageFormat::BGRA_UByte:			
 		// performance penalty. can happen with texture parameter images.
 		PlatformFormat = PF_R8G8B8A8;	
 		UE_LOG(LogMutable, Display, TEXT("Building instance: a texture was generated in a format not supported by the hardware (BGRA), this results in an additional conversion, so a performance penalty."));
 		break;
 
 	// Good cases:
-	case mu::EImageFormat::IF_RGBA_UBYTE:		PlatformFormat = PF_R8G8B8A8;	break;
-	case mu::EImageFormat::IF_BC1:				PlatformFormat = PF_DXT1;		break;
-	case mu::EImageFormat::IF_BC2:				PlatformFormat = PF_DXT3;		break;
-	case mu::EImageFormat::IF_BC3:				PlatformFormat = PF_DXT5;		break;
-	case mu::EImageFormat::IF_BC4:				PlatformFormat = PF_BC4;		break;
-	case mu::EImageFormat::IF_BC5:				PlatformFormat = PF_BC5;		break;
-	case mu::EImageFormat::IF_L_UBYTE:			PlatformFormat = PF_G8;			break;
-	case mu::EImageFormat::IF_ASTC_4x4_RGB_LDR:	PlatformFormat = PF_ASTC_4x4;	break;
-	case mu::EImageFormat::IF_ASTC_4x4_RGBA_LDR:PlatformFormat = PF_ASTC_4x4;	break;
-	case mu::EImageFormat::IF_ASTC_4x4_RG_LDR:	PlatformFormat = PF_ASTC_4x4;	break;
+	case mu::EImageFormat::RGBA_UByte:		PlatformFormat = PF_R8G8B8A8;	break;
+	case mu::EImageFormat::BC1:				PlatformFormat = PF_DXT1;		break;
+	case mu::EImageFormat::BC2:				PlatformFormat = PF_DXT3;		break;
+	case mu::EImageFormat::BC3:				PlatformFormat = PF_DXT5;		break;
+	case mu::EImageFormat::BC4:				PlatformFormat = PF_BC4;		break;
+	case mu::EImageFormat::BC5:				PlatformFormat = PF_BC5;		break;
+	case mu::EImageFormat::L_UByte:			PlatformFormat = PF_G8;			break;
+	case mu::EImageFormat::ASTC_4x4_RGB_LDR:	PlatformFormat = PF_ASTC_4x4;	break;
+	case mu::EImageFormat::ASTC_4x4_RGBA_LDR:PlatformFormat = PF_ASTC_4x4;	break;
+	case mu::EImageFormat::ASTC_4x4_RG_LDR:	PlatformFormat = PF_ASTC_4x4;	break;
 	default:
 		// Cannot prepare texture if it's not in the right format, this can happen if mutable is in debug mode or in case of bugs
 		UE_LOG(LogMutable, Warning, TEXT("Building instance: a texture was generated in an unsupported format, it will be converted to Unreal with a performance penalty."));
@@ -3318,13 +3318,13 @@ FTexturePlatformData* MutableCreateImagePlatformData(mu::Ptr<const mu::Image> Mu
 		{
 		case 1:
 			PlatformFormat = PF_R8;
-			MutableImage = ImOp.ImagePixelFormat(0, MutableImage.get(), mu::EImageFormat::IF_L_UBYTE);
+			MutableImage = ImOp.ImagePixelFormat(0, MutableImage.get(), mu::EImageFormat::L_UByte);
 			break;
 		case 2:
 		case 3:
 		case 4:
 			PlatformFormat = PF_R8G8B8A8;
-			MutableImage = ImOp.ImagePixelFormat(0, MutableImage.get(), mu::EImageFormat::IF_RGBA_UBYTE);
+			MutableImage = ImOp.ImagePixelFormat(0, MutableImage.get(), mu::EImageFormat::RGBA_UByte);
 			break;
 		default: 
 			// Absolutely worst case
@@ -3374,13 +3374,13 @@ FTexturePlatformData* MutableCreateImagePlatformData(mu::Ptr<const mu::Image> Mu
 			const uint8* MutableData = MutableImage->GetLODData(MipLevelMutable);
 			const uint32 SourceDataSize = MutableImage->GetLODDataSize(MipLevelMutable);
 
-			uint32 DestDataSize = (MutableFormat == mu::EImageFormat::IF_RGB_UBYTE)
+			uint32 DestDataSize = (MutableFormat == mu::EImageFormat::RGB_UByte)
 					? (SourceDataSize/3) * 4
 					: SourceDataSize;
 			void* pData = Mip->BulkData.Realloc(DestDataSize);
 
 			// Special inefficient cases
-			if (MutableFormat== mu::EImageFormat::IF_BGRA_UBYTE)
+			if (MutableFormat== mu::EImageFormat::BGRA_UByte)
 			{
 				check(SourceDataSize==DestDataSize);
 
@@ -3396,7 +3396,7 @@ FTexturePlatformData* MutableCreateImagePlatformData(mu::Ptr<const mu::Image> Mu
 				}
 			}
 
-			else if (MutableFormat == mu::EImageFormat::IF_RGB_UBYTE)
+			else if (MutableFormat == mu::EImageFormat::RGB_UByte)
 			{
 				MUTABLE_CPUPROFILER_SCOPE(Innefficent_RGB_Format_Conversion);
 
@@ -3477,20 +3477,20 @@ void ConvertImage(UTexture2D* Texture, mu::Ptr<const mu::Image> MutableImage, co
 	{
 		mu::FImageOperator ImOp = mu::FImageOperator::GetDefault(mu::FImageOperator::FImagePixelFormatFunc());
 
-		MutableImage = ImOp.ImagePixelFormat( 4, MutableImage.get(), mu::EImageFormat::IF_RGBA_UBYTE );
+		MutableImage = ImOp.ImagePixelFormat( 4, MutableImage.get(), mu::EImageFormat::RGBA_UByte );
 
 		uint8_t Channel = uint8_t( FMath::Clamp(ExtractChannel,0,3) );
-		MutableImage = ImOp.ImageSwizzle( mu::EImageFormat::IF_L_UBYTE, &MutableImage, &Channel );
-		MutableFormat = mu::EImageFormat::IF_L_UBYTE;
+		MutableImage = ImOp.ImageSwizzle( mu::EImageFormat::L_UByte, &MutableImage, &Channel );
+		MutableFormat = mu::EImageFormat::L_UByte;
 	}
 
 	// Hack: This format is unsupported in UE, but it shouldn't happen in production.
-	if (MutableFormat == mu::EImageFormat::IF_RGB_UBYTE)
+	if (MutableFormat == mu::EImageFormat::RGB_UByte)
 	{
 		UE_LOG(LogMutable, Warning, TEXT("Building instance: a texture was generated in RGB format, which is slow to convert to Unreal."));
 
 		// Expand the image.
-		mu::ImagePtr Converted = new mu::Image(MutableImage->GetSizeX(), MutableImage->GetSizeY(), MutableImage->GetLODCount(), mu::EImageFormat::IF_RGBA_UBYTE, mu::EInitializationType::NotInitialized);
+		mu::ImagePtr Converted = new mu::Image(MutableImage->GetSizeX(), MutableImage->GetSizeY(), MutableImage->GetLODCount(), mu::EImageFormat::RGBA_UByte, mu::EInitializationType::NotInitialized);
 
 		for (int32 LODIndex = 0; LODIndex < Converted->GetLODCount(); ++LODIndex)
 		{
@@ -3507,16 +3507,16 @@ void ConvertImage(UTexture2D* Texture, mu::Ptr<const mu::Image> MutableImage, co
 		}
 
 		MutableImage = Converted;
-		MutableFormat = mu::EImageFormat::IF_RGBA_UBYTE;
+		MutableFormat = mu::EImageFormat::RGBA_UByte;
 	}
-	else if (MutableFormat == mu::EImageFormat::IF_BGRA_UBYTE)
+	else if (MutableFormat == mu::EImageFormat::BGRA_UByte)
 	{
 		UE_LOG(LogMutable, Warning, TEXT("Building instance: a texture was generated in BGRA format, which is slow to convert to Unreal."));
 
 		MUTABLE_CPUPROFILER_SCOPE(Swizzle);
 		// Swizzle the image.
 		// \TODO: Raise a warning?
-		mu::ImagePtr Converted = new mu::Image(MutableImage->GetSizeX(), MutableImage->GetSizeY(), 1, mu::EImageFormat::IF_RGBA_UBYTE, mu::EInitializationType::NotInitialized);
+		mu::ImagePtr Converted = new mu::Image(MutableImage->GetSizeX(), MutableImage->GetSizeY(), 1, mu::EImageFormat::RGBA_UByte, mu::EInitializationType::NotInitialized);
 		int32 PixelCount = MutableImage->GetSizeX() * MutableImage->GetSizeY();
 
 		const uint8* pSource = MutableImage->GetLODData(0);
@@ -3530,7 +3530,7 @@ void ConvertImage(UTexture2D* Texture, mu::Ptr<const mu::Image> MutableImage, co
 		}
 
 		MutableImage = Converted;
-		MutableFormat = mu::EImageFormat::IF_RGBA_UBYTE;
+		MutableFormat = mu::EImageFormat::RGBA_UByte;
 	}
 
 	if (OnlyLOD >= 0)
@@ -4108,10 +4108,10 @@ void UCustomizableInstancePrivate::BuildOrCopyClothingData(const TSharedRef<FUpd
 				const mu::FMeshBufferSet& IndicesSet = MutableMesh->GetIndexBuffers();
 
 				int32 ClothingIndexBuffer, ClothingIndexChannel;
-				MeshSet.FindChannel(mu::MBS_OTHER, 2, &ClothingIndexBuffer, &ClothingIndexChannel);
+				MeshSet.FindChannel(mu::EMeshBufferSemantic::Other, 2, &ClothingIndexBuffer, &ClothingIndexChannel);
 
 				int32 ClothingResourceBuffer, ClothingResourceChannel;
-				MeshSet.FindChannel(mu::MBS_OTHER, 3, &ClothingResourceBuffer, &ClothingResourceChannel);
+				MeshSet.FindChannel(mu::EMeshBufferSemantic::Other, 3, &ClothingResourceBuffer, &ClothingResourceChannel);
 
 				if (ClothingIndexBuffer < 0 || ClothingResourceBuffer < 0)
 				{
@@ -5159,7 +5159,7 @@ bool UCustomizableInstancePrivate::BuildOrCopyRenderData(const TSharedRef<FUpdat
 					int32 SemanticIndex;
 					MutableMeshVertexBuffers.GetChannel(BufferIndex, 0, &Semantic, &SemanticIndex, nullptr, nullptr, nullptr);
 
-					if (Semantic != mu::EMeshBufferSemantic::MBS_ALTSKINWEIGHT)
+					if (Semantic != mu::EMeshBufferSemantic::AltSkinWeight)
 					{
 						continue;
 					}
diff --git a/Engine/Plugins/Mutable/Source/CustomizableObject/Private/MuCO/CustomizableObjectInstanceDescriptor.cpp b/Engine/Plugins/Mutable/Source/CustomizableObject/Private/MuCO/CustomizableObjectInstanceDescriptor.cpp
--- a/Engine/Plugins/Mutable/Source/CustomizableObject/Private/MuCO/CustomizableObjectInstanceDescriptor.cpp
+++ b/Engine/Plugins/Mutable/Source/CustomizableObject/Private/MuCO/CustomizableObjectInstanceDescriptor.cpp
@@ -532,11 +532,11 @@ mu::ParametersPtr FCustomizableObjectInstanceDescriptor::GetParameters() const
 	{
 		const FString& Name = MutableParameters->GetName(ParamIndex);
 		const FGuid& Uid = MutableParameters->GetUid(ParamIndex);
-		const mu::PARAMETER_TYPE MutableType = MutableParameters->GetType(ParamIndex);
+		const mu::EParameterType MutableType = MutableParameters->GetType(ParamIndex);
 
 		switch (MutableType)
 		{
-		case mu::PARAMETER_TYPE::T_BOOL:
+		case mu::EParameterType::Bool:
 		{
 			for (const FCustomizableObjectBoolParameterValue& BoolParameter : BoolParameters)
 			{
@@ -550,7 +550,7 @@ mu::ParametersPtr FCustomizableObjectInstanceDescriptor::GetParameters() const
 			break;
 		}
 
-		case mu::PARAMETER_TYPE::T_INT:
+		case mu::EParameterType::Int:
 		{
 			for (const FCustomizableObjectIntParameterValue& IntParameter : IntParameters)
 			{
@@ -580,7 +580,7 @@ mu::ParametersPtr FCustomizableObjectInstanceDescriptor::GetParameters() const
 			break;
 		}
 
-		case mu::PARAMETER_TYPE::T_FLOAT:
+		case mu::EParameterType::Float:
 		{
 			for (const FCustomizableObjectFloatParameterValue& FloatParameter : FloatParameters)
 			{
@@ -606,14 +606,13 @@ mu::ParametersPtr FCustomizableObjectInstanceDescriptor::GetParameters() const
 			break;
 		}
 
-		case mu::PARAMETER_TYPE::T_COLOUR:
+		case mu::EParameterType::Color:
 		{
 			for (const FCustomizableObjectVectorParameterValue& VectorParameter : VectorParameters)
 			{
 				if (VectorParameter.ParameterName == Name || (Uid.IsValid() && VectorParameter.Id == Uid))
 				{
-					MutableParameters->SetColourValue(ParamIndex, VectorParameter.ParameterValue.R, VectorParameter.ParameterValue.G, 
-																  VectorParameter.ParameterValue.B, VectorParameter.ParameterValue.A);
+					MutableParameters->SetColourValue(ParamIndex, VectorParameter.ParameterValue);
 
 					break;
 				}
@@ -622,7 +621,7 @@ mu::ParametersPtr FCustomizableObjectInstanceDescriptor::GetParameters() const
 			break;
 		}
 
-		case mu::PARAMETER_TYPE::T_MATRIX:
+		case mu::EParameterType::Matrix:
 			{
 				for (const FCustomizableObjectTransformParameterValue& TransformParameter : TransformParameters)
 				{
@@ -637,7 +636,7 @@ mu::ParametersPtr FCustomizableObjectInstanceDescriptor::GetParameters() const
 				break;
 			}
 			
-		case mu::PARAMETER_TYPE::T_PROJECTOR:
+		case mu::EParameterType::Projector:
 		{
 			for (const auto& ProjectorParameter : ProjectorParameters)
 			{
@@ -696,7 +695,7 @@ mu::ParametersPtr FCustomizableObjectInstanceDescriptor::GetParameters() const
 			break;
 		}
 
-		case mu::PARAMETER_TYPE::T_IMAGE:
+		case mu::EParameterType::Image:
 		{
 			for (const FCustomizableObjectTextureParameterValue& TextureParameter : TextureParameters)
 			{
@@ -794,11 +793,11 @@ void FCustomizableObjectInstanceDescriptor::ReloadParameters()
 	{
 		const FString& Name = MutableParameters->GetName(ParamIndex);
 		const FGuid& Uid = MutableParameters->GetUid(ParamIndex);
-		const mu::PARAMETER_TYPE MutableType = MutableParameters->GetType(ParamIndex);
+		const mu::EParameterType MutableType = MutableParameters->GetType(ParamIndex);
 
 		switch (MutableType)
 		{
-		case mu::PARAMETER_TYPE::T_BOOL:
+		case mu::EParameterType::Bool:
 		{
 			FCustomizableObjectBoolParameterValue Param;
 			Param.ParameterName = Name;
@@ -822,7 +821,7 @@ void FCustomizableObjectInstanceDescriptor::ReloadParameters()
 			break;
 		}
 
-		case mu::PARAMETER_TYPE::T_INT:
+		case mu::EParameterType::Int:
 		{
 			FCustomizableObjectIntParameterValue Param;
 			Param.ParameterName = Name;
@@ -926,7 +925,7 @@ void FCustomizableObjectInstanceDescriptor::ReloadParameters()
 			break;
 		}
 
-		case mu::PARAMETER_TYPE::T_FLOAT:
+		case mu::EParameterType::Float:
 		{
 			FCustomizableObjectFloatParameterValue Param;
 			Param.ParameterName = Name;
@@ -977,7 +976,7 @@ void FCustomizableObjectInstanceDescriptor::ReloadParameters()
 			break;
 		}
 
-		case mu::PARAMETER_TYPE::T_COLOUR:
+		case mu::EParameterType::Color:
 		{
 			FCustomizableObjectVectorParameterValue Param;
 			Param.ParameterName = Name;
@@ -994,15 +993,16 @@ void FCustomizableObjectInstanceDescriptor::ReloadParameters()
 			}
 			else // Not found in Instance Parameters. Use Mutable Parameters.
 			{
-				MutableParameters->GetColourValue(ParamIndex, &Param.ParameterValue.R, &Param.ParameterValue.G,
-															  &Param.ParameterValue.B, &Param.ParameterValue.A);
+				FVector4f V;
+				MutableParameters->GetColourValue(ParamIndex, V);
+				Param.ParameterValue = V;
 			}
 
 			VectorParameters.Add(Param);
 			break;
 		}
 
-		case mu::PARAMETER_TYPE::T_MATRIX:
+		case mu::EParameterType::Matrix:
 			{
 				FCustomizableObjectTransformParameterValue Param;
 				Param.ParameterName = Name;
@@ -1028,7 +1028,7 @@ void FCustomizableObjectInstanceDescriptor::ReloadParameters()
 				break;
 			}
 			
-		case mu::PARAMETER_TYPE::T_PROJECTOR:
+		case mu::EParameterType::Projector:
 		{
 			FCustomizableObjectProjectorParameterValue Param;
 			Param.ParameterName = Name;
@@ -1064,7 +1064,7 @@ void FCustomizableObjectInstanceDescriptor::ReloadParameters()
 			{
 				auto GetProjector = [&MutableParameters, ParamIndex](FCustomizableObjectProjector &Value, const mu::RangeIndexPtr RangeIndex = nullptr)
 				{
-					mu::PROJECTOR_TYPE Type;
+					mu::EProjectorType Type;
 					MutableParameters->GetProjectorValue(ParamIndex,
 						&Type,
 						&Value.Position,
@@ -1111,7 +1111,7 @@ void FCustomizableObjectInstanceDescriptor::ReloadParameters()
 			break;
 		}
 
-		case mu::PARAMETER_TYPE::T_IMAGE:
+		case mu::EParameterType::Image:
 		{
 			FCustomizableObjectTextureParameterValue Param;
 			Param.ParameterName = Name;
diff --git a/Engine/Plugins/Mutable/Source/CustomizableObject/Private/MuCO/CustomizableObjectSystem.cpp b/Engine/Plugins/Mutable/Source/CustomizableObject/Private/MuCO/CustomizableObjectSystem.cpp
--- a/Engine/Plugins/Mutable/Source/CustomizableObject/Private/MuCO/CustomizableObjectSystem.cpp
+++ b/Engine/Plugins/Mutable/Source/CustomizableObject/Private/MuCO/CustomizableObjectSystem.cpp
@@ -1923,10 +1923,10 @@ namespace impl
 				const mu::FMeshBufferSet& MeshSet = LOD.Mesh->GetVertexBuffers();
 
 				int32 VertexMorphsInfoIndexAndCountBufferIndex, VertexMorphsInfoIndexAndCountBufferChannel;
-				MeshSet.FindChannel(mu::MBS_OTHER, 0, &VertexMorphsInfoIndexAndCountBufferIndex, &VertexMorphsInfoIndexAndCountBufferChannel);
+				MeshSet.FindChannel(mu::EMeshBufferSemantic::Other, 0, &VertexMorphsInfoIndexAndCountBufferIndex, &VertexMorphsInfoIndexAndCountBufferChannel);
 
 				int32 VertexMorphsResourceIdBufferIndex, VertexMorphsResourceIdBufferChannel;
-				MeshSet.FindChannel(mu::MBS_OTHER, 1, &VertexMorphsResourceIdBufferIndex, &VertexMorphsResourceIdBufferChannel);
+				MeshSet.FindChannel(mu::EMeshBufferSemantic::Other, 1, &VertexMorphsResourceIdBufferIndex, &VertexMorphsResourceIdBufferChannel);
 
 				if (VertexMorphsInfoIndexAndCountBufferIndex < 0 || VertexMorphsResourceIdBufferIndex < 0)
 				{
diff --git a/Engine/Plugins/Mutable/Source/CustomizableObject/Private/MuCO/MutableMeshBufferUtils.cpp b/Engine/Plugins/Mutable/Source/CustomizableObject/Private/MuCO/MutableMeshBufferUtils.cpp
--- a/Engine/Plugins/Mutable/Source/CustomizableObject/Private/MuCO/MutableMeshBufferUtils.cpp
+++ b/Engine/Plugins/Mutable/Source/CustomizableObject/Private/MuCO/MutableMeshBufferUtils.cpp
@@ -12,9 +12,9 @@ void MutableMeshBufferUtils::SetupVertexPositionsBuffer(const int32& InCurrentVe
 	using namespace mu;
 	const int32 ElementSize = sizeof(FPositionVertex);
 	constexpr int32 ChannelCount = 1;
-	const EMeshBufferSemantic Semantics[ChannelCount] = {MBS_POSITION};
+	const EMeshBufferSemantic Semantics[ChannelCount] = {EMeshBufferSemantic::Position};
 	const int32 SemanticIndices[ChannelCount] = {0};
-	const EMeshBufferFormat Formats[ChannelCount] = {MBF_FLOAT32};
+	const EMeshBufferFormat Formats[ChannelCount] = {EMeshBufferFormat::Float32};
 	const int32 Components[ChannelCount] = {3};
 	const int32 Offsets[ChannelCount] =
 	{
@@ -36,9 +36,9 @@ void MutableMeshBufferUtils::SetupTangentBuffer(const int32& InCurrentVertexBuff
 	using namespace mu;
 	const int32 ElementSize = sizeof(TangentType);
 	constexpr int32 ChannelCount = 2;
-	const EMeshBufferSemantic Semantics[ChannelCount] = {MBS_TANGENT, MBS_NORMAL};
+	const EMeshBufferSemantic Semantics[ChannelCount] = {EMeshBufferSemantic::Tangent, EMeshBufferSemantic::Normal};
 	const int32 SemanticIndices[ChannelCount] = {0, 0};
-	const EMeshBufferFormat Formats[ChannelCount] = {MBF_PACKEDDIRS8, MBF_PACKEDDIRS8_W_TANGENTSIGN};
+	const EMeshBufferFormat Formats[ChannelCount] = {EMeshBufferFormat::PackedDirS8, EMeshBufferFormat::PackedDirS8_W_TangentSign};
 	const int32 Components[ChannelCount] = {4, 4};
 	const int32 Offsets[ChannelCount] =
 	{
@@ -62,9 +62,9 @@ void MutableMeshBufferUtils::SetupTexCoordinatesBuffer(const int32& InCurrentVer
 	const int32 ElementSize = sizeof(UVType) * InChannelCount;
 	constexpr int32 MaxChannelCount = MaxTexCordChannelCount;
 	const EMeshBufferSemantic Semantics[MaxChannelCount] = {
-		MBS_TEXCOORDS, MBS_TEXCOORDS, MBS_TEXCOORDS, MBS_TEXCOORDS
+		EMeshBufferSemantic::TexCoords, EMeshBufferSemantic::TexCoords, EMeshBufferSemantic::TexCoords, EMeshBufferSemantic::TexCoords
 	};
-	const EMeshBufferFormat Formats[MaxChannelCount] = {MBF_FLOAT32, MBF_FLOAT32, MBF_FLOAT32, MBF_FLOAT32};
+	const EMeshBufferFormat Formats[MaxChannelCount] = {EMeshBufferFormat::Float32, EMeshBufferFormat::Float32, EMeshBufferFormat::Float32, EMeshBufferFormat::Float32};
 	const int32 Components[MaxChannelCount] = {2, 2, 2, 2};
 	const int32 Offsets[MaxChannelCount] =
 	{
@@ -97,39 +97,39 @@ void MutableMeshBufferUtils::SetupSkinBuffer(const int32& InCurrentVertexBuffer,
 	using namespace mu;
 	const int32 ElementSize = (MaxBoneWeightTypeSizeBytes + MaxBoneIndexTypeSizeBytes) * MaxNumBonesPerVertex;
 	constexpr int32 ChannelCount = 2;
-	const EMeshBufferSemantic Semantics[ChannelCount] = {MBS_BONEINDICES, MBS_BONEWEIGHTS};
+	const EMeshBufferSemantic Semantics[ChannelCount] = {EMeshBufferSemantic::BoneIndices, EMeshBufferSemantic::BoneWeights};
 	const int32 SemanticIndices[ChannelCount] = {0, 0};
 
-	EMeshBufferFormat Formats[ChannelCount] = {MBF_UINT8, MBF_NUINT8};
+	EMeshBufferFormat Formats[ChannelCount] = {EMeshBufferFormat::UInt8, EMeshBufferFormat::NUInt8};
 	switch (MaxBoneIndexTypeSizeBytes)
 	{
 	case 0: // Fallback to something in this case.
-	case 1: Formats[0] = mu::MBF_UINT8;
+	case 1: Formats[0] = mu::EMeshBufferFormat::UInt8;
 		break;
-	case 2: Formats[0] = mu::MBF_UINT16;
+	case 2: Formats[0] = mu::EMeshBufferFormat::UInt16;
 		break;
-	case 4: Formats[0] = mu::MBF_UINT32;
+	case 4: Formats[0] = mu::EMeshBufferFormat::UInt32;
 		break;
 	default:
 		// unsupported bone index type
 		check(false);
-		Formats[0] = mu::MBF_NONE;
+		Formats[0] = mu::EMeshBufferFormat::None;
 		break;
 	}
 
 	switch (MaxBoneWeightTypeSizeBytes)
 	{
 	case 0: // Fallback to something in this case.
-	case 1: Formats[1] = mu::MBF_NUINT8;
+	case 1: Formats[1] = mu::EMeshBufferFormat::NUInt8;
 		break;
-	case 2: Formats[1] = mu::MBF_NUINT16;
+	case 2: Formats[1] = mu::EMeshBufferFormat::NUInt16;
 		break;
-	case 4: Formats[1] = mu::MBF_NUINT32;
+	case 4: Formats[1] = mu::EMeshBufferFormat::NUInt32;
 		break;
 	default:
 		// unsupported bone weight type
 		check(false);
-		Formats[1] = mu::MBF_NONE;
+		Formats[1] = mu::EMeshBufferFormat::None;
 		break;
 	}
 
@@ -151,9 +151,9 @@ void MutableMeshBufferUtils::SetupVertexColorBuffer(const int32& InCurrentVertex
 	using namespace mu;
 	const int32 ElementSize = sizeof(FColor);
 	constexpr int32 ChannelCount = 1;
-	const EMeshBufferSemantic Semantics[ChannelCount] = {MBS_COLOUR};
+	const EMeshBufferSemantic Semantics[ChannelCount] = {EMeshBufferSemantic::Color};
 	const int32 SemanticIndices[ChannelCount] = {0};
-	const EMeshBufferFormat Formats[ChannelCount] = {MBF_NUINT8};
+	const EMeshBufferFormat Formats[ChannelCount] = {EMeshBufferFormat::NUInt8};
 	const int32 Components[ChannelCount] = {4};
 	const int32 Offsets[ChannelCount] = {0};
 	check(ElementSize == 4);
@@ -171,12 +171,12 @@ void MutableMeshBufferUtils::SetupIndexBuffer(mu::FMeshBufferSet& OutTargetIndex
 	const int32 ElementSize = sizeof(uint32);
 	//SkeletalMesh->GetImportedResource()->LODModels[LOD].MultiSizeIndexContainer.GetDataTypeSize();
 	constexpr int32 ChannelCount = 1;
-	const EMeshBufferSemantic Semantics[ChannelCount] = {MBS_VERTEXINDEX};
+	const EMeshBufferSemantic Semantics[ChannelCount] = {EMeshBufferSemantic::VertexIndex};
 	const int32 SemanticIndices[ChannelCount] = {0};
 	// We force 32 bit indices, since merging meshes may create vertex buffers bigger than the initial mesh
 	// and for now the mutable runtime doesn't handle it.
 	// \TODO: go back to 16-bit indices when possible.
-	EMeshBufferFormat Formats[ChannelCount] = {MBF_UINT32};
+	EMeshBufferFormat Formats[ChannelCount] = {EMeshBufferFormat::UInt32};
 	const int32 Components[ChannelCount] = {1};
 	const int32 Offsets[ChannelCount] = {0};
 
@@ -195,41 +195,41 @@ void MutableMeshBufferUtils::SetupSkinWeightProfileBuffer(const int32& InCurrent
 	using namespace mu;
 	const int32 ElementSize = sizeof(int32) + (MaxBoneIndexTypeSizeBytes + 1) * MaxNumBonesPerVertex;
 	constexpr int32 ChannelCount = 3;
-	const EMeshBufferSemantic Semantics[ChannelCount] = { MBS_ALTSKINWEIGHT, MBS_BONEINDICES, MBS_BONEWEIGHTS };
+	const EMeshBufferSemantic Semantics[ChannelCount] = { EMeshBufferSemantic::AltSkinWeight, EMeshBufferSemantic::BoneIndices, EMeshBufferSemantic::BoneWeights };
 	const int32 SemanticIndices[ChannelCount] = { SemanticsIndex, SemanticsIndex, SemanticsIndex };
 
-	EMeshBufferFormat Formats[ChannelCount] = { MBF_INT32, MBF_UINT8, MBF_NUINT8 };
+	EMeshBufferFormat Formats[ChannelCount] = { EMeshBufferFormat::Int32, EMeshBufferFormat::UInt8, EMeshBufferFormat::NUInt8 };
 	switch (MaxBoneIndexTypeSizeBytes)
 	{
-	case 1: Formats[1] = mu::MBF_UINT8;
+	case 1: Formats[1] = mu::EMeshBufferFormat::UInt8;
 		break;
-	case 2: Formats[1] = mu::MBF_UINT16;
+	case 2: Formats[1] = mu::EMeshBufferFormat::UInt16;
 		break;
-	case 4: Formats[1] = mu::MBF_UINT32;
+	case 4: Formats[1] = mu::EMeshBufferFormat::UInt32;
 		break;
 	default:
 		// unsupported bone index type
 		check(false);
-		Formats[1] = mu::MBF_NONE;
+		Formats[1] = mu::EMeshBufferFormat::None;
 		break;
 	}
 
 	switch (MaxBoneWeightTypeSizeBytes)
 	{
-	case 1: Formats[2] = mu::MBF_NUINT8;
+	case 1: Formats[2] = mu::EMeshBufferFormat::NUInt8;
 		break;
 	case 2:
 	{
 		unimplemented()
-		Formats[2] = mu::MBF_NUINT16;
+		Formats[2] = mu::EMeshBufferFormat::NUInt16;
 		break;
 	}
-	case 4: Formats[2] = mu::MBF_NUINT32;
+	case 4: Formats[2] = mu::EMeshBufferFormat::NUInt32;
 		break;
 	default:
 		// unsupported bone weight type
 		check(false);
-		Formats[2] = mu::MBF_NONE;
+		Formats[2] = mu::EMeshBufferFormat::None;
 		break;
 	}
 
diff --git a/Engine/Plugins/Mutable/Source/CustomizableObject/Private/MuCO/MutableProjectorTypeUtils.cpp b/Engine/Plugins/Mutable/Source/CustomizableObject/Private/MuCO/MutableProjectorTypeUtils.cpp
--- a/Engine/Plugins/Mutable/Source/CustomizableObject/Private/MuCO/MutableProjectorTypeUtils.cpp
+++ b/Engine/Plugins/Mutable/Source/CustomizableObject/Private/MuCO/MutableProjectorTypeUtils.cpp
@@ -5,21 +5,21 @@
 #include "MuCO/CustomizableObjectParameterTypeDefinitions.h"
 #include "MuR/Parameters.h"
 
-ECustomizableObjectProjectorType ProjectorUtils::GetEquivalentProjectorType (mu::PROJECTOR_TYPE ProjectorType)
+ECustomizableObjectProjectorType ProjectorUtils::GetEquivalentProjectorType (mu::EProjectorType ProjectorType)
 {
 	// Translate projector type from Mutable Core enum type to CO enum type
 	switch (ProjectorType)
 	{
-	case mu::PROJECTOR_TYPE::PLANAR:
+	case mu::EProjectorType::Planar:
 		return ECustomizableObjectProjectorType::Planar;
 		
-	case mu::PROJECTOR_TYPE::CYLINDRICAL:
+	case mu::EProjectorType::Cylindrical:
 		return ECustomizableObjectProjectorType::Cylindrical;
 		
-	case mu::PROJECTOR_TYPE::WRAPPING:
+	case mu::EProjectorType::Wrapping:
 		return ECustomizableObjectProjectorType::Wrapping;
 		
-	case mu::PROJECTOR_TYPE::COUNT:
+	case mu::EProjectorType::Count:
 	default:
 		checkNoEntry();
 		return ECustomizableObjectProjectorType::Planar;
@@ -27,21 +27,21 @@ ECustomizableObjectProjectorType ProjectorUtils::GetEquivalentProjectorType (mu:
 }
 
 
-mu::PROJECTOR_TYPE ProjectorUtils::GetEquivalentProjectorType (ECustomizableObjectProjectorType ProjectorType)
+mu::EProjectorType ProjectorUtils::GetEquivalentProjectorType (ECustomizableObjectProjectorType ProjectorType)
 {
-	if (GetEquivalentProjectorType(mu::PROJECTOR_TYPE::PLANAR) == ProjectorType)
+	if (GetEquivalentProjectorType(mu::EProjectorType::Planar) == ProjectorType)
 	{
-		return mu::PROJECTOR_TYPE::PLANAR;
+		return mu::EProjectorType::Planar;
 	}
-	if (GetEquivalentProjectorType(mu::PROJECTOR_TYPE::WRAPPING) == ProjectorType)
+	if (GetEquivalentProjectorType(mu::EProjectorType::Wrapping) == ProjectorType)
 	{
-		return mu::PROJECTOR_TYPE::WRAPPING;
+		return mu::EProjectorType::Wrapping;
 	}
-	if (GetEquivalentProjectorType(mu::PROJECTOR_TYPE::CYLINDRICAL) == ProjectorType)
+	if (GetEquivalentProjectorType(mu::EProjectorType::Cylindrical) == ProjectorType)
 	{
-		return mu::PROJECTOR_TYPE::CYLINDRICAL;
+		return mu::EProjectorType::Cylindrical;
 	}
 
 	checkNoEntry();
-	return mu::PROJECTOR_TYPE::COUNT;		// Invalid 
+	return mu::EProjectorType::Count;		// Invalid 
 }
diff --git a/Engine/Plugins/Mutable/Source/CustomizableObject/Private/MuCO/UnrealConversionUtils.cpp b/Engine/Plugins/Mutable/Source/CustomizableObject/Private/MuCO/UnrealConversionUtils.cpp
--- a/Engine/Plugins/Mutable/Source/CustomizableObject/Private/MuCO/UnrealConversionUtils.cpp
+++ b/Engine/Plugins/Mutable/Source/CustomizableObject/Private/MuCO/UnrealConversionUtils.cpp
@@ -177,7 +177,7 @@ namespace UnrealConversionUtils
 		int32 NumBoneInfluences = 0;
 		int32 boneIndexBuffer = -1;
 		int32 boneIndexChannel = -1;
-		MutableMeshVertexBuffers.FindChannel(mu::MBS_BONEINDICES, 0, &boneIndexBuffer, &boneIndexChannel);
+		MutableMeshVertexBuffers.FindChannel(mu::EMeshBufferSemantic::BoneIndices, 0, &boneIndexBuffer, &boneIndexChannel);
 		if (boneIndexBuffer >= 0 || boneIndexChannel >= 0)
 		{
 			MutableMeshVertexBuffers.GetChannel(boneIndexBuffer, boneIndexChannel,
@@ -256,33 +256,33 @@ namespace UnrealConversionUtils
 			MutableMeshVertexBuffers.GetBufferData(MUTABLE_VERTEXBUFFER_TEXCOORDS)
 		);
 
-		mu::EMeshBufferFormat BoneIndexFormat = mu::MBF_NONE;
+		mu::EMeshBufferFormat BoneIndexFormat = mu::EMeshBufferFormat::None;
 		int32 NumBoneInfluences = 0;
 		int32 BoneIndexBuffer = -1;
 		int32 BoneIndexChannel = -1;
-		MutableMeshVertexBuffers.FindChannel(mu::MBS_BONEINDICES, 0, &BoneIndexBuffer, &BoneIndexChannel);
+		MutableMeshVertexBuffers.FindChannel(mu::EMeshBufferSemantic::BoneIndices, 0, &BoneIndexBuffer, &BoneIndexChannel);
 		if (BoneIndexBuffer >= 0 || BoneIndexChannel >= 0)
 		{
 			MutableMeshVertexBuffers.GetChannel(BoneIndexBuffer, BoneIndexChannel,
 				nullptr, nullptr, &BoneIndexFormat, &NumBoneInfluences, nullptr);
 		}
 
-		mu::EMeshBufferFormat BoneWeightFormat = mu::MBF_NONE;
+		mu::EMeshBufferFormat BoneWeightFormat = mu::EMeshBufferFormat::None;
 		int32 BoneWeightBuffer = -1;
 		int32 BoneWeightChannel = -1;
-		MutableMeshVertexBuffers.FindChannel(mu::MBS_BONEWEIGHTS, 0, &BoneWeightBuffer, &BoneWeightChannel);
+		MutableMeshVertexBuffers.FindChannel(mu::EMeshBufferSemantic::BoneWeights, 0, &BoneWeightBuffer, &BoneWeightChannel);
 		if (BoneWeightBuffer >= 0 || BoneWeightChannel >= 0)
 		{
 			MutableMeshVertexBuffers.GetChannel(BoneWeightBuffer, BoneWeightChannel,
 				nullptr, nullptr, &BoneWeightFormat, nullptr, nullptr);
 		}
 
-		if (BoneIndexFormat == mu::MBF_UINT16)
+		if (BoneIndexFormat == mu::EMeshBufferFormat::UInt16)
 		{
 			LODResource.SkinWeightVertexBuffer.SetUse16BitBoneIndex(true);
 		}
 
-		if (BoneWeightFormat == mu::MBF_NUINT16)
+		if (BoneWeightFormat == mu::EMeshBufferFormat::NUInt16)
 		{
 			LODResource.SkinWeightVertexBuffer.SetUse16BitBoneWeight(true);
 		}
@@ -311,7 +311,7 @@ namespace UnrealConversionUtils
 				MutableMeshVertexBuffers.GetChannel(Buffer, 0, &Semantic, &SemanticIndex, &Format, &ComponentCount, &Offset);
 
 				// colour buffer?
-				if (Semantic == mu::MBS_COLOUR)
+				if (Semantic == mu::EMeshBufferSemantic::Color)
 				{
 					const void* DataPtr = MutableMeshVertexBuffers.GetBufferData(Buffer);
 					FColorVertexBuffers_InitWithMutableData(LODResource.StaticVertexBuffers, NumVertices, DataPtr);
@@ -350,33 +350,33 @@ namespace UnrealConversionUtils
 			MutableMeshVertexBuffers.GetBufferData(MUTABLE_VERTEXBUFFER_TEXCOORDS)
 		);
 
-		mu::EMeshBufferFormat BoneIndexFormat = mu::MBF_NONE;
+		mu::EMeshBufferFormat BoneIndexFormat = mu::EMeshBufferFormat::None;
 		int32 NumBoneInfluences = 0;
 		int32 BoneIndexBuffer = -1;
 		int32 BoneIndexChannel = -1;
-		MutableMeshVertexBuffers.FindChannel(mu::MBS_BONEINDICES, 0, &BoneIndexBuffer, &BoneIndexChannel);
+		MutableMeshVertexBuffers.FindChannel(mu::EMeshBufferSemantic::BoneIndices, 0, &BoneIndexBuffer, &BoneIndexChannel);
 		if (BoneIndexBuffer >= 0 || BoneIndexChannel >= 0)
 		{
 			MutableMeshVertexBuffers.GetChannel(BoneIndexBuffer, BoneIndexChannel,
 				nullptr, nullptr, &BoneIndexFormat, &NumBoneInfluences, nullptr);
 		}
 
-		mu::EMeshBufferFormat BoneWeightFormat = mu::MBF_NONE;
+		mu::EMeshBufferFormat BoneWeightFormat = mu::EMeshBufferFormat::None;
 		int32 BoneWeightBuffer = -1;
 		int32 BoneWeightChannel = -1;
-		MutableMeshVertexBuffers.FindChannel(mu::MBS_BONEWEIGHTS, 0, &BoneWeightBuffer, &BoneWeightChannel);
+		MutableMeshVertexBuffers.FindChannel(mu::EMeshBufferSemantic::BoneWeights, 0, &BoneWeightBuffer, &BoneWeightChannel);
 		if (BoneWeightBuffer >= 0 || BoneWeightChannel >= 0)
 		{
 			MutableMeshVertexBuffers.GetChannel(BoneWeightBuffer, BoneWeightChannel,
 				nullptr, nullptr, &BoneWeightFormat, nullptr, nullptr);
 		}
 
-		if (BoneIndexFormat == mu::MBF_UINT16)
+		if (BoneIndexFormat == mu::EMeshBufferFormat::UInt16)
 		{
 			LODResource.SkinWeightVertexBuffer.SetUse16BitBoneIndex(true);
 		}
 
-		if (BoneWeightFormat == mu::MBF_NUINT16)
+		if (BoneWeightFormat == mu::EMeshBufferFormat::NUInt16)
 		{
 			LODResource.SkinWeightVertexBuffer.SetUse16BitBoneWeight(true);
 		}
@@ -405,7 +405,7 @@ namespace UnrealConversionUtils
 				MutableMeshVertexBuffers.GetChannel(Buffer, 0, &Semantic, &SemanticIndex, &Format, &ComponentCount, &Offset);
 
 				// colour buffer?
-				if (Semantic == mu::MBS_COLOUR)
+				if (Semantic == mu::EMeshBufferSemantic::Color)
 				{
 					const void* DataPtr = MutableMeshVertexBuffers.GetBufferData(Buffer);
 					FColorVertexBuffers_InitWithMutableData(LODResource.StaticVertexBuffers, NumVertices, DataPtr);
@@ -509,7 +509,7 @@ namespace UnrealConversionUtils
 		int32 Offset;
 		MutableMeshVertexBuffers.GetChannel(BoneIndexBuffer, 2, nullptr, nullptr, &Format, &MutableNumInfluences, &Offset);
 
-		const uint8 MutBoneIndexSize = Format == mu::MBF_UINT16 ? 2 : 1;
+		const uint8 MutBoneIndexSize = Format == mu::EMeshBufferFormat::UInt16 ? 2 : 1;
 		const uint8 MutBoneIndicesSize = Offset - sizeof(int32);
 		const uint8 MutBoneWeightsSize = MutableNumInfluences;
 		const uint8 MutBoneWeightVertexSize = MutBoneIndicesSize + MutBoneWeightsSize;
diff --git a/Engine/Plugins/Mutable/Source/CustomizableObject/Private/MuCO/UnrealMutableImageProvider.cpp b/Engine/Plugins/Mutable/Source/CustomizableObject/Private/MuCO/UnrealMutableImageProvider.cpp
--- a/Engine/Plugins/Mutable/Source/CustomizableObject/Private/MuCO/UnrealMutableImageProvider.cpp
+++ b/Engine/Plugins/Mutable/Source/CustomizableObject/Private/MuCO/UnrealMutableImageProvider.cpp
@@ -28,14 +28,14 @@ namespace
 		check(SizeX > 0 && SizeY > 0);
 
 		EPixelFormat Format = Texture->GetPlatformData()->PixelFormat;
-		mu::EImageFormat MutableFormat = mu::EImageFormat::IF_NONE;
+		mu::EImageFormat MutableFormat = mu::EImageFormat::None;
 
 		switch (Format)
 		{
-		case EPixelFormat::PF_B8G8R8A8: MutableFormat = mu::EImageFormat::IF_BGRA_UBYTE; break;
+		case EPixelFormat::PF_B8G8R8A8: MutableFormat = mu::EImageFormat::BGRA_UByte; break;
 			// This format is deprecated and using the enum fails to compile in some cases.
-			//case ETextureSourceFormat::TSF_RGBA8: MutableFormat = mu::EImageFormat::IF_RGBA_UBYTE; break;
-		case EPixelFormat::PF_G8: MutableFormat = mu::EImageFormat::IF_L_UBYTE; break;
+			//case ETextureSourceFormat::TSF_RGBA8: MutableFormat = mu::EImageFormat::RGBA_UByte; break;
+		case EPixelFormat::PF_G8: MutableFormat = mu::EImageFormat::L_UByte; break;
 		default:
 			break;
 		}
@@ -63,20 +63,20 @@ mu::EImageFormat GetMutablePixelFormat(EPixelFormat InTextureFormat)
 {
 	switch (InTextureFormat)
 	{
-	case PF_B8G8R8A8: return mu::EImageFormat::IF_BGRA_UBYTE;
-	case PF_R8G8B8A8: return mu::EImageFormat::IF_RGBA_UBYTE;
-	case PF_DXT1: return mu::EImageFormat::IF_BC1;
-	case PF_DXT3: return mu::EImageFormat::IF_BC2;
-	case PF_DXT5: return mu::EImageFormat::IF_BC3;
-	case PF_BC4: return mu::EImageFormat::IF_BC4;
-	case PF_BC5: return mu::EImageFormat::IF_BC5;
-	case PF_G8: return mu::EImageFormat::IF_L_UBYTE;
-	case PF_ASTC_4x4: return mu::EImageFormat::IF_ASTC_4x4_RGBA_LDR;
-	case PF_ASTC_6x6: return mu::EImageFormat::IF_ASTC_6x6_RGBA_LDR;
-	case PF_ASTC_8x8: return mu::EImageFormat::IF_ASTC_8x8_RGBA_LDR;
-	case PF_ASTC_10x10: return mu::EImageFormat::IF_ASTC_10x10_RGBA_LDR;
-	case PF_ASTC_12x12: return mu::EImageFormat::IF_ASTC_12x12_RGBA_LDR;
-	default: return mu::EImageFormat::IF_NONE;
+	case PF_B8G8R8A8: return mu::EImageFormat::BGRA_UByte;
+	case PF_R8G8B8A8: return mu::EImageFormat::RGBA_UByte;
+	case PF_DXT1: return mu::EImageFormat::BC1;
+	case PF_DXT3: return mu::EImageFormat::BC2;
+	case PF_DXT5: return mu::EImageFormat::BC3;
+	case PF_BC4: return mu::EImageFormat::BC4;
+	case PF_BC5: return mu::EImageFormat::BC5;
+	case PF_G8: return mu::EImageFormat::L_UByte;
+	case PF_ASTC_4x4: return mu::EImageFormat::ASTC_4x4_RGBA_LDR;
+	case PF_ASTC_6x6: return mu::EImageFormat::ASTC_6x6_RGBA_LDR;
+	case PF_ASTC_8x8: return mu::EImageFormat::ASTC_8x8_RGBA_LDR;
+	case PF_ASTC_10x10: return mu::EImageFormat::ASTC_10x10_RGBA_LDR;
+	case PF_ASTC_12x12: return mu::EImageFormat::ASTC_12x12_RGBA_LDR;
+	default: return mu::EImageFormat::None;
 	}
 }
 
@@ -94,7 +94,7 @@ TTuple<UE::Tasks::FTask, TFunction<void()>> FUnrealMutableResourceProvider::GetI
 	EPixelFormat Format = EPixelFormat::PF_Unknown;
 	int32 BulkDataSize = 0;
 
-	mu::EImageFormat MutImageFormat = mu::EImageFormat::IF_NONE;
+	mu::EImageFormat MutImageFormat = mu::EImageFormat::None;
 	int32 MutImageDataSize = 0;
 
 	auto TrivialReturn = []() -> TTuple<UE::Tasks::FTask, TFunction<void()>>
@@ -169,7 +169,7 @@ TTuple<UE::Tasks::FTask, TFunction<void()>> FUnrealMutableResourceProvider::GetI
 			MutImageFormat = GetMutablePixelFormat(Format);
 
 			// Check if it's a format we support
-			if (MutImageFormat == mu::EImageFormat::IF_NONE)
+			if (MutImageFormat == mu::EImageFormat::None)
 			{
 				UE_LOG(LogMutable, Warning, TEXT("Failed to get external image [%s]. Unexpected image format. EImageFormat [%s]."), *Id.ToString(), GetPixelFormatString(Format));
 				ResultCallback(CreateDummy());
@@ -425,7 +425,7 @@ mu::FImageDesc FUnrealMutableResourceProvider::GetImageDesc(FName Id, uint8 Mipm
 
 		const int32 Lods = 1;
 		
-		return mu::FImageDesc(ImageSize, mu::EImageFormat::IF_NONE, Lods);
+		return mu::FImageDesc(ImageSize, mu::EImageFormat::None, Lods);
 	}
 #else
 	if (UTexture2D* TextureToLoad = ImageInfo.TextureToLoad)
@@ -452,7 +452,7 @@ mu::FImageDesc FUnrealMutableResourceProvider::GetImageDesc(FName Id, uint8 Mipm
 		const mu::EImageFormat MutableFormat = GetMutablePixelFormat(Format);
 
 		// Check if it's a format we support
-		if (MutableFormat == mu::EImageFormat::IF_NONE)
+		if (MutableFormat == mu::EImageFormat::None)
 		{
 			UE_LOG(LogMutable, Warning, TEXT("Failed to get external image descriptor. Unexpected image format. EImageFormat [%s]."), GetPixelFormatString(Format));
 			return CreateDummyDesc();
@@ -565,7 +565,7 @@ void FUnrealMutableResourceProvider::CacheImage(FName Id, bool bUser)
 				case UCustomizableSystemImageProvider::ValueType::Raw:
 					{
 						FIntVector desc = Provider->GetTextureParameterValueSize(Id);
-						mu::Ptr<mu::Image> pResult = new mu::Image(desc[0], desc[1], 1, mu::EImageFormat::IF_RGBA_UBYTE, mu::EInitializationType::Black);
+						mu::Ptr<mu::Image> pResult = new mu::Image(desc[0], desc[1], 1, mu::EImageFormat::RGBA_UByte, mu::EInitializationType::Black);
 						Provider->GetTextureParameterValueData(Id, pResult->GetLODData(0));
 
 						ImageInfo = FUnrealMutableImageInfo(pResult);
@@ -683,7 +683,7 @@ void FUnrealMutableResourceProvider::CacheImages(const mu::Parameters& Parameter
 	const int32 NumParams = Parameters.GetCount();
 	for (int32 ParamIndex = 0; ParamIndex < NumParams; ++ParamIndex)
 	{
-		if (Parameters.GetType(ParamIndex) != mu::PARAMETER_TYPE::T_IMAGE)
+		if (Parameters.GetType(ParamIndex) != mu::EParameterType::Image)
 		{
 			continue;
 		}
@@ -710,7 +710,7 @@ void FUnrealMutableResourceProvider::UnCacheImages(const mu::Parameters& Paramet
 	const int32 NumParams = Parameters.GetCount();
 	for (int32 ParamIndex = 0; ParamIndex < NumParams; ++ParamIndex)
 	{
-		if (Parameters.GetType(ParamIndex) != mu::PARAMETER_TYPE::T_IMAGE)
+		if (Parameters.GetType(ParamIndex) != mu::EParameterType::Image)
 		{
 			continue;
 		}
@@ -776,7 +776,7 @@ mu::ImagePtr FUnrealMutableResourceProvider::CreateDummy()
 	mu::ImagePtr pResult = new mu::Image(Size, Size, DUMMY_IMAGE_DESC.m_lods, DUMMY_IMAGE_DESC.m_format, mu::EInitializationType::NotInitialized);
 
 	check(pResult->GetLODCount() == 1);
-	check(pResult->GetFormat() == mu::EImageFormat::IF_RGBA_UBYTE || pResult->GetFormat() == mu::EImageFormat::IF_BGRA_UBYTE);
+	check(pResult->GetFormat() == mu::EImageFormat::RGBA_UByte || pResult->GetFormat() == mu::EImageFormat::BGRA_UByte);
 	uint8* pData = pResult->GetLODData(0);
 	for (int32 X = 0; X < Size; ++X)
 	{
diff --git a/Engine/Plugins/Mutable/Source/CustomizableObject/Private/MuCO/UnrealMutableImageProvider.h b/Engine/Plugins/Mutable/Source/CustomizableObject/Private/MuCO/UnrealMutableImageProvider.h
--- a/Engine/Plugins/Mutable/Source/CustomizableObject/Private/MuCO/UnrealMutableImageProvider.h
+++ b/Engine/Plugins/Mutable/Source/CustomizableObject/Private/MuCO/UnrealMutableImageProvider.h
@@ -88,7 +88,7 @@ private:
 	};
 
 	static inline const mu::FImageDesc DUMMY_IMAGE_DESC = 
-			mu::FImageDesc(mu::FImageSize(32, 32), mu::EImageFormat::IF_RGBA_UBYTE, 1);
+			mu::FImageDesc(mu::FImageSize(32, 32), mu::EImageFormat::RGBA_UByte, 1);
 
 	/** This will be called if an image Id has been requested by Mutable core but it has not been provided by any provider. */
 	static mu::ImagePtr CreateDummy();
diff --git a/Engine/Plugins/Mutable/Source/CustomizableObject/Private/MuCO/UnrealMutableModelDiskStreamer.cpp b/Engine/Plugins/Mutable/Source/CustomizableObject/Private/MuCO/UnrealMutableModelDiskStreamer.cpp
--- a/Engine/Plugins/Mutable/Source/CustomizableObject/Private/MuCO/UnrealMutableModelDiskStreamer.cpp
+++ b/Engine/Plugins/Mutable/Source/CustomizableObject/Private/MuCO/UnrealMutableModelDiskStreamer.cpp
@@ -190,7 +190,7 @@ void FUnrealMutableModelBulkReader::EndStreaming()
 {
 	for (FObjectData& o : Objects)
 	{
-		for (TPair<OPERATION_ID, FReadRequest>& it : o.CurrentReadRequests)
+		for (TPair<FOperationID, FReadRequest>& it : o.CurrentReadRequests)
 		{
 			if (it.Value.FileReadRequest)
 			{
@@ -213,7 +213,7 @@ FAutoConsoleVariableRef CVarStreamPriority(
 	TEXT(""));
 
 
-mu::ModelReader::OPERATION_ID FUnrealMutableModelBulkReader::BeginReadBlock(const mu::Model* Model, uint32 Key, void* pBuffer, uint64 size, TFunction<void(bool bSuccess)>* CompletionCallback)
+mu::ModelReader::FOperationID FUnrealMutableModelBulkReader::BeginReadBlock(const mu::Model* Model, uint32 Key, void* pBuffer, uint64 size, TFunction<void(bool bSuccess)>* CompletionCallback)
 {
 	MUTABLE_CPUPROFILER_SCOPE(FUnrealMutableModelBulkStreamer::OpenReadFile);
 
@@ -249,7 +249,7 @@ mu::ModelReader::OPERATION_ID FUnrealMutableModelBulkReader::BeginReadBlock(cons
 		return -1;
 	}
 
-	OPERATION_ID Result = ++LastOperationID;
+	FOperationID Result = ++LastOperationID;
 
 	if(ObjectData->ModelStreamableBulkData->StreamableBulkData.IsValidIndex(Block->FileId))
 	{
@@ -399,7 +399,7 @@ mu::ModelReader::OPERATION_ID FUnrealMutableModelBulkReader::BeginReadBlock(cons
 }
 
 
-bool FUnrealMutableModelBulkReader::IsReadCompleted(mu::ModelReader::OPERATION_ID OperationId)
+bool FUnrealMutableModelBulkReader::IsReadCompleted(mu::ModelReader::FOperationID OperationId)
 {
 	MUTABLE_CPUPROFILER_SCOPE(FUnrealMutableModelBulkStreamer::IsReadCompleted);
 
@@ -430,7 +430,7 @@ bool FUnrealMutableModelBulkReader::IsReadCompleted(mu::ModelReader::OPERATION_I
 }
 
 
-bool FUnrealMutableModelBulkReader::EndRead(mu::ModelReader::OPERATION_ID OperationId)
+bool FUnrealMutableModelBulkReader::EndRead(mu::ModelReader::FOperationID OperationId)
 {
 	MUTABLE_CPUPROFILER_SCOPE(FUnrealMutableModelBulkStreamer::EndRead);
 
diff --git a/Engine/Plugins/Mutable/Source/CustomizableObject/Private/MuCO/UnrealToMutableTextureConversionUtils.cpp b/Engine/Plugins/Mutable/Source/CustomizableObject/Private/MuCO/UnrealToMutableTextureConversionUtils.cpp
--- a/Engine/Plugins/Mutable/Source/CustomizableObject/Private/MuCO/UnrealToMutableTextureConversionUtils.cpp
+++ b/Engine/Plugins/Mutable/Source/CustomizableObject/Private/MuCO/UnrealToMutableTextureConversionUtils.cpp
@@ -295,7 +295,7 @@ EUnrealToMutableConversionError ConvertTextureUnrealSourceToMutable(mu::Image* O
 		MUTABLE_CPUPROFILER_SCOPE(NoConvert);
         check(LODs == 1);
 
-		OutResult->Init(SizeX, SizeY, LODs, mu::EImageFormat::IF_L_UBYTE, mu::EInitializationType::NotInitialized);
+		OutResult->Init(SizeX, SizeY, LODs, mu::EImageFormat::L_UByte, mu::EInitializationType::NotInitialized);
 		OutResult->DataStorage.GetInternalArray(0) = MoveTemp(TempImage.RawData);
 		break;
 	}
@@ -312,7 +312,7 @@ EUnrealToMutableConversionError ConvertTextureUnrealSourceToMutable(mu::Image* O
 		if (bHasAlphaChannel)
 		{
 			MUTABLE_CPUPROFILER_SCOPE(ToRGBA);
-			OutResult->Init(SizeX, SizeY, LODs, mu::EImageFormat::IF_RGBA_UBYTE, mu::EInitializationType::NotInitialized);
+			OutResult->Init(SizeX, SizeY, LODs, mu::EImageFormat::RGBA_UByte, mu::EInitializationType::NotInitialized);
             check(LODs == 1);
             uint8* DataDest = OutResult->GetLODData(0);
 
@@ -332,7 +332,7 @@ EUnrealToMutableConversionError ConvertTextureUnrealSourceToMutable(mu::Image* O
 			MUTABLE_CPUPROFILER_SCOPE(ToRGB);
 
 			// TODO: add support for a mu::IF_RGBX_UBYTE?			
-			OutResult->Init(SizeX, SizeY, LODs, mu::EImageFormat::IF_RGB_UBYTE, mu::EInitializationType::NotInitialized);
+			OutResult->Init(SizeX, SizeY, LODs, mu::EImageFormat::RGB_UByte, mu::EInitializationType::NotInitialized);
 			check(LODs == 1);
             uint8* DataDest = OutResult->GetLODData(0);
 
diff --git a/Engine/Plugins/Mutable/Source/CustomizableObjectEditor/Private/MuCOE/GenerateMutableSource/GenerateMutableSource.cpp b/Engine/Plugins/Mutable/Source/CustomizableObjectEditor/Private/MuCOE/GenerateMutableSource/GenerateMutableSource.cpp
--- a/Engine/Plugins/Mutable/Source/CustomizableObjectEditor/Private/MuCOE/GenerateMutableSource/GenerateMutableSource.cpp
+++ b/Engine/Plugins/Mutable/Source/CustomizableObjectEditor/Private/MuCOE/GenerateMutableSource/GenerateMutableSource.cpp
@@ -1598,7 +1598,7 @@ mu::FImageDesc GenerateImageDescriptor(UTexture* Texture)
 	ImageDesc.m_size[1] = Texture->Source.GetSizeY();
 	ImageDesc.m_lods = Texture->Source.GetNumMips();
 
-	mu::EImageFormat MutableFormat = mu::EImageFormat::IF_RGBA_UBYTE;
+	mu::EImageFormat MutableFormat = mu::EImageFormat::RGBA_UByte;
 	ETextureSourceFormat SourceFormat = Texture->Source.GetFormat();
 	switch (SourceFormat)
 	{
@@ -1606,7 +1606,7 @@ mu::FImageDesc GenerateImageDescriptor(UTexture* Texture)
 	case ETextureSourceFormat::TSF_G16:
 	case ETextureSourceFormat::TSF_R16F:
 	case ETextureSourceFormat::TSF_R32F:
-		MutableFormat = mu::EImageFormat::IF_L_UBYTE;
+		MutableFormat = mu::EImageFormat::L_UByte;
 		break;
 
 	default:
diff --git a/Engine/Plugins/Mutable/Source/CustomizableObjectEditor/Private/MuCOE/GenerateMutableSource/GenerateMutableSourceGroupProjector.cpp b/Engine/Plugins/Mutable/Source/CustomizableObjectEditor/Private/MuCOE/GenerateMutableSource/GenerateMutableSourceGroupProjector.cpp
--- a/Engine/Plugins/Mutable/Source/CustomizableObjectEditor/Private/MuCOE/GenerateMutableSource/GenerateMutableSourceGroupProjector.cpp
+++ b/Engine/Plugins/Mutable/Source/CustomizableObjectEditor/Private/MuCOE/GenerateMutableSource/GenerateMutableSourceGroupProjector.cpp
@@ -219,7 +219,7 @@ mu::NodeImagePtr GenerateMutableSourceGroupProjector(const int32 NodeLOD, const
 	ZeroPlainColourNode->SetColour(ZeroColorNode);
 
 	mu::NodeImageSwizzlePtr ZeroChannelNode = new mu::NodeImageSwizzle;
-	ZeroChannelNode->SetFormat(mu::EImageFormat::IF_L_UBYTE);
+	ZeroChannelNode->SetFormat(mu::EImageFormat::L_UByte);
 	ZeroChannelNode->SetSource(0, ZeroPlainColourNode);
 	ZeroChannelNode->SetSourceChannel(0, 2); // Just take a zeroed channel for the base alpha
 
@@ -251,7 +251,7 @@ mu::NodeImagePtr GenerateMutableSourceGroupProjector(const int32 NodeLOD, const
 		}
 
 		mu::NodeImageSwizzlePtr ImageNodesAlphaChannelNode = new mu::NodeImageSwizzle;
-		ImageNodesAlphaChannelNode->SetFormat(mu::EImageFormat::IF_L_UBYTE);
+		ImageNodesAlphaChannelNode->SetFormat(mu::EImageFormat::L_UByte);
 		ImageNodesAlphaChannelNode->SetSource(0, ImageNodes[i]);
 		ImageNodesAlphaChannelNode->SetSourceChannel(0, 3);
 
@@ -268,7 +268,7 @@ mu::NodeImagePtr GenerateMutableSourceGroupProjector(const int32 NodeLOD, const
 		//OpacityMultiLayerNode->SetMask(OneChannelNode); // No mask needed
 
 		mu::NodeImageSwizzlePtr MultiplySwizzleNode = new mu::NodeImageSwizzle;
-		MultiplySwizzleNode->SetFormat(mu::EImageFormat::IF_L_UBYTE);
+		MultiplySwizzleNode->SetFormat(mu::EImageFormat::L_UByte);
 		MultiplySwizzleNode->SetSource(0, OpacityMultiLayerNode);
 		MultiplySwizzleNode->SetSourceChannel(0, 0);
 
@@ -290,22 +290,22 @@ mu::NodeImagePtr GenerateMutableSourceGroupProjector(const int32 NodeLOD, const
 	}
 
 	mu::NodeImageSwizzlePtr SwizzleNodeR = new mu::NodeImageSwizzle;
-	SwizzleNodeR->SetFormat(mu::EImageFormat::IF_L_UBYTE);
+	SwizzleNodeR->SetFormat(mu::EImageFormat::L_UByte);
 	SwizzleNodeR->SetSource(0, ResultImage);
 	SwizzleNodeR->SetSourceChannel(0, 0);
 
 	mu::NodeImageSwizzlePtr SwizzleNodeG = new mu::NodeImageSwizzle;
-	SwizzleNodeG->SetFormat(mu::EImageFormat::IF_L_UBYTE);
+	SwizzleNodeG->SetFormat(mu::EImageFormat::L_UByte);
 	SwizzleNodeG->SetSource(0, ResultImage);
 	SwizzleNodeG->SetSourceChannel(0, 1);
 
 	mu::NodeImageSwizzlePtr SwizzleNodeB = new mu::NodeImageSwizzle;
-	SwizzleNodeB->SetFormat(mu::EImageFormat::IF_L_UBYTE);
+	SwizzleNodeB->SetFormat(mu::EImageFormat::L_UByte);
 	SwizzleNodeB->SetSource(0, ResultImage);
 	SwizzleNodeB->SetSourceChannel(0, 2);
 
 	mu::NodeImageSwizzlePtr FinalSwizzleNode = new mu::NodeImageSwizzle;
-	FinalSwizzleNode->SetFormat(mu::EImageFormat::IF_RGBA_UBYTE);
+	FinalSwizzleNode->SetFormat(mu::EImageFormat::RGBA_UByte);
 	FinalSwizzleNode->SetSource(0, SwizzleNodeR);
 	FinalSwizzleNode->SetSourceChannel(0, 0);
 	FinalSwizzleNode->SetSource(1, SwizzleNodeG);
diff --git a/Engine/Plugins/Mutable/Source/CustomizableObjectEditor/Private/MuCOE/GenerateMutableSource/GenerateMutableSourceImage.cpp b/Engine/Plugins/Mutable/Source/CustomizableObjectEditor/Private/MuCOE/GenerateMutableSource/GenerateMutableSourceImage.cpp
--- a/Engine/Plugins/Mutable/Source/CustomizableObjectEditor/Private/MuCOE/GenerateMutableSource/GenerateMutableSourceImage.cpp
+++ b/Engine/Plugins/Mutable/Source/CustomizableObjectEditor/Private/MuCOE/GenerateMutableSource/GenerateMutableSourceImage.cpp
@@ -165,12 +165,12 @@ mu::NodeImagePtr GenerateMutableSourceImage(const UEdGraphPin* Pin, FMutableGrap
 		// \TODO: Take if from default?
 		//if (TypedNodeParam->DefaultValue)
 		//{
-		//	FormatNode->SetFormat(mu::EImageFormat::IF_RGBA_UBYTE);
+		//	FormatNode->SetFormat(mu::EImageFormat::RGBA_UByte);
 		//}
 		//else
 		{
 			// Force an "easy format" on the texture.
-			FormatNode->SetFormat(mu::EImageFormat::IF_RGBA_UBYTE);
+			FormatNode->SetFormat(mu::EImageFormat::RGBA_UByte);
 		}
 		FormatNode->SetSource(TextureNode);
 
@@ -530,7 +530,7 @@ mu::NodeImagePtr GenerateMutableSourceImage(const UEdGraphPin* Pin, FMutableGrap
 
 		if (RGB && ANode)
 		{
-			SwizzleNode->SetFormat(mu::EImageFormat::IF_RGBA_UBYTE);
+			SwizzleNode->SetFormat(mu::EImageFormat::RGBA_UByte);
 			SwizzleNode->SetSource(0, RNode);
 			SwizzleNode->SetSourceChannel(0, 0);
 			SwizzleNode->SetSource(1, GNode);
@@ -542,7 +542,7 @@ mu::NodeImagePtr GenerateMutableSourceImage(const UEdGraphPin* Pin, FMutableGrap
 		}
 		else if (RGB)
 		{
-			SwizzleNode->SetFormat(mu::EImageFormat::IF_RGB_UBYTE);
+			SwizzleNode->SetFormat(mu::EImageFormat::RGB_UByte);
 			SwizzleNode->SetSource(0, RNode);
 			SwizzleNode->SetSourceChannel(0, 0);
 			SwizzleNode->SetSource(1, GNode);
@@ -552,13 +552,13 @@ mu::NodeImagePtr GenerateMutableSourceImage(const UEdGraphPin* Pin, FMutableGrap
 		}
 		else if (RNode)
 		{
-			SwizzleNode->SetFormat(mu::EImageFormat::IF_L_UBYTE);
+			SwizzleNode->SetFormat(mu::EImageFormat::L_UByte);
 			SwizzleNode->SetSource(0, RNode);
 			SwizzleNode->SetSourceChannel(0, 0);
 		}
 		else if (ANode)
 		{
-			SwizzleNode->SetFormat(mu::EImageFormat::IF_L_UBYTE);
+			SwizzleNode->SetFormat(mu::EImageFormat::L_UByte);
 			SwizzleNode->SetSource(0, ANode);
 			SwizzleNode->SetSourceChannel(0, 0);
 		}
@@ -574,7 +574,7 @@ mu::NodeImagePtr GenerateMutableSourceImage(const UEdGraphPin* Pin, FMutableGrap
 
 		mu::NodeImageSwizzlePtr SwizzleNode = new mu::NodeImageSwizzle;
 		Result = SwizzleNode;
-		SwizzleNode->SetFormat(mu::EImageFormat::IF_L_UBYTE);
+		SwizzleNode->SetFormat(mu::EImageFormat::L_UByte);
 		SwizzleNode->SetSource(0, BaseNode);
 
 		if (Pin == TypedNodeTo->RPin())
diff --git a/Engine/Plugins/Mutable/Source/CustomizableObjectEditor/Private/MuCOE/GenerateMutableSource/GenerateMutableSourceMesh.cpp b/Engine/Plugins/Mutable/Source/CustomizableObjectEditor/Private/MuCOE/GenerateMutableSource/GenerateMutableSourceMesh.cpp
--- a/Engine/Plugins/Mutable/Source/CustomizableObjectEditor/Private/MuCOE/GenerateMutableSource/GenerateMutableSourceMesh.cpp
+++ b/Engine/Plugins/Mutable/Source/CustomizableObjectEditor/Private/MuCOE/GenerateMutableSource/GenerateMutableSourceMesh.cpp
@@ -958,12 +958,12 @@ mu::MeshPtr ConvertSkeletalMeshToMutable(
 		using namespace mu;
 		const int32 ElementSize = sizeof(FSoftSkinVertex);
 		const int32 ChannelCount = 11;
-		const EMeshBufferSemantic Semantics[ChannelCount] = { MBS_POSITION, MBS_TANGENT, MBS_BINORMAL, MBS_NORMAL, MBS_TEXCOORDS, MBS_TEXCOORDS, MBS_TEXCOORDS, MBS_TEXCOORDS, MBS_COLOUR, MBS_BONEINDICES, MBS_BONEWEIGHTS };
+		const EMeshBufferSemantic Semantics[ChannelCount] = { EMeshBufferSemantic::Position, EMeshBufferSemantic::Tangent, EMeshBufferSemantic::Binormal, EMeshBufferSemantic::Normal, EMeshBufferSemantic::TexCoords, EMeshBufferSemantic::TexCoords, EMeshBufferSemantic::TexCoords, EMeshBufferSemantic::TexCoords, EMeshBufferSemantic::Color, EMeshBufferSemantic::BoneIndices, EMeshBufferSemantic::BoneWeights };
 		const int32 SemanticIndices[ChannelCount] = { 0, 0, 0, 0, 0, 1, 2, 3, 0, 0, 0 };
 
 		// TODO: Remove BoneWeightFormat after merge
-		EMeshBufferFormat BoneWeightFormat = sizeof(TDecay<decltype(DeclVal<FSoftSkinVertex>().InfluenceWeights[0])>::Type) == 1 ? MBF_NUINT8 : MBF_NUINT16;
-		const EMeshBufferFormat Formats[ChannelCount] = { MBF_FLOAT32, MBF_FLOAT32, MBF_FLOAT32, MBF_FLOAT32, MBF_FLOAT32, MBF_FLOAT32, MBF_FLOAT32, MBF_FLOAT32, MBF_NUINT8, MBF_UINT16, BoneWeightFormat };
+		EMeshBufferFormat BoneWeightFormat = sizeof(TDecay<decltype(DeclVal<FSoftSkinVertex>().InfluenceWeights[0])>::Type) == 1 ? EMeshBufferFormat::NUInt8 : EMeshBufferFormat::NUInt16;
+		const EMeshBufferFormat Formats[ChannelCount] = { EMeshBufferFormat::Float32, EMeshBufferFormat::Float32, EMeshBufferFormat::Float32, EMeshBufferFormat::Float32, EMeshBufferFormat::Float32, EMeshBufferFormat::Float32, EMeshBufferFormat::Float32, EMeshBufferFormat::Float32, EMeshBufferFormat::NUInt8, EMeshBufferFormat::UInt16, BoneWeightFormat };
 
 		int32 Components[ChannelCount] = { 3, 3, 3, 4, 2, 2, 2, 2, 4, 4, 4 };
 		if (GenerationContext.Options.CustomizableObjectNumBoneInfluences != ECustomizableObjectNumBoneInfluences::Four &&
@@ -1011,9 +1011,9 @@ mu::MeshPtr ConvertSkeletalMeshToMutable(
 		// Base channels
 		{
 			const int32 ChannelCount = 5;
-			const EMeshBufferSemantic Semantics[ChannelCount] = { MBS_POSITION, MBS_TANGENT, MBS_BINORMAL, MBS_NORMAL, MBS_COLOUR };
+			const EMeshBufferSemantic Semantics[ChannelCount] = { EMeshBufferSemantic::Position, EMeshBufferSemantic::Tangent, EMeshBufferSemantic::Binormal, EMeshBufferSemantic::Normal, EMeshBufferSemantic::Color };
 			const int32 SemanticIndices[ChannelCount] = { 0, 0, 0, 0, 0 };
-			const EMeshBufferFormat Formats[ChannelCount] = { MBF_FLOAT32, MBF_FLOAT32, MBF_FLOAT32, MBF_FLOAT32, MBF_NUINT8 };
+			const EMeshBufferFormat Formats[ChannelCount] = { EMeshBufferFormat::Float32, EMeshBufferFormat::Float32, EMeshBufferFormat::Float32, EMeshBufferFormat::Float32, EMeshBufferFormat::NUInt8 };
 			int32 Components[ChannelCount] = { 3, 3, 3, 4, 4 };
 			const int32 Offsets[ChannelCount] =
 			{
@@ -1033,9 +1033,9 @@ mu::MeshPtr ConvertSkeletalMeshToMutable(
 		{
 			const int32 ElementSize = 4 * 2 * sizeof(float);
 			const int32 ChannelCount = 4;
-			const EMeshBufferSemantic Semantics[ChannelCount] = { MBS_TEXCOORDS, MBS_TEXCOORDS, MBS_TEXCOORDS, MBS_TEXCOORDS };
+			const EMeshBufferSemantic Semantics[ChannelCount] = { EMeshBufferSemantic::TexCoords, EMeshBufferSemantic::TexCoords, EMeshBufferSemantic::TexCoords, EMeshBufferSemantic::TexCoords };
 			const int32 SemanticIndices[ChannelCount] = { 0, 1, 2, 3 };
-			const EMeshBufferFormat Formats[ChannelCount] = { MBF_FLOAT32, MBF_FLOAT32, MBF_FLOAT32, MBF_FLOAT32 };
+			const EMeshBufferFormat Formats[ChannelCount] = { EMeshBufferFormat::Float32, EMeshBufferFormat::Float32, EMeshBufferFormat::Float32, EMeshBufferFormat::Float32 };
 			int32 Components[ChannelCount] = { 2, 2, 2, 2 };
 			const int32 Offsets[ChannelCount] =
 			{
@@ -1057,12 +1057,12 @@ mu::MeshPtr ConvertSkeletalMeshToMutable(
 		// Skinning buffer
 		using namespace mu;
 		const int32 ChannelCount = 2;
-		const EMeshBufferSemantic Semantics[ChannelCount] = { MBS_BONEINDICES, MBS_BONEWEIGHTS };
+		const EMeshBufferSemantic Semantics[ChannelCount] = { EMeshBufferSemantic::BoneIndices, EMeshBufferSemantic::BoneWeights };
 		const int32 SemanticIndices[ChannelCount] = { 0, 0 };
 
 		// TODO: Remove BoneWeightFormat after merge
-		EMeshBufferFormat BoneWeightFormat = sizeof(TDecay<decltype(DeclVal<FSoftSkinVertex>().InfluenceWeights[0])>::Type) == 1 ? MBF_NUINT8 : MBF_NUINT16;
-		const EMeshBufferFormat Formats[ChannelCount] = { MBF_UINT16, BoneWeightFormat };
+		EMeshBufferFormat BoneWeightFormat = sizeof(TDecay<decltype(DeclVal<FSoftSkinVertex>().InfluenceWeights[0])>::Type) == 1 ? EMeshBufferFormat::NUInt8 : EMeshBufferFormat::NUInt16;
+		const EMeshBufferFormat Formats[ChannelCount] = { EMeshBufferFormat::UInt16, BoneWeightFormat };
 
 		int32 Components[ChannelCount] = { 4, 4 };
 		if (GenerationContext.Options.CustomizableObjectNumBoneInfluences != ECustomizableObjectNumBoneInfluences::Four && 
@@ -1090,12 +1090,12 @@ mu::MeshPtr ConvertSkeletalMeshToMutable(
 			Components[0] * sizeof(uint16),
 		};
 
-		int32 ElementSize = Components[0] * sizeof(uint16) + Components[1] * ((BoneWeightFormat == MBF_NUINT8) ? 1 : 2);
+		int32 ElementSize = Components[0] * sizeof(uint16) + Components[1] * ((BoneWeightFormat == EMeshBufferFormat::NUInt8) ? 1 : 2);
 
 		MutableMesh->GetVertexBuffers().SetBuffer(2, ElementSize, ChannelCount, Semantics, SemanticIndices, Formats, Components, Offsets);
 
 		// Fix bone weights if required (uint8 -> uint16)
-		if (BoneWeightFormat == MBF_NUINT16 && Vertices.IsValidIndex(VertexStart))
+		if (BoneWeightFormat == EMeshBufferFormat::NUInt16 && Vertices.IsValidIndex(VertexStart))
 		{
 			FSoftSkinVertex FirstVertex = Vertices[VertexStart];
 
@@ -1136,7 +1136,7 @@ mu::MeshPtr ConvertSkeletalMeshToMutable(
 				const int32 MaxMutableWeights = 4;
 				int32 MaxOrderedWeighsIndices[MaxMutableWeights] = { -1, -1, -1, -1 };
 
-				const int32 MaxBoneWeightValue = BoneWeightFormat == MBF_NUINT16 ? 65535 : 255;
+				const int32 MaxBoneWeightValue = BoneWeightFormat == EMeshBufferFormat::NUInt16 ? 65535 : 255;
 				NormalizeWeights(&Vertex.InfluenceBones[0], &Vertex.InfluenceWeights[0], MaxSectionInfluences, MaxMutableWeights,
 					&MaxOrderedWeighsIndices[0], MaxSectionBoneMapIndex, MaxBoneWeightValue);
 			}
@@ -1277,9 +1277,9 @@ mu::MeshPtr ConvertSkeletalMeshToMutable(
 				using namespace mu;
 				const int32 ElementSize = sizeof(uint32);
 				const int32 ChannelCount = 1;
-				const EMeshBufferSemantic Semantics[ChannelCount] = { MBS_OTHER };
+				const EMeshBufferSemantic Semantics[ChannelCount] = { EMeshBufferSemantic::Other };
 				const int32 SemanticIndices[ChannelCount] = { 0 };
-				const EMeshBufferFormat Formats[ChannelCount] = { MBF_UINT32 };
+				const EMeshBufferFormat Formats[ChannelCount] = { EMeshBufferFormat::UInt32 };
 				int32 Components[ChannelCount] = { 1 };
 				const int32 Offsets[ChannelCount] = { 0 };
 
@@ -1293,9 +1293,9 @@ mu::MeshPtr ConvertSkeletalMeshToMutable(
 				using namespace mu;
 				const int32 ElementSize = sizeof(uint32);
 				const int32 ChannelCount = 1;
-				const EMeshBufferSemantic Semantics[ChannelCount] = { MBS_OTHER };
+				const EMeshBufferSemantic Semantics[ChannelCount] = { EMeshBufferSemantic::Other };
 				const int32 SemanticIndices[ChannelCount] = { 1 };
-				const EMeshBufferFormat Formats[ChannelCount] = { MBF_UINT32 };
+				const EMeshBufferFormat Formats[ChannelCount] = { EMeshBufferFormat::UInt32 };
 				int32 Components[ChannelCount] = { 1 };
 				const int32 Offsets[ChannelCount] = { 0 };
 
@@ -1626,9 +1626,9 @@ mu::MeshPtr ConvertSkeletalMeshToMutable(
 				using namespace mu;
 				const int32 ElementSize = sizeof(int32);
 				const int32 ChannelCount = 1;
-				const EMeshBufferSemantic Semantics[ChannelCount] = { MBS_OTHER };
+				const EMeshBufferSemantic Semantics[ChannelCount] = { EMeshBufferSemantic::Other };
 				const int32 SemanticIndices[ChannelCount] = { 2 };
-				const EMeshBufferFormat Formats[ChannelCount] = { MBF_INT32 };
+				const EMeshBufferFormat Formats[ChannelCount] = { EMeshBufferFormat::Int32 };
 				int32 Components[ChannelCount] = { 1 };
 				const int32 Offsets[ChannelCount] = { 0 };
 
@@ -1639,9 +1639,9 @@ mu::MeshPtr ConvertSkeletalMeshToMutable(
 				using namespace mu;
 				const int32 ElementSize = sizeof(uint32);
 				const int32 ChannelCount = 1;
-				const EMeshBufferSemantic Semantics[ChannelCount] = { MBS_OTHER };
+				const EMeshBufferSemantic Semantics[ChannelCount] = { EMeshBufferSemantic::Other };
 				const int32 SemanticIndices[ChannelCount] = { 3 };
-				const EMeshBufferFormat Formats[ChannelCount] = { MBF_UINT32 };
+				const EMeshBufferFormat Formats[ChannelCount] = { EMeshBufferFormat::UInt32 };
 				int32 Components[ChannelCount] = { 1 };
 				const int32 Offsets[ChannelCount] = { 0 };
 
@@ -1813,7 +1813,7 @@ mu::MeshPtr ConvertSkeletalMeshToMutable(
 
 		// TODO: Remove BoneWeightFormat after merge
 		const int32 BoneWeightTypeSizeBytes = sizeof(TDecay<decltype(DeclVal<FRawSkinWeight>().InfluenceWeights[0])>::Type);
-		EMeshBufferFormat BoneWeightFormat = BoneWeightTypeSizeBytes == 1 ? MBF_NUINT8 : MBF_NUINT16;
+		EMeshBufferFormat BoneWeightFormat = BoneWeightTypeSizeBytes == 1 ? EMeshBufferFormat::NUInt8 : EMeshBufferFormat::NUInt16;
 
 		// Limit skinning weights if necessary
 		const int32 MutableBonesPerVertex = bUseUnlimitedInfluences ?
@@ -1855,7 +1855,7 @@ mu::MeshPtr ConvertSkeletalMeshToMutable(
 					const int32 MaxMutableWeights = 4;
 					int32 MaxOrderedWeighsIndices[MaxMutableWeights] = { -1, -1, -1, -1 };
 
-					const int32 MaxBoneWeightValue = BoneWeightFormat == MBF_NUINT16 ? 65535 : 255;
+					const int32 MaxBoneWeightValue = BoneWeightFormat == EMeshBufferFormat::NUInt16 ? 65535 : 255;
 					NormalizeWeights(&SkinWeight.InfluenceBones[0], &SkinWeight.InfluenceWeights[0], MaxSectionInfluences, MaxMutableWeights,
 						&MaxOrderedWeighsIndices[0], MaxSectionBoneMapIndex, MaxBoneWeightValue);
 				}
@@ -1898,9 +1898,9 @@ mu::MeshPtr ConvertSkeletalMeshToMutable(
 			// Set up SkinWeightPRofile BufferData
 			const int32 ElementSize = sizeof(int32) + sizeof(FBoneIndexType) + BoneWeightTypeSizeBytes;
 			const int32 ChannelCount = 3;
-			const EMeshBufferSemantic Semantics[ChannelCount] = { MBS_ALTSKINWEIGHT, MBS_BONEINDICES, MBS_BONEWEIGHTS };
+			const EMeshBufferSemantic Semantics[ChannelCount] = { EMeshBufferSemantic::AltSkinWeight, EMeshBufferSemantic::BoneIndices, EMeshBufferSemantic::BoneWeights };
 			const int32 SemanticIndices[ChannelCount] = { static_cast<int32>(ProfileId), static_cast<int32>(ProfileId), static_cast<int32>(ProfileId) };
-			const EMeshBufferFormat Formats[ChannelCount] = { MBF_INT32, MBF_UINT16, BoneWeightFormat };
+			const EMeshBufferFormat Formats[ChannelCount] = { EMeshBufferFormat::Int32, EMeshBufferFormat::UInt16, BoneWeightFormat };
 			const int32 Components[ChannelCount] = { 1, MutableBonesPerVertex, MutableBonesPerVertex };
 			const int32 Offsets[ChannelCount] = { 0, sizeof(int32), sizeof(int32) + BoneIndicesSize };
 
@@ -1925,12 +1925,12 @@ mu::MeshPtr ConvertSkeletalMeshToMutable(
 
 		const int32 FinalElementSize = sizeof(uint32_t);
 		const int32 ChannelCount = 1;
-		const EMeshBufferSemantic Semantics[ChannelCount] = { MBS_VERTEXINDEX };
+		const EMeshBufferSemantic Semantics[ChannelCount] = { EMeshBufferSemantic::VertexIndex };
 		const int32 SemanticIndices[ChannelCount] = { 0 };
 		// We force 32 bit indices, since merging meshes may create vertex buffers bigger than the initial mesh
 		// and for now the mutable runtime doesn't handle it.
 		// \TODO: go back to 16-bit indices when possible.
-		EMeshBufferFormat Formats[ChannelCount] = { MBF_UINT32 };
+		EMeshBufferFormat Formats[ChannelCount] = { EMeshBufferFormat::UInt32 };
 		const int32 Components[ChannelCount] = { 1 };
 		const int32 Offsets[ChannelCount] = { 0 };
 
@@ -2280,9 +2280,9 @@ mu::MeshPtr ConvertStaticMeshToMutable(const UStaticMesh* StaticMesh, int32 LODI
 
 			const int32 ElementSize = 12;
 			const int32 ChannelCount = 1;
-			const EMeshBufferSemantic Semantics[ChannelCount] = { MBS_POSITION };
+			const EMeshBufferSemantic Semantics[ChannelCount] = { EMeshBufferSemantic::Position };
 			const int32 SemanticIndices[ChannelCount] = { 0 };
-			const EMeshBufferFormat Formats[ChannelCount] = { MBF_FLOAT32 };
+			const EMeshBufferFormat Formats[ChannelCount] = { EMeshBufferFormat::Float32 };
 			const int32 Components[ChannelCount] = { 3 };
 			const int32 Offsets[ChannelCount] = { 0 };
 
@@ -2306,17 +2306,17 @@ mu::MeshPtr ConvertStaticMeshToMutable(const UStaticMesh* StaticMesh, int32 LODI
 			int32 currentChannel = 0;
 			int32 currentOffset = 0;
 
-			Semantics[currentChannel] = MBS_TANGENT;
+			Semantics[currentChannel] = EMeshBufferSemantic::Tangent;
 			SemanticIndices[currentChannel] = 0;
-			Formats[currentChannel] = MBF_PACKEDDIRS8;
+			Formats[currentChannel] = EMeshBufferFormat::PackedDirS8;
 			Components[currentChannel] = 4;
 			Offsets[currentChannel] = currentOffset;
 			currentOffset += 4;
 			++currentChannel;
 
-			Semantics[currentChannel] = MBS_NORMAL;
+			Semantics[currentChannel] = EMeshBufferSemantic::Normal;
 			SemanticIndices[currentChannel] = 0;
-			Formats[currentChannel] = MBF_PACKEDDIRS8;
+			Formats[currentChannel] = EMeshBufferFormat::PackedDirS8;
 
 			Components[currentChannel] = 4;
 			Offsets[currentChannel] = currentOffset;
@@ -2353,17 +2353,17 @@ mu::MeshPtr ConvertStaticMeshToMutable(const UStaticMesh* StaticMesh, int32 LODI
 			if (VertexBuffer.GetUseFullPrecisionUVs())
 			{
 				texChannelSize = 2 * 4;
-				texChannelFormat = MBF_FLOAT32;
+				texChannelFormat = EMeshBufferFormat::Float32;
 			}
 			else
 			{
 				texChannelSize = 2 * 2;
-				texChannelFormat = MBF_FLOAT16;
+				texChannelFormat = EMeshBufferFormat::Float16;
 			}
 
 			for (int32 c = 0; c < texChannels; ++c)
 			{
-				Semantics[currentChannel] = MBS_TEXCOORDS;
+				Semantics[currentChannel] = EMeshBufferSemantic::TexCoords;
 				SemanticIndices[currentChannel] = c;
 				Formats[currentChannel] = texChannelFormat;
 				Components[currentChannel] = 2;
@@ -2398,9 +2398,9 @@ mu::MeshPtr ConvertStaticMeshToMutable(const UStaticMesh* StaticMesh, int32 LODI
 		using namespace mu;
 		const int ElementSize = 2;
 		const int ChannelCount = 1;
-		const EMeshBufferSemantic Semantics[ChannelCount] = { MBS_VERTEXINDEX };
+		const EMeshBufferSemantic Semantics[ChannelCount] = { EMeshBufferSemantic::VertexIndex };
 		const int SemanticIndices[ChannelCount] = { 0 };
-		EMeshBufferFormat Formats[ChannelCount] = { MBF_UINT16 };
+		EMeshBufferFormat Formats[ChannelCount] = { EMeshBufferFormat::UInt16 };
 		const int Components[ChannelCount] = { 1 };
 		const int Offsets[ChannelCount] = { 0 };
 
@@ -2523,7 +2523,7 @@ mu::Ptr<mu::Mesh> BuildMorphedMutableMeshFromMesh(mu::Ptr<mu::Mesh> BaseSourceMe
 	{
 		int32 PosBuf = -1;
 		int32 PosChannel = -1;
-		MorphedSourceMesh->GetVertexBuffers().FindChannel(mu::MBS_POSITION, 0, &PosBuf, &PosChannel);
+		MorphedSourceMesh->GetVertexBuffers().FindChannel(mu::EMeshBufferSemantic::Position, 0, &PosBuf, &PosChannel);
 		check(PosBuf >= 0 && PosChannel >= 0);
 
 		int32 PosElemSize = MorphedSourceMesh->GetVertexBuffers().GetElementSize(PosBuf);
@@ -2532,7 +2532,7 @@ mu::Ptr<mu::Mesh> BuildMorphedMutableMeshFromMesh(mu::Ptr<mu::Mesh> BaseSourceMe
 
 		int32 NorBuf = -1;
 		int32 NorChannel = -1;
-		MorphedSourceMesh->GetVertexBuffers().FindChannel(mu::MBS_NORMAL, 0, &NorBuf, &NorChannel);
+		MorphedSourceMesh->GetVertexBuffers().FindChannel(mu::EMeshBufferSemantic::Normal, 0, &NorBuf, &NorChannel);
 
 		const bool bHasNormals = NorBuf >= 0 && NorChannel >= 0;
 
@@ -3501,18 +3501,18 @@ mu::Ptr<mu::NodeMesh> GenerateMutableSourceMesh(const UEdGraphPin* Pin,
 					const mu::FMeshBufferSet& MeshSet = MutableMesh->GetVertexBuffers();
 
 					int32 MorphIndexBuffer, MorphIndexChannel;
-					MeshSet.FindChannel(mu::MBS_OTHER, 0, &MorphIndexBuffer, &MorphIndexChannel);
+					MeshSet.FindChannel(mu::EMeshBufferSemantic::Other, 0, &MorphIndexBuffer, &MorphIndexChannel);
 
 					int32 MorphResourceBuffer, MorphResourceChannel;
-					MeshSet.FindChannel(mu::MBS_OTHER, 1, &MorphResourceBuffer, &MorphResourceChannel);
+					MeshSet.FindChannel(mu::EMeshBufferSemantic::Other, 1, &MorphResourceBuffer, &MorphResourceChannel);
 
 					MeshData.bHasRealTimeMorphs = MorphIndexBuffer >= 0 && MorphResourceBuffer >= 0; 
 
 					int32 ClothIndexBuffer, ClothIndexChannel;
-					MeshSet.FindChannel(mu::MBS_OTHER, 2, &ClothIndexBuffer, &ClothIndexChannel);
+					MeshSet.FindChannel(mu::EMeshBufferSemantic::Other, 2, &ClothIndexBuffer, &ClothIndexChannel);
 
 					int32 ClothResourceBuffer, ClothResourceChannel;
-					MeshSet.FindChannel(mu::MBS_OTHER, 3, &ClothResourceBuffer, &ClothResourceChannel);
+					MeshSet.FindChannel(mu::EMeshBufferSemantic::Other, 3, &ClothResourceBuffer, &ClothResourceChannel);
 
 					MeshData.bHasClothing = ClothIndexBuffer >= 0 && ClothResourceBuffer >= 0; 
 				}
diff --git a/Engine/Plugins/Mutable/Source/CustomizableObjectEditor/Private/MuCOE/GenerateMutableSource/GenerateMutableSourceProjector.cpp b/Engine/Plugins/Mutable/Source/CustomizableObjectEditor/Private/MuCOE/GenerateMutableSource/GenerateMutableSourceProjector.cpp
--- a/Engine/Plugins/Mutable/Source/CustomizableObjectEditor/Private/MuCOE/GenerateMutableSource/GenerateMutableSourceProjector.cpp
+++ b/Engine/Plugins/Mutable/Source/CustomizableObjectEditor/Private/MuCOE/GenerateMutableSource/GenerateMutableSourceProjector.cpp
@@ -33,10 +33,10 @@ mu::NodeProjectorPtr GenerateMutableSourceProjector(const UEdGraphPin* Pin, FMut
 		mu::NodeProjectorConstantPtr ProjectorNode = new mu::NodeProjectorConstant();
 		Result = ProjectorNode;
 
-		const mu::PROJECTOR_TYPE ProjectorType = ProjectorUtils::GetEquivalentProjectorType(TypedNodeConst->Value.ProjectionType);
+		const mu::EProjectorType ProjectorType = ProjectorUtils::GetEquivalentProjectorType(TypedNodeConst->Value.ProjectionType);
 		switch (ProjectorType) 
 		{ 
-			case mu::PROJECTOR_TYPE::PLANAR:
+			case mu::EProjectorType::Planar:
 				ProjectorNode->SetValue(
 					ProjectorType,
 					TypedNodeConst->Value.Position,
@@ -46,7 +46,7 @@ mu::NodeProjectorPtr GenerateMutableSourceProjector(const UEdGraphPin* Pin, FMut
 					TypedNodeConst->Value.Angle);
 				break;
 			
-			case mu::PROJECTOR_TYPE::CYLINDRICAL:
+			case mu::EProjectorType::Cylindrical:
 				{
 					// Apply strange swizzle for scales
 					// TODO: try to avoid this
@@ -62,7 +62,7 @@ mu::NodeProjectorPtr GenerateMutableSourceProjector(const UEdGraphPin* Pin, FMut
 				}
 				break;
 			
-			case mu::PROJECTOR_TYPE::WRAPPING:
+			case mu::EProjectorType::Wrapping:
 				ProjectorNode->SetValue(
 					ProjectorType,
 					TypedNodeConst->Value.Position,
@@ -72,7 +72,7 @@ mu::NodeProjectorPtr GenerateMutableSourceProjector(const UEdGraphPin* Pin, FMut
 					TypedNodeConst->Value.Angle);
 				break;
 			
-			case mu::PROJECTOR_TYPE::COUNT: 
+			case mu::EProjectorType::Count: 
 			default:
 				checkNoEntry();
 		}
@@ -89,7 +89,7 @@ mu::NodeProjectorPtr GenerateMutableSourceProjector(const UEdGraphPin* Pin, FMut
 		{
 		case 0:
 			ProjectorNode->SetDefaultValue(
-				mu::PROJECTOR_TYPE::PLANAR,
+				mu::EProjectorType::Planar,
 				TypedNodeParam->DefaultValue.Position,
 				TypedNodeParam->DefaultValue.Direction,
 				TypedNodeParam->DefaultValue.Up,
@@ -104,7 +104,7 @@ mu::NodeProjectorPtr GenerateMutableSourceProjector(const UEdGraphPin* Pin, FMut
 			float Radius = FMath::Abs(TypedNodeParam->DefaultValue.Scale[0] / 2.0f);
 			float Height = TypedNodeParam->DefaultValue.Scale[2];
 			ProjectorNode->SetDefaultValue(
-				mu::PROJECTOR_TYPE::CYLINDRICAL,
+				mu::EProjectorType::Cylindrical,
 				TypedNodeParam->DefaultValue.Position,
 				-TypedNodeParam->DefaultValue.Direction,
 				-TypedNodeParam->DefaultValue.Up,
@@ -115,7 +115,7 @@ mu::NodeProjectorPtr GenerateMutableSourceProjector(const UEdGraphPin* Pin, FMut
 
 		case 2:
 			ProjectorNode->SetDefaultValue(
-				mu::PROJECTOR_TYPE::WRAPPING,
+				mu::EProjectorType::Wrapping,
 				TypedNodeParam->DefaultValue.Position,
 				TypedNodeParam->DefaultValue.Direction,
 				TypedNodeParam->DefaultValue.Up,
diff --git a/Engine/Plugins/Mutable/Source/CustomizableObjectEditor/Private/MuCOE/GenerateMutableSource/GenerateMutableSourceSurface.cpp b/Engine/Plugins/Mutable/Source/CustomizableObjectEditor/Private/MuCOE/GenerateMutableSource/GenerateMutableSourceSurface.cpp
--- a/Engine/Plugins/Mutable/Source/CustomizableObjectEditor/Private/MuCOE/GenerateMutableSource/GenerateMutableSourceSurface.cpp
+++ b/Engine/Plugins/Mutable/Source/CustomizableObjectEditor/Private/MuCOE/GenerateMutableSource/GenerateMutableSourceSurface.cpp
@@ -137,9 +137,9 @@ void SetSurfaceFormat( FMutableGraphGenerationContext& GenerationContext,
 		{
 			const int32 ElementSize = sizeof(uint32);
 			constexpr int32 ChannelCount = 1;
-			const EMeshBufferSemantic Semantics[ChannelCount] = { MBS_OTHER };
+			const EMeshBufferSemantic Semantics[ChannelCount] = { EMeshBufferSemantic::Other };
 			const int32 SemanticIndices[ChannelCount] = { 0 };
-			const EMeshBufferFormat Formats[ChannelCount] = { MBF_UINT32 };
+			const EMeshBufferFormat Formats[ChannelCount] = { EMeshBufferFormat::UInt32 };
 			const int32 Components[ChannelCount] = { 1 };
 			const int32 Offsets[ChannelCount] = { 0 };
 
@@ -150,9 +150,9 @@ void SetSurfaceFormat( FMutableGraphGenerationContext& GenerationContext,
 		{
 			const int32 ElementSize = sizeof(uint32);
 			constexpr int32 ChannelCount = 1;
-			const EMeshBufferSemantic Semantics[ChannelCount] = { MBS_OTHER };
+			const EMeshBufferSemantic Semantics[ChannelCount] = { EMeshBufferSemantic::Other };
 			const int32 SemanticIndices[ChannelCount] = { 1 };
-			const EMeshBufferFormat Formats[ChannelCount] = { MBF_UINT32 };
+			const EMeshBufferFormat Formats[ChannelCount] = { EMeshBufferFormat::UInt32 };
 			const int32 Components[ChannelCount] = { 1 };
 			const int32 Offsets[ChannelCount] = { 0 };
 
@@ -168,9 +168,9 @@ void SetSurfaceFormat( FMutableGraphGenerationContext& GenerationContext,
 			using namespace mu;
 			const int32 ElementSize = sizeof(int32);
 			constexpr int32 ChannelCount = 1;
-			const EMeshBufferSemantic Semantics[ChannelCount] = { MBS_OTHER };
+			const EMeshBufferSemantic Semantics[ChannelCount] = { EMeshBufferSemantic::Other };
 			const int32 SemanticIndices[ChannelCount] = { 2 };
-			const EMeshBufferFormat Formats[ChannelCount] = { MBF_INT32 };
+			const EMeshBufferFormat Formats[ChannelCount] = { EMeshBufferFormat::Int32 };
 			const int32 Components[ChannelCount] = { 1 };
 			const int32 Offsets[ChannelCount] = { 0 };
 
@@ -182,9 +182,9 @@ void SetSurfaceFormat( FMutableGraphGenerationContext& GenerationContext,
 			using namespace mu;
 			const int32 ElementSize = sizeof(uint32);
 			constexpr int32 ChannelCount = 1;
-			const EMeshBufferSemantic Semantics[ChannelCount] = { MBS_OTHER };
+			const EMeshBufferSemantic Semantics[ChannelCount] = { EMeshBufferSemantic::Other };
 			const int32 SemanticIndices[ChannelCount] = { 3 };
-			const EMeshBufferFormat Formats[ChannelCount] = { MBF_UINT32 };
+			const EMeshBufferFormat Formats[ChannelCount] = { EMeshBufferFormat::UInt32 };
 			const int32 Components[ChannelCount] = { 1 };
 			const int32 Offsets[ChannelCount] = { 0 };
 
@@ -709,7 +709,7 @@ mu::Ptr<mu::NodeSurface> GenerateMutableSourceSurface(const UEdGraphPin * Pin, F
 							mu::Ptr<mu::NodeImage> FormatSource = LastImage;
 							mu::NodeImageFormatPtr FormatImage = new mu::NodeImageFormat();
 							FormatImage->SetSource(LastImage.get());
-							FormatImage->SetFormat(mu::EImageFormat::IF_RGBA_UBYTE);
+							FormatImage->SetFormat(mu::EImageFormat::RGBA_UByte);
 							FormatImage->SetMessageContext(Node);
 							LastImage = FormatImage;
 
@@ -719,7 +719,7 @@ mu::Ptr<mu::NodeSurface> GenerateMutableSourceSurface(const UEdGraphPin * Pin, F
 								ReferenceTexture->GetTargetPlatformBuildSettings(GenerationContext.Options.TargetPlatform, BuildSettingsPerFormatPerLayer);
 								if (BuildSettingsPerFormatPerLayer.IsEmpty())
 								{
-									const FString ReplacedImageFormatMsg = FString::Printf(TEXT("In object [%s] for platform [%s] the unsupported image format of texture [%s] is used, IF_RGBA_UBYTE will be used instead."),
+									const FString ReplacedImageFormatMsg = FString::Printf(TEXT("In object [%s] for platform [%s] the unsupported image format of texture [%s] is used, RGBA_UByte will be used instead."),
 										*GenerationContext.Object->GetName(),
 										*GenerationContext.Options.TargetPlatform->PlatformName(),
 										*ReferenceTexture->GetName());
@@ -769,12 +769,12 @@ mu::Ptr<mu::NodeSurface> GenerateMutableSourceSurface(const UEdGraphPin * Pin, F
 											// See GetQualityFormat in TextureFormatASTC.cpp to understand why
 											if (UnrealTargetPlatformFormat == PF_ASTC_6x6 || UnrealTargetPlatformFormat == PF_ASTC_6x6_NORM_RG)
 											{
-												MutableFormat = mu::EImageFormat::IF_ASTC_4x4_RGBA_LDR;
-												MutableFormatIfAlpha = mu::EImageFormat::IF_ASTC_4x4_RGBA_LDR;
+												MutableFormat = mu::EImageFormat::ASTC_4x4_RGBA_LDR;
+												MutableFormatIfAlpha = mu::EImageFormat::ASTC_4x4_RGBA_LDR;
 
 												// Insert a channel swizzle
 												mu::Ptr<mu::NodeImageSwizzle> Swizzle = new mu::NodeImageSwizzle;
-												Swizzle->SetFormat( mu::EImageFormat::IF_RGBA_UBYTE );
+												Swizzle->SetFormat( mu::EImageFormat::RGBA_UByte );
 												Swizzle->SetSource(0, FormatSource);
 												Swizzle->SetSource(1, FormatSource);
 												Swizzle->SetSource(2, FormatSource);
@@ -790,21 +790,21 @@ mu::Ptr<mu::NodeSurface> GenerateMutableSourceSurface(const UEdGraphPin * Pin, F
 									}
 
 									// Unsupported format: look for something generic
-									if (MutableFormat == mu::EImageFormat::IF_NONE)
+									if (MutableFormat == mu::EImageFormat::None)
 									{
-										const FString ReplacedImageFormatMsg = FString::Printf(TEXT("In object [%s] the unsupported image format %d is used, IF_RGBA_UBYTE will be used instead."), *GenerationContext.Object->GetName(), UnrealTargetPlatformFormat );
+										const FString ReplacedImageFormatMsg = FString::Printf(TEXT("In object [%s] the unsupported image format %d is used, RGBA_UByte will be used instead."), *GenerationContext.Object->GetName(), UnrealTargetPlatformFormat );
 										const FText ReplacedImageFormatText = FText::FromString(ReplacedImageFormatMsg);
 										GenerationContext.Log(ReplacedImageFormatText, Node, EMessageSeverity::Info);
 										UE_LOG(LogMutable, Log, TEXT("%s"), *ReplacedImageFormatMsg);
-										MutableFormat = mu::EImageFormat::IF_RGBA_UBYTE;
+										MutableFormat = mu::EImageFormat::RGBA_UByte;
 									}
-									if (MutableFormatIfAlpha == mu::EImageFormat::IF_NONE)
+									if (MutableFormatIfAlpha == mu::EImageFormat::None)
 									{
-										const FString ReplacedImageFormatMsg = FString::Printf(TEXT("In object [%s] the unsupported image format %d is used, IF_RGBA_UBYTE will be used instead."), *GenerationContext.Object->GetName(), UnrealTargetPlatformFormatAlpha);
+										const FString ReplacedImageFormatMsg = FString::Printf(TEXT("In object [%s] the unsupported image format %d is used, RGBA_UByte will be used instead."), *GenerationContext.Object->GetName(), UnrealTargetPlatformFormatAlpha);
 										const FText ReplacedImageFormatText = FText::FromString(ReplacedImageFormatMsg);
 										GenerationContext.Log(ReplacedImageFormatText, Node, EMessageSeverity::Info);
 										UE_LOG(LogMutable, Log, TEXT("%s"), *ReplacedImageFormatMsg);
-										MutableFormatIfAlpha = mu::EImageFormat::IF_RGBA_UBYTE;
+										MutableFormatIfAlpha = mu::EImageFormat::RGBA_UByte;
 									}
 
 									FormatImage->SetFormat(MutableFormat, MutableFormatIfAlpha);
diff --git a/Engine/Plugins/Mutable/Source/CustomizableObjectEditor/Private/MuCOE/MutableMeshPreviewUtils.cpp b/Engine/Plugins/Mutable/Source/CustomizableObjectEditor/Private/MuCOE/MutableMeshPreviewUtils.cpp
--- a/Engine/Plugins/Mutable/Source/CustomizableObjectEditor/Private/MuCOE/MutableMeshPreviewUtils.cpp
+++ b/Engine/Plugins/Mutable/Source/CustomizableObjectEditor/Private/MuCOE/MutableMeshPreviewUtils.cpp
@@ -59,8 +59,8 @@ namespace MutableMeshPreviewUtils
 					const int32 BufferChannelCount = OriginVertexBuffers.GetBufferChannelCount(BufferIndex);
 					for (int32 ChannelIndex = 0; ChannelIndex < BufferChannelCount; ChannelIndex ++)
 					{
-						mu::EMeshBufferSemantic BufferSemantic = mu::EMeshBufferSemantic::MBS_NONE;
-						mu::EMeshBufferFormat BufferFormat = mu::EMeshBufferFormat::MBF_NONE;
+						mu::EMeshBufferSemantic BufferSemantic = mu::EMeshBufferSemantic::None;
+						mu::EMeshBufferFormat BufferFormat = mu::EMeshBufferFormat::None;
 						int32 BufferComponentCount = 0;
 						int32 SemanticIndex = 0;
 
@@ -78,14 +78,14 @@ namespace MutableMeshPreviewUtils
 
 						switch (BufferSemantic)
 						{
-							case mu::EMeshBufferSemantic::MBS_COLOUR:
+							case mu::EMeshBufferSemantic::Color:
 							{
 								// Does the original mesh have vert color data?
 								bOriginDoesHaveVertexColorData = true;
 								break;
 							}
 
-							case mu::EMeshBufferSemantic::MBS_TEXCOORDS:
+							case mu::EMeshBufferSemantic::TexCoords:
 							{
 								// Store the texture semantic index to be later used during the buffer formatting process on  mu::MeshFormat
 								TextureSemanticIndices[NumbOfTextCoordChannels] = SemanticIndex;
@@ -95,18 +95,18 @@ namespace MutableMeshPreviewUtils
 								break;
 							}
 
-							case mu::EMeshBufferSemantic::MBS_BONEINDICES:
+							case mu::EMeshBufferSemantic::BoneIndices:
 							{
 								// Store the amount of bones a vertex can be skinned to.
 								MaxNumBonesPerVertex = BufferComponentCount;
-								BoneIndicesSizeBytes = BufferFormat == mu::EMeshBufferFormat::MBF_INT16 ? 2 : 1;
+								BoneIndicesSizeBytes = BufferFormat == mu::EMeshBufferFormat::Int16 ? 2 : 1;
 								break;
 							}
 
-							case mu::EMeshBufferSemantic::MBS_BONEWEIGHTS:
+							case mu::EMeshBufferSemantic::BoneWeights:
 							{
 								// Store the amount of bones a vertex can be skinned to.
-								BoneWeightsSizeBytes = BufferFormat == mu::EMeshBufferFormat::MBF_NUINT16 ? 2 : 1;
+								BoneWeightsSizeBytes = BufferFormat == mu::EMeshBufferFormat::NUInt16 ? 2 : 1;
 								break;
 							}
 
@@ -202,7 +202,7 @@ namespace MutableMeshPreviewUtils
 			int32 Channel = -1;
 
 			// POSITION BUFFER INDEX CHECK
-			VertexBuffers.FindChannel(mu::EMeshBufferSemantic::MBS_POSITION, 0, &Buffer, &Channel);
+			VertexBuffers.FindChannel(mu::EMeshBufferSemantic::Position, 0, &Buffer, &Channel);
 			if (Buffer == -1)
 			{
 				if (bLogFindingsToConsole)
@@ -226,7 +226,7 @@ namespace MutableMeshPreviewUtils
 			}
 
 			// TANGENT BUFFER INDEX CHECK
-			VertexBuffers.FindChannel(mu::EMeshBufferSemantic::MBS_TANGENT, 0, &Buffer, &Channel);
+			VertexBuffers.FindChannel(mu::EMeshBufferSemantic::Tangent, 0, &Buffer, &Channel);
 			if (Buffer == -1)
 			{
 				if (bLogFindingsToConsole)
@@ -250,7 +250,7 @@ namespace MutableMeshPreviewUtils
 			}
 
 			// TEX COORDS BUFFER INDEX CHECK
-			VertexBuffers.FindChannel(mu::EMeshBufferSemantic::MBS_TEXCOORDS, 0, &Buffer, &Channel);
+			VertexBuffers.FindChannel(mu::EMeshBufferSemantic::TexCoords, 0, &Buffer, &Channel);
 			if (Buffer == -1)
 			{
 				if (bLogFindingsToConsole)
diff --git a/Engine/Plugins/Mutable/Source/CustomizableObjectEditor/Private/MuCOE/SMutableCodeViewer.cpp b/Engine/Plugins/Mutable/Source/CustomizableObjectEditor/Private/MuCOE/SMutableCodeViewer.cpp
--- a/Engine/Plugins/Mutable/Source/CustomizableObjectEditor/Private/MuCOE/SMutableCodeViewer.cpp
+++ b/Engine/Plugins/Mutable/Source/CustomizableObjectEditor/Private/MuCOE/SMutableCodeViewer.cpp
@@ -2472,7 +2472,7 @@ namespace
 	class TestResourceProvider : public mu::ExternalResourceProvider
 	{
 		static inline const mu::FImageDesc IMAGE_DESC = 
-			mu::FImageDesc(mu::FImageSize(1024, 1024), mu::EImageFormat::IF_RGBA_UBYTE, 1);
+			mu::FImageDesc(mu::FImageSize(1024, 1024), mu::EImageFormat::RGBA_UByte, 1);
 
 	public:
 
diff --git a/Engine/Plugins/Mutable/Source/CustomizableObjectEditor/Private/MuCOE/SMutableImageViewer.cpp b/Engine/Plugins/Mutable/Source/CustomizableObjectEditor/Private/MuCOE/SMutableImageViewer.cpp
--- a/Engine/Plugins/Mutable/Source/CustomizableObjectEditor/Private/MuCOE/SMutableImageViewer.cpp
+++ b/Engine/Plugins/Mutable/Source/CustomizableObjectEditor/Private/MuCOE/SMutableImageViewer.cpp
@@ -265,9 +265,9 @@ void SMutableImageViewer::Tick(const FGeometry& AllottedGeometry, const double I
 			mu::Ptr<const mu::Image> ImageToConvert = MutableImage;
 
 			// \TODO: This test is a bit weak
-			if (MutableImage->GetFormat() >= mu::EImageFormat::IF_ASTC_4x4_RGB_LDR
+			if (MutableImage->GetFormat() >= mu::EImageFormat::ASTC_4x4_RGB_LDR
 				&&
-				MutableImage->GetFormat() <= mu::EImageFormat::IF_ASTC_10x10_RG_LDR)
+				MutableImage->GetFormat() <= mu::EImageFormat::ASTC_10x10_RG_LDR)
 			{
 				mu::FImageOperator ImOp = mu::FImageOperator::GetDefault(mu::FImageOperator::FImagePixelFormatFunc());
 				int32 Quality = 4;
diff --git a/Engine/Plugins/Mutable/Source/CustomizableObjectEditor/Private/MuCOE/SMutableMeshViewer.cpp b/Engine/Plugins/Mutable/Source/CustomizableObjectEditor/Private/MuCOE/SMutableMeshViewer.cpp
--- a/Engine/Plugins/Mutable/Source/CustomizableObjectEditor/Private/MuCOE/SMutableMeshViewer.cpp
+++ b/Engine/Plugins/Mutable/Source/CustomizableObjectEditor/Private/MuCOE/SMutableMeshViewer.cpp
@@ -637,8 +637,8 @@ void SMutableMeshViewer::FillTargetBufferSetDataArray(const mu::FMeshBufferSet&
 		{
 			TSharedPtr<FBufferChannelElement> NewChannelDefinition = MakeShareable(new FBufferChannelElement());
 
-			mu::EMeshBufferSemantic BufferChannelSemantic = mu::EMeshBufferSemantic::MBS_NONE;
-			mu::EMeshBufferFormat BufferFormat = mu::EMeshBufferFormat::MBF_NONE;
+			mu::EMeshBufferSemantic BufferChannelSemantic = mu::EMeshBufferSemantic::None;
+			mu::EMeshBufferFormat BufferFormat = mu::EMeshBufferFormat::None;
 			int32 BufferComponentCount = 0;
 			int32 SemanticIndex = -1;
 
