[mutable] Cleanup MutableRuntime API moving to use more unreal types and naming standards.

diff --git a/Engine/Plugins/Mutable/Source/CustomizableObjectEditor/Private/MuCOE/SMutableParametersWidget.cpp b/Engine/Plugins/Mutable/Source/CustomizableObjectEditor/Private/MuCOE/SMutableParametersWidget.cpp
--- a/Engine/Plugins/Mutable/Source/CustomizableObjectEditor/Private/MuCOE/SMutableParametersWidget.cpp
+++ b/Engine/Plugins/Mutable/Source/CustomizableObjectEditor/Private/MuCOE/SMutableParametersWidget.cpp
@@ -128,12 +128,12 @@ void SMutableParametersWidget::Tick(const FGeometry& AllottedGeometry, const dou
 			if (ParameterControllingValueCount >= 0)
 			{
 				// Get the numerical value of that parameters
-				mu::PARAMETER_TYPE ParameterType = MutableParameters->GetType(ParameterControllingValueCount);
-				if (ParameterType == mu::PARAMETER_TYPE::T_INT)
+				mu::EParameterType ParameterType = MutableParameters->GetType(ParameterControllingValueCount);
+				if (ParameterType == mu::EParameterType::Int)
 				{
 					PositionCount = MutableParameters->GetIntValue(ParameterControllingValueCount);
 				}
-				else if (ParameterType == mu::PARAMETER_TYPE::T_FLOAT)
+				else if (ParameterType == mu::EParameterType::Float)
 				{
 					PositionCount = MutableParameters->GetFloatValue(ParameterControllingValueCount);
 				}
@@ -231,7 +231,7 @@ void SMutableParametersWidget::GenerateAndAttachParameterSlate(const int32 Param
 {
 	switch (MutableParameters->GetType(ParamIndex) )
 	{
-	case mu::PARAMETER_TYPE::T_BOOL:
+	case mu::EParameterType::Bool:
 		ParameterHorizontalBox->AddSlot()
 			.Padding(4.0f)
 			.VAlign(VAlign_Center)
@@ -243,7 +243,7 @@ void SMutableParametersWidget::GenerateAndAttachParameterSlate(const int32 Param
 			];
 		break;
 
-	case mu::PARAMETER_TYPE::T_FLOAT:
+	case mu::EParameterType::Float:
 		ParameterHorizontalBox->AddSlot()
 			.Padding(4.0f)
 			.VAlign(VAlign_Center)
@@ -259,7 +259,7 @@ void SMutableParametersWidget::GenerateAndAttachParameterSlate(const int32 Param
 			];
 		break;
 
-	case mu::PARAMETER_TYPE::T_COLOUR:
+	case mu::EParameterType::Color:
 		ParameterHorizontalBox->AddSlot()
 			.Padding(4.0f)
 			.VAlign(VAlign_Center)
@@ -274,7 +274,7 @@ void SMutableParametersWidget::GenerateAndAttachParameterSlate(const int32 Param
 			];
 		break;
 
-	case mu::PARAMETER_TYPE::T_INT:
+	case mu::EParameterType::Int:
 	{
 		// If we have a list of options, add a combo box too
 		TSharedPtr<SSearchableComboBox> ParamComboBox;
@@ -339,7 +339,7 @@ void SMutableParametersWidget::GenerateAndAttachParameterSlate(const int32 Param
 		break;
 	}
 
-	case mu::PARAMETER_TYPE::T_PROJECTOR:
+	case mu::EParameterType::Projector:
 		ParameterHorizontalBox->AddSlot()
 			.Padding(4.0f)
 			.VAlign(VAlign_Center)
@@ -432,7 +432,7 @@ ECheckBoxState SMutableParametersWidget::GetBoolParameterValue(int32 ParamIndex,
 {
 	if ( !MutableParameters 
 		|| ParamIndex>=MutableParameters->GetCount()
-		|| MutableParameters->GetType(ParamIndex)!=mu::PARAMETER_TYPE::T_BOOL )
+		|| MutableParameters->GetType(ParamIndex)!=mu::EParameterType::Bool )
 	{
 		return ECheckBoxState::Undetermined;
 	}	
@@ -445,7 +445,7 @@ void SMutableParametersWidget::OnBoolParameterChanged(ECheckBoxState InCheckboxS
 {
 	if (!MutableParameters
 		|| ParamIndex >= MutableParameters->GetCount()
-		|| MutableParameters->GetType(ParamIndex) != mu::PARAMETER_TYPE::T_BOOL)
+		|| MutableParameters->GetType(ParamIndex) != mu::EParameterType::Bool)
 	{
 		return;
 	}
@@ -460,7 +460,7 @@ TOptional<float> SMutableParametersWidget::GetFloatParameterValue(int32 ParamInd
 {
 	if (!MutableParameters
 		|| ParamIndex >= MutableParameters->GetCount()
-		|| MutableParameters->GetType(ParamIndex) != mu::PARAMETER_TYPE::T_FLOAT)
+		|| MutableParameters->GetType(ParamIndex) != mu::EParameterType::Float)
 	{
 		return TOptional<float>();
 	}
@@ -473,7 +473,7 @@ void SMutableParametersWidget::OnFloatParameterChanged(float InValue, int32 Para
 {
 	if (!MutableParameters
 		|| ParamIndex >= MutableParameters->GetCount()
-		|| MutableParameters->GetType(ParamIndex) != mu::PARAMETER_TYPE::T_FLOAT)
+		|| MutableParameters->GetType(ParamIndex) != mu::EParameterType::Float)
 	{
 		return;
 	}
@@ -497,7 +497,7 @@ TOptional<int32> SMutableParametersWidget::GetIntParameterValue (int32 ParamInde
 {
 	if (!MutableParameters
 		|| ParamIndex >= MutableParameters->GetCount()
-		|| MutableParameters->GetType(ParamIndex) != mu::PARAMETER_TYPE::T_INT)
+		|| MutableParameters->GetType(ParamIndex) != mu::EParameterType::Int)
 	{
 		return TOptional<int32>();
 	}
@@ -513,7 +513,7 @@ TOptional<int32> SMutableParametersWidget::GetIntParameterValueMax(int32 ParamIn
 {
 	if (!MutableParameters
 		|| ParamIndex >= MutableParameters->GetCount()
-		|| MutableParameters->GetType(ParamIndex) != mu::PARAMETER_TYPE::T_INT)
+		|| MutableParameters->GetType(ParamIndex) != mu::EParameterType::Int)
 	{
 		return TOptional<int32>();
 	}
@@ -533,7 +533,7 @@ void SMutableParametersWidget::OnIntParameterChanged(int32 InValue, int32 ParamI
 {
 	if (!MutableParameters
 		|| ParamIndex >= MutableParameters->GetCount()
-		|| MutableParameters->GetType(ParamIndex) != mu::PARAMETER_TYPE::T_INT)
+		|| MutableParameters->GetType(ParamIndex) != mu::EParameterType::Int)
 	{
 		return;
 	}
@@ -564,7 +564,7 @@ void SMutableParametersWidget::OnIntParameterTextChanged(TSharedPtr<FString> Sel
 	if (!Selection
 		|| !MutableParameters
 		|| ParamIndex >= MutableParameters->GetCount()
-		|| MutableParameters->GetType(ParamIndex) != mu::PARAMETER_TYPE::T_INT)
+		|| MutableParameters->GetType(ParamIndex) != mu::EParameterType::Int)
 	{
 		return;
 	}
@@ -592,13 +592,13 @@ FLinearColor SMutableParametersWidget::GetColorParameterValue(int32 ParamIndex,
 {
 	if (!MutableParameters
 		|| ParamIndex >= MutableParameters->GetCount()
-		|| MutableParameters->GetType(ParamIndex) != mu::PARAMETER_TYPE::T_COLOUR)
+		|| MutableParameters->GetType(ParamIndex) != mu::EParameterType::Color)
 	{
 		return FLinearColor();
 	}
 
-	FLinearColor Result;
-	MutableParameters->GetColourValue(ParamIndex, &Result.R, &Result.G, &Result.B, &Result.A, RangeIndex);
+	FVector4f Result;
+	MutableParameters->GetColourValue(ParamIndex, Result, RangeIndex);
 
 	return Result;
 }
@@ -629,12 +629,12 @@ void SMutableParametersWidget::OnSetColorFromColorPicker(FLinearColor NewColor,
 {
 	if (!MutableParameters
 		|| ParamIndex >= MutableParameters->GetCount()
-		|| MutableParameters->GetType(ParamIndex) != mu::PARAMETER_TYPE::T_COLOUR)
+		|| MutableParameters->GetType(ParamIndex) != mu::EParameterType::Color)
 	{
 		return;
 	}
 	
-	MutableParameters->SetColourValue(ParamIndex, NewColor.R, NewColor.G, NewColor.B, NewColor.A, RangeIndex);
+	MutableParameters->SetColourValue(ParamIndex, NewColor, RangeIndex);
 
 	OnParametersValueChanged.ExecuteIfBound(ParamIndex);
 }
@@ -644,12 +644,12 @@ TOptional<FVector::FReal> SMutableParametersWidget::GetProjectorLocation(EAxis::
 { 
 	if (!MutableParameters
 		|| ParamIndex >= MutableParameters->GetCount()
-		|| MutableParameters->GetType(ParamIndex) != mu::PARAMETER_TYPE::T_PROJECTOR)
+		|| MutableParameters->GetType(ParamIndex) != mu::EParameterType::Projector)
 	{
 		return 0.0f;
 	}
 	
-	mu::PROJECTOR_TYPE Type;
+	mu::EProjectorType Type;
 	FVector3f Pos;
 	FVector3f Dir;
 	FVector3f Up;
@@ -665,14 +665,14 @@ void SMutableParametersWidget::SetProjectorLocation(FVector::FReal NewValue, ETe
 {
 	if (!MutableParameters
 		|| ParamIndex >= MutableParameters->GetCount()
-		|| MutableParameters->GetType(ParamIndex) != mu::PARAMETER_TYPE::T_PROJECTOR)
+		|| MutableParameters->GetType(ParamIndex) != mu::EParameterType::Projector)
 	{
 		return;
 	}
 	
 	//mu::RangeIndexPtr UpdatedRange = RangeIndexPayload.GetRangeIndex();
 	
-	mu::PROJECTOR_TYPE Type;
+	mu::EProjectorType Type;
 	FVector3f Pos;
 	FVector3f Dir;
 	FVector3f Up;
@@ -692,12 +692,12 @@ TOptional<FVector::FReal> SMutableParametersWidget::GetProjectorScale(EAxis::Typ
 {
 	if (!MutableParameters
 		|| ParamIndex >= MutableParameters->GetCount()
-		|| MutableParameters->GetType(ParamIndex) != mu::PARAMETER_TYPE::T_PROJECTOR)
+		|| MutableParameters->GetType(ParamIndex) != mu::EParameterType::Projector)
 	{
 		return 0.0f;
 	}
 
-	mu::PROJECTOR_TYPE Type;
+	mu::EProjectorType Type;
 	FVector3f Pos;
 	FVector3f Dir;
 	FVector3f Up;
@@ -713,13 +713,13 @@ void SMutableParametersWidget::SetProjectorScale(FVector::FReal NewValue, ETextC
 {
 	if (!MutableParameters
 		|| ParamIndex >= MutableParameters->GetCount()
-		|| MutableParameters->GetType(ParamIndex) != mu::PARAMETER_TYPE::T_PROJECTOR)
+		|| MutableParameters->GetType(ParamIndex) != mu::EParameterType::Projector)
 	{
 		return;
 	}
 
 	
-	mu::PROJECTOR_TYPE Type;
+	mu::EProjectorType Type;
 	FVector3f Pos;
 	FVector3f Dir;
 	FVector3f Up;
diff --git a/Engine/Plugins/Mutable/Source/CustomizableObjectEditor/Private/MuCOE/SMutableProjectorViewer.cpp b/Engine/Plugins/Mutable/Source/CustomizableObjectEditor/Private/MuCOE/SMutableProjectorViewer.cpp
--- a/Engine/Plugins/Mutable/Source/CustomizableObjectEditor/Private/MuCOE/SMutableProjectorViewer.cpp
+++ b/Engine/Plugins/Mutable/Source/CustomizableObjectEditor/Private/MuCOE/SMutableProjectorViewer.cpp
@@ -229,16 +229,16 @@ FText SMutableProjectorViewer::GetProjectorAngleAsText() const
 }
 
 
-mu::PROJECTOR_TYPE SMutableProjectorViewer::GetProjectorType() const
+mu::EProjectorType SMutableProjectorViewer::GetProjectorType() const
 {
-	const mu::PROJECTOR_TYPE ProjectorType = MutableProjector.type;
+	const mu::EProjectorType ProjectorType = MutableProjector.type;
 	return ProjectorType;
 }
 
 
 FText SMutableProjectorViewer::GetProjectorTypeAsText() const
 {
-	const mu::PROJECTOR_TYPE ProjectorType =  GetProjectorType();
+	const mu::EProjectorType ProjectorType =  GetProjectorType();
 	const uint32 TypeIndex = static_cast<uint32> (ProjectorType);
 
 	return FText::FromString(* FString(mu::TypeInfo::s_projectorTypeName[TypeIndex]));
diff --git a/Engine/Plugins/Mutable/Source/CustomizableObjectEditor/Private/MuCOE/SMutableProjectorViewer.h b/Engine/Plugins/Mutable/Source/CustomizableObjectEditor/Private/MuCOE/SMutableProjectorViewer.h
--- a/Engine/Plugins/Mutable/Source/CustomizableObjectEditor/Private/MuCOE/SMutableProjectorViewer.h
+++ b/Engine/Plugins/Mutable/Source/CustomizableObjectEditor/Private/MuCOE/SMutableProjectorViewer.h
@@ -5,7 +5,7 @@
 #include "MuR/ParametersPrivate.h"
 #include "Widgets/SCompoundWidget.h"
 
-namespace mu { enum class PROJECTOR_TYPE : uint32; }
+namespace mu { enum class EProjectorType : uint32; }
 
 
 class SMutableProjectorViewer final : public SCompoundWidget
@@ -72,7 +72,7 @@ private:
 	 * It accesses the set projector and returns the type of it.
 	 * @return - The projector type.
 	 */
-	mu::PROJECTOR_TYPE GetProjectorType() const;
+	mu::EProjectorType GetProjectorType() const;
 	
 	/**
 	 * Provides the UI with the mutable type in a format it can display on an STextBlock
diff --git a/Engine/Plugins/Mutable/Source/CustomizableObjectEditor/Private/MuCOE/UnrealBakeHelpers.cpp b/Engine/Plugins/Mutable/Source/CustomizableObjectEditor/Private/MuCOE/UnrealBakeHelpers.cpp
--- a/Engine/Plugins/Mutable/Source/CustomizableObjectEditor/Private/MuCOE/UnrealBakeHelpers.cpp
+++ b/Engine/Plugins/Mutable/Source/CustomizableObjectEditor/Private/MuCOE/UnrealBakeHelpers.cpp
@@ -187,14 +187,14 @@ UTexture2D* FUnrealBakeHelpers::BakeHelper_CreateAssetTexture(UTexture2D* SrcTex
 	SrcTex->GetPlatformData()->Mips[MipIndex].BulkData.Unlock();
 
 	// Reformat the mutable image
-	mu::EImageFormat UncompressedMutableFormat = mu::EImageFormat::IF_RGBA_UBYTE;
+	mu::EImageFormat UncompressedMutableFormat = mu::EImageFormat::RGBA_UByte;
 	switch (SrcTex->GetPlatformData()->PixelFormat)
 	{
 	case PF_G8:
 	case PF_L8:
 	case PF_A8:
 	case PF_BC4:
-		UncompressedMutableFormat = mu::EImageFormat::IF_L_UBYTE;
+		UncompressedMutableFormat = mu::EImageFormat::L_UByte;
 		break;
 
 	default:
diff --git a/Engine/Plugins/Mutable/Source/MutableRuntime/Internal/MuR/ConvertData.h b/Engine/Plugins/Mutable/Source/MutableRuntime/Internal/MuR/ConvertData.h
--- a/Engine/Plugins/Mutable/Source/MutableRuntime/Internal/MuR/ConvertData.h
+++ b/Engine/Plugins/Mutable/Source/MutableRuntime/Internal/MuR/ConvertData.h
@@ -23,7 +23,7 @@ namespace mu
 	{
 		switch ( resultFormat )
 		{
-		case MBF_FLOAT64:
+		case EMeshBufferFormat::Float64:
 		{
 			double* pTypedResult = reinterpret_cast<double*>(pResult);
 			uint8* pByteResult = reinterpret_cast<uint8*>(pResult);
@@ -31,7 +31,7 @@ namespace mu
 
 			switch (sourceFormat)
 			{
-			case MBF_FLOAT64:
+			case EMeshBufferFormat::Float64:
 			{
 				// Just dereferencing is not safe in all architectures. some like ARM require the
 				// floats to be 4-byte aligned and it may not be the case. We need memcpy.
@@ -39,63 +39,63 @@ namespace mu
 				break;
 			}
 
-			case MBF_FLOAT32:
+			case EMeshBufferFormat::Float32:
 			{
 				const float* pTypedSource = reinterpret_cast<const float*>(pSource);
 				pTypedResult[channel] = pTypedSource[channel];
 				break;
 			}
 
-			case MBF_FLOAT16:
+			case EMeshBufferFormat::Float16:
 			{
 				const FFloat16* pTypedSource = reinterpret_cast<const FFloat16*>(pSource);
 				pTypedResult[channel] = pTypedSource[channel];
 				break;
 			}
 
-			case MBF_INT32:
+			case EMeshBufferFormat::Int32:
 			{
 				const int32* pTypedSource = reinterpret_cast<const int32*>(pSource);
 				pTypedResult[channel] = (double)(pTypedSource[channel]);
 				break;
 			}
 
-			case MBF_UINT32:
+			case EMeshBufferFormat::UInt32:
 			{
 				const uint32* pTypedSource = reinterpret_cast<const uint32*>(pSource);
 				pTypedResult[channel] = (double)(pTypedSource[channel]);
 				break;
 			}
 
-			case MBF_INT16:
+			case EMeshBufferFormat::Int16:
 			{
 				const int16* pTypedSource = reinterpret_cast<const int16*>(pSource);
 				pTypedResult[channel] = (double)(pTypedSource[channel]);
 				break;
 			}
 
-			case MBF_UINT16:
+			case EMeshBufferFormat::UInt16:
 			{
 				const uint16* pTypedSource = reinterpret_cast<const uint16*>(pSource);
 				pTypedResult[channel] = (double)(pTypedSource[channel]);
 				break;
 			}
 
-			case MBF_INT8:
+			case EMeshBufferFormat::Int8:
 			{
 				const int8* pTypedSource = reinterpret_cast<const int8*>(pSource);
 				pTypedResult[channel] = (double)(pTypedSource[channel]);
 				break;
 			}
 
-			case MBF_UINT8:
+			case EMeshBufferFormat::UInt8:
 			{
 				const uint8* pTypedSource = reinterpret_cast<const uint8*>(pSource);
 				pTypedResult[channel] = (double)(pTypedSource[channel]);
 				break;
 			}
 
-			case MBF_NINT32:
+			case EMeshBufferFormat::NInt32:
 			{
 				const int32* pTypedSource = reinterpret_cast<const int32*>(pSource);
 				pTypedResult[channel] = (double)(pTypedSource[channel]);
@@ -103,7 +103,7 @@ namespace mu
 				break;
 			}
 
-			case MBF_NUINT32:
+			case EMeshBufferFormat::NUInt32:
 			{
 				const uint32* pTypedSource = reinterpret_cast<const uint32*>(pSource);
 				pTypedResult[channel] = (double)(pTypedSource[channel]);
@@ -111,7 +111,7 @@ namespace mu
 				break;
 			}
 
-			case MBF_NINT16:
+			case EMeshBufferFormat::NInt16:
 			{
 				const int16* pTypedSource = reinterpret_cast<const int16*>(pSource);
 				pTypedResult[channel] = (double)(pTypedSource[channel]);
@@ -119,7 +119,7 @@ namespace mu
 				break;
 			}
 
-			case MBF_NUINT16:
+			case EMeshBufferFormat::NUInt16:
 			{
 				const uint16* pTypedSource = reinterpret_cast<const uint16*>(pSource);
 				pTypedResult[channel] = (double)(pTypedSource[channel]);
@@ -127,7 +127,7 @@ namespace mu
 				break;
 			}
 
-			case MBF_NINT8:
+			case EMeshBufferFormat::NInt8:
 			{
 				const int8* pTypedSource = reinterpret_cast<const int8*>(pSource);
 				pTypedResult[channel] = (double)(pTypedSource[channel]);
@@ -135,7 +135,7 @@ namespace mu
 				break;
 			}
 
-			case MBF_NUINT8:
+			case EMeshBufferFormat::NUInt8:
 			{
 				const uint8* pTypedSource = reinterpret_cast<const uint8*>(pSource);
 				pTypedResult[channel] = (double)(pTypedSource[channel]);
@@ -143,8 +143,8 @@ namespace mu
 				break;
 			}
 
-			case MBF_PACKEDDIR8:
-			case MBF_PACKEDDIR8_W_TANGENTSIGN:
+			case EMeshBufferFormat::PackedDir8:
+			case EMeshBufferFormat::PackedDir8_W_TangentSign:
 			{
 				const uint8* pTypedSource = reinterpret_cast<const uint8*>(pSource);
 				pTypedResult[channel] = (double)(pTypedSource[channel]);
@@ -153,8 +153,8 @@ namespace mu
 				break;
 			}
 
-			case MBF_PACKEDDIRS8:
-			case MBF_PACKEDDIRS8_W_TANGENTSIGN:
+			case EMeshBufferFormat::PackedDirS8:
+			case EMeshBufferFormat::PackedDirS8_W_TangentSign:
 			{
 				const int8* pTypedSource = reinterpret_cast<const int8*>(pSource);
 				pTypedResult[channel] = (double)(pTypedSource[channel]);
@@ -169,7 +169,7 @@ namespace mu
 			break;
 		}
 
-		case MBF_FLOAT32:
+		case EMeshBufferFormat::Float32:
 		{
 			float* pTypedResult = reinterpret_cast<float*>( pResult );
             uint8* pByteResult = reinterpret_cast<uint8*>( pResult );
@@ -177,14 +177,14 @@ namespace mu
 
 			switch ( sourceFormat )
 			{
-			case MBF_FLOAT64:
+			case EMeshBufferFormat::Float64:
 			{
 				const double* pTypedSource = reinterpret_cast<const double*>(pSource);
 				pTypedResult[channel] = float(pTypedSource[channel]);
 				break;
 			}
 
-			case MBF_FLOAT32:
+			case EMeshBufferFormat::Float32:
 			{
 				// Just dereferencing is not safe in all architectures. some like ARM require the
 				// floats to be 4-byte aligned and it may not be the case. We need memcpy.
@@ -192,56 +192,56 @@ namespace mu
 				break;
 			}
 
-			case MBF_FLOAT16:
+			case EMeshBufferFormat::Float16:
 			{
 				const FFloat16* pTypedSource = reinterpret_cast<const FFloat16*>( pSource );
 				pTypedResult[channel] = pTypedSource[channel];
 				break;
 			}
 
-			case MBF_INT32:
+			case EMeshBufferFormat::Int32:
 			{
                 const int32* pTypedSource = reinterpret_cast<const int32*>( pSource );
 				pTypedResult[channel] = (float)(pTypedSource[channel]);
 				break;
 			}
 
-			case MBF_UINT32:
+			case EMeshBufferFormat::UInt32:
 			{
                 const uint32* pTypedSource = reinterpret_cast<const uint32*>( pSource );
 				pTypedResult[channel] = (float)(pTypedSource[channel]);
 				break;
 			}
 
-			case MBF_INT16:
+			case EMeshBufferFormat::Int16:
 			{
                 const int16* pTypedSource = reinterpret_cast<const int16*>( pSource );
 				pTypedResult[channel] = (float)(pTypedSource[channel]);
 				break;
 			}
 
-			case MBF_UINT16:
+			case EMeshBufferFormat::UInt16:
 			{
                 const uint16* pTypedSource = reinterpret_cast<const uint16*>( pSource );
 				pTypedResult[channel] = (float)(pTypedSource[channel]);
 				break;
 			}
 
-			case MBF_INT8:
+			case EMeshBufferFormat::Int8:
 			{
                 const int8* pTypedSource = reinterpret_cast<const int8*>( pSource );
 				pTypedResult[channel] = (float)(pTypedSource[channel]);
 				break;
 			}
 
-			case MBF_UINT8:
+			case EMeshBufferFormat::UInt8:
 			{
                 const uint8* pTypedSource = reinterpret_cast<const uint8*>( pSource );
 				pTypedResult[channel] = (float)(pTypedSource[channel]);
 				break;
 			}
 
-			case MBF_NINT32:
+			case EMeshBufferFormat::NInt32:
 			{
                 const int32* pTypedSource = reinterpret_cast<const int32*>( pSource );
 				pTypedResult[channel] = (float)(pTypedSource[channel]);
@@ -249,7 +249,7 @@ namespace mu
 				break;
 			}
 
-			case MBF_NUINT32:
+			case EMeshBufferFormat::NUInt32:
 			{
                 const uint32* pTypedSource = reinterpret_cast<const uint32*>( pSource );
 				pTypedResult[channel] = (float)(pTypedSource[channel]);
@@ -257,7 +257,7 @@ namespace mu
 				break;
 			}
 
-			case MBF_NINT16:
+			case EMeshBufferFormat::NInt16:
 			{
                 const int16* pTypedSource = reinterpret_cast<const int16*>( pSource );
 				pTypedResult[channel] = (float)(pTypedSource[channel]);
@@ -265,7 +265,7 @@ namespace mu
 				break;
 			}
 
-			case MBF_NUINT16:
+			case EMeshBufferFormat::NUInt16:
 			{
                 const uint16* pTypedSource = reinterpret_cast<const uint16*>( pSource );
 				pTypedResult[channel] = (float)(pTypedSource[channel]);
@@ -273,7 +273,7 @@ namespace mu
 				break;
 			}
 
-			case MBF_NINT8:
+			case EMeshBufferFormat::NInt8:
 			{
                 const int8* pTypedSource = reinterpret_cast<const int8*>( pSource );
 				pTypedResult[channel] = (float)(pTypedSource[channel]);
@@ -281,7 +281,7 @@ namespace mu
 				break;
 			}
 
-			case MBF_NUINT8:
+			case EMeshBufferFormat::NUInt8:
 			{
                 const uint8* pTypedSource = reinterpret_cast<const uint8*>( pSource );
 				pTypedResult[channel] = (float)(pTypedSource[channel]);
@@ -289,8 +289,8 @@ namespace mu
 				break;
 			}
 
-            case MBF_PACKEDDIR8:
-            case MBF_PACKEDDIR8_W_TANGENTSIGN:
+            case EMeshBufferFormat::PackedDir8:
+            case EMeshBufferFormat::PackedDir8_W_TangentSign:
             {
                 const uint8* pTypedSource = reinterpret_cast<const uint8*>( pSource );
                 pTypedResult[channel] = (float)(pTypedSource[channel]);
@@ -299,8 +299,8 @@ namespace mu
                 break;
             }
 
-            case MBF_PACKEDDIRS8:
-            case MBF_PACKEDDIRS8_W_TANGENTSIGN:
+            case EMeshBufferFormat::PackedDirS8:
+            case EMeshBufferFormat::PackedDirS8_W_TangentSign:
             {
                 const int8* pTypedSource = reinterpret_cast<const int8*>( pSource );
                 pTypedResult[channel] = (float)(pTypedSource[channel]);
@@ -316,62 +316,62 @@ namespace mu
 		}
 
 		//-----------------------------------------------------------------------------------------
-		case MBF_FLOAT16:
+		case EMeshBufferFormat::Float16:
 		{
 			FFloat16* pTypedResult = reinterpret_cast<FFloat16*>( pResult );
 
 			switch ( sourceFormat )
 			{
-			case MBF_FLOAT32:
+			case EMeshBufferFormat::Float32:
 			{
 				const float* pTypedSource = reinterpret_cast<const float*>( pSource );
 				pTypedResult[channel] = pTypedSource[channel];
 				break;
 			}
 
-			case MBF_FLOAT16:
+			case EMeshBufferFormat::Float16:
 			{
 				const FFloat16* pTypedSource = reinterpret_cast<const FFloat16*>( pSource );
 				pTypedResult[channel] = pTypedSource[channel];
 				break;
 			}
 
-			case MBF_UINT32:
+			case EMeshBufferFormat::UInt32:
 			{
                 const uint32* pTypedSource = reinterpret_cast<const uint32*>( pSource );
 				pTypedResult[channel] = (float)(pTypedSource[channel]);
 				break;
 			}
 
-			case MBF_INT32:
+			case EMeshBufferFormat::Int32:
 			{
                 const int32* pTypedSource = reinterpret_cast<const int32*>( pSource );
 				pTypedResult[channel] = (float)(pTypedSource[channel]);
 				break;
 			}
 
-			case MBF_UINT16:
+			case EMeshBufferFormat::UInt16:
 			{
                 const uint16* pTypedSource = reinterpret_cast<const uint16*>( pSource );
 				pTypedResult[channel] = (float)(pTypedSource[channel]);
 				break;
 			}
 
-			case MBF_INT16:
+			case EMeshBufferFormat::Int16:
 			{
                 const int16* pTypedSource = reinterpret_cast<const int16*>( pSource );
 				pTypedResult[channel] = (float)(pTypedSource[channel]);
 				break;
 			}
 
-			case MBF_UINT8:
+			case EMeshBufferFormat::UInt8:
 			{
                 const uint8* pTypedSource = reinterpret_cast<const uint8*>( pSource );
 				pTypedResult[channel] = (float)(pTypedSource[channel]);
 				break;
 			}
 
-			case MBF_INT8:
+			case EMeshBufferFormat::Int8:
 			{
                 const int8* pTypedSource = reinterpret_cast<const int8*>( pSource );
 				pTypedResult[channel] = (float)(pTypedSource[channel]);
@@ -386,13 +386,13 @@ namespace mu
 		}
 
 		//-----------------------------------------------------------------------------------------
-		case MBF_UINT8:
+		case EMeshBufferFormat::UInt8:
 		{
             uint8* pTypedResult = reinterpret_cast<uint8*>( pResult );
 
 			switch ( sourceFormat )
 			{
-			case MBF_FLOAT32:
+			case EMeshBufferFormat::Float32:
 			{
 				const float* pTypedSource = reinterpret_cast<const float*>( pSource );
                 pTypedResult[channel] = (uint8)
@@ -406,7 +406,7 @@ namespace mu
 				break;
 			}
 
-			case MBF_FLOAT16:
+			case EMeshBufferFormat::Float16:
 			{
 				const FFloat16* pTypedSource = reinterpret_cast<const FFloat16*>( pSource );
                 pTypedResult[channel] = (uint8)
@@ -420,21 +420,21 @@ namespace mu
 				break;
 			}
 
-			case MBF_INT8:
+			case EMeshBufferFormat::Int8:
 			{
                 const int8* pTypedSource = reinterpret_cast<const int8*>( pSource );
                 pTypedResult[channel] =  (uint8)FMath::Max<int8>( 0, pTypedSource[channel] );
 				break;
 			}
 
-			case MBF_UINT8:
+			case EMeshBufferFormat::UInt8:
 			{
                 const uint8* pTypedSource = reinterpret_cast<const uint8*>( pSource );
 				pTypedResult[channel] = pTypedSource[channel];
 				break;
 			}
 
-			case MBF_INT16:
+			case EMeshBufferFormat::Int16:
 			{
                 const int16* pTypedSource = reinterpret_cast<const int16*>( pSource );
                 pTypedResult[channel] =  (uint8)
@@ -448,7 +448,7 @@ namespace mu
 				break;
 			}
 
-			case MBF_UINT16:
+			case EMeshBufferFormat::UInt16:
 			{
 				// Clamp
                 const uint16* pTypedSource = reinterpret_cast<const uint16*>( pSource );
@@ -463,7 +463,7 @@ namespace mu
 				break;
 			}
 
-			case MBF_INT32:
+			case EMeshBufferFormat::Int32:
 			{
                 const int32* pTypedSource = reinterpret_cast<const int32*>( pSource );
                 pTypedResult[channel] =  (uint8)
@@ -477,7 +477,7 @@ namespace mu
 				break;
 			}
 
-			case MBF_UINT32:
+			case EMeshBufferFormat::UInt32:
 			{
 				// Clamp
                 const uint32* pTypedSource = reinterpret_cast<const uint32*>( pSource );
@@ -492,12 +492,12 @@ namespace mu
 				break;
 			}
 
-			case MBF_NUINT8:
-			case MBF_NUINT16:
-			case MBF_NUINT32:
-			case MBF_NINT8:
-			case MBF_NINT16:
-			case MBF_NINT32:
+			case EMeshBufferFormat::NUInt8:
+			case EMeshBufferFormat::NUInt16:
+			case EMeshBufferFormat::NUInt32:
+			case EMeshBufferFormat::NInt8:
+			case EMeshBufferFormat::NInt16:
+			case EMeshBufferFormat::NInt32:
 				pTypedResult = 0;
 				break;
 
@@ -508,13 +508,13 @@ namespace mu
 			break;
 		}
 
-		case MBF_UINT16:
+		case EMeshBufferFormat::UInt16:
 		{
             uint16* pTypedResult = reinterpret_cast<uint16*>( pResult );
 
 			switch ( sourceFormat )
 			{
-			case MBF_FLOAT32:
+			case EMeshBufferFormat::Float32:
 			{
 				const float* pTypedSource = reinterpret_cast<const float*>( pSource );
                 pTypedResult[channel] = (uint16)
@@ -528,7 +528,7 @@ namespace mu
 				break;
 			}
 
-			case MBF_FLOAT16:
+			case EMeshBufferFormat::Float16:
 			{
 				const FFloat16* pTypedSource = reinterpret_cast<const FFloat16*>( pSource );
                 pTypedResult[channel] = (uint16)
@@ -542,35 +542,35 @@ namespace mu
 				break;
 			}
 
-			case MBF_UINT8:
+			case EMeshBufferFormat::UInt8:
 			{
                 const uint8* pTypedSource = reinterpret_cast<const uint8*>( pSource );
 				pTypedResult[channel] =  pTypedSource[channel];
 				break;
 			}
 
-			case MBF_INT8:
+			case EMeshBufferFormat::Int8:
 			{
                 const int8* pTypedSource = reinterpret_cast<const int8*>( pSource );
                 pTypedResult[channel] =  (uint16)FMath::Max<int8>( 0, pTypedSource[channel] );
 				break;
 			}
 
-			case MBF_UINT16:
+			case EMeshBufferFormat::UInt16:
 			{
                 const uint16* pTypedSource = reinterpret_cast<const uint16*>( pSource );
 				pTypedResult[channel] =  pTypedSource[channel];
 				break;
 			}
 
-			case MBF_INT16:
+			case EMeshBufferFormat::Int16:
 			{
                 const int16* pTypedSource = reinterpret_cast<const int16*>( pSource );
                 pTypedResult[channel] =  (uint16)FMath::Max<int16>( 0, pTypedSource[channel] );
 				break;
 			}
 
-			case MBF_UINT32:
+			case EMeshBufferFormat::UInt32:
 			{
 				// Clamp
                 const uint32* pTypedSource = reinterpret_cast<const uint32*>( pSource );
@@ -585,7 +585,7 @@ namespace mu
 				break;
 			}
 
-			case MBF_INT32:
+			case EMeshBufferFormat::Int32:
 			{
 				// Clamp
                 const int32* pTypedSource = reinterpret_cast<const int32*>( pSource );
@@ -600,12 +600,12 @@ namespace mu
 				break;
 			}
 
-			case MBF_NUINT8:
-			case MBF_NUINT16:
-			case MBF_NUINT32:
-			case MBF_NINT8:
-			case MBF_NINT16:
-			case MBF_NINT32:
+			case EMeshBufferFormat::NUInt8:
+			case EMeshBufferFormat::NUInt16:
+			case EMeshBufferFormat::NUInt32:
+			case EMeshBufferFormat::NInt8:
+			case EMeshBufferFormat::NInt16:
+			case EMeshBufferFormat::NInt32:
 				pTypedResult = 0;
 				break;
 
@@ -616,13 +616,13 @@ namespace mu
 			break;
 		}
 
-		case MBF_UINT32:
+		case EMeshBufferFormat::UInt32:
 		{
             uint32* pTypedResult = reinterpret_cast<uint32*>( pResult );
 
 			switch ( sourceFormat )
 			{
-			case MBF_FLOAT32:
+			case EMeshBufferFormat::Float32:
 			{
 				const float* pTypedSource = reinterpret_cast<const float*>( pSource );
 				pTypedResult[channel] =
@@ -636,7 +636,7 @@ namespace mu
 				break;
 			}
 
-			case MBF_FLOAT16:
+			case EMeshBufferFormat::Float16:
 			{
 				const FFloat16* pTypedSource = reinterpret_cast<const FFloat16*>( pSource );
 				pTypedResult[channel] =
@@ -650,42 +650,42 @@ namespace mu
 				break;
 			}
 
-			case MBF_UINT8:
+			case EMeshBufferFormat::UInt8:
 			{
                 const uint8* pTypedSource = reinterpret_cast<const uint8*>( pSource );
 				pTypedResult[channel] =  pTypedSource[channel];
 				break;
 			}
 
-			case MBF_INT8:
+			case EMeshBufferFormat::Int8:
 			{
                 const int8* pTypedSource = reinterpret_cast<const int8*>( pSource );
                 pTypedResult[channel] =  (uint16)FMath::Max<int8>( 0, pTypedSource[channel] );
 				break;
 			}
 
-			case MBF_UINT16:
+			case EMeshBufferFormat::UInt16:
 			{
                 const uint16* pTypedSource = reinterpret_cast<const uint16*>( pSource );
 				pTypedResult[channel] =  pTypedSource[channel];
 				break;
 			}
 
-			case MBF_INT16:
+			case EMeshBufferFormat::Int16:
 			{
                 const int16* pTypedSource = reinterpret_cast<const int16*>( pSource );
                 pTypedResult[channel] =  (uint16)FMath::Max<int16>( 0, pTypedSource[channel] );
 				break;
 			}
 
-			case MBF_UINT32:
+			case EMeshBufferFormat::UInt32:
 			{
                 const uint32* pTypedSource = reinterpret_cast<const uint32*>( pSource );
 				pTypedResult[channel] = pTypedSource[channel];
 				break;
 			}
 
-			case MBF_INT32:
+			case EMeshBufferFormat::Int32:
 			{
 				// Clamp
                 const int32* pTypedSource = reinterpret_cast<const int32*>( pSource );
@@ -693,12 +693,12 @@ namespace mu
 				break;
 			}
 
-			case MBF_NUINT8:
-			case MBF_NUINT16:
-			case MBF_NUINT32:
-			case MBF_NINT8:
-			case MBF_NINT16:
-			case MBF_NINT32:
+			case EMeshBufferFormat::NUInt8:
+			case EMeshBufferFormat::NUInt16:
+			case EMeshBufferFormat::NUInt32:
+			case EMeshBufferFormat::NInt8:
+			case EMeshBufferFormat::NInt16:
+			case EMeshBufferFormat::NInt32:
 				pTypedResult = 0;
 				break;
 
@@ -709,34 +709,34 @@ namespace mu
 			break;
 		}
 
-		case MBF_UINT64:
+		case EMeshBufferFormat::UInt64:
 		{
 			uint64* pTypedResult = reinterpret_cast<uint64*>(pResult);
 
 			switch (sourceFormat)
 			{
-			case MBF_UINT8:
+			case EMeshBufferFormat::UInt8:
 			{
 				const uint8* pTypedSource = reinterpret_cast<const uint8*>(pSource);
 				pTypedResult[channel] = pTypedSource[channel];
 				break;
 			}
 
-			case MBF_UINT16:
+			case EMeshBufferFormat::UInt16:
 			{
 				const uint16* pTypedSource = reinterpret_cast<const uint16*>(pSource);
 				pTypedResult[channel] = pTypedSource[channel];
 				break;
 			}
 
-			case MBF_UINT32:
+			case EMeshBufferFormat::UInt32:
 			{
 				const uint32* pTypedSource = reinterpret_cast<const uint32*>(pSource);
 				pTypedResult[channel] = pTypedSource[channel];
 				break;
 			}
 
-			case MBF_UINT64:
+			case EMeshBufferFormat::UInt64:
 			{
 				const uint64* pTypedSource = reinterpret_cast<const uint64*>(pSource);
 				pTypedResult[channel] = pTypedSource[channel];
@@ -750,13 +750,13 @@ namespace mu
 			break;
 		}
 
-		case MBF_INT8:
+		case EMeshBufferFormat::Int8:
 		{
             int8* pTypedResult = reinterpret_cast<int8*>( pResult );
 
 			switch ( sourceFormat )
 			{
-			case MBF_FLOAT32:
+			case EMeshBufferFormat::Float32:
 			{
 				const float* pTypedSource = reinterpret_cast<const float*>( pSource );
                 pTypedResult[channel] = (int8)
@@ -770,7 +770,7 @@ namespace mu
 				break;
 			}
 
-			case MBF_FLOAT16:
+			case EMeshBufferFormat::Float16:
 			{
 				const FFloat16* pTypedSource = reinterpret_cast<const FFloat16*>( pSource );
                 pTypedResult[channel] = (int8)
@@ -784,19 +784,19 @@ namespace mu
 				break;
 			}
 
-			case MBF_INT8:
+			case EMeshBufferFormat::Int8:
 			{
                 const int8* pTypedSource = reinterpret_cast<const int8*>( pSource );
 				pTypedResult[channel] = pTypedSource[channel];
 				break;
 			}
 
-			case MBF_NUINT8:
-			case MBF_NUINT16:
-			case MBF_NUINT32:
-			case MBF_NINT8:
-			case MBF_NINT16:
-			case MBF_NINT32:
+			case EMeshBufferFormat::NUInt8:
+			case EMeshBufferFormat::NUInt16:
+			case EMeshBufferFormat::NUInt32:
+			case EMeshBufferFormat::NInt8:
+			case EMeshBufferFormat::NInt16:
+			case EMeshBufferFormat::NInt32:
 				pTypedResult = 0;
 				break;
 
@@ -807,13 +807,13 @@ namespace mu
 			break;
 		}
 
-		case MBF_INT16:
+		case EMeshBufferFormat::Int16:
 		{
             int16* pTypedResult = reinterpret_cast<int16*>( pResult );
 
 			switch ( sourceFormat )
 			{
-			case MBF_FLOAT32:
+			case EMeshBufferFormat::Float32:
 			{
 				const float* pTypedSource = reinterpret_cast<const float*>( pSource );
                 pTypedResult[channel] = (int16)
@@ -827,7 +827,7 @@ namespace mu
 				break;
 			}
 
-			case MBF_FLOAT16:
+			case EMeshBufferFormat::Float16:
 			{
 				const FFloat16* pTypedSource = reinterpret_cast<const FFloat16*>( pSource );
                 pTypedResult[channel] = (int16)
@@ -841,21 +841,21 @@ namespace mu
 				break;
 			}
 
-			case MBF_INT8:
+			case EMeshBufferFormat::Int8:
 			{
                 const int8* pTypedSource = reinterpret_cast<const int8*>( pSource );
                 pTypedResult[channel] = (int16)pTypedSource[channel];
 				break;
 			}
 
-			case MBF_UINT8:
+			case EMeshBufferFormat::UInt8:
 			{
                 const uint8* pTypedSource = reinterpret_cast<const uint8*>( pSource );
                 pTypedResult[channel] = (int16)pTypedSource[channel];
 				break;
 			}
 
-			case MBF_UINT16:
+			case EMeshBufferFormat::UInt16:
 			{
                 const uint16* pTypedSource = reinterpret_cast<const uint16*>( pSource );
                 pTypedResult[channel] = (int16)
@@ -865,7 +865,7 @@ namespace mu
 				break;
 			}
 
-			case MBF_INT32:
+			case EMeshBufferFormat::Int32:
 			{
                 const int32* pTypedSource = reinterpret_cast<const int32*>( pSource );
                 pTypedResult[channel] = (int16)
@@ -879,7 +879,7 @@ namespace mu
 				break;
 			}
 
-			case MBF_UINT32:
+			case EMeshBufferFormat::UInt32:
 			{
                 const uint32* pTypedSource = reinterpret_cast<const uint32*>( pSource );
                 pTypedResult[channel] = (int16)
@@ -889,12 +889,12 @@ namespace mu
 				break;
 			}
 
-			case MBF_NUINT8:
-			case MBF_NUINT16:
-			case MBF_NUINT32:
-			case MBF_NINT8:
-			case MBF_NINT16:
-			case MBF_NINT32:
+			case EMeshBufferFormat::NUInt8:
+			case EMeshBufferFormat::NUInt16:
+			case EMeshBufferFormat::NUInt32:
+			case EMeshBufferFormat::NInt8:
+			case EMeshBufferFormat::NInt16:
+			case EMeshBufferFormat::NInt32:
 				pTypedResult = 0;
 				break;
 
@@ -905,74 +905,74 @@ namespace mu
 			break;
 		}
 
-		case MBF_INT32:
+		case EMeshBufferFormat::Int32:
 		{
             int32* pTypedResult = reinterpret_cast<int32*>( pResult );
 
 			switch ( sourceFormat )
 			{
-			case MBF_FLOAT32:
+			case EMeshBufferFormat::Float32:
 			{
 				const float* pTypedSource = reinterpret_cast<const float*>( pSource );
                 pTypedResult[channel] = (int32)(pTypedSource[channel]);
 				break;
 			}
 
-			case MBF_FLOAT16:
+			case EMeshBufferFormat::Float16:
 			{
 				const FFloat16* pTypedSource = reinterpret_cast<const FFloat16*>( pSource );
                 pTypedResult[channel] = (int32)float(pTypedSource[channel]);
 				break;
 			}
 
-			case MBF_INT8:
+			case EMeshBufferFormat::Int8:
 			{
                 const int8* pTypedSource = reinterpret_cast<const int8*>( pSource );
                 pTypedResult[channel] = (int32)pTypedSource[channel];
 				break;
 			}
 
-			case MBF_UINT8:
+			case EMeshBufferFormat::UInt8:
 			{
                 const uint8* pTypedSource = reinterpret_cast<const uint8*>( pSource );
                 pTypedResult[channel] = (int32)pTypedSource[channel];
 				break;
 			}
 
-			case MBF_INT16:
+			case EMeshBufferFormat::Int16:
 			{
                 const int16* pTypedSource = reinterpret_cast<const int16*>( pSource );
                 pTypedResult[channel] = (int32)pTypedSource[channel];
 				break;
 			}
 
-			case MBF_UINT16:
+			case EMeshBufferFormat::UInt16:
 			{
                 const uint16* pTypedSource = reinterpret_cast<const uint16*>( pSource );
                 pTypedResult[channel] = (int32)pTypedSource[channel];
 				break;
 			}
 
-			case MBF_UINT32:
+			case EMeshBufferFormat::UInt32:
 			{
                 const uint32* pTypedSource = reinterpret_cast<const uint32*>( pSource );
                 pTypedResult[channel] = (int32)pTypedSource[channel];
 				break;
 			}
 
-			case MBF_INT32:
+			case EMeshBufferFormat::Int32:
 			{
                 const int32* pTypedSource = reinterpret_cast<const int32*>( pSource );
                 pTypedResult[channel] = pTypedSource[channel];
 				break;
 			}
 
-			case MBF_NUINT8:
-			case MBF_NUINT16:
-			case MBF_NUINT32:
-			case MBF_NINT8:
-			case MBF_NINT16:
-			case MBF_NINT32:
+			case EMeshBufferFormat::NUInt8:
+			case EMeshBufferFormat::NUInt16:
+			case EMeshBufferFormat::NUInt32:
+			case EMeshBufferFormat::NInt8:
+			case EMeshBufferFormat::NInt16:
+			case EMeshBufferFormat::NInt32:
 				pTypedResult = 0;
 				break;
 
@@ -984,27 +984,27 @@ namespace mu
 		}
 
 		//-----------------------------------------------------------------------------------------
-		case MBF_NUINT8:
+		case EMeshBufferFormat::NUInt8:
 		{
             uint8* pTypedResult = reinterpret_cast<uint8*>( pResult );
 
 			switch ( sourceFormat )
 			{
-			case MBF_NUINT8:
+			case EMeshBufferFormat::NUInt8:
 			{
 				const uint8* pTypedSource = reinterpret_cast<const uint8*>(pSource);
 				pTypedResult[channel] = pTypedSource[channel];
 				break;
 			}
 
-			case MBF_NUINT16:
+			case EMeshBufferFormat::NUInt16:
 			{
 				const uint16* pTypedSource = reinterpret_cast<const uint16*>(pSource);
 				pTypedResult[channel] = pTypedSource[channel] / (65535 / 255);
 				break;
 			}
 
-			case MBF_FLOAT32:
+			case EMeshBufferFormat::Float32:
 			{
 				const float* pTypedSource = reinterpret_cast<const float*>(pSource);
 				pTypedResult[channel] = (uint8)
@@ -1018,7 +1018,7 @@ namespace mu
 				break;
 			}
 
-			case MBF_FLOAT16:
+			case EMeshBufferFormat::Float16:
 			{
 				const FFloat16* pTypedSource = reinterpret_cast<const FFloat16*>( pSource );
                 pTypedResult[channel] = (uint8)
@@ -1032,12 +1032,12 @@ namespace mu
 				break;
 			}
 
-			case MBF_UINT8:
-			case MBF_UINT16:
-			case MBF_UINT32:
-			case MBF_INT8:
-			case MBF_INT16:
-			case MBF_INT32:
+			case EMeshBufferFormat::UInt8:
+			case EMeshBufferFormat::UInt16:
+			case EMeshBufferFormat::UInt32:
+			case EMeshBufferFormat::Int8:
+			case EMeshBufferFormat::Int16:
+			case EMeshBufferFormat::Int32:
 				// TODO: 1 or 0
 				pTypedResult = 0;
 				break;
@@ -1049,27 +1049,27 @@ namespace mu
 			break;
 		}
 
-		case MBF_NUINT16:
+		case EMeshBufferFormat::NUInt16:
 		{
             uint16* pTypedResult = reinterpret_cast<uint16*>( pResult );
 
 			switch ( sourceFormat )
 			{
-			case MBF_NUINT16:
+			case EMeshBufferFormat::NUInt16:
 			{
 				const uint16* pTypedSource = reinterpret_cast<const uint16*>(pSource);
 				pTypedResult[channel] = pTypedSource[channel];
 				break;
 			}
 
-			case MBF_NUINT8:
+			case EMeshBufferFormat::NUInt8:
 			{
 				const uint8* pTypedSource = reinterpret_cast<const uint8*>(pSource);
 				pTypedResult[channel] = pTypedSource[channel] * (65535 / 255);
 				break;
 			}
 
-			case MBF_FLOAT32:
+			case EMeshBufferFormat::Float32:
 			{
 				const float* pTypedSource = reinterpret_cast<const float*>( pSource );
                 pTypedResult[channel] = (uint16)
@@ -1083,7 +1083,7 @@ namespace mu
 				break;
 			}
 
-			case MBF_FLOAT16:
+			case EMeshBufferFormat::Float16:
 			{
 				const FFloat16* pTypedSource = reinterpret_cast<const FFloat16*>( pSource );
                 pTypedResult[channel] = (uint16)
@@ -1097,12 +1097,12 @@ namespace mu
 				break;
 			}
 
-			case MBF_UINT8:
-			case MBF_UINT16:
-			case MBF_UINT32:
-			case MBF_INT8:
-			case MBF_INT16:
-			case MBF_INT32:
+			case EMeshBufferFormat::UInt8:
+			case EMeshBufferFormat::UInt16:
+			case EMeshBufferFormat::UInt32:
+			case EMeshBufferFormat::Int8:
+			case EMeshBufferFormat::Int16:
+			case EMeshBufferFormat::Int32:
 				// TODO: 1 or 0
 				pTypedResult = 0;
 				break;
@@ -1114,13 +1114,13 @@ namespace mu
 			break;
 		}
 
-		case MBF_NUINT32:
+		case EMeshBufferFormat::NUInt32:
 		{
             uint32* pTypedResult = reinterpret_cast<uint32*>( pResult );
 
 			switch ( sourceFormat )
 			{
-			case MBF_FLOAT32:
+			case EMeshBufferFormat::Float32:
 			{
 				const float* pTypedSource = reinterpret_cast<const float*>( pSource );
                 pTypedResult[channel] = (uint32)
@@ -1134,7 +1134,7 @@ namespace mu
 				break;
 			}
 
-			case MBF_FLOAT16:
+			case EMeshBufferFormat::Float16:
 			{
 				const FFloat16* pTypedSource = reinterpret_cast<const FFloat16*>( pSource );
                 pTypedResult[channel] = (uint32)
@@ -1148,12 +1148,12 @@ namespace mu
 				break;
 			}
 
-			case MBF_UINT8:
-			case MBF_UINT16:
-			case MBF_UINT32:
-			case MBF_INT8:
-			case MBF_INT16:
-			case MBF_INT32:
+			case EMeshBufferFormat::UInt8:
+			case EMeshBufferFormat::UInt16:
+			case EMeshBufferFormat::UInt32:
+			case EMeshBufferFormat::Int8:
+			case EMeshBufferFormat::Int16:
+			case EMeshBufferFormat::Int32:
 				// TODO: 1 or 0
 				pTypedResult = 0;
 				break;
@@ -1165,13 +1165,13 @@ namespace mu
 			break;
 		}
 
-		case MBF_NINT8:
+		case EMeshBufferFormat::NInt8:
 		{
             int8* pTypedResult = reinterpret_cast<int8*>( pResult );
 
 			switch ( sourceFormat )
 			{
-			case MBF_FLOAT32:
+			case EMeshBufferFormat::Float32:
 			{
 				const float* pTypedSource = reinterpret_cast<const float*>( pSource );
                 pTypedResult[channel] = (int8)
@@ -1185,7 +1185,7 @@ namespace mu
 				break;
 			}
 
-			case MBF_FLOAT16:
+			case EMeshBufferFormat::Float16:
 			{
 				const FFloat16* pTypedSource = reinterpret_cast<const FFloat16*>( pSource );
                 pTypedResult[channel] = (int8)
@@ -1199,12 +1199,12 @@ namespace mu
 				break;
 			}
 
-			case MBF_UINT8:
-			case MBF_UINT16:
-			case MBF_UINT32:
-			case MBF_INT8:
-			case MBF_INT16:
-			case MBF_INT32:
+			case EMeshBufferFormat::UInt8:
+			case EMeshBufferFormat::UInt16:
+			case EMeshBufferFormat::UInt32:
+			case EMeshBufferFormat::Int8:
+			case EMeshBufferFormat::Int16:
+			case EMeshBufferFormat::Int32:
 				// TODO: -1, 1 or 0
 				pTypedResult = 0;
 				break;
@@ -1216,13 +1216,13 @@ namespace mu
 			break;
 		}
 
-		case MBF_NINT16:
+		case EMeshBufferFormat::NInt16:
 		{
             int16* pTypedResult = reinterpret_cast<int16*>( pResult );
 
 			switch ( sourceFormat )
 			{
-			case MBF_FLOAT32:
+			case EMeshBufferFormat::Float32:
 			{
 				const float* pTypedSource = reinterpret_cast<const float*>( pSource );
                 pTypedResult[channel] = (int16)
@@ -1236,7 +1236,7 @@ namespace mu
 				break;
 			}
 
-			case MBF_FLOAT16:
+			case EMeshBufferFormat::Float16:
 			{
 				const FFloat16* pTypedSource = reinterpret_cast<const FFloat16*>( pSource );
                 pTypedResult[channel] = (int16)
@@ -1250,12 +1250,12 @@ namespace mu
 				break;
 			}
 
-			case MBF_UINT8:
-			case MBF_UINT16:
-			case MBF_UINT32:
-			case MBF_INT8:
-			case MBF_INT16:
-			case MBF_INT32:
+			case EMeshBufferFormat::UInt8:
+			case EMeshBufferFormat::UInt16:
+			case EMeshBufferFormat::UInt32:
+			case EMeshBufferFormat::Int8:
+			case EMeshBufferFormat::Int16:
+			case EMeshBufferFormat::Int32:
 				// TODO: -1, 1 or 0
 				pTypedResult = 0;
 				break;
@@ -1268,32 +1268,32 @@ namespace mu
 		}
 
 
-		case MBF_NINT32:
+		case EMeshBufferFormat::NInt32:
 		{
             int32* pTypedResult = reinterpret_cast<int32*>( pResult );
 
 			switch ( sourceFormat )
 			{
-			case MBF_FLOAT32:
+			case EMeshBufferFormat::Float32:
 			{
 				const float* pTypedSource = reinterpret_cast<const float*>( pSource );
                 pTypedResult[channel] = (int32)(2147483648.0f*pTypedSource[channel]+0.5f);
 				break;
 			}
 
-			case MBF_FLOAT16:
+			case EMeshBufferFormat::Float16:
 			{
 				const FFloat16* pTypedSource = reinterpret_cast<const FFloat16*>( pSource );
                 pTypedResult[channel] = (int32)(2147483648.0f*float(pTypedSource[channel])+0.5f);
 				break;
 			}
 
-			case MBF_UINT8:
-			case MBF_UINT16:
-			case MBF_UINT32:
-			case MBF_INT8:
-			case MBF_INT16:
-			case MBF_INT32:
+			case EMeshBufferFormat::UInt8:
+			case EMeshBufferFormat::UInt16:
+			case EMeshBufferFormat::UInt32:
+			case EMeshBufferFormat::Int8:
+			case EMeshBufferFormat::Int16:
+			case EMeshBufferFormat::Int32:
 				// TODO: -1, 1 or 0
 				pTypedResult = 0;
 				break;
@@ -1305,15 +1305,15 @@ namespace mu
 			break;
 		}
 
-        case MBF_PACKEDDIR8:
-        case MBF_PACKEDDIR8_W_TANGENTSIGN:
+        case EMeshBufferFormat::PackedDir8:
+        case EMeshBufferFormat::PackedDir8_W_TangentSign:
         {
             uint8* pTypedResult = reinterpret_cast<uint8*>( pResult );
 
             switch ( sourceFormat )
             {
-            case MBF_PACKEDDIR8:
-            case MBF_PACKEDDIR8_W_TANGENTSIGN:
+            case EMeshBufferFormat::PackedDir8:
+            case EMeshBufferFormat::PackedDir8_W_TangentSign:
             {
                 const uint8* pTypedSource = reinterpret_cast<const uint8*>( pSource );
                 uint8 source = pTypedSource[channel];
@@ -1321,7 +1321,7 @@ namespace mu
                 break;
             }
 
-            case MBF_FLOAT32:
+            case EMeshBufferFormat::Float32:
             {
                 const float* pTypedSource = reinterpret_cast<const float*>( pSource );
                 float source = pTypedSource[channel];
@@ -1338,15 +1338,15 @@ namespace mu
             break;
         }
 
-        case MBF_PACKEDDIRS8:
-        case MBF_PACKEDDIRS8_W_TANGENTSIGN:
+        case EMeshBufferFormat::PackedDirS8:
+        case EMeshBufferFormat::PackedDirS8_W_TangentSign:
         {
             int8* pTypedResult = reinterpret_cast<int8*>( pResult );
 
             switch ( sourceFormat )
             {
-            case MBF_PACKEDDIRS8:
-            case MBF_PACKEDDIRS8_W_TANGENTSIGN:
+            case EMeshBufferFormat::PackedDirS8:
+            case EMeshBufferFormat::PackedDirS8_W_TangentSign:
             {
                 const int8* pTypedSource = reinterpret_cast<const int8*>( pSource );
                 int8 source = pTypedSource[channel];
@@ -1354,7 +1354,7 @@ namespace mu
                 break;
             }
 
-            case MBF_FLOAT32:
+            case EMeshBufferFormat::Float32:
             {
                 const float* pTypedSource = reinterpret_cast<const float*>( pSource );
                 float source = pTypedSource[channel];
diff --git a/Engine/Plugins/Mutable/Source/MutableRuntime/Internal/MuR/ImagePrivate.h b/Engine/Plugins/Mutable/Source/MutableRuntime/Internal/MuR/ImagePrivate.h
--- a/Engine/Plugins/Mutable/Source/MutableRuntime/Internal/MuR/ImagePrivate.h
+++ b/Engine/Plugins/Mutable/Source/MutableRuntime/Internal/MuR/ImagePrivate.h
@@ -23,12 +23,12 @@ namespace mu
 			return FormatA;
 		}
 
-		if (FormatA == EImageFormat::IF_NONE)
+		if (FormatA == EImageFormat::None)
 		{
 			return FormatA;
 		}
 
-		if (FormatB == EImageFormat::IF_NONE)
+		if (FormatB == EImageFormat::None)
 		{
 			return FormatB;
 		}
@@ -43,14 +43,14 @@ namespace mu
 			return FormatB;
 		}
 
-        if (FormatA == EImageFormat::IF_BC2 || FormatA == EImageFormat::IF_BC3 
-			|| FormatA == EImageFormat::IF_ASTC_4x4_RGBA_LDR || FormatA == EImageFormat::IF_ASTC_6x6_RGBA_LDR || FormatA == EImageFormat::IF_ASTC_8x8_RGBA_LDR || FormatA == EImageFormat::IF_ASTC_10x10_RGBA_LDR)
+        if (FormatA == EImageFormat::BC2 || FormatA == EImageFormat::BC3 
+			|| FormatA == EImageFormat::ASTC_4x4_RGBA_LDR || FormatA == EImageFormat::ASTC_6x6_RGBA_LDR || FormatA == EImageFormat::ASTC_8x8_RGBA_LDR || FormatA == EImageFormat::ASTC_10x10_RGBA_LDR)
 		{	
 			return FormatA;
 		}
 
-        if (FormatB == EImageFormat::IF_BC2 || FormatB == EImageFormat::IF_BC3
-			|| FormatB == EImageFormat::IF_ASTC_4x4_RGBA_LDR || FormatB == EImageFormat::IF_ASTC_6x6_RGBA_LDR || FormatB == EImageFormat::IF_ASTC_8x8_RGBA_LDR || FormatB == EImageFormat::IF_ASTC_10x10_RGBA_LDR)
+        if (FormatB == EImageFormat::BC2 || FormatB == EImageFormat::BC3
+			|| FormatB == EImageFormat::ASTC_4x4_RGBA_LDR || FormatB == EImageFormat::ASTC_6x6_RGBA_LDR || FormatB == EImageFormat::ASTC_8x8_RGBA_LDR || FormatB == EImageFormat::ASTC_10x10_RGBA_LDR)
 
 		{
 			return FormatB;
@@ -67,20 +67,20 @@ namespace mu
     {
 		InFormat = GetUncompressedFormat(InFormat);
 
-		if (InFormat == EImageFormat::IF_NONE)
+		if (InFormat == EImageFormat::None)
 		{
 			return InFormat;
 		}
 
 		switch (InFormat)
 		{
-		case EImageFormat::IF_L_UBYTE: 
+		case EImageFormat::L_UByte: 
 		{
-			return EImageFormat::IF_RGB_UBYTE;
+			return EImageFormat::RGB_UByte;
 		}
-		case EImageFormat::IF_RGB_UBYTE:
-		case EImageFormat::IF_RGBA_UBYTE:
-		case EImageFormat::IF_BGRA_UBYTE:
+		case EImageFormat::RGB_UByte:
+		case EImageFormat::RGBA_UByte:
+		case EImageFormat::BGRA_UByte:
 		{
 			return InFormat;
 		}
@@ -90,7 +90,7 @@ namespace mu
 		}
 		}
 
-		return EImageFormat::IF_NONE;
+		return EImageFormat::None;
     }
 
     //---------------------------------------------------------------------------------------------
@@ -121,264 +121,264 @@ namespace mu
 
 		switch(SrcFormat)
 		{
-			case EImageFormat::IF_BC1:
+			case EImageFormat::BC1:
 			{
-				if (DestFormat == EImageFormat::IF_RGBA_UBYTE)
+				if (DestFormat == EImageFormat::RGBA_UByte)
 				{
 					return DecompressionFuncRefType(miro::SubImageDecompression::BC1_To_RGBASubImage);
 				}
-				else if (DestFormat == EImageFormat::IF_RGB_UBYTE)
+				else if (DestFormat == EImageFormat::RGB_UByte)
 				{
 					return DecompressionFuncRefType(miro::SubImageDecompression::BC1_To_RGBSubImage);
 				}
 				break;
 			}
 
-			case EImageFormat::IF_BC2:
+			case EImageFormat::BC2:
 			{
-				if (DestFormat == EImageFormat::IF_RGBA_UBYTE)
+				if (DestFormat == EImageFormat::RGBA_UByte)
 				{
 					return DecompressionFuncRefType(miro::SubImageDecompression::BC2_To_RGBASubImage);
 				}
-				else if (DestFormat == EImageFormat::IF_RGB_UBYTE)
+				else if (DestFormat == EImageFormat::RGB_UByte)
 				{
 					return DecompressionFuncRefType(miro::SubImageDecompression::BC2_To_RGBSubImage);
 				}
 				break;
 			}
 
-			case EImageFormat::IF_BC3:
+			case EImageFormat::BC3:
 			{
-				if (DestFormat == EImageFormat::IF_RGBA_UBYTE)
+				if (DestFormat == EImageFormat::RGBA_UByte)
 				{
 					return DecompressionFuncRefType(miro::SubImageDecompression::BC3_To_RGBASubImage);
 				}
-				else if (DestFormat == EImageFormat::IF_RGB_UBYTE)
+				else if (DestFormat == EImageFormat::RGB_UByte)
 				{
 					return DecompressionFuncRefType(miro::SubImageDecompression::BC3_To_RGBSubImage);
 				}
 				break;
 			}
 
-			case EImageFormat::IF_BC4:
+			case EImageFormat::BC4:
 			{
-				if (DestFormat == EImageFormat::IF_RGBA_UBYTE)
+				if (DestFormat == EImageFormat::RGBA_UByte)
 				{
 					return DecompressionFuncRefType(miro::SubImageDecompression::BC4_To_RGBASubImage);
 				}
-				else if (DestFormat == EImageFormat::IF_RGB_UBYTE)
+				else if (DestFormat == EImageFormat::RGB_UByte)
 				{
 					return DecompressionFuncRefType(miro::SubImageDecompression::BC4_To_RGBSubImage);
 				}
-				if (DestFormat == EImageFormat::IF_L_UBYTE)
+				if (DestFormat == EImageFormat::L_UByte)
 				{
 					return DecompressionFuncRefType(miro::SubImageDecompression::BC4_To_LSubImage);
 				}
 				break;
 			}
 
-			case EImageFormat::IF_BC5:
+			case EImageFormat::BC5:
 			{
-				if (DestFormat == EImageFormat::IF_RGBA_UBYTE)
+				if (DestFormat == EImageFormat::RGBA_UByte)
 				{
 					return DecompressionFuncRefType(miro::SubImageDecompression::BC5_To_RGBASubImage);
 				}
-				else if (DestFormat == EImageFormat::IF_RGB_UBYTE)
+				else if (DestFormat == EImageFormat::RGB_UByte)
 				{
 					return DecompressionFuncRefType(miro::SubImageDecompression::BC5_To_RGBSubImage);
 				}
 				break;
 			}
 	
-			case EImageFormat::IF_ASTC_4x4_RGBA_LDR:
+			case EImageFormat::ASTC_4x4_RGBA_LDR:
 			{
-				if (DestFormat == EImageFormat::IF_RGBA_UBYTE)
+				if (DestFormat == EImageFormat::RGBA_UByte)
 				{
 					return DecompressionFuncRefType(miro::SubImageDecompression::ASTC4x4RGBAL_To_RGBASubImage);
 				}
-				else if (DestFormat == EImageFormat::IF_RGB_UBYTE)
+				else if (DestFormat == EImageFormat::RGB_UByte)
 				{
 					return DecompressionFuncRefType(miro::SubImageDecompression::ASTC4x4RGBAL_To_RGBSubImage);
 				}
 				break;
 			}
 
-			case EImageFormat::IF_ASTC_4x4_RGB_LDR:
+			case EImageFormat::ASTC_4x4_RGB_LDR:
 			{
-				if (DestFormat == EImageFormat::IF_RGBA_UBYTE)
+				if (DestFormat == EImageFormat::RGBA_UByte)
 				{
 					return DecompressionFuncRefType(miro::SubImageDecompression::ASTC4x4RGBL_To_RGBASubImage);
 				}
-				else if (DestFormat == EImageFormat::IF_RGB_UBYTE)
+				else if (DestFormat == EImageFormat::RGB_UByte)
 				{
 					return DecompressionFuncRefType(miro::SubImageDecompression::ASTC4x4RGBL_To_RGBSubImage);
 				}
 				break;
 			}
 
-			case EImageFormat::IF_ASTC_4x4_RG_LDR:
+			case EImageFormat::ASTC_4x4_RG_LDR:
 			{
-				if (DestFormat == EImageFormat::IF_RGBA_UBYTE)
+				if (DestFormat == EImageFormat::RGBA_UByte)
 				{
 					return DecompressionFuncRefType(miro::SubImageDecompression::ASTC4x4RGL_To_RGBASubImage);
 				}
-				else if (DestFormat == EImageFormat::IF_RGB_UBYTE)
+				else if (DestFormat == EImageFormat::RGB_UByte)
 				{
 					return DecompressionFuncRefType(miro::SubImageDecompression::ASTC4x4RGL_To_RGBSubImage);
 				}
 				break;
 			}
 
-			case EImageFormat::IF_ASTC_6x6_RGBA_LDR:
+			case EImageFormat::ASTC_6x6_RGBA_LDR:
 			{
-				if (DestFormat == EImageFormat::IF_RGBA_UBYTE)
+				if (DestFormat == EImageFormat::RGBA_UByte)
 				{
 					return DecompressionFuncRefType(miro::SubImageDecompression::ASTC6x6RGBAL_To_RGBASubImage);
 				}
-				else if (DestFormat == EImageFormat::IF_RGB_UBYTE)
+				else if (DestFormat == EImageFormat::RGB_UByte)
 				{
 					return DecompressionFuncRefType(miro::SubImageDecompression::ASTC6x6RGBAL_To_RGBSubImage);
 				}
 				break;
 			}
 
-			case EImageFormat::IF_ASTC_6x6_RGB_LDR:
+			case EImageFormat::ASTC_6x6_RGB_LDR:
 			{
-				if (DestFormat == EImageFormat::IF_RGBA_UBYTE)
+				if (DestFormat == EImageFormat::RGBA_UByte)
 				{
 					return DecompressionFuncRefType(miro::SubImageDecompression::ASTC6x6RGBL_To_RGBASubImage);
 				}
-				else if (DestFormat == EImageFormat::IF_RGB_UBYTE)
+				else if (DestFormat == EImageFormat::RGB_UByte)
 				{
 					return DecompressionFuncRefType(miro::SubImageDecompression::ASTC6x6RGBL_To_RGBSubImage);
 				}
 				break;
 			}
 
-			case EImageFormat::IF_ASTC_6x6_RG_LDR:
+			case EImageFormat::ASTC_6x6_RG_LDR:
 			{
-				if (DestFormat == EImageFormat::IF_RGBA_UBYTE)
+				if (DestFormat == EImageFormat::RGBA_UByte)
 				{
 					return DecompressionFuncRefType(miro::SubImageDecompression::ASTC6x6RGL_To_RGBASubImage);
 				}
-				else if (DestFormat == EImageFormat::IF_RGB_UBYTE)
+				else if (DestFormat == EImageFormat::RGB_UByte)
 				{
 					return DecompressionFuncRefType(miro::SubImageDecompression::ASTC6x6RGL_To_RGBSubImage);
 				}
 				break;
 			}
 
-			case EImageFormat::IF_ASTC_8x8_RGBA_LDR:
+			case EImageFormat::ASTC_8x8_RGBA_LDR:
 			{
-				if (DestFormat == EImageFormat::IF_RGBA_UBYTE)
+				if (DestFormat == EImageFormat::RGBA_UByte)
 				{
 					return DecompressionFuncRefType(miro::SubImageDecompression::ASTC8x8RGBAL_To_RGBASubImage);
 				}
-				else if (DestFormat == EImageFormat::IF_RGB_UBYTE)
+				else if (DestFormat == EImageFormat::RGB_UByte)
 				{
 					return DecompressionFuncRefType(miro::SubImageDecompression::ASTC8x8RGBAL_To_RGBSubImage);
 				}
 				break;
 			}
 
-			case EImageFormat::IF_ASTC_8x8_RGB_LDR:
+			case EImageFormat::ASTC_8x8_RGB_LDR:
 			{
-				if (DestFormat == EImageFormat::IF_RGBA_UBYTE)
+				if (DestFormat == EImageFormat::RGBA_UByte)
 				{
 					return DecompressionFuncRefType(miro::SubImageDecompression::ASTC8x8RGBL_To_RGBASubImage);
 				}
-				else if (DestFormat == EImageFormat::IF_RGB_UBYTE)
+				else if (DestFormat == EImageFormat::RGB_UByte)
 				{
 					return DecompressionFuncRefType(miro::SubImageDecompression::ASTC8x8RGBL_To_RGBSubImage);
 				}
 				break;
 			}
 
-			case EImageFormat::IF_ASTC_8x8_RG_LDR:
+			case EImageFormat::ASTC_8x8_RG_LDR:
 			{
-				if (DestFormat == EImageFormat::IF_RGBA_UBYTE)
+				if (DestFormat == EImageFormat::RGBA_UByte)
 				{
 					return DecompressionFuncRefType(miro::SubImageDecompression::ASTC8x8RGL_To_RGBASubImage);
 				}
-				else if (DestFormat == EImageFormat::IF_RGB_UBYTE)
+				else if (DestFormat == EImageFormat::RGB_UByte)
 				{
 					return DecompressionFuncRefType(miro::SubImageDecompression::ASTC8x8RGL_To_RGBSubImage);
 				}
 				break;
 			}
 
-			case EImageFormat::IF_ASTC_10x10_RGBA_LDR:
+			case EImageFormat::ASTC_10x10_RGBA_LDR:
 			{
-				if (DestFormat == EImageFormat::IF_RGBA_UBYTE)
+				if (DestFormat == EImageFormat::RGBA_UByte)
 				{
 					return DecompressionFuncRefType(miro::SubImageDecompression::ASTC10x10RGBAL_To_RGBASubImage);
 				}
-				else if (DestFormat == EImageFormat::IF_RGB_UBYTE)
+				else if (DestFormat == EImageFormat::RGB_UByte)
 				{
 					return DecompressionFuncRefType(miro::SubImageDecompression::ASTC10x10RGBAL_To_RGBSubImage);
 				}
 				break;
 			}
 
-			case EImageFormat::IF_ASTC_10x10_RGB_LDR:
+			case EImageFormat::ASTC_10x10_RGB_LDR:
 			{
-				if (DestFormat == EImageFormat::IF_RGBA_UBYTE)
+				if (DestFormat == EImageFormat::RGBA_UByte)
 				{
 					return DecompressionFuncRefType(miro::SubImageDecompression::ASTC10x10RGBL_To_RGBASubImage);
 				}
-				else if (DestFormat == EImageFormat::IF_RGB_UBYTE)
+				else if (DestFormat == EImageFormat::RGB_UByte)
 				{
 					return DecompressionFuncRefType(miro::SubImageDecompression::ASTC10x10RGBL_To_RGBSubImage);
 				}
 				break;
 			}
 
-			case EImageFormat::IF_ASTC_10x10_RG_LDR:
+			case EImageFormat::ASTC_10x10_RG_LDR:
 			{
-				if (DestFormat == EImageFormat::IF_RGBA_UBYTE)
+				if (DestFormat == EImageFormat::RGBA_UByte)
 				{
 					return DecompressionFuncRefType(miro::SubImageDecompression::ASTC10x10RGL_To_RGBASubImage);
 				}
-				else if (DestFormat == EImageFormat::IF_RGB_UBYTE)
+				else if (DestFormat == EImageFormat::RGB_UByte)
 				{
 					return DecompressionFuncRefType(miro::SubImageDecompression::ASTC10x10RGL_To_RGBSubImage);
 				}
 				break;
 			}
 
-			case EImageFormat::IF_ASTC_12x12_RGBA_LDR:
+			case EImageFormat::ASTC_12x12_RGBA_LDR:
 			{
-				if (DestFormat == EImageFormat::IF_RGBA_UBYTE)
+				if (DestFormat == EImageFormat::RGBA_UByte)
 				{
 					return DecompressionFuncRefType(miro::SubImageDecompression::ASTC12x12RGBAL_To_RGBASubImage);
 				}
-				else if (DestFormat == EImageFormat::IF_RGB_UBYTE)
+				else if (DestFormat == EImageFormat::RGB_UByte)
 				{
 					return DecompressionFuncRefType(miro::SubImageDecompression::ASTC12x12RGBAL_To_RGBSubImage);
 				}
 				break;
 			}
 
-			case EImageFormat::IF_ASTC_12x12_RGB_LDR:
+			case EImageFormat::ASTC_12x12_RGB_LDR:
 			{
-				if (DestFormat == EImageFormat::IF_RGBA_UBYTE)
+				if (DestFormat == EImageFormat::RGBA_UByte)
 				{
 					return DecompressionFuncRefType(miro::SubImageDecompression::ASTC12x12RGBL_To_RGBASubImage);
 				}
-				else if (DestFormat == EImageFormat::IF_RGB_UBYTE)
+				else if (DestFormat == EImageFormat::RGB_UByte)
 				{
 					return DecompressionFuncRefType(miro::SubImageDecompression::ASTC12x12RGBL_To_RGBSubImage);
 				}
 				break;
 			}
 
-			case EImageFormat::IF_ASTC_12x12_RG_LDR:
+			case EImageFormat::ASTC_12x12_RG_LDR:
 			{
-				if (DestFormat == EImageFormat::IF_RGBA_UBYTE)
+				if (DestFormat == EImageFormat::RGBA_UByte)
 				{
 					return DecompressionFuncRefType(miro::SubImageDecompression::ASTC12x12RGL_To_RGBASubImage);
 				}
-				else if (DestFormat == EImageFormat::IF_RGB_UBYTE)
+				else if (DestFormat == EImageFormat::RGB_UByte)
 				{
 					return DecompressionFuncRefType(miro::SubImageDecompression::ASTC12x12RGL_To_RGBSubImage);
 				}
diff --git a/Engine/Plugins/Mutable/Source/MutableRuntime/Internal/MuR/MeshPrivate.h b/Engine/Plugins/Mutable/Source/MutableRuntime/Internal/MuR/MeshPrivate.h
--- a/Engine/Plugins/Mutable/Source/MutableRuntime/Internal/MuR/MeshPrivate.h
+++ b/Engine/Plugins/Mutable/Source/MutableRuntime/Internal/MuR/MeshPrivate.h
@@ -35,9 +35,9 @@ namespace mu
 		pMesh->GetVertexBuffers().FindChannel( semantic, 0, &buffer, &channel );
 		check( buffer>=0 && channel>=0 );
 
-		EMeshBufferSemantic realSemantic = MBS_NONE;
+		EMeshBufferSemantic realSemantic = EMeshBufferSemantic::None;
 		int realSemanticIndex = 0;
-		EMeshBufferFormat format = MBF_NONE;
+		EMeshBufferFormat format = EMeshBufferFormat::None;
 		int components = 0;
 		int offset = 0;
 		pMesh->GetVertexBuffers().GetChannel( buffer, channel, &realSemantic, &realSemanticIndex, &format, &components, &offset );
@@ -70,9 +70,9 @@ namespace mu
 		pMesh->GetVertexBuffers().FindChannel( semantic, 0, &buffer, &channel );
 		check( buffer>=0 && channel>=0 );
 
-		EMeshBufferSemantic realSemantic = MBS_NONE;
+		EMeshBufferSemantic realSemantic = EMeshBufferSemantic::None;
 		int realSemanticIndex = 0;
-		EMeshBufferFormat format = MBF_NONE;
+		EMeshBufferFormat format = EMeshBufferFormat::None;
 		int components = 0;
 		int offset = 0;
 		pMesh->GetVertexBuffers().GetChannel
@@ -96,7 +96,7 @@ namespace mu
 	public:
 		inline UntypedMeshBufferIterator()
 		{
-			m_format = MBF_NONE;
+			m_format = EMeshBufferFormat::None;
 			m_components = 0;
 			m_elementSize = 0;
 			m_pBuf = 0;
@@ -116,9 +116,9 @@ namespace mu
 
 			if ( buffer>=0 && channel>=0 )
 			{
-				EMeshBufferSemantic realSemantic = MBS_NONE;
+				EMeshBufferSemantic realSemantic = EMeshBufferSemantic::None;
 				int realSemanticIndex = 0;
-				m_format = MBF_NONE;
+				m_format = EMeshBufferFormat::None;
 				m_components = 0;
 				int offset = 0;
 				bufferSet.GetChannel
@@ -138,7 +138,7 @@ namespace mu
 			}
 			else
 			{
-				m_format = MBF_NONE;
+				m_format = EMeshBufferFormat::None;
 				m_components = 0;
 				m_elementSize = 0;
 				m_pBuf = 0;
@@ -206,7 +206,7 @@ namespace mu
 
             for (int c=0; c< FMath::Min(m_components,4); ++c)
             {
-                ConvertData( c, &res[0], MBF_FLOAT32, ptr(), m_format );
+                ConvertData( c, &res[0], EMeshBufferFormat::Float32, ptr(), m_format );
             }
 
             return res;
@@ -217,7 +217,7 @@ namespace mu
 			FVector3f res = { 0,0,0 };
 			for (int c = 0; c < FMath::Min(m_components, 3); ++c)
 			{
-				ConvertData(c, &res[0], MBF_FLOAT32, ptr(), m_format);
+				ConvertData(c, &res[0], EMeshBufferFormat::Float32, ptr(), m_format);
 			}
 			return res;
 		}
@@ -227,7 +227,7 @@ namespace mu
 			FVector3d res = { 0,0,0 };
 			for (int c = 0; c < FMath::Min(m_components, 3); ++c)
 			{
-				ConvertData(c, &res[0], MBF_FLOAT64, ptr(), m_format);
+				ConvertData(c, &res[0], EMeshBufferFormat::Float64, ptr(), m_format);
 			}
 			return res;
 		}
@@ -237,7 +237,7 @@ namespace mu
 			FVector2f res = {0.0f, 0.0f};
 			for (int32 c = 0; c < FMath::Min(m_components, 2); ++c)
 			{
-				ConvertData(c, &res[0], MBF_FLOAT32, ptr(), m_format);
+				ConvertData(c, &res[0], EMeshBufferFormat::Float32, ptr(), m_format);
 			}
 
 			return res;
@@ -246,27 +246,27 @@ namespace mu
 		uint32 GetAsUINT32() const
 		{
 			uint32 res;
-			ConvertData(0, &res, MBF_UINT32, ptr(), m_format);
+			ConvertData(0, &res, EMeshBufferFormat::UInt32, ptr(), m_format);
 			return res;
 		}
 
 		uint64 GetAsUINT64() const
 		{
 			uint64 res;
-			ConvertData(0, &res, MBF_UINT64, ptr(), m_format);
+			ConvertData(0, &res, EMeshBufferFormat::UInt64, ptr(), m_format);
 			return res;
 		}
 
 		void SetFromUINT32(uint32_t v)
 		{
-			ConvertData(0, ptr(), m_format, &v, MBF_UINT32);
+			ConvertData(0, ptr(), m_format, &v, EMeshBufferFormat::UInt32);
 		}
 
 		void SetFromVec3f(const FVector3f& v)
 		{
 			for (int c = 0; c < FMath::Min(m_components, 3); ++c)
 			{
-				ConvertData(c, ptr(), m_format, &v, MBF_FLOAT32);
+				ConvertData(c, ptr(), m_format, &v, EMeshBufferFormat::Float32);
 			}
 		}
 
@@ -274,7 +274,7 @@ namespace mu
 		{
 			for (int c = 0; c < FMath::Min(m_components, 3); ++c)
 			{
-				ConvertData(c, ptr(), m_format, &v, MBF_FLOAT64);
+				ConvertData(c, ptr(), m_format, &v, EMeshBufferFormat::Float64);
 			}
 		}
 
@@ -314,9 +314,9 @@ namespace mu
 				
 				if (buffer >= 0 && channel >= 0)
 				{
-					EMeshBufferSemantic realSemantic = MBS_NONE;
+					EMeshBufferSemantic realSemantic = EMeshBufferSemantic::None;
 					int realSemanticIndex = 0;
-					EMeshBufferFormat format = MBF_NONE;
+					EMeshBufferFormat format = EMeshBufferFormat::None;
 					int components = 0;
 					int offset = 0;
 					bufferSet.GetChannel
@@ -330,7 +330,7 @@ namespace mu
 					if (format!=FORMAT || components!=COMPONENTS)
 					{
 						// Invalidate
-						m_format = MBF_NONE;
+						m_format = EMeshBufferFormat::None;
 						m_components = 0;
 						m_elementSize = 0;
 						m_pBuf = nullptr;
@@ -339,7 +339,7 @@ namespace mu
 				else
 				{
 					// Invalidate
-					m_format = MBF_NONE;
+					m_format = EMeshBufferFormat::None;
 					m_components = 0;
 					m_elementSize = 0;
 					m_pBuf = nullptr;
@@ -389,9 +389,9 @@ namespace mu
 
 			if ( buffer>=0 && channel>=0 )
 			{
-				EMeshBufferSemantic realSemantic = MBS_NONE;
+				EMeshBufferSemantic realSemantic = EMeshBufferSemantic::None;
 				int realSemanticIndex = 0;
-				m_format = MBF_NONE;
+				m_format = EMeshBufferFormat::None;
 				m_components = 0;
 				int offset = 0;
 				bufferSet.GetChannel
@@ -411,7 +411,7 @@ namespace mu
 			}
 			else
 			{
-				m_format = MBF_NONE;
+				m_format = EMeshBufferFormat::None;
 				m_components = 0;
 				m_elementSize = 0;
 				m_pBuf = nullptr;
@@ -482,7 +482,7 @@ namespace mu
 			FVector4f res(0.0f,0.0f,0.0f,0.0f);
 			for (int c = 0; c < FMath::Min(m_components, 4); ++c)
 			{
-				ConvertData(c, &res[0], MBF_FLOAT32, ptr(), m_format);
+				ConvertData(c, &res[0], EMeshBufferFormat::Float32, ptr(), m_format);
 			}
 			return res;
 		}
@@ -492,7 +492,7 @@ namespace mu
 			FVector3f res = { 0,0,0 };
 			for (int c = 0; c < FMath::Min(m_components, 3); ++c)
 			{
-				ConvertData(c, &res[0], MBF_FLOAT32, ptr(), m_format);
+				ConvertData(c, &res[0], EMeshBufferFormat::Float32, ptr(), m_format);
 			}
 			return res;
 		}
@@ -502,7 +502,7 @@ namespace mu
 			FVector3d res = { 0,0,0 };
 			for (int c = 0; c < FMath::Min(m_components, 3); ++c)
 			{
-				ConvertData(c, &res[0], MBF_FLOAT64, ptr(), m_format);
+				ConvertData(c, &res[0], EMeshBufferFormat::Float64, ptr(), m_format);
 			}
 			return res;
 		}
@@ -512,7 +512,7 @@ namespace mu
 			FVector2f res = {0.0f, 0.0f};
 			for (int32 c = 0; c < FMath::Min(m_components, 2); ++c)
 			{
-				ConvertData(c, &res[0], MBF_FLOAT32, ptr(), m_format);
+				ConvertData(c, &res[0], EMeshBufferFormat::Float32, ptr(), m_format);
 			}
 
 			return res;
@@ -521,14 +521,14 @@ namespace mu
 		uint32 GetAsUINT32() const
 		{
 			uint32 res = 0;
-			ConvertData(0, &res, MBF_UINT32, ptr(), m_format);
+			ConvertData(0, &res, EMeshBufferFormat::UInt32, ptr(), m_format);
 			return res;
 		}
 
 		uint64 GetAsUINT64() const
 		{
 			uint64 res;
-			ConvertData(0, &res, MBF_UINT64, ptr(), m_format);
+			ConvertData(0, &res, EMeshBufferFormat::UInt64, ptr(), m_format);
 			return res;
 		}
 
@@ -536,7 +536,7 @@ namespace mu
 		{
 			for (int32 c = 0; c < FMath::Min(m_components, Count); ++c)
 			{
-				ConvertData(c, Data, MBF_INT32, ptr(), m_format);
+				ConvertData(c, Data, EMeshBufferFormat::Int32, ptr(), m_format);
 			}
 		}
 
@@ -551,7 +551,7 @@ namespace mu
 
 		int m_elementSize = 0;
         const uint8_t* m_pBuf = nullptr;
-		EMeshBufferFormat m_format = MBF_NONE;
+		EMeshBufferFormat m_format = EMeshBufferFormat::None;
 		int m_components = 0;
 	};
 
@@ -588,9 +588,9 @@ namespace mu
 
 				if (buffer >= 0 && channel >= 0)
 				{
-					EMeshBufferSemantic realSemantic = MBS_NONE;
+					EMeshBufferSemantic realSemantic = EMeshBufferSemantic::None;
 					int realSemanticIndex = 0;
-					EMeshBufferFormat format = MBF_NONE;
+					EMeshBufferFormat format = EMeshBufferFormat::None;
 					int components = 0;
 					int offset = 0;
 					bufferSet.GetChannel
@@ -604,7 +604,7 @@ namespace mu
 					if (format!=FORMAT || components!=COMPONENTS)
 					{
 						// Invalidate
-						m_format = MBF_NONE;
+						m_format = EMeshBufferFormat::None;
 						m_components = 0;
 						m_elementSize = 0;
 						m_pBuf = nullptr;
@@ -613,7 +613,7 @@ namespace mu
 				else
 				{
 					// Invalidate
-					m_format = MBF_NONE;
+					m_format = EMeshBufferFormat::None;
 					m_components = 0;
 					m_elementSize = 0;
 					m_pBuf = nullptr;
@@ -666,7 +666,7 @@ namespace mu
 			Mesh = InMesh;
 			if (Mesh)
 			{
-				BufferIterator = UntypedMeshBufferIteratorConst(Mesh->VertexBuffers, MBS_VERTEXINDEX, 0 );
+				BufferIterator = UntypedMeshBufferIteratorConst(Mesh->VertexBuffers, EMeshBufferSemantic::VertexIndex, 0 );
 			}
 		}		
 
@@ -736,7 +736,7 @@ namespace mu
 			}
 
 			// Is it relative?
-			else if (BufferIterator.GetFormat() == MBF_UINT32)
+			else if (BufferIterator.GetFormat() == EMeshBufferFormat::UInt32)
 			{
 				// There is a buffer storing IDs without prefix because it is the same for all vertices.
 				uint32 RelativeId = BufferIterator.GetAsUINT32();
@@ -744,7 +744,7 @@ namespace mu
 			}
 
 			// Is it explicit?
-			else if (BufferIterator.GetFormat() == MBF_UINT64)
+			else if (BufferIterator.GetFormat() == EMeshBufferFormat::UInt64)
 			{
 				Id = BufferIterator.GetAsUINT64();
 			}
diff --git a/Engine/Plugins/Mutable/Source/MutableRuntime/Internal/MuR/ModelPrivate.h b/Engine/Plugins/Mutable/Source/MutableRuntime/Internal/MuR/ModelPrivate.h
--- a/Engine/Plugins/Mutable/Source/MutableRuntime/Internal/MuR/ModelPrivate.h
+++ b/Engine/Plugins/Mutable/Source/MutableRuntime/Internal/MuR/ModelPrivate.h
@@ -72,7 +72,7 @@ namespace mu
 		uint16 ImageSizeY = 0;
 		uint16 _Padding = 0;
 		uint8 LODCount = 0;
-		EImageFormat ImageFormat = EImageFormat::IF_NONE;
+		EImageFormat ImageFormat = EImageFormat::None;
 	};
 	MUTABLE_DEFINE_POD_SERIALISABLE(FImageLODRange);
 
diff --git a/Engine/Plugins/Mutable/Source/MutableRuntime/Internal/MuR/OpImageGrow.h b/Engine/Plugins/Mutable/Source/MutableRuntime/Internal/MuR/OpImageGrow.h
--- a/Engine/Plugins/Mutable/Source/MutableRuntime/Internal/MuR/OpImageGrow.h
+++ b/Engine/Plugins/Mutable/Source/MutableRuntime/Internal/MuR/OpImageGrow.h
@@ -15,7 +15,7 @@ namespace mu
 	//---------------------------------------------------------------------------------------------
 	inline void ImageGrow( Image* pImage )
 	{
-		check(pImage && pImage->GetFormat()== EImageFormat::IF_L_UBYTE );
+		check(pImage && pImage->GetFormat()== EImageFormat::L_UByte );
 
 		MUTABLE_CPUPROFILER_SCOPE(ImageGrow1)
 
@@ -69,7 +69,7 @@ namespace mu
 	{
 		MUTABLE_CPUPROFILER_SCOPE(ImageGrow2)
 
-		check( pMask->GetFormat()== EImageFormat::IF_L_UBYTE );
+		check( pMask->GetFormat()== EImageFormat::L_UByte );
 
 		int sizeX = pImage->GetSizeX();
 		int sizeY = pImage->GetSizeY();
diff --git a/Engine/Plugins/Mutable/Source/MutableRuntime/Internal/MuR/OpMeshMerge.h b/Engine/Plugins/Mutable/Source/MutableRuntime/Internal/MuR/OpMeshMerge.h
--- a/Engine/Plugins/Mutable/Source/MutableRuntime/Internal/MuR/OpMeshMerge.h
+++ b/Engine/Plugins/Mutable/Source/MutableRuntime/Internal/MuR/OpMeshMerge.h
@@ -55,7 +55,7 @@ namespace mu
 			(void)SecondIndexBuffer;
 
 			// This will be changed below if need to change the format of the index buffers.
-			EMeshBufferFormat IndexBufferFormat = MBF_NONE;
+			EMeshBufferFormat IndexBufferFormat = EMeshBufferFormat::None;
 
 			if (FirstCount && SecondCount)
 			{
@@ -70,18 +70,18 @@ namespace mu
 				
 				if (totalVertexCount > maxSupportedVertices)
 				{
-					IndexBufferFormat = totalVertexCount > MAX_uint16 ? MBF_UINT32 : MBF_UINT16;
+					IndexBufferFormat = totalVertexCount > MAX_uint16 ? EMeshBufferFormat::UInt32 : EMeshBufferFormat::UInt16;
 				}
 
 			}
 			
-			if (IndexBufferFormat != MBF_NONE)
+			if (IndexBufferFormat != EMeshBufferFormat::None)
 			{
 				// We only support vertex indices in case of having to change the format.
 				check(FirstIndexBuffer.Channels.Num() == 1);
 
 				ResultIndexBuffer.Channels.SetNum(1);
-				ResultIndexBuffer.Channels[0].Semantic = MBS_VERTEXINDEX;
+				ResultIndexBuffer.Channels[0].Semantic = EMeshBufferSemantic::VertexIndex;
 				ResultIndexBuffer.Channels[0].Format = IndexBufferFormat;
 				ResultIndexBuffer.Channels[0].ComponentCount = 1;
 				ResultIndexBuffer.Channels[0].SemanticIndex = 0;
@@ -102,13 +102,13 @@ namespace mu
 			ResultIndexBuffer.Data.SetNum(ResultIndexBuffer.ElementSize * (FirstCount + SecondCount));
 
 			check(ResultIndexBuffer.Channels.Num() == 1);
-			check(ResultIndexBuffer.Channels[0].Semantic == MBS_VERTEXINDEX);
+			check(ResultIndexBuffer.Channels[0].Semantic == EMeshBufferSemantic::VertexIndex);
 
 			if (!ResultIndexBuffer.Data.IsEmpty())
 			{
 				if (FirstCount)
 				{
-					if (IndexBufferFormat == MBF_NONE
+					if (IndexBufferFormat == EMeshBufferFormat::None
 						|| IndexBufferFormat == FirstIndexBuffer.Channels[0].Format)
 					{
 						FMemory::Memcpy(&ResultIndexBuffer.Data[0],
@@ -122,11 +122,11 @@ namespace mu
 						uint8_t* pDest = &ResultIndexBuffer.Data[0];
 						switch (IndexBufferFormat)
 						{
-						case MBF_UINT32:
+						case EMeshBufferFormat::UInt32:
 						{
 							switch (FirstIndexBuffer.Channels[0].Format)
 							{
-							case MBF_UINT16:
+							case EMeshBufferFormat::UInt16:
 							{
 								for (int v = 0; v < FirstCount; ++v)
 								{
@@ -137,7 +137,7 @@ namespace mu
 								break;
 							}
 
-							case MBF_UINT8:
+							case EMeshBufferFormat::UInt8:
 							{
 								for (int v = 0; v < FirstCount; ++v)
 								{
@@ -155,12 +155,12 @@ namespace mu
 							break;
 						}
 
-						case MBF_UINT16:
+						case EMeshBufferFormat::UInt16:
 						{
 							switch (FirstIndexBuffer.Channels[0].Format)
 							{
 
-							case MBF_UINT8:
+							case EMeshBufferFormat::UInt8:
 							{
 								for (int v = 0; v < FirstCount; ++v)
 								{
@@ -192,15 +192,15 @@ namespace mu
 
 					uint32_t firstVertexCount = pFirst->GetVertexBuffers().GetElementCount();
 
-					if (IndexBufferFormat == MBF_NONE
+					if (IndexBufferFormat == EMeshBufferFormat::None
 						|| IndexBufferFormat == SecondIndexBuffer.Channels[0].Format)
 					{
 						switch (SecondIndexBuffer.Channels[0].Format)
 						{
-						case MBF_INT32:
-						case MBF_UINT32:
-						case MBF_NINT32:
-						case MBF_NUINT32:
+						case EMeshBufferFormat::Int32:
+						case EMeshBufferFormat::UInt32:
+						case EMeshBufferFormat::NInt32:
+						case EMeshBufferFormat::NUInt32:
 						{
 							for (int v = 0; v < SecondCount; ++v)
 							{
@@ -211,10 +211,10 @@ namespace mu
 							break;
 						}
 
-						case MBF_INT16:
-						case MBF_UINT16:
-						case MBF_NINT16:
-						case MBF_NUINT16:
+						case EMeshBufferFormat::Int16:
+						case EMeshBufferFormat::UInt16:
+						case EMeshBufferFormat::NInt16:
+						case EMeshBufferFormat::NUInt16:
 						{
 							for (int v = 0; v < SecondCount; ++v)
 							{
@@ -225,10 +225,10 @@ namespace mu
 							break;
 						}
 
-						case MBF_INT8:
-						case MBF_UINT8:
-						case MBF_NINT8:
-						case MBF_NUINT8:
+						case EMeshBufferFormat::Int8:
+						case EMeshBufferFormat::UInt8:
+						case EMeshBufferFormat::NInt8:
+						case EMeshBufferFormat::NUInt8:
 						{
 							for (int v = 0; v < SecondCount; ++v)
 							{
@@ -250,14 +250,14 @@ namespace mu
 						switch (IndexBufferFormat)
 						{
 
-						case MBF_UINT32:
+						case EMeshBufferFormat::UInt32:
 						{
 							switch (SecondIndexBuffer.Channels[0].Format)
 							{
-							case MBF_INT16:
-							case MBF_UINT16:
-							case MBF_NINT16:
-							case MBF_NUINT16:
+							case EMeshBufferFormat::Int16:
+							case EMeshBufferFormat::UInt16:
+							case EMeshBufferFormat::NInt16:
+							case EMeshBufferFormat::NUInt16:
 							{
 								for (int v = 0; v < SecondCount; ++v)
 								{
@@ -268,10 +268,10 @@ namespace mu
 								break;
 							}
 
-							case MBF_INT8:
-							case MBF_UINT8:
-							case MBF_NINT8:
-							case MBF_NUINT8:
+							case EMeshBufferFormat::Int8:
+							case EMeshBufferFormat::UInt8:
+							case EMeshBufferFormat::NInt8:
+							case EMeshBufferFormat::NUInt8:
 							{
 								for (int v = 0; v < SecondCount; ++v)
 								{
@@ -290,14 +290,14 @@ namespace mu
 							break;
 						}
 
-						case MBF_UINT16:
+						case EMeshBufferFormat::UInt16:
 						{
 							switch (SecondIndexBuffer.Channels[0].Format)
 							{
-							case MBF_INT8:
-							case MBF_UINT8:
-							case MBF_NINT8:
-							case MBF_NUINT8:
+							case EMeshBufferFormat::Int8:
+							case EMeshBufferFormat::UInt8:
+							case EMeshBufferFormat::NInt8:
+							case EMeshBufferFormat::NUInt8:
 							{
 								for (int v = 0; v < SecondCount; ++v)
 								{
@@ -731,14 +731,14 @@ namespace mu
 		// undergo different operations.
 		if (!bNeedsExplicitIds)
 		{
-			UntypedMeshBufferIteratorConst FirstIDs(pFirst->VertexBuffers, MBS_VERTEXINDEX, 0);
-			UntypedMeshBufferIteratorConst SecondIDs(pSecond->VertexBuffers, MBS_VERTEXINDEX, 0);
+			UntypedMeshBufferIteratorConst FirstIDs(pFirst->VertexBuffers, EMeshBufferSemantic::VertexIndex, 0);
+			UntypedMeshBufferIteratorConst SecondIDs(pSecond->VertexBuffers, EMeshBufferSemantic::VertexIndex, 0);
 			bNeedsExplicitIds = FirstIDs.GetFormat() != SecondIDs.GetFormat();
 		}
 		if (!bNeedsExplicitIds)
 		{
-			UntypedMeshBufferIteratorConst FirstIDs(pFirst->VertexBuffers, MBS_LAYOUTBLOCK, 0);
-			UntypedMeshBufferIteratorConst SecondIDs(pSecond->VertexBuffers, MBS_LAYOUTBLOCK, 0);
+			UntypedMeshBufferIteratorConst FirstIDs(pFirst->VertexBuffers, EMeshBufferSemantic::LayoutBlock, 0);
+			UntypedMeshBufferIteratorConst SecondIDs(pSecond->VertexBuffers, EMeshBufferSemantic::LayoutBlock, 0);
 			bNeedsExplicitIds = FirstIDs.GetFormat() != SecondIDs.GetFormat();
 		}
 
@@ -760,11 +760,11 @@ namespace mu
 
 			// Check if the format of the BoneIndex buffer has to change
 			bool bChangeBoneIndicesFormat = false;
-			EMeshBufferFormat BoneIndexFormat = MaxNumBonesInBoneMaps > MAX_uint8 ? MBF_UINT16 : MBF_UINT8;
-			bChangeBoneIndicesFormat = pFirst->GetVertexBuffers().HasAnySemanticWithDifferentFormat(MBS_BONEINDICES, BoneIndexFormat);
+			EMeshBufferFormat BoneIndexFormat = MaxNumBonesInBoneMaps > MAX_uint8 ? EMeshBufferFormat::UInt16 : EMeshBufferFormat::UInt8;
+			bChangeBoneIndicesFormat = pFirst->GetVertexBuffers().HasAnySemanticWithDifferentFormat(EMeshBufferSemantic::BoneIndices, BoneIndexFormat);
 			if (!bChangeBoneIndicesFormat)
 			{
-				bChangeBoneIndicesFormat = pSecond->GetVertexBuffers().HasAnySemanticWithDifferentFormat(MBS_BONEINDICES, BoneIndexFormat);
+				bChangeBoneIndicesFormat = pSecond->GetVertexBuffers().HasAnySemanticWithDifferentFormat(EMeshBufferSemantic::BoneIndices, BoneIndexFormat);
 			}
 
 			// Step 1: Set up the initial vertex buffer structure of the result mesh.
@@ -789,9 +789,9 @@ namespace mu
 					for (const FMeshBufferChannel& SecondChan : SecondBuf.Channels)
 					{
 						// Skip system buffers
-						if (SecondChan.Semantic == MBS_VERTEXINDEX
+						if (SecondChan.Semantic == EMeshBufferSemantic::VertexIndex
 							||
-							SecondChan.Semantic == MBS_LAYOUTBLOCK)
+							SecondChan.Semantic == EMeshBufferSemantic::LayoutBlock)
 						{
 							continue;
 						}
@@ -809,13 +809,13 @@ namespace mu
 						}
 
 						// If there are additional UV channels try to add them
-						if (!bAllNewChannels && SecondChan.Semantic == MBS_TEXCOORDS
+						if (!bAllNewChannels && SecondChan.Semantic == EMeshBufferSemantic::TexCoords
 							&&
 							SecondChan.SemanticIndex > 0)
 						{
 							// Add additional UV channels if the previous one is found.
 							FMeshBufferSet& ResultVertexBuffers = Result->GetVertexBuffers();
-							ResultVertexBuffers.FindChannel(MBS_TEXCOORDS, SecondChan.SemanticIndex - 1, &FoundBuffer, &FoundChannel);
+							ResultVertexBuffers.FindChannel(EMeshBufferSemantic::TexCoords, SecondChan.SemanticIndex - 1, &FoundBuffer, &FoundChannel);
 
 							if (FoundBuffer >= 0)
 							{
@@ -886,7 +886,7 @@ namespace mu
 						const int32 ChannelsCount = VertexBuffers.GetBufferChannelCount(VertexBufferIndex);
 						for (int32 ChannelIndex = 0; ChannelIndex < ChannelsCount; ++ChannelIndex)
 						{
-							if (result.Channels[ChannelIndex].Semantic == MBS_BONEINDICES)
+							if (result.Channels[ChannelIndex].Semantic == EMeshBufferSemantic::BoneIndices)
 							{
 								result.Channels[ChannelIndex].Format = BoneIndexFormat;
 
@@ -994,7 +994,7 @@ namespace mu
 						const int32 ChannelsCount = VertexBuffers.GetBufferChannelCount(VertexBufferIndex);
 						for (int32 ChannelIndex = 0; ChannelIndex < ChannelsCount; ++ChannelIndex)
 						{
-							if (ResultBuffer.Channels[ChannelIndex].Semantic != MBS_BONEINDICES)
+							if (ResultBuffer.Channels[ChannelIndex].Semantic != EMeshBufferSemantic::BoneIndices)
 							{
 								continue;
 							}
@@ -1008,8 +1008,8 @@ namespace mu
 							{
 								switch (BoneIndexFormat)
 								{
-								case MBF_INT8:
-								case MBF_UINT8:
+								case EMeshBufferFormat::Int8:
+								case EMeshBufferFormat::UInt8:
 								{
 									uint8* pD = reinterpret_cast<uint8*>(&ResultBuffer.Data[ResultOffset]);
 
@@ -1032,8 +1032,8 @@ namespace mu
 									break;
 								}
 
-								case MBF_INT16:
-								case MBF_UINT16:
+								case EMeshBufferFormat::Int16:
+								case EMeshBufferFormat::UInt16:
 								{
 									uint16* pD = reinterpret_cast<uint16*>(&ResultBuffer.Data[ResultOffset]);
 
@@ -1056,8 +1056,8 @@ namespace mu
 									break;
 								}
 
-								case MBF_INT32:
-								case MBF_UINT32:
+								case EMeshBufferFormat::Int32:
+								case EMeshBufferFormat::UInt32:
 								{
 									// Unreal does not support 32 bit bone indices
 									checkf(false, TEXT("Format not supported."));
diff --git a/Engine/Plugins/Mutable/Source/MutableRuntime/Internal/MuR/Operations.h b/Engine/Plugins/Mutable/Source/MutableRuntime/Internal/MuR/Operations.h
--- a/Engine/Plugins/Mutable/Source/MutableRuntime/Internal/MuR/Operations.h
+++ b/Engine/Plugins/Mutable/Source/MutableRuntime/Internal/MuR/Operations.h
@@ -1222,7 +1222,7 @@ namespace mu
         //! For image instructions, for every image format, true if it is supported as the base
         //! format of the operation.
         //! TODO: Move to tools library?
-		bool supportedBasePixelFormats[uint8(EImageFormat::IF_COUNT)];
+		bool supportedBasePixelFormats[uint8(EImageFormat::Count)];
     };
 
 	MUTABLERUNTIME_API extern const FOpDesc& GetOpDesc( OP_TYPE type );
diff --git a/Engine/Plugins/Mutable/Source/MutableRuntime/Internal/MuR/ParametersPrivate.h b/Engine/Plugins/Mutable/Source/MutableRuntime/Internal/MuR/ParametersPrivate.h
--- a/Engine/Plugins/Mutable/Source/MutableRuntime/Internal/MuR/ParametersPrivate.h
+++ b/Engine/Plugins/Mutable/Source/MutableRuntime/Internal/MuR/ParametersPrivate.h
@@ -137,14 +137,14 @@ namespace Private
 
 }
 
-    MUTABLE_DEFINE_ENUM_SERIALISABLE(PARAMETER_TYPE)
-    MUTABLE_DEFINE_ENUM_SERIALISABLE(PROJECTOR_TYPE)
+    MUTABLE_DEFINE_ENUM_SERIALISABLE(EParameterType)
+    MUTABLE_DEFINE_ENUM_SERIALISABLE(EProjectorType)
 
 
     /** Description of a projector to project an image on a mesh. */
     struct FProjector
     {
-        PROJECTOR_TYPE type = PROJECTOR_TYPE::PLANAR;
+        EProjectorType type = EProjectorType::Planar;
         FVector3f position = {0,0,0};
 		FVector3f direction = {0,0,0};
 		FVector3f up = {0,0,0};
@@ -248,23 +248,23 @@ namespace Private
 	};
 
 
-	using ParamBoolType = bool;
-	using ParamIntType = int32;
-	using ParamFloatType = float;
-	using ParamColorType = FVector4f;
-	using ParamProjectorType = Private::TIndirectObject<FProjector>;
-	using ParamImageType = FName;
-	using ParamStringType = Private::TIndirectObject<FString>;
-	using ParamMatrixType = Private::TIndirectObject<FMatrix44f>;
+	using FParamBoolType = bool;
+	using FParamIntType = int32;
+	using FParamFloatType = float;
+	using FParamColorType = FVector4f;
+	using FParamProjectorType = Private::TIndirectObject<FProjector>;
+	using FParamImageType = FName;
+	using FParamStringType = Private::TIndirectObject<FString>;
+	using FParamMatrixType = Private::TIndirectObject<FMatrix44f>;
 	
-	using PARAMETER_VALUE = TVariant<
-            ParamBoolType, ParamIntType, ParamFloatType, ParamColorType, ParamProjectorType, ParamImageType, ParamStringType, ParamMatrixType>;
+	using FParameterValue = TVariant<
+            FParamBoolType, FParamIntType, FParamFloatType, FParamColorType, FParamProjectorType, FParamImageType, FParamStringType, FParamMatrixType>;
 
-    // static_assert to track PARAMETER_VALUE size changes. It is ok to change if needed.
-    static_assert(sizeof(PARAMETER_VALUE) == 8*4, "PARAMETER_VALUE size has changed.");
+    // static_assert to track ParameterValue size changes. It is ok to change if needed.
+    static_assert(sizeof(FParameterValue) == 8*4, "ParameterValue size has changed.");
 
 	// TVariant currently does not support this operator. Once supported remove it.
-	inline bool operator==(const PARAMETER_VALUE& ValueA, const PARAMETER_VALUE& ValueB)
+	inline bool operator==(const FParameterValue& ValueA, const FParameterValue& ValueB)
 	{
 		const int32 IndexA = ValueA.GetIndex();
 		const int32 IndexB = ValueB.GetIndex();
@@ -289,9 +289,9 @@ namespace Private
         //! Unique id (provided externally, so no actual guarantee that it is unique.)
 		FGuid m_uid;
 
-        PARAMETER_TYPE m_type = PARAMETER_TYPE::T_NONE;
+        EParameterType m_type = EParameterType::None;
 
-        PARAMETER_VALUE m_defaultValue;
+        FParameterValue m_defaultValue;
 
         //! Ranges, if the parameter is multi-dimensional. The indices refer to the Model's program
         //! vector of range descriptors.
@@ -437,12 +437,12 @@ namespace Private
         TSharedPtr<const Model> m_pModel;
  
         //! Values for the parameters if they are not multidimensional.
-		TArray<PARAMETER_VALUE> m_values;
+		TArray<FParameterValue> m_values;
 
 
         //! If the parameter is multidemensional, the values are stored here.
         //! The key of the map is the vector of values stored in a RangeIndex
-		TArray< TMap< TArray<int32>, PARAMETER_VALUE > > m_multiValues;
+		TArray< TMap< TArray<int32>, FParameterValue > > m_multiValues;
 
 
         //!
diff --git a/Engine/Plugins/Mutable/Source/MutableRuntime/Private/MuR/CodeRunner.cpp b/Engine/Plugins/Mutable/Source/MutableRuntime/Private/MuR/CodeRunner.cpp
--- a/Engine/Plugins/Mutable/Source/MutableRuntime/Private/MuR/CodeRunner.cpp
+++ b/Engine/Plugins/Mutable/Source/MutableRuntime/Private/MuR/CodeRunner.cpp
@@ -4135,7 +4135,7 @@ namespace mu
             		
                 Ptr<const Image> Base = LoadImage( FCacheAddress(args.base,item) );
 
-				Ptr<Image> Result = CreateImage(Base->GetSizeX(), Base->GetSizeY(), Base->GetLODCount(), EImageFormat::IF_L_UBYTE, EInitializationType::NotInitialized);
+				Ptr<Image> Result = CreateImage(Base->GetSizeX(), Base->GetSizeY(), Base->GetLODCount(), EImageFormat::L_UByte, EInitializationType::NotInitialized);
                 ImageLuminance( Result.get(),Base.get() );
 
 				Release(Base);
@@ -4249,7 +4249,7 @@ namespace mu
 				FVector4f colour0 = LoadColor(FScheduledOp::FromOpAndOptions(args.colour0, item, 0));
 				FVector4f colour1 = LoadColor(FScheduledOp::FromOpAndOptions(args.colour1, item, 0));
 
-				ImagePtr pResult = CreateImage(args.size[0], args.size[1], 1, EImageFormat::IF_RGB_UBYTE, EInitializationType::NotInitialized);
+				ImagePtr pResult = CreateImage(args.size[0], args.size[1], 1, EImageFormat::RGB_UByte, EInitializationType::NotInitialized);
                 ImageGradient( pResult.get(), colour0, colour1 );
 
 				StoreImage( item, pResult );
@@ -4282,7 +4282,7 @@ namespace mu
 
                 float c = LoadScalar(FScheduledOp::FromOpAndOptions(args.threshold, item, 0));
 
-                Ptr<Image> Result = CreateImage(pA->GetSizeX(), pA->GetSizeY(), pA->GetLODCount(), EImageFormat::IF_L_UBYTE, EInitializationType::NotInitialized);
+                Ptr<Image> Result = CreateImage(pA->GetSizeX(), pA->GetSizeY(), pA->GetLODCount(), EImageFormat::L_UByte, EInitializationType::NotInitialized);
 				ImageBinarise( Result.get(), pA.get(), c );
 
 				Release(pA);
@@ -4608,7 +4608,7 @@ namespace mu
 					}
 
                     // Flat mesh UV raster
-					Ptr<Image> ResultImage = CreateImage(SizeX, SizeY, 1, EImageFormat::IF_L_UBYTE, EInitializationType::Black);
+					Ptr<Image> ResultImage = CreateImage(SizeX, SizeY, 1, EImageFormat::L_UByte, EInitializationType::Black);
 					if (pMesh)
 					{
 						ImageRasterMesh(pMesh.get(), ResultImage.get(), args.LayoutIndex, args.BlockId, CropMin, UncroppedSize);
@@ -4751,7 +4751,7 @@ namespace mu
 				const float FadeStartRad = FMath::DegreesToRadians(fadeStart);
 				const float FadeEndRad = FMath::DegreesToRadians(fadeEnd);
 
-				EImageFormat Format = Source ? GetUncompressedFormat(Source->GetFormat()) : EImageFormat::IF_L_UBYTE;
+				EImageFormat Format = Source ? GetUncompressedFormat(Source->GetFormat()) : EImageFormat::L_UByte;
 
 				if (Source && Source->GetFormat()!=Format)
 				{
@@ -4840,7 +4840,7 @@ namespace mu
 
 					switch (Projector.type)
 					{
-					case PROJECTOR_TYPE::PLANAR:
+					case EProjectorType::Planar:
 						ImageRasterProjectedPlanar(pMesh.get(), New.get(),
 							Source.get(), Mask.get(),
 							args.bIsRGBFadingEnabled, args.bIsAlphaFadingEnabled,
@@ -4851,7 +4851,7 @@ namespace mu
 							&Scratch, bUseProjectionVectorImpl);
 						break;
 
-					case PROJECTOR_TYPE::WRAPPING:
+					case EProjectorType::Wrapping:
 						ImageRasterProjectedWrapping(pMesh.get(), New.get(),
 							Source.get(), Mask.get(),
 							args.bIsRGBFadingEnabled, args.bIsAlphaFadingEnabled,
@@ -4862,7 +4862,7 @@ namespace mu
 							&Scratch, bUseProjectionVectorImpl);
 						break;
 
-					case PROJECTOR_TYPE::CYLINDRICAL:
+					case EProjectorType::Cylindrical:
 						ImageRasterProjectedCylindrical(pMesh.get(), New.get(),
 							Source.get(), Mask.get(),
 							args.bIsRGBFadingEnabled, args.bIsAlphaFadingEnabled,
@@ -4910,7 +4910,7 @@ namespace mu
 
                 Ptr<const Image> Mask = LoadImage( FCacheAddress(args.mask,item) );
 
-                Ptr<Image> Result = CreateImage( Mask->GetSizeX(), Mask->GetSizeY(), Mask->GetLODCount(), EImageFormat::IF_L_UBYTE, EInitializationType::NotInitialized);
+                Ptr<Image> Result = CreateImage( Mask->GetSizeX(), Mask->GetSizeY(), Mask->GetLODCount(), EImageFormat::L_UByte, EInitializationType::NotInitialized);
 
                 ImageMakeGrowMap(Result.get(), Mask.get(), args.border );
 				Result->m_flags |= Image::IF_CANNOT_BE_SCALED;
@@ -5729,12 +5729,9 @@ namespace mu
         {
 			OP::ParameterArgs args = Program.GetOpArgs<OP::ParameterArgs>(item.At);
 			Ptr<RangeIndex> index = BuildCurrentOpRangeIndex( item, pParams, pModel, args.variable );
-            float R = 0.f;
-            float G = 0.f;
-			float B = 0.f;
-			float A = 0.f;
-            pParams->GetColourValue( args.variable, &R, &G, &B, &A, index );
-            StoreColor( item, FVector4f(R, G, B, A) );
+            FVector4f V;
+            pParams->GetColourValue( args.variable, V, index );
+            StoreColor( item, V );
             break;
         }
 
@@ -5953,7 +5950,7 @@ namespace mu
             FProjector Result = pParams->GetPrivate()->GetProjectorValue(args.variable,index);
 
             // The type cannot be changed, take it from the default value
-            const FProjector& def = Program.m_parameters[args.variable].m_defaultValue.Get<ParamProjectorType>();
+            const FProjector& def = Program.m_parameters[args.variable].m_defaultValue.Get<FParamProjectorType>();
             Result.type = def.type;
 
             StoreProjector( item, Result );
@@ -6500,7 +6497,7 @@ namespace mu
 				// Update directly in the heap
 				EImageFormat OldFormat = m_heapImageDesc[item.CustomState].m_format;
 				EImageFormat NewFormat = args.format;
-				if (args.formatIfAlpha != EImageFormat::IF_NONE
+				if (args.formatIfAlpha != EImageFormat::None
 					&&
 					GetImageFormatData(OldFormat).Channels > 3)
 				{
@@ -6739,7 +6736,7 @@ namespace mu
 				break;
 
 			case 1:
-				m_heapImageDesc[item.CustomState].m_format = EImageFormat::IF_L_UBYTE;
+				m_heapImageDesc[item.CustomState].m_format = EImageFormat::L_UByte;
 				StoreValidDesc(item);
 				break;
 
@@ -6789,7 +6786,7 @@ namespace mu
 			m_heapImageDesc[item.CustomState].m_size[0] = args.size[0];
 			m_heapImageDesc[item.CustomState].m_size[1] = args.size[1];
 			m_heapImageDesc[item.CustomState].m_lods = 1;
-			m_heapImageDesc[item.CustomState].m_format = EImageFormat::IF_RGB_UBYTE;
+			m_heapImageDesc[item.CustomState].m_format = EImageFormat::RGB_UByte;
 			StoreValidDesc(item);
 			break;
 		}
@@ -6804,7 +6801,7 @@ namespace mu
 				break;
 
 			case 1:
-				m_heapImageDesc[item.CustomState].m_format = EImageFormat::IF_L_UBYTE;
+				m_heapImageDesc[item.CustomState].m_format = EImageFormat::L_UByte;
 				StoreValidDesc(item);
 				break;
 
@@ -6877,7 +6874,7 @@ namespace mu
 			m_heapImageDesc[item.CustomState].m_size[0] = args.sizeX;
 			m_heapImageDesc[item.CustomState].m_size[1] = args.sizeY;
 			m_heapImageDesc[item.CustomState].m_lods = 1;
-			m_heapImageDesc[item.CustomState].m_format = EImageFormat::IF_L_UBYTE;
+			m_heapImageDesc[item.CustomState].m_format = EImageFormat::L_UByte;
 			StoreValidDesc(item);
 			break;
 		}
@@ -6892,7 +6889,7 @@ namespace mu
 				break;
 
 			case 1:
-				m_heapImageDesc[item.CustomState].m_format = EImageFormat::IF_L_UBYTE;
+				m_heapImageDesc[item.CustomState].m_format = EImageFormat::L_UByte;
 				m_heapImageDesc[item.CustomState].m_lods = 1;
 				StoreValidDesc(item);
 				break;
diff --git a/Engine/Plugins/Mutable/Source/MutableRuntime/Private/MuR/CodeRunner.h b/Engine/Plugins/Mutable/Source/MutableRuntime/Private/MuR/CodeRunner.h
--- a/Engine/Plugins/Mutable/Source/MutableRuntime/Private/MuR/CodeRunner.h
+++ b/Engine/Plugins/Mutable/Source/MutableRuntime/Private/MuR/CodeRunner.h
@@ -253,7 +253,7 @@ namespace  mu
 			using StreamingDataContainerType = TArray<uint8, FDefaultMemoryTrackingAllocator<MemoryCounters::FStreamingMemoryCounter>>;
 
 			int32 RomIndex = -1;
-			ModelReader::OPERATION_ID m_streamID = -1;
+			ModelReader::FOperationID m_streamID = -1;
 			StreamingDataContainerType m_streamBuffer;
 			UE::Tasks::FTask Event;
 		};
