[mutable] Cleanup MutableRuntime API moving to use more unreal types and naming standards.

diff --git a/Engine/Plugins/Mutable/Source/MutableRuntime/Private/MuR/CodeRunnerTasked.cpp b/Engine/Plugins/Mutable/Source/MutableRuntime/Private/MuR/CodeRunnerTasked.cpp
--- a/Engine/Plugins/Mutable/Source/MutableRuntime/Private/MuR/CodeRunnerTasked.cpp
+++ b/Engine/Plugins/Mutable/Source/MutableRuntime/Private/MuR/CodeRunnerTasked.cpp
@@ -655,7 +655,7 @@ namespace mu
 		Ptr<const Image> Blended;
 		Ptr<const Image> Mask;
 		Ptr<Image> Result;
-		EImageFormat InitialFormat = EImageFormat::IF_NONE;
+		EImageFormat InitialFormat = EImageFormat::None;
 	};
 
 
@@ -668,13 +668,13 @@ namespace mu
 		ImageCompressionQuality = Runner->m_pSettings->ImageCompressionQuality;
 
 		Ptr<const Image> Base = Runner->LoadImage({ Args.base, Op.ExecutionIndex, Op.ExecutionOptions });
-		check(!Base || Base->GetFormat() < EImageFormat::IF_COUNT);
+		check(!Base || Base->GetFormat() < EImageFormat::Count);
 
 		Blended = Runner->LoadImage({ Args.blended, Op.ExecutionIndex, Op.ExecutionOptions });
 		if (Args.mask)
 		{
 			Mask = Runner->LoadImage({ Args.mask, Op.ExecutionIndex, Op.ExecutionOptions });
-			check(!Mask || Mask->GetFormat() < EImageFormat::IF_COUNT);
+			check(!Mask || Mask->GetFormat() < EImageFormat::Count);
 		}
 
 		// Shortcuts
@@ -966,7 +966,7 @@ namespace mu
 		FVector4f Color;
 		Ptr<const Image> Mask;
 		Ptr<Image> Result;
-		EImageFormat InitialFormat = EImageFormat::IF_NONE;
+		EImageFormat InitialFormat = EImageFormat::None;
 	};
 
 
@@ -979,13 +979,13 @@ namespace mu
 		ImageCompressionQuality = Runner->m_pSettings->ImageCompressionQuality;
 
 		Ptr<const Image> Base = Runner->LoadImage({ Args.base, Op.ExecutionIndex, Op.ExecutionOptions });
-		check(!Base || Base->GetFormat() < EImageFormat::IF_COUNT);
+		check(!Base || Base->GetFormat() < EImageFormat::Count);
 
 		Color = Runner->LoadColor({ Args.colour, Op.ExecutionIndex, 0 });
 		if (Args.mask)
 		{
 			Mask = Runner->LoadImage({ Args.mask, Op.ExecutionIndex, Op.ExecutionOptions });
-			check(!Mask || Mask->GetFormat() < EImageFormat::IF_COUNT);
+			check(!Mask || Mask->GetFormat() < EImageFormat::Count);
 		}
 
 		// Shortcuts
@@ -1006,7 +1006,7 @@ namespace mu
 
 		// Fix input data
 		InitialFormat = Base->GetFormat();
-		check(InitialFormat < EImageFormat::IF_COUNT);
+		check(InitialFormat < EImageFormat::Count);
 
 		if (Args.mask && Mask)
 		{
@@ -1202,7 +1202,7 @@ namespace mu
 	private:
 		int ImageCompressionQuality = 0;
 		OP::ImagePixelFormatArgs Args;
-		EImageFormat TargetFormat = EImageFormat::IF_NONE;
+		EImageFormat TargetFormat = EImageFormat::None;
 		Ptr<const Image> Base;
 		Ptr<Image> Result;
 		FImageOperator::FImagePixelFormatFunc ImagePixelFormatFunc;
@@ -1218,7 +1218,7 @@ namespace mu
 		ImageCompressionQuality = Runner->m_pSettings->ImageCompressionQuality;
 
 		Base = Runner->LoadImage({ Args.source, Op.ExecutionIndex, Op.ExecutionOptions });
-		check(!Base || Base->GetFormat() < EImageFormat::IF_COUNT);		
+		check(!Base || Base->GetFormat() < EImageFormat::Count);		
 
 		// Shortcuts
 		if (!Base)
@@ -1235,14 +1235,14 @@ namespace mu
 		}
 
 		TargetFormat = Args.format;
-		if (Args.formatIfAlpha != EImageFormat::IF_NONE
+		if (Args.formatIfAlpha != EImageFormat::None
 			&&
 			GetImageFormatData(Base->GetFormat()).Channels > 3)
 		{
 			TargetFormat = Args.formatIfAlpha;
 		}
 
-		if (TargetFormat == EImageFormat::IF_NONE || TargetFormat == Base->GetFormat())
+		if (TargetFormat == EImageFormat::None || TargetFormat == Base->GetFormat())
 		{
 			Runner->StoreImage(Op, Base);
 			return false;
@@ -1315,7 +1315,7 @@ namespace mu
 		ImageCompressionQuality = Runner->m_pSettings->ImageCompressionQuality;
 
 		Base = Runner->LoadImage({ Args.source, Op.ExecutionIndex, Op.ExecutionOptions });
-		check(!Base || Base->GetFormat() < EImageFormat::IF_COUNT);
+		check(!Base || Base->GetFormat() < EImageFormat::Count);
 
 		// Shortcuts
 		if (!Base)
diff --git a/Engine/Plugins/Mutable/Source/MutableRuntime/Private/MuR/Image.cpp b/Engine/Plugins/Mutable/Source/MutableRuntime/Private/MuR/Image.cpp
--- a/Engine/Plugins/Mutable/Source/MutableRuntime/Private/MuR/Image.cpp
+++ b/Engine/Plugins/Mutable/Source/MutableRuntime/Private/MuR/Image.cpp
@@ -343,7 +343,7 @@ namespace mu
 		const TArrayView<const uint8> Data = DataStorage.GetLOD(0);
 
 		// Non-generic part
-		if (Format == EImageFormat::IF_RGB_UBYTE)
+		if (Format == EImageFormat::RGB_UByte)
 		{
 			int32 ByteOffset = BlockOffset * FormatData.BytesPerBlock
 					+ (BlockPixelY * FormatData.PixelsPerBlockX + BlockPixelX) * 3;
@@ -353,7 +353,7 @@ namespace mu
 			Result[2] = Data[ByteOffset + 2] / 255.0f;
 			Result[3] = 1.0f;
 		}
-		else if (Format == EImageFormat::IF_RGBA_UBYTE)
+		else if (Format == EImageFormat::RGBA_UByte)
 		{
 			int32 ByteOffset = BlockOffset * FormatData.BytesPerBlock
 					+ (BlockPixelY * FormatData.PixelsPerBlockX + BlockPixelX) * 4;
@@ -363,7 +363,7 @@ namespace mu
 			Result[2] = Data[ByteOffset + 2] / 255.0f;
 			Result[3] = Data[ByteOffset + 3] / 255.0f;
 		}
-		else if (Format == EImageFormat::IF_BGRA_UBYTE)
+		else if (Format == EImageFormat::BGRA_UByte)
 		{
 			int32 ByteOffset = BlockOffset * FormatData.BytesPerBlock
 					+ (BlockPixelY * FormatData.PixelsPerBlockX + BlockPixelX) * 4;
@@ -373,7 +373,7 @@ namespace mu
 			Result[2] = Data[ByteOffset + 0] / 255.0f;
 			Result[3] = Data[ByteOffset + 3] / 255.0f;
 		}
-		else if (Format== EImageFormat::IF_L_UBYTE)
+		else if (Format== EImageFormat::L_UByte)
 		{
 			int32 ByteOffset = BlockOffset * FormatData.BytesPerBlock
 					+ (BlockPixelY * FormatData.PixelsPerBlockX + BlockPixelX) * 1;
@@ -402,7 +402,7 @@ namespace mu
 		{
 			switch(GetFormat())
 			{
-			case EImageFormat::IF_L_UBYTE:
+			case EImageFormat::L_UByte:
 			{
 				uint8 const * const DataPtr = DataView.GetData();
 
@@ -422,7 +422,7 @@ namespace mu
 				break;
 			}
 
-			case EImageFormat::IF_RGB_UBYTE:
+			case EImageFormat::RGB_UByte:
 			{
 				uint8 const * const DataPtr = DataView.GetData();
 
@@ -442,8 +442,8 @@ namespace mu
 				break;
 			}
 
-			case EImageFormat::IF_RGBA_UBYTE:
-			case EImageFormat::IF_BGRA_UBYTE:
+			case EImageFormat::RGBA_UByte:
+			case EImageFormat::BGRA_UByte:
 			{
 				uint8 const * const DataPtr = DataView.GetData();
 
@@ -458,11 +458,11 @@ namespace mu
 				}
 
 				Result = I >= NumElems;
-				if (GetFormat() == EImageFormat::IF_RGBA_UBYTE)
+				if (GetFormat() == EImageFormat::RGBA_UByte)
 				{
 					OutColor = FVector4f(DataPtr[0], DataPtr[1], DataPtr[2], DataPtr[3]) / 255.0f;
 				}
-				else if (GetFormat() == EImageFormat::IF_BGRA_UBYTE)
+				else if (GetFormat() == EImageFormat::BGRA_UByte)
 				{
 					OutColor = FVector4f(DataPtr[2], DataPtr[1], DataPtr[0], DataPtr[3]) / 255.0f;
 				}
@@ -496,8 +496,8 @@ namespace mu
 		{
 			switch(GetFormat())
 			{
-			case EImageFormat::IF_RGBA_UBYTE:
-			case EImageFormat::IF_BGRA_UBYTE:
+			case EImageFormat::RGBA_UByte:
+			case EImageFormat::BGRA_UByte:
 			{
 				const uint8* DataPtr = DataView.GetData() + 3;
 
@@ -514,7 +514,7 @@ namespace mu
 				break;
 			}
 
-			case EImageFormat::IF_RGB_UBYTE:
+			case EImageFormat::RGB_UByte:
 			{
 				return true;
 				break;
@@ -576,13 +576,13 @@ namespace mu
 
 				switch (Format)
 				{
-				case EImageFormat::IF_L_UBYTE:
+				case EImageFormat::L_UByte:
 				{
 					bIsBlack = DataPtr[Y*Size.X + X] == 0;
 					break;
 				}
 
-				case EImageFormat::IF_RGB_UBYTE:
+				case EImageFormat::RGB_UByte:
 				{
 					constexpr uint32 ZeroValue = 0;
 					bIsBlack = FMemory::Memcmp(&ZeroValue, DataPtr + (Y*Size.X + X)*3, 3) == 0;
@@ -590,8 +590,8 @@ namespace mu
 					break;
 				}
 
-				case EImageFormat::IF_RGBA_UBYTE:
-				case EImageFormat::IF_BGRA_UBYTE:
+				case EImageFormat::RGBA_UByte:
+				case EImageFormat::BGRA_UByte:
 				{
 					constexpr uint32 ZeroValue = 0;
 					bIsBlack = FMemory::Memcmp(&ZeroValue, DataPtr + (Y*Size.X + X)*4, 4) == 0;
@@ -657,7 +657,7 @@ namespace mu
 
   //      switch ( m_format )
   //      {
-  //      case IF_L_UBYTE:
+  //      case L_UByte:
   //      {
   //          size_t rowStride = sx;
 
@@ -711,9 +711,9 @@ namespace mu
   //          break;
   //      }
 
-  //      case IF_RGB_UBYTE:
-  //      case IF_RGBA_UBYTE:
-  //      case IF_BGRA_UBYTE:
+  //      case RGB_UByte:
+  //      case RGBA_UByte:
+  //      case BGRA_UByte:
   //      {
   //          size_t bytesPerPixel = GetImageFormatData( m_format ).BytesPerBlock;
   //          size_t rowStride = sx * bytesPerPixel;
@@ -805,7 +805,7 @@ namespace mu
 
 		switch (Format)
 		{
-		case EImageFormat::IF_RGB_UBYTE:
+		case EImageFormat::RGB_UByte:
 		{
 			constexpr int32 BatchSizeInElems = 1 << 14; 
 			constexpr int32 ElemSizeInBytes  = 3;
@@ -850,8 +850,8 @@ namespace mu
 			break;
 		}
 
-		case EImageFormat::IF_RGBA_UBYTE:
-		case EImageFormat::IF_BGRA_UBYTE:
+		case EImageFormat::RGBA_UByte:
+		case EImageFormat::BGRA_UByte:
 		{
 			constexpr int32 BatchSizeInElems = 1 << 14; 
 			constexpr int32 ElemSizeInBytes  = 4;
@@ -862,7 +862,7 @@ namespace mu
 			uint32 B = uint32(FMath::Clamp(255.0f * Color[2], 0.0f, 255.0f));
 			uint32 A = uint32(FMath::Clamp(255.0f * Color[3], 0.0f, 255.0f));
 		
-			const uint32 PixelData = Format == EImageFormat::IF_RGBA_UBYTE
+			const uint32 PixelData = Format == EImageFormat::RGBA_UByte
 					? R | (G << 8) | (B << 16) | (A << 24)
 					: B | (G << 8) | (R << 16) | (A << 24);
 
@@ -901,7 +901,7 @@ namespace mu
 			break;
 		}
 
-		case EImageFormat::IF_L_UBYTE:
+		case EImageFormat::L_UByte:
 		{
 			constexpr int32 BatchSizeInElems = 1 << 14; 
 			constexpr int32 ElemSizeInBytes  = 1;
@@ -933,7 +933,7 @@ namespace mu
 			// Generic case that supports compressed formats.
 			const FImageFormatData& FormatData = GetImageFormatData(Format);
 			Ptr<Image> BlockImage = CreateImage(
-						FormatData.PixelsPerBlockX, FormatData.PixelsPerBlockY, 1, EImageFormat::IF_RGBA_UBYTE, EInitializationType::NotInitialized);
+						FormatData.PixelsPerBlockX, FormatData.PixelsPerBlockY, 1, EImageFormat::RGBA_UByte, EInitializationType::NotInitialized);
 
 			const uint32 PixelData = 
 				  (uint32(FMath::Clamp(255.0f * Color[0], 0.0f, 255.0f)) << 0 ) |
diff --git a/Engine/Plugins/Mutable/Source/MutableRuntime/Private/MuR/ImageDataStorage.cpp b/Engine/Plugins/Mutable/Source/MutableRuntime/Private/MuR/ImageDataStorage.cpp
--- a/Engine/Plugins/Mutable/Source/MutableRuntime/Private/MuR/ImageDataStorage.cpp
+++ b/Engine/Plugins/Mutable/Source/MutableRuntime/Private/MuR/ImageDataStorage.cpp
@@ -148,7 +148,7 @@ namespace mu
 	void FImageDataStorage::CopyFrom(const FImageDataStorage& Other)
 	{
 		// Copy images that have different format is not allowed.
-		check(ImageFormat == EImageFormat::IF_NONE || Other.ImageFormat == ImageFormat);
+		check(ImageFormat == EImageFormat::None || Other.ImageFormat == ImageFormat);
 
 		ImageSize = Other.ImageSize;
 		ImageFormat = Other.ImageFormat;
diff --git a/Engine/Plugins/Mutable/Source/MutableRuntime/Private/MuR/ImageTypes.cpp b/Engine/Plugins/Mutable/Source/MutableRuntime/Private/MuR/ImageTypes.cpp
--- a/Engine/Plugins/Mutable/Source/MutableRuntime/Private/MuR/ImageTypes.cpp
+++ b/Engine/Plugins/Mutable/Source/MutableRuntime/Private/MuR/ImageTypes.cpp
@@ -13,11 +13,11 @@ namespace mu
 	MUTABLE_IMPLEMENT_ENUM_SERIALISABLE(EMinFilterMethod);
 	MUTABLE_IMPLEMENT_ENUM_SERIALISABLE(EImageFormat);
 
-    static FImageFormatData s_imageFormatData[uint32(EImageFormat::IF_COUNT)] =
+    static FImageFormatData s_imageFormatData[uint32(EImageFormat::Count)] =
     {
-		FImageFormatData(0, 0, 0, 0),	// IF_NONE
-		FImageFormatData(1, 1, 3, 3),	// IF_RGB_UBYTE
-		FImageFormatData(1, 1, 4, 4),	// IF_RGBA_UBYTE
+		FImageFormatData(0, 0, 0, 0),	// None
+		FImageFormatData(1, 1, 3, 3),	// RGB_UByte
+		FImageFormatData(1, 1, 4, 4),	// RGBA_UByte
 		FImageFormatData(1, 1, 1, 1),	// IF_U_UBYTE
 
 		FImageFormatData(0, 0, 0, 0),	// IF_PVRTC2 (deprecated)
@@ -25,42 +25,42 @@ namespace mu
         FImageFormatData(0, 0, 0, 0),	// IF_ETC1 (deprecated)
         FImageFormatData(0, 0, 0, 0),	// IF_ETC2 (deprecated)
 
-        FImageFormatData(0, 0, 0, 1),	// IF_L_UBYTE_RLE
-        FImageFormatData(0, 0, 0, 3),	// IF_RGB_UBYTE_RLE
-        FImageFormatData(0, 0, 0, 4),	// IF_RGBA_UBYTE_RLE
-        FImageFormatData(0, 0, 0, 1),	// IF_L_UBIT_RLE
+        FImageFormatData(0, 0, 0, 1),	// L_UByteRLE
+        FImageFormatData(0, 0, 0, 3),	// RGB_UByteRLE
+        FImageFormatData(0, 0, 0, 4),	// RGBA_UByteRLE
+        FImageFormatData(0, 0, 0, 1),	// L_UBitRLE
 
-        FImageFormatData(4, 4, 8,  4),	// IF_BC1
-        FImageFormatData(4, 4, 16, 4),	// IF_BC2
-        FImageFormatData(4, 4, 16, 4),	// IF_BC3
-        FImageFormatData(4, 4, 8,  1),	// IF_BC4
-        FImageFormatData(4, 4, 16, 2),	// IF_BC5
-        FImageFormatData(4, 4, 16, 3),	// IF_BC6
-        FImageFormatData(4, 4, 16, 4),	// IF_BC7
+        FImageFormatData(4, 4, 8,  4),	// BC1
+        FImageFormatData(4, 4, 16, 4),	// BC2
+        FImageFormatData(4, 4, 16, 4),	// BC3
+        FImageFormatData(4, 4, 8,  1),	// BC4
+        FImageFormatData(4, 4, 16, 2),	// BC5
+        FImageFormatData(4, 4, 16, 3),	// BC6
+        FImageFormatData(4, 4, 16, 4),	// BC7
 
-        FImageFormatData(1, 1, 4, 4 ),		// IF_BGRA_UBYTE
+        FImageFormatData(1, 1, 4, 4 ),		// BGRA_UByte
 
-		FImageFormatData(4, 4, 16, 3, {252, 253, 255, 255, 255, 255, 255, 255, 0, 0, 0, 0, 0, 0, 255, 255}), // IF_ASTC_4x4_RGB_LDR
-		FImageFormatData(4, 4, 16, 4, {252, 253, 255, 255, 255, 255, 255, 255, 0, 0, 0, 0, 0, 0, 0, 0}),     // IF_ASTC_4x4_RGBA_LDR
-        FImageFormatData(4, 4, 16, 2),	// IF_ASTC_4x4_RG_LDR // TODO: check black block for RG.
+		FImageFormatData(4, 4, 16, 3, {252, 253, 255, 255, 255, 255, 255, 255, 0, 0, 0, 0, 0, 0, 255, 255}), // ASTC_4x4_RGB_LDR
+		FImageFormatData(4, 4, 16, 4, {252, 253, 255, 255, 255, 255, 255, 255, 0, 0, 0, 0, 0, 0, 0, 0}),     // ASTC_4x4_RGBA_LDR
+        FImageFormatData(4, 4, 16, 2),	// ASTC_4x4_RG_LDR // TODO: check black block for RG.
 		
-		FImageFormatData(8, 8, 16, 3, {252, 253, 255, 255, 255, 255, 255, 255, 0, 0, 0, 0, 0, 0, 255, 255}),	// IF_ASTC_8x8_RGB_LDR,
-		FImageFormatData(8, 8, 16, 4, {252, 253, 255, 255, 255, 255, 255, 255, 0, 0, 0, 0, 0, 0, 0, 0}),		// IF_ASTC_8x8_RGBA_LDR,
-		FImageFormatData(8, 8, 16, 2),		// IF_ASTC_8x8_RG_LDR,
-		FImageFormatData(12, 12, 16, 3, {252, 253, 255, 255, 255, 255, 255, 255, 0, 0, 0, 0, 0, 0, 255, 255}),	// IF_ASTC_12x12_RGB_LDR
-		FImageFormatData(12, 12, 16, 4, {252, 253, 255, 255, 255, 255, 255, 255, 0, 0, 0, 0, 0, 0, 0, 0}),		// IF_ASTC_12x12_RGBA_LDR
-		FImageFormatData(12, 12, 16, 2),	// IF_ASTC_12x12_RG_LDR
-		FImageFormatData(6, 6, 16, 3, {252, 253, 255, 255, 255, 255, 255, 255, 0, 0, 0, 0, 0, 0, 255, 255}),	// IF_ASTC_6x6_RGB_LDR,
-		FImageFormatData(6, 6, 16, 4, {252, 253, 255, 255, 255, 255, 255, 255, 0, 0, 0, 0, 0, 0, 0, 0}),		// IF_ASTC_6x6_RGBA_LDR,
-		FImageFormatData(6, 6, 16, 2),		// IF_ASTC_6x6_RG_LDR,
-		FImageFormatData(10, 10, 16, 3, {252, 253, 255, 255, 255, 255, 255, 255, 0, 0, 0, 0, 0, 0, 255, 255}),	// IF_ASTC_6x6_RGB_LDR,
-		FImageFormatData(10, 10, 16, 4, {252, 253, 255, 255, 255, 255, 255, 255, 0, 0, 0, 0, 0, 0, 0, 0}),		// IF_ASTC_6x6_RGBA_LDR,
-		FImageFormatData(10, 10, 16, 2),	// IF_ASTC_6x6_RG_LDR,
+		FImageFormatData(8, 8, 16, 3, {252, 253, 255, 255, 255, 255, 255, 255, 0, 0, 0, 0, 0, 0, 255, 255}),	// ASTC_8x8_RGB_LDR,
+		FImageFormatData(8, 8, 16, 4, {252, 253, 255, 255, 255, 255, 255, 255, 0, 0, 0, 0, 0, 0, 0, 0}),		// ASTC_8x8_RGBA_LDR,
+		FImageFormatData(8, 8, 16, 2),		// ASTC_8x8_RG_LDR,
+		FImageFormatData(12, 12, 16, 3, {252, 253, 255, 255, 255, 255, 255, 255, 0, 0, 0, 0, 0, 0, 255, 255}),	// ASTC_12x12_RGB_LDR
+		FImageFormatData(12, 12, 16, 4, {252, 253, 255, 255, 255, 255, 255, 255, 0, 0, 0, 0, 0, 0, 0, 0}),		// ASTC_12x12_RGBA_LDR
+		FImageFormatData(12, 12, 16, 2),	// ASTC_12x12_RG_LDR
+		FImageFormatData(6, 6, 16, 3, {252, 253, 255, 255, 255, 255, 255, 255, 0, 0, 0, 0, 0, 0, 255, 255}),	// ASTC_6x6_RGB_LDR,
+		FImageFormatData(6, 6, 16, 4, {252, 253, 255, 255, 255, 255, 255, 255, 0, 0, 0, 0, 0, 0, 0, 0}),		// ASTC_6x6_RGBA_LDR,
+		FImageFormatData(6, 6, 16, 2),		// ASTC_6x6_RG_LDR,
+		FImageFormatData(10, 10, 16, 3, {252, 253, 255, 255, 255, 255, 255, 255, 0, 0, 0, 0, 0, 0, 255, 255}),	// ASTC_6x6_RGB_LDR,
+		FImageFormatData(10, 10, 16, 4, {252, 253, 255, 255, 255, 255, 255, 255, 0, 0, 0, 0, 0, 0, 0, 0}),		// ASTC_6x6_RGBA_LDR,
+		FImageFormatData(10, 10, 16, 2),	// ASTC_6x6_RG_LDR,
     };
 
     const FImageFormatData& GetImageFormatData( EImageFormat format )
     {
-        check( format < EImageFormat::IF_COUNT );
+        check( format < EImageFormat::Count );
         return s_imageFormatData[ uint8(format) ];
     }
 
diff --git a/Engine/Plugins/Mutable/Source/MutableRuntime/Private/MuR/Mesh.cpp b/Engine/Plugins/Mutable/Source/MutableRuntime/Private/MuR/Mesh.cpp
--- a/Engine/Plugins/Mutable/Source/MutableRuntime/Private/MuR/Mesh.cpp
+++ b/Engine/Plugins/Mutable/Source/MutableRuntime/Private/MuR/Mesh.cpp
@@ -326,7 +326,7 @@ bool Mesh::AreVertexIdsImplicit() const
 	// Is there a buffer for vertex ids?
 	int32 BufferIndex = -1;
 	int32 ChannelIndex = -1;
-	VertexBuffers.FindChannel(MBS_VERTEXINDEX, 0, &BufferIndex, &ChannelIndex);
+	VertexBuffers.FindChannel(EMeshBufferSemantic::VertexIndex, 0, &BufferIndex, &ChannelIndex);
 
 	return (MeshIDPrefix != 0) && (BufferIndex < 0) && (ChannelIndex < 0);
 }
@@ -337,11 +337,11 @@ bool Mesh::AreVertexIdsExplicit() const
 	// Is there a buffer for vertex ids?
 	int32 BufferIndex = -1;
 	int32 ChannelIndex = -1;
-	VertexBuffers.FindChannel(MBS_VERTEXINDEX, 0, &BufferIndex, &ChannelIndex);
+	VertexBuffers.FindChannel(EMeshBufferSemantic::VertexIndex, 0, &BufferIndex, &ChannelIndex);
 
 	bool bExplicit = 
 			(BufferIndex >= 0) && (ChannelIndex >= 0) && 
-			(VertexBuffers.Buffers[BufferIndex].Channels[ChannelIndex].Format == MBF_UINT64);
+			(VertexBuffers.Buffers[BufferIndex].Channels[ChannelIndex].Format == EMeshBufferFormat::UInt64);
 
 	if (bExplicit)
 	{
@@ -358,9 +358,9 @@ void Mesh::MakeVertexIdsRelative()
 
 	int32 NewBuffer = VertexBuffers.GetBufferCount();
 	VertexBuffers.SetBufferCount(NewBuffer + 1);
-	EMeshBufferSemantic Semantic = MBS_VERTEXINDEX;
+	EMeshBufferSemantic Semantic = EMeshBufferSemantic::VertexIndex;
 	int32 SemanticIndex = 0;
-	EMeshBufferFormat Format = MBF_UINT32;
+	EMeshBufferFormat Format = EMeshBufferFormat::UInt32;
 	int32 Components = 1;
 	int32 Offset = 0;
 	VertexBuffers.SetBuffer(NewBuffer, sizeof(uint32), 1, &Semantic, &SemanticIndex, &Format, &Components, &Offset);
@@ -387,14 +387,14 @@ void Mesh::MakeIdsExplicit()
 
 		int32 OldBufferIndex = -1;
 		int32 OldChannelIndex = -1;
-		VertexBuffers.FindChannel(MBS_VERTEXINDEX, 0, &OldBufferIndex, &OldChannelIndex);
+		VertexBuffers.FindChannel(EMeshBufferSemantic::VertexIndex, 0, &OldBufferIndex, &OldChannelIndex);
 		bool bHasVertexIndices = (OldBufferIndex >= 0 && OldChannelIndex >= 0);
 		if (bHasVertexIndices)
 		{
 			check(OldChannelIndex == 0 && VertexBuffers.Buffers[OldBufferIndex].Channels.Num() == 1);
 
 			FMeshBuffer& Buffer = VertexBuffers.Buffers[OldBufferIndex];
-			Buffer.Channels[0].Format = MBF_UINT64;
+			Buffer.Channels[0].Format = EMeshBufferFormat::UInt64;
 			Buffer.ElementSize = sizeof(uint64);
 		}
 
@@ -405,9 +405,9 @@ void Mesh::MakeIdsExplicit()
 			FMeshBuffer& Buffer = VertexBuffers.Buffers.Emplace_GetRef();
 			Buffer.ElementSize = sizeof(uint64);
 			FMeshBufferChannel& Channel = Buffer.Channels.Emplace_GetRef();
-			Channel.Semantic = MBS_VERTEXINDEX;
+			Channel.Semantic = EMeshBufferSemantic::VertexIndex;
 			Channel.SemanticIndex= 0;
-			Channel.Format = MBF_UINT64;
+			Channel.Format = EMeshBufferFormat::UInt64;
 			Channel.ComponentCount = 1;
 			Channel.Offset = 0;
 		}
@@ -419,7 +419,7 @@ void Mesh::MakeIdsExplicit()
 		{
 			for (FMeshBufferChannel& Channel : Buffer.Channels)
 			{
-				if (Channel.Semantic != MBS_LAYOUTBLOCK)
+				if (Channel.Semantic != EMeshBufferSemantic::LayoutBlock)
 				{
 					continue;
 				}
@@ -427,7 +427,7 @@ void Mesh::MakeIdsExplicit()
 				check(Buffer.Channels.Num() == 1);
 				check(Buffer.Channels[0].Offset == 0);
 
-				Buffer.Channels[0].Format = MBF_UINT64;
+				Buffer.Channels[0].Format = EMeshBufferFormat::UInt64;
 				Buffer.ElementSize = sizeof(uint64);
 			}
 		}
@@ -802,7 +802,7 @@ UE::Math::TIntVector3<uint32> Mesh::GetFaceVertexIndices(int32 FaceIndex) const
 {
 	UE::Math::TIntVector3<uint32> Result;
 
-    MeshBufferIteratorConst<MBF_UINT32, uint32, 1> Iter(IndexBuffers, MBS_VERTEXINDEX);
+    MeshBufferIteratorConst<EMeshBufferFormat::UInt32, uint32, 1> Iter(IndexBuffers, EMeshBufferSemantic::VertexIndex);
     Iter += FaceIndex*3;
 
     Result[0] = (*Iter)[0];
@@ -857,8 +857,8 @@ void Mesh::GetVertexMap(const Mesh& Other, FVertexMatchMap& VertexMap, float Tol
     }
 
 
-    MeshBufferIteratorConst<MBF_FLOAT32, float, 3> ItPosition(VertexBuffers, MBS_POSITION);
-    MeshBufferIteratorConst<MBF_FLOAT32, float, 3> ItOtherPositionBegin(Other.VertexBuffers, MBS_POSITION);
+    MeshBufferIteratorConst<EMeshBufferFormat::Float32, float, 3> ItPosition(VertexBuffers, EMeshBufferSemantic::Position);
+    MeshBufferIteratorConst<EMeshBufferFormat::Float32, float, 3> ItOtherPositionBegin(Other.VertexBuffers, EMeshBufferSemantic::Position);
 
 
     // Bucket the other mesh
@@ -867,7 +867,7 @@ void Mesh::GetVertexMap(const Mesh& Other, FVertexMatchMap& VertexMap, float Tol
 
     float RangeMin = TNumericLimits<float>::Max();
     float RangeMax = -TNumericLimits<float>::Max();
-    MeshBufferIteratorConst< MBF_FLOAT32, float, 3 >  ItOtherPosition = ItOtherPositionBegin;
+    MeshBufferIteratorConst< EMeshBufferFormat::Float32, float, 3 >  ItOtherPosition = ItOtherPositionBegin;
     
 	for (int32 OtherVertex = 0; OtherVertex < OtherVertexCount; ++OtherVertex)
     {
@@ -1144,15 +1144,15 @@ void Mesh::CheckIntegrity() const
 	{
 		int32 BufferIndex = -1;
 		int32 ChannelIndex = -1;
-		VertexBuffers.FindChannel(MBS_VERTEXINDEX, 0, &BufferIndex, &ChannelIndex);
+		VertexBuffers.FindChannel(EMeshBufferSemantic::VertexIndex, 0, &BufferIndex, &ChannelIndex);
 		if (BufferIndex >= 0 && ChannelIndex >= 0)
 		{
 			EMeshBufferFormat IdFormat = VertexBuffers.m_buffers[BufferIndex].m_channels[ChannelIndex].m_format;
-			if (IdFormat == MBF_UINT64)
+			if (IdFormat == EMeshBufferFormat::UInt64)
 			{
 				check(MeshIDPrefix==0);
 			}
-			else if (IdFormat == MBF_UINT32)
+			else if (IdFormat == EMeshBufferFormat::UInt32)
 			{
 				check(MeshIDPrefix != 0);
 			}
@@ -1178,7 +1178,7 @@ void Mesh::CheckIntegrity() const
                 int32 offset = 0;
                 IndexBuffers.GetChannel( b, c, &semantic, &semanticIndex, &format, &components, &offset );
 
-                if ( semantic==MBS_VERTEXINDEX )
+                if ( semantic==EMeshBufferSemantic::VertexIndex )
                 {
                     int32 icount = IndexBuffers.GetElementCount();
                     int32 elemCount = VertexBuffers.GetElementCount();
@@ -1188,19 +1188,19 @@ void Mesh::CheckIntegrity() const
 
                         switch (format)
                         {
-                        case MBF_UINT32:
+                        case EMeshBufferFormat::UInt32:
                         {
                             uint32 index = *(const uint32*)pData;
                             check( index < uint32( elemCount ) );
                             break;
                         }
-                        case MBF_UINT16:
+                        case EMeshBufferFormat::UInt16:
                         {
                             uint16 index = *(const uint16*)pData;
                             check( index < uint16( elemCount ) );
                             break;
                         }
-                        case MBF_UINT8:
+                        case EMeshBufferFormat::UInt8:
                         {
                             uint8 index = *(const uint8*)pData;
                             check( index < uint8( elemCount ) );
@@ -1236,13 +1236,13 @@ void Mesh::CheckIntegrity() const
 
                 // If it is not one of the relevant semantics
                 if (
-                        //semantic!=MBS_POSITION &&
-                        //semantic!=MBS_TEXCOORDS &&
-                        //semantic!=MBS_NORMAL &&
-                        //semantic!=MBS_TANGENT &&
-                        //semantic!=MBS_BINORMAL &&
-                        semantic!=MBS_BONEINDICES
-                        // && semantic!=MBS_BONEWEIGHTS
+                        //semantic!=EMeshBufferSemantic::Position &&
+                        //semantic!=EMeshBufferSemantic::TexCoords &&
+                        //semantic!=EMeshBufferSemantic::Normal &&
+                        //semantic!=EMeshBufferSemantic::Tangent &&
+                        //semantic!=EMeshBufferSemantic::Binormal &&
+                        semantic!=EMeshBufferSemantic::BoneIndices
+                        // && semantic!=EMeshBufferSemantic::BoneWeights
                         )
                 {
                     continue;
@@ -1257,7 +1257,7 @@ void Mesh::CheckIntegrity() const
 
                     switch (format)
                     {
-                    case MBF_UINT8:
+                    case EMeshBufferFormat::UInt8:
                     {
                         for ( int32 d = 0; d < components; ++d )
                         {
@@ -1268,7 +1268,7 @@ void Mesh::CheckIntegrity() const
                         break;
                     }
 
-                    case MBF_UINT16:
+                    case EMeshBufferFormat::UInt16:
                     {
                         for ( int32 d = 0; d < components; ++d )
                         {
@@ -1279,7 +1279,7 @@ void Mesh::CheckIntegrity() const
                         break;
                     }
 
-                    case MBF_UINT32:
+                    case EMeshBufferFormat::UInt32:
                     {
                         for ( int32 d = 0; d < components; ++d )
                         {
@@ -1319,8 +1319,8 @@ static bool StaticMeshFormatIdentify_Project(const Mesh* InMesh)
     {
         const FMeshBufferChannel& Channel = InMesh->VertexBuffers.Buffers[0].Channels[0];
 
-        bResult &= Channel.Semantic == MBS_TEXCOORDS;
-        bResult &= Channel.Format == MBF_FLOAT32;
+        bResult &= Channel.Semantic == EMeshBufferSemantic::TexCoords;
+        bResult &= Channel.Format == EMeshBufferFormat::Float32;
         bResult &= Channel.ComponentCount == 2;
         //we don't really care about the semantic index
         //bResult &= Channel.SemanticIndex == 0;
@@ -1331,8 +1331,8 @@ static bool StaticMeshFormatIdentify_Project(const Mesh* InMesh)
     {
         const FMeshBufferChannel& Channel = InMesh->VertexBuffers.Buffers[0].Channels[1];
 
-        bResult &= Channel.Semantic == MBS_POSITION;
-        bResult &= Channel.Format == MBF_FLOAT32;
+        bResult &= Channel.Semantic == EMeshBufferSemantic::Position;
+        bResult &= Channel.Format == EMeshBufferFormat::Float32;
         bResult &= Channel.ComponentCount == 3;
         bResult &= Channel.SemanticIndex == 0;
         bResult &= Channel.Offset == 8;
@@ -1342,8 +1342,8 @@ static bool StaticMeshFormatIdentify_Project(const Mesh* InMesh)
     {
         const FMeshBufferChannel& Channel = InMesh->VertexBuffers.Buffers[0].Channels[2];
 
-        bResult &= Channel.Semantic == MBS_NORMAL;
-        bResult &= Channel.Format == MBF_FLOAT32;
+        bResult &= Channel.Semantic == EMeshBufferSemantic::Normal;
+        bResult &= Channel.Format == EMeshBufferFormat::Float32;
         bResult &= Channel.ComponentCount == 3;
         bResult &= Channel.SemanticIndex == 0;
         bResult &= Channel.Offset == 20;
@@ -1359,8 +1359,8 @@ static bool StaticMeshFormatIdentify_Project(const Mesh* InMesh)
     {
         const FMeshBufferChannel& Channel = InMesh->IndexBuffers.Buffers[0].Channels[0];
 
-        bResult &= Channel.Semantic == MBS_VERTEXINDEX;
-        bResult &= Channel.Format == MBF_UINT32;
+        bResult &= Channel.Semantic == EMeshBufferSemantic::VertexIndex;
+        bResult &= Channel.Format == EMeshBufferFormat::UInt32;
         bResult &= Channel.ComponentCount == 1;
         bResult &= Channel.SemanticIndex == 0;
         bResult &= Channel.Offset == 0;
@@ -1388,8 +1388,8 @@ static bool StaticMeshFormatIdentify_ProjectWrapping(const Mesh* InMesh)
     {
         const FMeshBufferChannel& Channel = InMesh->VertexBuffers.Buffers[0].Channels[0];
 
-        bResult &= Channel.Semantic == MBS_TEXCOORDS;
-        bResult &= Channel.Format == MBF_FLOAT32;
+        bResult &= Channel.Semantic == EMeshBufferSemantic::TexCoords;
+        bResult &= Channel.Format == EMeshBufferFormat::Float32;
         bResult &= Channel.ComponentCount == 2;
         // we don't really care about the semantic index as long as there is only one
         // bResult &= Channel.SemanticIndex == 0;
@@ -1400,8 +1400,8 @@ static bool StaticMeshFormatIdentify_ProjectWrapping(const Mesh* InMesh)
     {
         const FMeshBufferChannel& Channel = InMesh->VertexBuffers.Buffers[0].Channels[1];
 
-        bResult &= Channel.Semantic == MBS_POSITION;
-        bResult &= Channel.Format == MBF_FLOAT32;
+        bResult &= Channel.Semantic == EMeshBufferSemantic::Position;
+        bResult &= Channel.Format == EMeshBufferFormat::Float32;
         bResult &= Channel.ComponentCount == 3;
         bResult &= Channel.SemanticIndex == 0;
         bResult &= Channel.Offset == 8;
@@ -1411,8 +1411,8 @@ static bool StaticMeshFormatIdentify_ProjectWrapping(const Mesh* InMesh)
     {
         const FMeshBufferChannel& Channel = InMesh->VertexBuffers.Buffers[0].Channels[2];
 
-        bResult &= Channel.Semantic == MBS_NORMAL;
-        bResult &= Channel.Format == MBF_FLOAT32;
+        bResult &= Channel.Semantic == EMeshBufferSemantic::Normal;
+        bResult &= Channel.Format == EMeshBufferFormat::Float32;
         bResult &= Channel.ComponentCount == 3;
         bResult &= Channel.SemanticIndex == 0;
         bResult &= Channel.Offset == 20;
@@ -1423,9 +1423,9 @@ static bool StaticMeshFormatIdentify_ProjectWrapping(const Mesh* InMesh)
     {
         const FMeshBufferChannel& Channel = InMesh->VertexBuffers.Buffers[1].Channels[0];
 
-        bResult &= Channel.Semantic == MBS_LAYOUTBLOCK;
+        bResult &= Channel.Semantic == EMeshBufferSemantic::LayoutBlock;
 		// We don't care about the layout block id format. We need to support them all.
-        //bResult &= Channel.Format == MBF_UINT64;
+        //bResult &= Channel.Format == EMeshBufferFormat::UInt64;
         bResult &= Channel.ComponentCount == 1;
         // we don't really care about the semantic index as long as there is only one
         // bResult &= Channel.SemanticIndex == 0;
@@ -1442,8 +1442,8 @@ static bool StaticMeshFormatIdentify_ProjectWrapping(const Mesh* InMesh)
     {
         const FMeshBufferChannel& Channel = InMesh->IndexBuffers.Buffers[0].Channels[0];
 
-        bResult &= Channel.Semantic == MBS_VERTEXINDEX;
-        bResult &= Channel.Format == MBF_UINT32;
+        bResult &= Channel.Semantic == EMeshBufferSemantic::VertexIndex;
+        bResult &= Channel.Format == EMeshBufferFormat::UInt32;
         bResult &= Channel.ComponentCount == 1;
         bResult &= Channel.SemanticIndex == 0;
         bResult &= Channel.Offset == 0;
@@ -1498,34 +1498,34 @@ namespace
                         
 						switch (Channel.Format)
                         {
-                        case MBF_UINT32:
-                        case MBF_NUINT32:
+                        case EMeshBufferFormat::UInt32:
+                        case EMeshBufferFormat::NUInt32:
 						{
 							Out += FString::Printf(TEXT("%d"), *(const uint32*)ChanDataPtr); 
 							ChanDataPtr += 4; 
 							break;
 						}
-                        case MBF_UINT16:
-                        case MBF_NUINT16: 
+                        case EMeshBufferFormat::UInt16:
+                        case EMeshBufferFormat::NUInt16: 
 						{
 							Out += FString::Printf(TEXT("%d"), *(const uint16*)ChanDataPtr); 
 							ChanDataPtr += 2; 
 							break;
 						}
-                        case MBF_UINT8:
-                        case MBF_NUINT8: 
+                        case EMeshBufferFormat::UInt8:
+                        case EMeshBufferFormat::NUInt8: 
 						{
 							Out += FString::Printf(TEXT("%d"), *(const uint8*)ChanDataPtr); 
 							ChanDataPtr += 1; 
 							break;
 						}
-						case MBF_FLOAT32: 
+						case EMeshBufferFormat::Float32: 
 						{
 							Out += FString::Printf(TEXT("%.3f"), *(const float*)ChanDataPtr); 
 							ChanDataPtr += 4; 
 							break;
 						}
-                        case MBF_FLOAT16: 
+                        case EMeshBufferFormat::Float16: 
 						{
 							Out += FString::Printf(TEXT("%d"), *(const uint16*)ChanDataPtr); 
 							ChanDataPtr += 2; 
diff --git a/Engine/Plugins/Mutable/Source/MutableRuntime/Private/MuR/MeshBufferSet.cpp b/Engine/Plugins/Mutable/Source/MutableRuntime/Private/MuR/MeshBufferSet.cpp
--- a/Engine/Plugins/Mutable/Source/MutableRuntime/Private/MuR/MeshBufferSet.cpp
+++ b/Engine/Plugins/Mutable/Source/MutableRuntime/Private/MuR/MeshBufferSet.cpp
@@ -19,7 +19,7 @@ namespace mu
 	MUTABLE_IMPLEMENT_ENUM_SERIALISABLE(EMeshBufferSemantic);
 
 	
-	static FMeshBufferFormatData s_meshBufferFormatData[] = // MBF_COUNT entries
+	static FMeshBufferFormatData s_meshBufferFormatData[] = // EMeshBufferFormat::Count entries
 	{
 		{ 0, 0 },
 		{ 2, 0 },
@@ -52,13 +52,13 @@ namespace mu
 		{ 8, 0 }
 	};
 
-	static_assert(sizeof(s_meshBufferFormatData) / sizeof(FMeshBufferFormatData) == int32(MBF_COUNT));
+	static_assert(sizeof(s_meshBufferFormatData) / sizeof(FMeshBufferFormatData) == int32(EMeshBufferFormat::Count));
 
 	const FMeshBufferFormatData& GetMeshFormatData(EMeshBufferFormat Format)
 	{
-		check(Format >= 0);
-		check(Format < MBF_COUNT);
-		return s_meshBufferFormatData[Format];
+		check(uint32(Format) >= 0);
+		check(uint32(Format) < uint32(EMeshBufferFormat::Count));
+		return s_meshBufferFormatData[uint32(Format)];
 	}
 
 	void FMeshBufferSet::Serialise(OutputArchive& Arch) const
@@ -195,9 +195,9 @@ namespace mu
 		for (int32 ChannelIndex = 0; ChannelIndex < ChannelCount; ++ChannelIndex)
 		{
 			FMeshBufferChannel& Channel = Buffer.Channels[ChannelIndex];
-			Channel.Semantic = SemanticsPtr ? SemanticsPtr[ChannelIndex] : MBS_NONE;
+			Channel.Semantic = SemanticsPtr ? SemanticsPtr[ChannelIndex] : EMeshBufferSemantic::None;
 			Channel.SemanticIndex = SemanticIndicesPtr ? SemanticIndicesPtr[ChannelIndex] : 0;
-			Channel.Format = FormatsPtr ? FormatsPtr[ChannelIndex] : MBF_NONE;
+			Channel.Format = FormatsPtr ? FormatsPtr[ChannelIndex] : EMeshBufferFormat::None;
 			Channel.ComponentCount = ComponentCountPtr ? ((uint16)ComponentCountPtr[ChannelIndex]) : 0;
 			Channel.Offset = OffsetsPtr ? ((uint8)OffsetsPtr[ChannelIndex]) : 0;
 
@@ -414,12 +414,12 @@ namespace mu
 
 	bool FMeshBufferSet::IsSpecialBufferToIgnoreInSimilar(const FMeshBuffer& Buffer) const
 	{
-		if (Buffer.Channels.Num() == 1 && Buffer.Channels[0].Semantic == MBS_VERTEXINDEX)
+		if (Buffer.Channels.Num() == 1 && Buffer.Channels[0].Semantic == EMeshBufferSemantic::VertexIndex)
 		{
 			return true;
 		}
 
-		if (Buffer.Channels.Num() == 1 && Buffer.Channels[0].Semantic == MBS_LAYOUTBLOCK)
+		if (Buffer.Channels.Num() == 1 && Buffer.Channels[0].Semantic == EMeshBufferSemantic::LayoutBlock)
 		{
 			return true;
 		}
@@ -472,7 +472,7 @@ namespace mu
 				return false;
 			}
 
-			bool bBufferHasTexCoords = ThisBuffer.HasSemantic(MBS_TEXCOORDS);
+			bool bBufferHasTexCoords = ThisBuffer.HasSemantic(EMeshBufferSemantic::TexCoords);
 			bool bCanBeFullCompared = !ThisBuffer.HasPadding() && (!bBufferHasTexCoords || bCompareUVs);
 			if (bCanBeFullCompared)
 			{
@@ -492,7 +492,7 @@ namespace mu
 				{
 					for (int32 C = 0; C < ThisNumChannels; ++C)
 					{
-						if (!bCompareUVs && ThisBuffer.Channels[C].Semantic == MBS_TEXCOORDS)
+						if (!bCompareUVs && ThisBuffer.Channels[C].Semantic == EMeshBufferSemantic::TexCoords)
 						{
 							continue;
 						}
@@ -602,13 +602,13 @@ namespace mu
 
 	void FMeshBufferSet::ResetBufferIndices()
 	{
-		int32 CurrentIndices[MBS_COUNT] = {0};
+		int32 CurrentIndices[int32(EMeshBufferSemantic::Count)] = {0};
 		for (FMeshBuffer& Buffer : Buffers)
 		{
 			for (FMeshBufferChannel& Channel : Buffer.Channels)
 			{
-				Channel.SemanticIndex = CurrentIndices[Channel.Semantic];
-				CurrentIndices[Channel.Semantic]++;
+				Channel.SemanticIndex = CurrentIndices[int32(Channel.Semantic)];
+				CurrentIndices[int32(Channel.Semantic)]++;
 			}
 		}
 	}
diff --git a/Engine/Plugins/Mutable/Source/MutableRuntime/Private/MuR/Model.cpp b/Engine/Plugins/Mutable/Source/MutableRuntime/Private/MuR/Model.cpp
--- a/Engine/Plugins/Mutable/Source/MutableRuntime/Private/MuR/Model.cpp
+++ b/Engine/Plugins/Mutable/Source/MutableRuntime/Private/MuR/Model.cpp
@@ -235,29 +235,26 @@ namespace mu
 	}
 
 
-    //---------------------------------------------------------------------------------------------
-    bool Model::HasExternalData() const
+	bool Model::HasExternalData() const
     {
 		return m_pD->m_program.m_roms.Num() > 0;
     }
 
 #if WITH_EDITOR
-	//---------------------------------------------------------------------------------------------
 	bool Model::IsValid() const
 	{
 		return m_pD->m_program.bIsValid;
 	}
 
 
-	//---------------------------------------------------------------------------------------------
 	void Model::Invalidate()
     {
 		m_pD->m_program.bIsValid = false;
     }
 #endif
 
-    //---------------------------------------------------------------------------------------------
-    void Model::UnloadExternalData()
+
+	void Model::UnloadExternalData()
     {
 		LLM_SCOPE_BYNAME(TEXT("MutableRuntime"));
 
@@ -279,7 +276,6 @@ namespace mu
 	}
 
 
-    //---------------------------------------------------------------------------------------------
 	TSharedPtr<Model> Model::StaticUnserialise( InputArchive& arch )
     {
 		LLM_SCOPE_BYNAME(TEXT("MutableRuntime"));
@@ -289,129 +285,109 @@ namespace mu
         return pResult;
     }
 
-    //---------------------------------------------------------------------------------------------
-    Model::Private* Model::GetPrivate() const
+
+	Model::Private* Model::GetPrivate() const
     {
         return m_pD;
     }
 
 
-    //---------------------------------------------------------------------------------------------
     bool Model::GetBoolDefaultValue(int32 Index) const
     {
     	check(m_pD->m_program.m_parameters.IsValidIndex(Index));
-		check(m_pD->m_program.m_parameters[Index].m_type == PARAMETER_TYPE::T_BOOL);
+		check(m_pD->m_program.m_parameters[Index].m_type == EParameterType::Bool);
 
         // Early out in case of invalid parameters
         if (!m_pD->m_program.m_parameters.IsValidIndex(Index) ||
-            m_pD->m_program.m_parameters[Index].m_type != PARAMETER_TYPE::T_BOOL)
+            m_pD->m_program.m_parameters[Index].m_type != EParameterType::Bool)
         {
             return false;
         }
 		
-        return m_pD->m_program.m_parameters[Index].m_defaultValue.Get<ParamBoolType>();
+        return m_pD->m_program.m_parameters[Index].m_defaultValue.Get<FParamBoolType>();
     }
 
 
-    //---------------------------------------------------------------------------------------------
 	int32 Model::GetIntDefaultValue(int32 Index) const
 	{
 		check(m_pD->m_program.m_parameters.IsValidIndex(Index));
-		check(m_pD->m_program.m_parameters[Index].m_type == PARAMETER_TYPE::T_INT);
+		check(m_pD->m_program.m_parameters[Index].m_type == EParameterType::Int);
 
         // Early out in case of invalid parameters
         if (!m_pD->m_program.m_parameters.IsValidIndex(Index) ||
-            m_pD->m_program.m_parameters[Index].m_type != PARAMETER_TYPE::T_INT)
+            m_pD->m_program.m_parameters[Index].m_type != EParameterType::Int)
         {
             return 0;
         }
 		
-        return m_pD->m_program.m_parameters[Index].m_defaultValue.Get<ParamIntType>();
+        return m_pD->m_program.m_parameters[Index].m_defaultValue.Get<FParamIntType>();
 	}
 
 
-    //---------------------------------------------------------------------------------------------
 	float Model::GetFloatDefaultValue(int32 Index) const
 	{
     	check(m_pD->m_program.m_parameters.IsValidIndex(Index));
-		check(m_pD->m_program.m_parameters[Index].m_type == PARAMETER_TYPE::T_FLOAT);
+		check(m_pD->m_program.m_parameters[Index].m_type == EParameterType::Float);
 
         // Early out in case of invalid parameters
         if (!m_pD->m_program.m_parameters.IsValidIndex(Index) ||
-            m_pD->m_program.m_parameters[Index].m_type != PARAMETER_TYPE::T_FLOAT)
+            m_pD->m_program.m_parameters[Index].m_type != EParameterType::Float)
         {
             return 0.0f;
         }
 		
-        return m_pD->m_program.m_parameters[Index].m_defaultValue.Get<ParamFloatType>();
+        return m_pD->m_program.m_parameters[Index].m_defaultValue.Get<FParamFloatType>();
 	}
 
 
-	void Model::GetColourDefaultValue(int32 Index, float* R, float* G, float* B, float* A) const
+	void Model::GetColourDefaultValue(int32 Index, FVector4f& OutValue) const
     {
     	check(m_pD->m_program.m_parameters.IsValidIndex(Index));
-		check(m_pD->m_program.m_parameters[Index].m_type == PARAMETER_TYPE::T_COLOUR);
+		check(m_pD->m_program.m_parameters[Index].m_type == EParameterType::Color);
 
         // Early out in case of invalid parameters
         if (!m_pD->m_program.m_parameters.IsValidIndex(Index) ||
-            m_pD->m_program.m_parameters[Index].m_type != PARAMETER_TYPE::T_COLOUR)
+            m_pD->m_program.m_parameters[Index].m_type != EParameterType::Color)
         {
             return;
         }
 
-        ParamColorType& Color = m_pD->m_program.m_parameters[Index].m_defaultValue.Get<ParamColorType>();
+        FParamColorType& Color = m_pD->m_program.m_parameters[Index].m_defaultValue.Get<FParamColorType>();
 		
-		if (R)
-		{
-			*R = Color[0];
-		}
-
-		if (G)
-		{
-			*G = Color[1];
-		}
-
-		if (B)
-		{
-			*B = Color[2];
-		}
-
-		if (A)
-		{
-			*A = Color[3];
-		}
+		OutValue = Color;
     }
 
+
 	FMatrix44f Model::GetMatrixDefaultValue(int32 Index) const
 	{
     	check(m_pD->m_program.m_parameters.IsValidIndex(Index));
-    	check(m_pD->m_program.m_parameters[Index].m_type == PARAMETER_TYPE::T_MATRIX);
+    	check(m_pD->m_program.m_parameters[Index].m_type == EParameterType::Matrix);
 
     	// Early out in case of invalid parameters
     	if (!m_pD->m_program.m_parameters.IsValidIndex(Index) ||
-			m_pD->m_program.m_parameters[Index].m_type != PARAMETER_TYPE::T_MATRIX)
+			m_pD->m_program.m_parameters[Index].m_type != EParameterType::Matrix)
     	{
     		return FMatrix44f::Identity;
     	}
 
-    	return m_pD->m_program.m_parameters[Index].m_defaultValue.Get<ParamMatrixType>();
+    	return m_pD->m_program.m_parameters[Index].m_defaultValue.Get<FParamMatrixType>();
 	}
 
 
-	void Model::GetProjectorDefaultValue(int32 Index, PROJECTOR_TYPE* OutProjectionType, FVector3f* OutPos,
+	void Model::GetProjectorDefaultValue(int32 Index, EProjectorType* OutProjectionType, FVector3f* OutPos,
 	                                     FVector3f* OutDir, FVector3f* OutUp, FVector3f* OutScale, float* OutProjectionAngle) const
 	{
     	check(m_pD->m_program.m_parameters.IsValidIndex(Index));
-		check(m_pD->m_program.m_parameters[Index].m_type == PARAMETER_TYPE::T_PROJECTOR);
+		check(m_pD->m_program.m_parameters[Index].m_type == EParameterType::Projector);
 
         // Early out in case of invalid parameters
         if (!m_pD->m_program.m_parameters.IsValidIndex(Index) ||
-            m_pD->m_program.m_parameters[Index].m_type != PARAMETER_TYPE::T_PROJECTOR)
+            m_pD->m_program.m_parameters[Index].m_type != EParameterType::Projector)
         {
             return;
         }
 
-        const FProjector& Projector = m_pD->m_program.m_parameters[Index].m_defaultValue.Get<ParamProjectorType>();
+        const FProjector& Projector = m_pD->m_program.m_parameters[Index].m_defaultValue.Get<FParamProjectorType>();
         if (OutProjectionType) *OutProjectionType = Projector.type;
     	if (OutPos) *OutPos = Projector.position;
 		if (OutDir) *OutDir = Projector.direction;
@@ -421,24 +397,22 @@ namespace mu
 	}
 
 
-	//---------------------------------------------------------------------------------------------
 	FName Model::GetImageDefaultValue(int32 Index) const
     {
 	    check(m_pD->m_program.m_parameters.IsValidIndex(Index));
-		check(m_pD->m_program.m_parameters[Index].m_type == PARAMETER_TYPE::T_IMAGE);
+		check(m_pD->m_program.m_parameters[Index].m_type == EParameterType::Image);
 
         // Early out in case of invalid parameters
         if (!m_pD->m_program.m_parameters.IsValidIndex(Index) ||
-            m_pD->m_program.m_parameters[Index].m_type != PARAMETER_TYPE::T_IMAGE)
+            m_pD->m_program.m_parameters[Index].m_type != EParameterType::Image)
         {
             return {};
         }
 		
-        return m_pD->m_program.m_parameters[Index].m_defaultValue.Get<ParamImageType>();
+        return m_pD->m_program.m_parameters[Index].m_defaultValue.Get<FParamImageType>();
     }
 
 
-    //---------------------------------------------------------------------------------------------
     int32 Model::GetRomCount() const
     {
     	return m_pD->m_program.m_roms.Num();
@@ -472,7 +446,6 @@ namespace mu
 	}
 
 
-    //---------------------------------------------------------------------------------------------
     ParametersPtr Model::NewParameters(TSharedPtr<const Model> Model, const Parameters* pOld )
     {
 		LLM_SCOPE_BYNAME(TEXT("MutableRuntime"));
@@ -491,9 +464,9 @@ namespace mu
         // Copy old values
         if ( pOld )
         {
-            for ( int p=0; p<pOld->GetCount(); ++p )
+            for ( int32 p=0; p<pOld->GetCount(); ++p )
             {
-                int thisP = pRes->GetPrivate()->Find( pOld->GetName(p) );
+                int32 thisP = pRes->GetPrivate()->Find( pOld->GetName(p) );
 
                 if ( thisP>=0 )
                 {
@@ -501,35 +474,35 @@ namespace mu
                     {
                         switch ( pRes->GetType(thisP) )
                         {
-						case PARAMETER_TYPE::T_BOOL:
+						case EParameterType::Bool:
                             pRes->SetBoolValue( thisP, pOld->GetBoolValue(p) );
                             break;
 
-                        case PARAMETER_TYPE::T_INT:
+                        case EParameterType::Int:
                             pRes->SetIntValue( thisP, pOld->GetIntValue(p) );
                             break;
 
-                        case PARAMETER_TYPE::T_FLOAT:
+                        case EParameterType::Float:
                             pRes->SetFloatValue( thisP, pOld->GetFloatValue(p) );
                             break;
 
-                        case PARAMETER_TYPE::T_COLOUR:
+                        case EParameterType::Color:
                         {
-                            float R, G, B, A;
-                            pOld->GetColourValue( p, &R, &G, &B, &A );
-                            pRes->SetColourValue( thisP, R, G, B, A );
+                            FVector4f V;
+                            pOld->GetColourValue( p, V );
+                            pRes->SetColourValue( thisP, V );
                             break;
                         }
 
-                        case PARAMETER_TYPE::T_PROJECTOR:
+                        case EParameterType::Projector:
                         {
 //							float m[16];
 //							pOld->GetProjectorValue( p, m );
-                            pRes->GetPrivate()->m_values[thisP].Set<ParamProjectorType>(pOld->GetPrivate()->m_values[p].Get<ParamProjectorType>());
+                            pRes->GetPrivate()->m_values[thisP].Set<FParamProjectorType>(pOld->GetPrivate()->m_values[p].Get<FParamProjectorType>());
                             break;
                         }
 
-                        case PARAMETER_TYPE::T_IMAGE:
+                        case EParameterType::Image:
                             pRes->SetImageValue( thisP, pOld->GetImageValue(p) );
                             break;
 
@@ -559,9 +532,9 @@ namespace mu
 	
 
     //---------------------------------------------------------------------------------------------
-    int Model::GetStateCount() const
+    int32 Model::GetStateCount() const
     {
-        return (int)m_pD->m_program.m_states.Num();
+        return (int32)m_pD->m_program.m_states.Num();
     }
 
 
@@ -570,7 +543,7 @@ namespace mu
     {
         const char* strRes = 0;
 
-        if ( index>=0 && index<(int)m_pD->m_program.m_states.Num() )
+        if ( index>=0 && index<(int32)m_pD->m_program.m_states.Num() )
         {
             return m_pD->m_program.m_states[index].Name;
         }
@@ -583,9 +556,9 @@ namespace mu
     //---------------------------------------------------------------------------------------------
     int32 Model::FindState( const FString& Name ) const
     {
-        int res = -1;
+        int32 res = -1;
 
-        for ( int i=0; res<0 && i<(int)m_pD->m_program.m_states.Num(); ++i )
+        for ( int32 i=0; res<0 && i<(int32)m_pD->m_program.m_states.Num(); ++i )
         {
             if ( m_pD->m_program.m_states[i].Name == Name )
             {
@@ -598,13 +571,13 @@ namespace mu
 
 
     //---------------------------------------------------------------------------------------------
-    int Model::GetStateParameterCount( int stateIndex ) const
+    int32 Model::GetStateParameterCount( int32 stateIndex ) const
     {
-        int res = -1;
+        int32 res = -1;
 
-        if ( stateIndex>=0 && stateIndex<(int)m_pD->m_program.m_states.Num() )
+        if ( stateIndex>=0 && stateIndex<(int32)m_pD->m_program.m_states.Num() )
         {
-            res = (int)m_pD->m_program.m_states[stateIndex].m_runtimeParameters.Num();
+            res = (int32)m_pD->m_program.m_states[stateIndex].m_runtimeParameters.Num();
         }
 
         return res;
@@ -612,16 +585,16 @@ namespace mu
 
 
     //---------------------------------------------------------------------------------------------
-    int Model::GetStateParameterIndex( int stateIndex, int paramIndex ) const
+    int32 Model::GetStateParameterIndex( int32 stateIndex, int32 paramIndex ) const
     {
-        int res = -1;
+        int32 res = -1;
 
-        if ( stateIndex>=0 && stateIndex<(int)m_pD->m_program.m_states.Num() )
+        if ( stateIndex>=0 && stateIndex<(int32)m_pD->m_program.m_states.Num() )
         {
             const FProgram::FState& state = m_pD->m_program.m_states[stateIndex];
-            if ( paramIndex>=0 && paramIndex<(int)state.m_runtimeParameters.Num() )
+            if ( paramIndex>=0 && paramIndex<(int32)state.m_runtimeParameters.Num() )
             {
-                res = (int)state.m_runtimeParameters[paramIndex];
+                res = (int32)state.m_runtimeParameters[paramIndex];
             }
         }
 
@@ -654,13 +627,13 @@ namespace mu
             switch ( param.m_type )
             {
 
-            case PARAMETER_TYPE::T_INT:
+            case EParameterType::Int:
             {
                 m_pD->m_instanceCount *= param.m_possibleValues.Num();
                 break;
             }
 
-            case PARAMETER_TYPE::T_BOOL:
+            case EParameterType::Bool:
             {
                 m_pD->m_instanceCount *= 2;
                 break;
@@ -707,9 +680,9 @@ namespace mu
 		Ptr<Parameters> res = Model::NewParameters(m_pD->m_pModel);
 
 		const FProgram& Program = m_pD->m_pModel->GetPrivate()->m_program;
-        int paramCount = Program.m_parameters.Num();
+        int32 paramCount = Program.m_parameters.Num();
         int64 currentInstance = index;
-        for (int i=0;i<paramCount;++i)
+        for (int32 i=0;i<paramCount;++i)
         {
             const FParameterDesc& param = Program.m_parameters[ i ];
 			Ptr<RangeIndex> RangeIndex = res->NewRangeIndex(i);
@@ -717,7 +690,7 @@ namespace mu
             switch ( param.m_type )
             {
 
-            case PARAMETER_TYPE::T_INT:
+            case EParameterType::Int:
             {
 				bool bIsRangeSize = Program.m_ranges.ContainsByPredicate([i](const FRangeDesc& r) { return r.m_dimensionParameter == i; });
 				if (bIsRangeSize)
@@ -726,22 +699,22 @@ namespace mu
 				}
 				else
 				{
-					int numOptions = res->GetIntPossibleValueCount(i);
-					int value = res->GetIntPossibleValue(i, (int)(currentInstance % numOptions));
+					int32 numOptions = res->GetIntPossibleValueCount(i);
+					int32 value = res->GetIntPossibleValue(i, (int32)(currentInstance % numOptions));
 					res->SetIntValue(i, value);
 					currentInstance /= numOptions;
 				}
                 break;
             }
 
-            case PARAMETER_TYPE::T_BOOL:
+            case EParameterType::Bool:
             {
                 res->SetBoolValue( i, currentInstance%2!=0 );
                 currentInstance /= 2;
                 break;
             }
 
-            case PARAMETER_TYPE::T_FLOAT:
+            case EParameterType::Float:
 			{
 				bool bIsRangeSize = Program.m_ranges.ContainsByPredicate([i](const FRangeDesc& r) { return r.m_dimensionParameter == i; });
 				if (bIsRangeSize)
@@ -783,7 +756,7 @@ namespace mu
 				break;
 			}
 
-            case PARAMETER_TYPE::T_COLOUR:
+            case EParameterType::Color:
                 if (RandomGenerator)
                 {
 					if (RangeIndex)
@@ -793,18 +766,20 @@ namespace mu
 							for (int32 RangePosition = 0; RangePosition < m_pD->DefaultRangeDimension; ++RangePosition)
 							{
 								RangeIndex->SetPosition(Dimensions, RangePosition);
-								res->SetColourValue(i, RandomGenerator(), RandomGenerator(), RandomGenerator(), RandomGenerator(), RangeIndex);
+								FVector4f Random(RandomGenerator(), RandomGenerator(), RandomGenerator(), RandomGenerator());
+								res->SetColourValue(i, Random, RangeIndex);
 							}
 						}
 					}
 					else
 					{
-						res->SetColourValue(i, RandomGenerator(), RandomGenerator(), RandomGenerator(), RandomGenerator());
+						FVector4f Random(RandomGenerator(), RandomGenerator(), RandomGenerator(), RandomGenerator());
+						res->SetColourValue(i, Random);
 					}
                 }
                 break;
 
-			case PARAMETER_TYPE::T_PROJECTOR:
+			case EParameterType::Projector:
 				if (RandomGenerator)
 				{
 					// For projectors we just warp the position a little bit just to get something different
@@ -847,7 +822,7 @@ namespace mu
 
 
     //---------------------------------------------------------------------------------------------
-    ParametersPtr ModelParametersGenerator::GetRandomInstance( TFunctionRef<float()> randomGenerator)
+    ParametersPtr ModelParametersGenerator::GetRandomInstance( TFunctionRef<float()> RandomGenerator)
     {
 		LLM_SCOPE_BYNAME(TEXT("MutableRuntime"));
 
@@ -856,7 +831,7 @@ namespace mu
 		const FProgram& Program = m_pD->m_pModel->GetPrivate()->m_program;
         int32 ParamCount = Program.m_parameters.Num();
 
-        for (int i=0;i<ParamCount;++i)
+        for (int32 i=0;i<ParamCount;++i)
         {
             const FParameterDesc& param = Program.m_parameters[ i ];
 			Ptr<RangeIndex> RangeIndex = res->NewRangeIndex(i);
@@ -864,7 +839,7 @@ namespace mu
             switch ( param.m_type )
             {
 
-            case PARAMETER_TYPE::T_INT:
+            case EParameterType::Int:
             {
 				bool bIsRangeSize = Program.m_ranges.ContainsByPredicate([i](const FRangeDesc& r) { return r.m_dimensionParameter == i; });
 				if (bIsRangeSize)
@@ -893,31 +868,31 @@ namespace mu
 						{
 							for (int32 RangePosition = 0; RangePosition < m_pD->DefaultRangeDimension; ++RangePosition)
 							{
-								int numOptions = res->GetIntPossibleValueCount(i);
-								int valueIndex = (int)(FMath::Min(numOptions - 1, int(randomGenerator() * numOptions)));
-								int value = res->GetIntPossibleValue(i, valueIndex);
+								int32 numOptions = res->GetIntPossibleValueCount(i);
+								int32 valueIndex = (int32)(FMath::Min(numOptions - 1, int32(RandomGenerator() * numOptions)));
+								int32 value = res->GetIntPossibleValue(i, valueIndex);
 								res->SetIntValue(i, value, RangeIndex);
 							}
 						}
 					}
 					else
 					{
-						int numOptions = res->GetIntPossibleValueCount(i);
-						int valueIndex = (int)(FMath::Min(numOptions - 1, int(randomGenerator() * numOptions)));
-						int value = res->GetIntPossibleValue(i, valueIndex);
+						int32 numOptions = res->GetIntPossibleValueCount(i);
+						int32 valueIndex = (int32)(FMath::Min(numOptions - 1, int32(RandomGenerator() * numOptions)));
+						int32 value = res->GetIntPossibleValue(i, valueIndex);
 						res->SetIntValue(i, value);
 					}
 				}
                 break;
             }
 
-            case PARAMETER_TYPE::T_BOOL:
+            case EParameterType::Bool:
             {
-                res->SetBoolValue( i, randomGenerator()>0.5f );
+                res->SetBoolValue( i, RandomGenerator()>0.5f );
                 break;
             }
 
-            case PARAMETER_TYPE::T_FLOAT:
+            case EParameterType::Float:
 			{
 				bool bIsRangeSize = Program.m_ranges.ContainsByPredicate([i](const FRangeDesc& r) { return r.m_dimensionParameter == i; });
 				if (bIsRangeSize)
@@ -945,37 +920,38 @@ namespace mu
 						{
 							for (int32 RangePosition = 0; RangePosition < m_pD->DefaultRangeDimension; ++RangePosition)
 							{
-								res->SetFloatValue(i, randomGenerator(), RangeIndex);
+								res->SetFloatValue(i, RandomGenerator(), RangeIndex);
 							}
 						}
 					}
 					else
 					{
-						res->SetFloatValue(i, randomGenerator());
+						res->SetFloatValue(i, RandomGenerator());
 					}
 				}
 				break;
 			}
 
-            case PARAMETER_TYPE::T_COLOUR:
+            case EParameterType::Color:
 				if (RangeIndex)
 				{
 					for (int32 Dimensions = 0; Dimensions < RangeIndex->GetRangeCount(); ++Dimensions)
 					{
 						for (int32 RangePosition = 0; RangePosition < m_pD->DefaultRangeDimension; ++RangePosition)
 						{
-							res->SetColourValue(i, randomGenerator(), randomGenerator(), randomGenerator(), randomGenerator(), 
-												RangeIndex);
+							FVector4f Random(RandomGenerator(), RandomGenerator(), RandomGenerator(), RandomGenerator());
+							res->SetColourValue(i, Random, RangeIndex);
 						}
 					}
 				}
 				else
 				{
-					res->SetColourValue(i, randomGenerator(), randomGenerator(), randomGenerator(), randomGenerator());
+					FVector4f Random(RandomGenerator(), RandomGenerator(), RandomGenerator(), RandomGenerator());
+					res->SetColourValue(i, Random);
 				}
                 break;
 
-			case PARAMETER_TYPE::T_PROJECTOR:
+			case EParameterType::Projector:
 			{
 				// For projectors we just warp the position a little bit just to get something different
 				FVector3f Position, Direction, Up, Scale;
@@ -987,18 +963,18 @@ namespace mu
 					{
 						for (int32 RangePosition = 0; RangePosition < m_pD->DefaultRangeDimension; ++RangePosition)
 						{
-							Position.X *= 0.9 + randomGenerator() * 0.2f;
-							Position.Y *= 0.9 + randomGenerator() * 0.2f;
-							Position.Z *= 0.9 + randomGenerator() * 0.2f;
+							Position.X *= 0.9 + RandomGenerator() * 0.2f;
+							Position.Y *= 0.9 + RandomGenerator() * 0.2f;
+							Position.Z *= 0.9 + RandomGenerator() * 0.2f;
 							res->SetProjectorValue(i, Position, Direction, Up, Scale, Angle, RangeIndex);
 						}
 					}
 				}
 				else
 				{
-					Position.X *= 0.9 + randomGenerator() * 0.2f;
-					Position.Y *= 0.9 + randomGenerator() * 0.2f;
-					Position.Z *= 0.9 + randomGenerator() * 0.2f;
+					Position.X *= 0.9 + RandomGenerator() * 0.2f;
+					Position.Y *= 0.9 + RandomGenerator() * 0.2f;
+					Position.Z *= 0.9 + RandomGenerator() * 0.2f;
 					res->SetProjectorValue(i, Position, Direction, Up, Scale, Angle);
 				}
 				break;
diff --git a/Engine/Plugins/Mutable/Source/MutableRuntime/Private/MuR/OpImageBinarise.h b/Engine/Plugins/Mutable/Source/MutableRuntime/Private/MuR/OpImageBinarise.h
--- a/Engine/Plugins/Mutable/Source/MutableRuntime/Private/MuR/OpImageBinarise.h
+++ b/Engine/Plugins/Mutable/Source/MutableRuntime/Private/MuR/OpImageBinarise.h
@@ -21,7 +21,7 @@ namespace mu
 			return;
 		}
 
-		check(DestImage->GetFormat() == EImageFormat::IF_L_UBYTE);
+		check(DestImage->GetFormat() == EImageFormat::L_UByte);
 
 		// Generic implementation
         const uint16 ThresholdUNorm = static_cast<uint16>(FMath::Clamp(Threshold, 0.0f, 1.0f) * 255.0f);
@@ -45,7 +45,7 @@ namespace mu
 
 			switch(AImage->GetFormat())
 			{
-			case EImageFormat::IF_L_UBYTE:
+			case EImageFormat::L_UByte:
 			{
 				for (int32 I = 0; I < NumElems; ++I)
 				{
@@ -54,7 +54,7 @@ namespace mu
 				break;
 			}
 
-			case EImageFormat::IF_RGB_UBYTE:
+			case EImageFormat::RGB_UByte:
 			{
 				for (int32 I = 0; I < NumElems; ++I)
 				{
@@ -66,8 +66,8 @@ namespace mu
 				break;
 			}
 
-			case EImageFormat::IF_RGBA_UBYTE:
-			case EImageFormat::IF_BGRA_UBYTE:
+			case EImageFormat::RGBA_UByte:
+			case EImageFormat::BGRA_UByte:
 			{
 				for (int32 I = 0; I < NumElems; ++I)
 				{
diff --git a/Engine/Plugins/Mutable/Source/MutableRuntime/Private/MuR/OpImageColourMap.h b/Engine/Plugins/Mutable/Source/MutableRuntime/Private/MuR/OpImageColourMap.h
--- a/Engine/Plugins/Mutable/Source/MutableRuntime/Private/MuR/OpImageColourMap.h
+++ b/Engine/Plugins/Mutable/Source/MutableRuntime/Private/MuR/OpImageColourMap.h
@@ -17,7 +17,7 @@ namespace mu
 		check(SourceImage->GetLODCount() == MaskImage->GetLODCount() || bOnlyOneMip);
 		check(SourceImage->GetFormat() == DestImage->GetFormat());
 		check(SourceImage->GetSize() == DestImage->GetSize());
-		check(MaskImage->GetFormat() == EImageFormat::IF_L_UBYTE);
+		check(MaskImage->GetFormat() == EImageFormat::L_UByte);
 
 		// Generic implementation
 
@@ -58,7 +58,7 @@ namespace mu
 
 			switch (SourceImage->GetFormat())
 			{
-			case EImageFormat::IF_L_UBYTE:
+			case EImageFormat::L_UByte:
 			{
 				for (int32 I = 0; I < NumElems; ++I)
 				{
@@ -73,7 +73,7 @@ namespace mu
 				}
 				break;
 			}
-			case EImageFormat::IF_RGB_UBYTE:
+			case EImageFormat::RGB_UByte:
 			{
 				for (int32 I = 0; I < NumElems; ++I)
 				{
@@ -93,7 +93,7 @@ namespace mu
 				break;
 			}
 
-			case EImageFormat::IF_RGBA_UBYTE:
+			case EImageFormat::RGBA_UByte:
 			{
 				for (int32 I = 0; I < NumElems; ++I)
 				{
@@ -115,7 +115,7 @@ namespace mu
 				break;
 			}
 
-			case EImageFormat::IF_BGRA_UBYTE:
+			case EImageFormat::BGRA_UByte:
 			{
 				for (int32 I = 0; I < NumElems; ++I)
 				{
diff --git a/Engine/Plugins/Mutable/Source/MutableRuntime/Private/MuR/OpImageDisplace.h b/Engine/Plugins/Mutable/Source/MutableRuntime/Private/MuR/OpImageDisplace.h
--- a/Engine/Plugins/Mutable/Source/MutableRuntime/Private/MuR/OpImageDisplace.h
+++ b/Engine/Plugins/Mutable/Source/MutableRuntime/Private/MuR/OpImageDisplace.h
@@ -36,8 +36,8 @@ namespace mu
 	{
 		MUTABLE_CPUPROFILER_SCOPE(ImageMakeGrowMap);
 
-		check(MaskImage->GetFormat() == EImageFormat::IF_L_UBYTE);
-		check(ResultImage->GetFormat() == EImageFormat::IF_L_UBYTE);
+		check(MaskImage->GetFormat() == EImageFormat::L_UByte);
+		check(ResultImage->GetFormat() == EImageFormat::L_UByte);
 		check(ResultImage->GetSizeX() == MaskImage->GetSizeX());
 		check(ResultImage->GetSizeY() == MaskImage->GetSizeY());
 		check(ResultImage->GetLODCount() == MaskImage->GetLODCount());
@@ -138,7 +138,7 @@ namespace mu
 		MUTABLE_CPUPROFILER_SCOPE(ImageDisplace)
 
 		check(ResultImage->GetFormat() == SourceImage->GetFormat());
-		check(MapImage->GetFormat() == EImageFormat::IF_L_UBYTE || MapImage->GetFormat() == EImageFormat::IF_L_UBYTE_RLE);
+		check(MapImage->GetFormat() == EImageFormat::L_UByte || MapImage->GetFormat() == EImageFormat::L_UByteRLE);
 		check(ResultImage->GetSizeX() == SourceImage->GetSizeX());
 		check(ResultImage->GetSizeY() == SourceImage->GetSizeY());
 		check(ResultImage->GetSizeX() == MapImage->GetSizeX());
@@ -157,13 +157,13 @@ namespace mu
         //}
         
 		EImageFormat MapFormat = MapImage->GetFormat();
-        bool bIsUncompressed = MapFormat == EImageFormat::IF_L_UBYTE;
+        bool bIsUncompressed = MapFormat == EImageFormat::L_UByte;
 
         if (bIsUncompressed)
 		{
 			switch (ResultImage->GetFormat())
 			{
-			case EImageFormat::IF_L_UBYTE:
+			case EImageFormat::L_UByte:
 			{
 				const uint8* SourceData = SourceImage->GetLODData(0);
 				const uint8* MapBuf = MapImage->GetLODData(0);
@@ -201,7 +201,7 @@ namespace mu
 				break;
 			}
 
-			case EImageFormat::IF_RGB_UBYTE:
+			case EImageFormat::RGB_UByte:
 			{
         		const uint8* SourceData = SourceImage->GetLODData(0);
 				const uint8* MapBuf = MapImage->GetLODData(0);
@@ -277,7 +277,7 @@ namespace mu
 			}
 		}
 
-        else if (MapFormat == EImageFormat::IF_L_UBYTE_RLE)
+        else if (MapFormat == EImageFormat::L_UByteRLE)
 		{
             int32 PixelSize = GetImageFormatData(ResultImage->GetFormat()).BytesPerBlock;
 			
diff --git a/Engine/Plugins/Mutable/Source/MutableRuntime/Private/MuR/OpImageGradient.h b/Engine/Plugins/Mutable/Source/MutableRuntime/Private/MuR/OpImageGradient.h
--- a/Engine/Plugins/Mutable/Source/MutableRuntime/Private/MuR/OpImageGradient.h
+++ b/Engine/Plugins/Mutable/Source/MutableRuntime/Private/MuR/OpImageGradient.h
@@ -11,7 +11,7 @@ namespace mu
 
 	inline void ImageGradient( Image* pDest, const FVector4f& c0, const FVector4f& c1 )
 	{
-		check(pDest && pDest->GetFormat()==EImageFormat::IF_RGBA_UBYTE);
+		check(pDest && pDest->GetFormat()==EImageFormat::RGBA_UByte);
 
         uint8* pDestBuf = pDest->GetLODData(0);
 
diff --git a/Engine/Plugins/Mutable/Source/MutableRuntime/Private/MuR/OpImageInterpolate.h b/Engine/Plugins/Mutable/Source/MutableRuntime/Private/MuR/OpImageInterpolate.h
--- a/Engine/Plugins/Mutable/Source/MutableRuntime/Private/MuR/OpImageInterpolate.h
+++ b/Engine/Plugins/Mutable/Source/MutableRuntime/Private/MuR/OpImageInterpolate.h
@@ -44,10 +44,10 @@ namespace mu
 
 			switch (DestImage->GetFormat())
 			{
-			case EImageFormat::IF_L_UBYTE:
-			case EImageFormat::IF_RGB_UBYTE:
-			case EImageFormat::IF_BGRA_UBYTE:
-			case EImageFormat::IF_RGBA_UBYTE:
+			case EImageFormat::L_UByte:
+			case EImageFormat::RGB_UByte:
+			case EImageFormat::BGRA_UByte:
+			case EImageFormat::RGBA_UByte:
 			{
 				for (int32 I = 0; I < NumBytes; ++I)
 				{
diff --git a/Engine/Plugins/Mutable/Source/MutableRuntime/Private/MuR/OpImageInvert.h b/Engine/Plugins/Mutable/Source/MutableRuntime/Private/MuR/OpImageInvert.h
--- a/Engine/Plugins/Mutable/Source/MutableRuntime/Private/MuR/OpImageInvert.h
+++ b/Engine/Plugins/Mutable/Source/MutableRuntime/Private/MuR/OpImageInvert.h
@@ -36,8 +36,8 @@ namespace mu
 
 			switch (AImage->GetFormat())
 			{
-			case EImageFormat::IF_L_UBYTE:
-			case EImageFormat::IF_RGB_UBYTE:
+			case EImageFormat::L_UByte:
+			case EImageFormat::RGB_UByte:
 			{
 				if constexpr (!bUseVectorImpl)
 				{
@@ -72,8 +72,8 @@ namespace mu
 				}
 				break;
 			}
-			case EImageFormat::IF_RGBA_UBYTE:
-			case EImageFormat::IF_BGRA_UBYTE:
+			case EImageFormat::RGBA_UByte:
+			case EImageFormat::BGRA_UByte:
 			{
 				if constexpr (!bUseVectorImpl)
 				{
diff --git a/Engine/Plugins/Mutable/Source/MutableRuntime/Private/MuR/OpImageLayer.h b/Engine/Plugins/Mutable/Source/MutableRuntime/Private/MuR/OpImageLayer.h
--- a/Engine/Plugins/Mutable/Source/MutableRuntime/Private/MuR/OpImageLayer.h
+++ b/Engine/Plugins/Mutable/Source/MutableRuntime/Private/MuR/OpImageLayer.h
@@ -73,25 +73,25 @@ namespace mu
 
 			switch (BaseFormat)
 			{
-			case EImageFormat::IF_L_UBYTE:
+			case EImageFormat::L_UByte:
 			{
 				check(BytesPerElem == 1);
 				Private::BufferLayerColourGenericChannel<1, BLEND_FUNC, CLAMP>(ResultView.GetData(), BaseView.GetData(), NumElems, ColorValue);
 				break;
 			}
-			case EImageFormat::IF_RGB_UBYTE:
+			case EImageFormat::RGB_UByte:
 			{
 				check(BytesPerElem == 3);
 				Private::BufferLayerColourGenericChannel<3, BLEND_FUNC, CLAMP>(ResultView.GetData(), BaseView.GetData(), NumElems, ColorValue);
 				break;
 			}
-			case EImageFormat::IF_RGBA_UBYTE:
+			case EImageFormat::RGBA_UByte:
 			{
 				check(BytesPerElem == 4);
 				Private::BufferLayerColourGenericChannel<4, BLEND_FUNC, CLAMP>(ResultView.GetData(), BaseView.GetData(), NumElems, ColorValue);
 				break;
 			}
-			case EImageFormat::IF_BGRA_UBYTE:
+			case EImageFormat::BGRA_UByte:
 			{
 				check(BytesPerElem == 4);
 				FIntVector BGRAColorValue = FIntVector(ColorValue.Z, ColorValue.Y, ColorValue.X);
@@ -197,25 +197,25 @@ namespace mu
 			const int32 NumElems = ResultView.Num() / BytesPerElem;
 			switch (BaseFormat)
 			{
-			case EImageFormat::IF_L_UBYTE:
+			case EImageFormat::L_UByte:
 			{
 				check(BytesPerElem == 1);
 				Private::BufferLayerColourFromAlphaGenericChannel<1, BLEND_FUNC, CLAMP>(ResultView.GetData(), BaseView.GetData(), NumElems, ColorValue);
 				break;
 			}
-			case EImageFormat::IF_RGB_UBYTE:
+			case EImageFormat::RGB_UByte:
 			{
 				check(BytesPerElem == 3);
 				Private::BufferLayerColourFromAlphaGenericChannel<3, BLEND_FUNC, CLAMP>(ResultView.GetData(), BaseView.GetData(), NumElems, ColorValue);
 				break;
 			}
-			case EImageFormat::IF_RGBA_UBYTE:
+			case EImageFormat::RGBA_UByte:
 			{
 				check(BytesPerElem == 4);
 				Private::BufferLayerColourFromAlphaGenericChannel<4, BLEND_FUNC, CLAMP>(ResultView.GetData(), BaseView.GetData(), NumElems, ColorValue);
 				break;
 			}
-			case EImageFormat::IF_BGRA_UBYTE:
+			case EImageFormat::BGRA_UByte:
 			{
 				check(BytesPerElem == 4);
 				FIntVector BGRAColorValues = FIntVector(ColorValue.Z, ColorValue.Y, ColorValue.X);
@@ -276,7 +276,7 @@ namespace mu
 
 		int32 NumLODs = bOnlyFirstLOD ? 1 : BaseImage->GetLODCount();
 
-        const bool bIsMaskUncompressed = (MaskFormat == EImageFormat::IF_L_UBYTE);
+        const bool bIsMaskUncompressed = (MaskFormat == EImageFormat::L_UByte);
 
 		constexpr uint32 NumColorChannels = FMath::Min(CHANNELS_TO_BLEND, 3u);
 		check(NumColorChannels + ColorChannelOffset < 4);
@@ -352,7 +352,7 @@ namespace mu
 				ParallelFor(NumBatches, ProcessBatch);
 			}
         }
-        else if (MaskFormat == EImageFormat::IF_L_UBYTE_RLE)
+        else if (MaskFormat == EImageFormat::L_UByteRLE)
         {
             int32 Rows = BaseImage->GetSizeY();
             int32 Width = BaseImage->GetSizeX();
@@ -501,34 +501,34 @@ namespace mu
 	{
 		check(BaseImage->GetSizeX() == MaskImage->GetSizeX());
 		check(BaseImage->GetSizeY() == MaskImage->GetSizeY());
-		check(MaskImage->GetFormat() == EImageFormat::IF_L_UBYTE ||
-			  MaskImage->GetFormat() == EImageFormat::IF_L_UBYTE_RLE);
+		check(MaskImage->GetFormat() == EImageFormat::L_UByte ||
+			  MaskImage->GetFormat() == EImageFormat::L_UByteRLE);
 
 		const bool bValid = (BaseImage->GetSizeX() == MaskImage->GetSizeX()) &&
 						    (BaseImage->GetSizeY() == MaskImage->GetSizeY()) &&
-							(MaskImage->GetFormat() == EImageFormat::IF_L_UBYTE || MaskImage->GetFormat() == EImageFormat::IF_L_UBYTE_RLE);
+							(MaskImage->GetFormat() == EImageFormat::L_UByte || MaskImage->GetFormat() == EImageFormat::L_UByteRLE);
 		if (!bValid)
 		{
 			return;
 		}
 
 		EImageFormat BaseFormat = BaseImage->GetFormat();
-		if (BaseFormat == EImageFormat::IF_RGB_UBYTE)
+		if (BaseFormat == EImageFormat::RGB_UByte)
 		{
 			BufferLayerColourFormat<BLEND_FUNC_MASKED, BLEND_FUNC, CLAMP, 3, 3>(ResultImage, BaseImage, MaskImage, Col, 0, 0, false);
 		}
-        else if (BaseFormat == EImageFormat::IF_RGBA_UBYTE)
+        else if (BaseFormat == EImageFormat::RGBA_UByte)
         {
             BufferLayerColourFormat<BLEND_FUNC_MASKED, BLEND_FUNC, CLAMP, 3, 4>(ResultImage, BaseImage, MaskImage, Col, 0, 0, false);
         }
-        else if (BaseFormat == EImageFormat::IF_BGRA_UBYTE)
+        else if (BaseFormat == EImageFormat::BGRA_UByte)
         {
             float Temp = Col[0];
             Col[0] = Col[2];
             Col[2] = Temp;
             BufferLayerColourFormat<BLEND_FUNC_MASKED, BLEND_FUNC, CLAMP, 3, 4>(ResultImage, BaseImage, MaskImage, Col, 0, 0, false);
         }
-        else if (BaseFormat == EImageFormat::IF_L_UBYTE)
+        else if (BaseFormat == EImageFormat::L_UByte)
 		{
 			BufferLayerColourFormat<BLEND_FUNC_MASKED, BLEND_FUNC, CLAMP, 1, 1>(ResultImage, BaseImage, MaskImage, Col, 0, 0, false);
 		}
@@ -627,19 +627,19 @@ namespace mu
 	{
 		EImageFormat BaseFormat = BaseImage->GetFormat();
 
-		if (BaseFormat == EImageFormat::IF_RGB_UBYTE)
+		if (BaseFormat == EImageFormat::RGB_UByte)
 		{
 			check(BaseOffset + CHANNEL_COUNT <= 3);
 			BufferLayerColourFormatInPlace<BLEND_FUNC, CLAMP, CHANNEL_COUNT, 3>
 					(BaseImage, Col, BaseOffset, ColOffset, bOnlyOneMip);
 		}
-		else if (BaseFormat == EImageFormat::IF_RGBA_UBYTE)
+		else if (BaseFormat == EImageFormat::RGBA_UByte)
 		{
 			check(BaseOffset + CHANNEL_COUNT <= 4);
 			BufferLayerColourFormatInPlace<BLEND_FUNC, CLAMP, CHANNEL_COUNT, 4>
 					(BaseImage, Col, BaseOffset, ColOffset, bOnlyOneMip);
 		}
-		else if (BaseFormat == EImageFormat::IF_BGRA_UBYTE)
+		else if (BaseFormat == EImageFormat::BGRA_UByte)
 		{
 			float Temp = Col[0];
 			Col[0] = Col[2];
@@ -648,7 +648,7 @@ namespace mu
 					(BaseImage, Col, BaseOffset, ColOffset, bOnlyOneMip);
 		}
 
-		else if (BaseFormat == EImageFormat::IF_L_UBYTE)
+		else if (BaseFormat == EImageFormat::L_UByte)
 		{
 			check(BaseOffset + CHANNEL_COUNT <= 1);
 			BufferLayerColourFormatInPlace< BLEND_FUNC, CLAMP, CHANNEL_COUNT, 1>
@@ -690,19 +690,19 @@ namespace mu
 
 		EImageFormat BaseFormat = BaseImage->GetFormat();
 
-		if (BaseFormat == EImageFormat::IF_RGB_UBYTE)
+		if (BaseFormat == EImageFormat::RGB_UByte)
 		{
 			check(BaseOffset + CHANNEL_COUNT <= 3);
 			BufferLayerColourFormat<BLEND_FUNC_MASKED, BLEND_FUNC, CLAMP, CHANNEL_COUNT, 3>
 					(BaseImage, BaseImage, MaskImage, Color, BaseOffset, ColorOffset, bOnlyOneMip);
 		}
-		else if (BaseFormat == EImageFormat::IF_RGBA_UBYTE)
+		else if (BaseFormat == EImageFormat::RGBA_UByte)
 		{
 			check(BaseOffset + CHANNEL_COUNT <= 4);
 			BufferLayerColourFormat<BLEND_FUNC_MASKED, BLEND_FUNC, CLAMP, CHANNEL_COUNT, 4>
 					(BaseImage, BaseImage, MaskImage, Color, BaseOffset, ColorOffset, bOnlyOneMip);
 		}
-		else if (BaseFormat == EImageFormat::IF_BGRA_UBYTE)
+		else if (BaseFormat == EImageFormat::BGRA_UByte)
 		{
 			check(BaseOffset + CHANNEL_COUNT <= 4);
 			float Temp = Color[0];
@@ -712,7 +712,7 @@ namespace mu
 			BufferLayerColourFormat<BLEND_FUNC_MASKED, BLEND_FUNC, CLAMP, CHANNEL_COUNT, 4>
 					(BaseImage, BaseImage, MaskImage, Color, BaseOffset, ColorOffset, bOnlyOneMip);
 		}
-		else if (BaseFormat == EImageFormat::IF_L_UBYTE)
+		else if (BaseFormat == EImageFormat::L_UByte)
 		{
 			check(BaseOffset + CHANNEL_COUNT <= 1);
 			BufferLayerColourFormat<BLEND_FUNC_MASKED, BLEND_FUNC, CLAMP, CHANNEL_COUNT, 1>
@@ -949,9 +949,9 @@ namespace mu
 
 		if (bUseBlendSourceFromBlendAlpha)
 		{
-			if (BlendedFormat == EImageFormat::IF_RGBA_UBYTE || BlendedFormat == EImageFormat::IF_BGRA_UBYTE)
+			if (BlendedFormat == EImageFormat::RGBA_UByte || BlendedFormat == EImageFormat::BGRA_UByte)
 			{
-				if (BaseFormat == EImageFormat::IF_L_UBYTE)
+				if (BaseFormat == EImageFormat::L_UByte)
 				{
 					BufferLayerFormat<BLEND_FUNC, CLAMP, 1, 1, 4, 3>(ResultImage, BaseImage, BlendedImage, bOnlyOneMip);
 				}
@@ -960,7 +960,7 @@ namespace mu
 					checkf(false, TEXT("Unsupported format."));
 				}
 			}
-			else if (BlendedFormat == EImageFormat::IF_L_UBYTE)
+			else if (BlendedFormat == EImageFormat::L_UByte)
 			{
 				BufferLayerFormat<BLEND_FUNC, CLAMP, 1, 1, 1, 0>(ResultImage, BaseImage, BlendedImage, bOnlyOneMip);
 			}
@@ -968,12 +968,12 @@ namespace mu
 		else
 		{
 			check(BaseFormat == BlendedFormat);
-			if (BaseFormat == EImageFormat::IF_RGB_UBYTE)
+			if (BaseFormat == EImageFormat::RGB_UByte)
 			{
 				check(!bUseBlendSourceFromBlendAlpha);
 				BufferLayerFormat<BLEND_FUNC, CLAMP, 3, 3, 3, 0>(ResultImage, BaseImage, BlendedImage, bOnlyOneMip);
 			}
-			else if (BaseFormat == EImageFormat::IF_RGBA_UBYTE || BaseFormat == EImageFormat::IF_BGRA_UBYTE)
+			else if (BaseFormat == EImageFormat::RGBA_UByte || BaseFormat == EImageFormat::BGRA_UByte)
 			{
 				check(!bUseBlendSourceFromBlendAlpha);
 				if (bApplyToAlpha)
@@ -985,7 +985,7 @@ namespace mu
 					BufferLayerFormat<BLEND_FUNC, CLAMP, 3, 4, 4, 0>(ResultImage, BaseImage, BlendedImage, bOnlyOneMip);
 				}
 			}
-			else if (BaseFormat == EImageFormat::IF_L_UBYTE)
+			else if (BaseFormat == EImageFormat::L_UByte)
 			{
 				BufferLayerFormat<BLEND_FUNC, CLAMP, 1, 1, 1, 0>(ResultImage, BaseImage, BlendedImage, bOnlyOneMip);
 			}
@@ -1008,29 +1008,29 @@ namespace mu
 		EImageFormat BaseFormat = BaseImage->GetFormat();
 		EImageFormat BlendFormat = BlendedImage->GetFormat();
 
-		if (BaseFormat == EImageFormat::IF_RGB_UBYTE && BlendFormat==EImageFormat::IF_RGB_UBYTE)
+		if (BaseFormat == EImageFormat::RGB_UByte && BlendFormat==EImageFormat::RGB_UByte)
 		{
 			check(BaseOffset + CHANNEL_COUNT <= 3);
 			check(BlendedOffset + CHANNEL_COUNT <= 3);
 			BufferLayerFormatInPlace<BLEND_FUNC, CLAMP, CHANNEL_COUNT, 3, 3>(BaseImage, BlendedImage, BaseOffset, BlendedOffset, bOnlyOneMip);
 		}
-		else if ((BaseFormat == EImageFormat::IF_RGBA_UBYTE || BaseFormat == EImageFormat::IF_BGRA_UBYTE) &&
-				 (BlendFormat == EImageFormat::IF_RGBA_UBYTE || BlendFormat == EImageFormat::IF_BGRA_UBYTE) )
+		else if ((BaseFormat == EImageFormat::RGBA_UByte || BaseFormat == EImageFormat::BGRA_UByte) &&
+				 (BlendFormat == EImageFormat::RGBA_UByte || BlendFormat == EImageFormat::BGRA_UByte) )
 		{
 			check(BaseOffset + CHANNEL_COUNT <= 4);
 			check(BlendedOffset + CHANNEL_COUNT <= 4);
 			BufferLayerFormatInPlace< BLEND_FUNC, CLAMP, CHANNEL_COUNT, 4, 4>
 					(BaseImage, BlendedImage, BaseOffset, BlendedOffset, bOnlyOneMip);
 		}
-		else if ((BaseFormat == EImageFormat::IF_RGBA_UBYTE || BaseFormat == EImageFormat::IF_BGRA_UBYTE) &&
-				 BlendFormat == EImageFormat::IF_L_UBYTE )
+		else if ((BaseFormat == EImageFormat::RGBA_UByte || BaseFormat == EImageFormat::BGRA_UByte) &&
+				 BlendFormat == EImageFormat::L_UByte )
 		{
 			check(BaseOffset + CHANNEL_COUNT <= 4);
 			check(BlendedOffset + CHANNEL_COUNT <= 1);
 			BufferLayerFormatInPlace< BLEND_FUNC, CLAMP, CHANNEL_COUNT, 4, 1>
 					(BaseImage, BlendedImage, BaseOffset, BlendedOffset, bOnlyOneMip);
 		}
-		else if (BaseFormat == EImageFormat::IF_L_UBYTE && BlendFormat==EImageFormat::IF_L_UBYTE)
+		else if (BaseFormat == EImageFormat::L_UByte && BlendFormat==EImageFormat::L_UByte)
 		{
 			check(BaseOffset + CHANNEL_COUNT <= 1);
 			BufferLayerFormatInPlace< BLEND_FUNC, CLAMP, CHANNEL_COUNT, 1, 1>
@@ -1067,8 +1067,8 @@ namespace mu
 		check(bOnlyFirstLOD || BaseImage->GetLODCount() <= MaskImage->GetLODCount());
 		check(bOnlyFirstLOD || BaseImage->GetLODCount() <= BlendImage->GetLODCount());
 		check(BaseImage->GetFormat() == BlendImage->GetFormat());
-		check(MaskImage->GetFormat() == EImageFormat::IF_L_UBYTE ||
-			  MaskImage->GetFormat() == EImageFormat::IF_L_UBYTE_RLE);
+		check(MaskImage->GetFormat() == EImageFormat::L_UByte ||
+			  MaskImage->GetFormat() == EImageFormat::L_UByteRLE);
 
 		const EImageFormat MaskFormat = MaskImage->GetFormat();
 
@@ -1077,7 +1077,7 @@ namespace mu
 
 		const int32 NumLODs = bOnlyFirstLOD ? 1 : BaseImage->GetLODCount();
 
-		const bool bIsMaskUncompressed = MaskFormat == EImageFormat::IF_L_UBYTE;
+		const bool bIsMaskUncompressed = MaskFormat == EImageFormat::L_UByte;
 		constexpr int32 UnblendedChannels = BASE_CHANNEL_STRIDE - CHANNELS_TO_BLEND;
 
 		if (bIsMaskUncompressed)
@@ -1156,7 +1156,7 @@ namespace mu
 				ParallelFor(NumBatches, ProcessBatch);
 			}
 		}
-		else if (MaskFormat == EImageFormat::IF_L_UBYTE_RLE)
+		else if (MaskFormat == EImageFormat::L_UByteRLE)
 		{
 			int32 Rows = BaseImage->GetSizeY();
 			int32 Width = BaseImage->GetSizeX();
@@ -1368,7 +1368,7 @@ namespace mu
 			check(NumElems == DestView.Num() / BytesPerElem);
 			check(NumElems == BlendedView.Num() / BytesPerElem);
 
-			if (BlendImage->GetFormat() == EImageFormat::IF_RGBA_UBYTE)
+			if (BlendImage->GetFormat() == EImageFormat::RGBA_UByte)
 			{
 				const uint8* MaskBuf = BlendedBuf + 3;
 
@@ -1457,7 +1457,7 @@ namespace mu
 		uint8* DestBuf = DestImage->GetLODData(0) + DestOffset;
 
 		EImageFormat MaskFormat = MaskImage->GetFormat();
-        bool bIsUncompressed = (MaskFormat == EImageFormat::IF_L_UBYTE);
+        bool bIsUncompressed = (MaskFormat == EImageFormat::L_UByte);
 
         if (bIsUncompressed)
 		{
@@ -1498,7 +1498,7 @@ namespace mu
 				DestBuf += Stride;
 			}
 		}
-        else if (MaskFormat == EImageFormat::IF_L_UBIT_RLE)
+        else if (MaskFormat == EImageFormat::L_UBitRLE)
 		{
 			int32 Rows = MaskImage->GetSizeY();
 			int32 Width = MaskImage->GetSizeX();
@@ -1554,13 +1554,13 @@ namespace mu
                             bool bApplyToAlpha,
 							bool bOnlyFirstLOD)
 	{
-		if (BaseImage->GetFormat() == EImageFormat::IF_RGB_UBYTE)
+		if (BaseImage->GetFormat() == EImageFormat::RGB_UByte)
 		{
             BufferLayerFormat<BLEND_FUNC_MASKED, BLEND_FUNC, CLAMP, 3, 3, 3>
                     (DestImage, BaseImage, MaskImage, BlendImage, 0, 0, 0, bOnlyFirstLOD);
 		}
-        else if (BaseImage->GetFormat() == EImageFormat::IF_RGBA_UBYTE || 
-				 BaseImage->GetFormat() == EImageFormat::IF_BGRA_UBYTE)
+        else if (BaseImage->GetFormat() == EImageFormat::RGBA_UByte || 
+				 BaseImage->GetFormat() == EImageFormat::BGRA_UByte)
 		{
 			if (bApplyToAlpha)
 			{
@@ -1573,7 +1573,7 @@ namespace mu
 					(DestImage, BaseImage, MaskImage, BlendImage, 0, 0, 0, bOnlyFirstLOD);
 			}
 		}
-		else if (BaseImage->GetFormat() == EImageFormat::IF_L_UBYTE)
+		else if (BaseImage->GetFormat() == EImageFormat::L_UByte)
 		{
             BufferLayerFormat<BLEND_FUNC_MASKED, BLEND_FUNC, CLAMP, 1, 1, 1>
                     (DestImage, BaseImage, MaskImage, BlendImage, 0, 0, 0, bOnlyFirstLOD);
@@ -1595,13 +1595,13 @@ namespace mu
 		bool bApplyToAlpha,
 		bool bOnlyFirstLOD)
 	{
-		if (BaseImage->GetFormat() == EImageFormat::IF_RGB_UBYTE)
+		if (BaseImage->GetFormat() == EImageFormat::RGB_UByte)
 		{
 			BufferLayerFormatEmbeddedMask<BLEND_FUNC_MASKED, BLEND_FUNC, CLAMP, 3, 3, 3>
 				(DestImage, BaseImage, BlendImage, 0, 0, bOnlyFirstLOD);
 		}
-		else if (BaseImage->GetFormat() == EImageFormat::IF_RGBA_UBYTE || 
-				 BaseImage->GetFormat() == EImageFormat::IF_BGRA_UBYTE)
+		else if (BaseImage->GetFormat() == EImageFormat::RGBA_UByte || 
+				 BaseImage->GetFormat() == EImageFormat::BGRA_UByte)
 		{
 			if (bApplyToAlpha)
 			{
@@ -1631,8 +1631,8 @@ namespace mu
 		bool bOnlyFirstLOD,
 		uint8 BlendAlphaSourceChannel)
 	{
-		check(BaseImage->GetFormat() == EImageFormat::IF_RGBA_UBYTE);
-		check(BlendImage->GetFormat() == EImageFormat::IF_RGBA_UBYTE);
+		check(BaseImage->GetFormat() == EImageFormat::RGBA_UByte);
+		check(BlendImage->GetFormat() == EImageFormat::RGBA_UByte);
 		check(BaseImage->GetSizeX() == BlendImage->GetSizeX() && BaseImage->GetSizeY() == BlendImage->GetSizeY());
 		check(bOnlyFirstLOD || BaseImage->GetLODCount() <= BlendImage->GetLODCount());
 
@@ -1748,8 +1748,8 @@ namespace mu
 		bool bOnlyFirstLOD,
 		uint8 BlendAlphaSourceChannel)
 	{
-		check(BaseImage->GetFormat() == EImageFormat::IF_RGBA_UBYTE);
-		check(BlendImage->GetFormat() == EImageFormat::IF_RGBA_UBYTE);
+		check(BaseImage->GetFormat() == EImageFormat::RGBA_UByte);
+		check(BlendImage->GetFormat() == EImageFormat::RGBA_UByte);
 		check(BaseImage->GetSizeX() == BlendImage->GetSizeX() && BaseImage->GetSizeY() == BlendImage->GetSizeY());
 		check(bOnlyFirstLOD || BaseImage->GetLODCount() <= BlendImage->GetLODCount());
 		check(BlendAlphaSourceChannel == 3);
@@ -1847,8 +1847,8 @@ namespace mu
 		bool bOnlyFirstLOD,
 		uint8 BlendAlphaSourceChannel)
 	{
-		check(BaseImage->GetFormat() == EImageFormat::IF_RGBA_UBYTE);
-		check(BlendImage->GetFormat() == EImageFormat::IF_RGBA_UBYTE);
+		check(BaseImage->GetFormat() == EImageFormat::RGBA_UByte);
+		check(BlendImage->GetFormat() == EImageFormat::RGBA_UByte);
 		check(BaseImage->GetSizeX() == BlendImage->GetSizeX() && BaseImage->GetSizeY() == BlendImage->GetSizeY());
 		check(bOnlyFirstLOD || BaseImage->GetLODCount() <= BlendImage->GetLODCount());
 
@@ -1942,18 +1942,18 @@ namespace mu
 			 bool CLAMP>
     inline void BufferLayerStrideNoAlpha(Image* DestImage, int32 DestOffset, int32 Stride, const Image* MaskImage, const Image* BlendImage/*, int32 LODCount*/)
 	{
-		if (BlendImage->GetFormat() == EImageFormat::IF_RGB_UBYTE)
+		if (BlendImage->GetFormat() == EImageFormat::RGB_UByte)
 		{
 			BufferLayerFormatStrideNoAlpha<BLEND_FUNC, CLAMP, 3>
 					(DestImage, DestOffset, Stride, MaskImage, BlendImage/*, LODCount*/);
 		}
-        else if (BlendImage->GetFormat() == EImageFormat::IF_RGBA_UBYTE || 
-				 BlendImage->GetFormat() == EImageFormat::IF_BGRA_UBYTE)
+        else if (BlendImage->GetFormat() == EImageFormat::RGBA_UByte || 
+				 BlendImage->GetFormat() == EImageFormat::BGRA_UByte)
 		{
 			BufferLayerFormatStrideNoAlpha<BLEND_FUNC, CLAMP, 4>
 					(DestImage, DestOffset, Stride, MaskImage, BlendImage/*, LODCount*/);
 		}
-		else if (BlendImage->GetFormat() == EImageFormat::IF_L_UBYTE)
+		else if (BlendImage->GetFormat() == EImageFormat::L_UByte)
 		{
 			BufferLayerFormatStrideNoAlpha<BLEND_FUNC, CLAMP, 1>
 					(DestImage, DestOffset, Stride, MaskImage, BlendImage/*, LODCount*/);
@@ -1983,10 +1983,10 @@ namespace mu
 		check(MaskImage->GetSizeX() == BlendedImage->GetSizeX());
 		check(MaskImage->GetSizeY() == BlendedImage->GetSizeY());
 		check(BaseImage->GetFormat() == BlendedImage->GetFormat());
-		check(MaskImage->GetFormat() == EImageFormat::IF_L_UBYTE ||
+		check(MaskImage->GetFormat() == EImageFormat::L_UByte ||
 			  //UBYTE_RLE does not look to be supported.
-			  //MaskImage->GetFormat() == EImageFormat::IF_L_UBYTE_RLE || 
-			  MaskImage->GetFormat() == EImageFormat::IF_L_UBIT_RLE);
+			  //MaskImage->GetFormat() == EImageFormat::L_UByteRLE || 
+			  MaskImage->GetFormat() == EImageFormat::L_UBitRLE);
         check(BaseImage->GetLODCount() <= MaskImage->GetLODCount());
         check(BaseImage->GetLODCount() <= BlendedImage->GetLODCount());
 
@@ -2277,15 +2277,15 @@ namespace mu
 
 		const EImageFormat BaseFormat = BaseImage->GetFormat();
 
-		if (BaseFormat == EImageFormat::IF_L_UBYTE)
+		if (BaseFormat == EImageFormat::L_UByte)
 		{
 			BufferLayerCombine<BLEND_FUNC, 1>(ResultImage, BaseImage, BlendedImage, bOnlyFirstLOD);
 		}
-		else if (BaseFormat == EImageFormat::IF_RGB_UBYTE)
+		else if (BaseFormat == EImageFormat::RGB_UByte)
 		{
 			BufferLayerCombine<BLEND_FUNC, 3>(ResultImage, BaseImage, BlendedImage, bOnlyFirstLOD);
 		}
-		else if (BaseFormat == EImageFormat::IF_RGBA_UBYTE || BaseFormat == EImageFormat::IF_BGRA_UBYTE)
+		else if (BaseFormat == EImageFormat::RGBA_UByte || BaseFormat == EImageFormat::BGRA_UByte)
 		{
 			// \todo: pass swizzle template argument if BGRA_UBYTE, not yet supported.
 			BufferLayerCombine<BLEND_FUNC, 4>(ResultImage, BaseImage, BlendedImage, bOnlyFirstLOD);
@@ -2313,25 +2313,25 @@ namespace mu
 		const EImageFormat BaseFormat = BaseImage->GetFormat();
 
 		Ptr<Image> TempMaskImage;
-		if (MaskImage->GetFormat() != EImageFormat::IF_L_UBYTE)
+		if (MaskImage->GetFormat() != EImageFormat::L_UByte)
 		{
 			UE_LOG(LogMutableCore, Log, TEXT("Image layer format not supported. A generic one will be used. "));
 
 			FImageOperator ImOp = FImageOperator::GetDefault(nullptr);
 			constexpr int32 Quality = 4;
-			TempMaskImage = ImOp.ImagePixelFormat( Quality, MaskImage, EImageFormat::IF_L_UBYTE );
+			TempMaskImage = ImOp.ImagePixelFormat( Quality, MaskImage, EImageFormat::L_UByte );
 			MaskImage = TempMaskImage.get();
 		}
 
-		if (BaseFormat == EImageFormat::IF_L_UBYTE)
+		if (BaseFormat == EImageFormat::L_UByte)
 		{
 			BufferLayerCombine<BLEND_FUNC_MASKED, 1>(ResultImage, BaseImage, MaskImage, BlendedImage, bOnlyFirstLOD);
 		}
-		else if (BaseFormat == EImageFormat::IF_RGB_UBYTE)
+		else if (BaseFormat == EImageFormat::RGB_UByte)
 		{
 			BufferLayerCombine<BLEND_FUNC_MASKED, 3>(ResultImage, BaseImage, MaskImage, BlendedImage, bOnlyFirstLOD);
 		}
-		else if (BaseFormat == EImageFormat::IF_RGBA_UBYTE || BaseFormat == EImageFormat::IF_BGRA_UBYTE)
+		else if (BaseFormat == EImageFormat::RGBA_UByte || BaseFormat == EImageFormat::BGRA_UByte)
 		{
 			// \todo: pass swizzle template argument if BGRA_UBYTE, not yet supported.
 			BufferLayerCombine<BLEND_FUNC_MASKED, 4>(ResultImage, BaseImage, MaskImage, BlendedImage, bOnlyFirstLOD);
@@ -2342,8 +2342,8 @@ namespace mu
 
 			FImageOperator ImOp = FImageOperator::GetDefault(nullptr);
 			constexpr int32 Quality = 4;
-			Ptr<Image> TempBaseImage = ImOp.ImagePixelFormat(Quality, BaseImage, EImageFormat::IF_RGBA_UBYTE);
-			Ptr<Image> TempBlededImage = ImOp.ImagePixelFormat(Quality, BlendedImage, EImageFormat::IF_RGBA_UBYTE);
+			Ptr<Image> TempBaseImage = ImOp.ImagePixelFormat(Quality, BaseImage, EImageFormat::RGBA_UByte);
+			Ptr<Image> TempBlededImage = ImOp.ImagePixelFormat(Quality, BlendedImage, EImageFormat::RGBA_UByte);
 			BufferLayerCombine<BLEND_FUNC_MASKED, 4>(ResultImage, TempBaseImage.get(), MaskImage, TempBlededImage.get(), bOnlyFirstLOD);
 		}
 	}
@@ -2358,15 +2358,15 @@ namespace mu
 
 		const EImageFormat BaseFormat = BaseImage->GetFormat();
 
-		if (BaseFormat == EImageFormat::IF_L_UBYTE)
+		if (BaseFormat == EImageFormat::L_UByte)
 		{
 			BufferLayerCombineColour<BLEND_FUNC, 1>(ResultImage, BaseImage, Color);
 		}
-		else if (BaseFormat == EImageFormat::IF_RGB_UBYTE)
+		else if (BaseFormat == EImageFormat::RGB_UByte)
 		{
 			BufferLayerCombineColour<BLEND_FUNC, 3>(ResultImage, BaseImage, Color);
 		}
-		else if (BaseFormat == EImageFormat::IF_RGBA_UBYTE || BaseFormat == EImageFormat::IF_BGRA_UBYTE)
+		else if (BaseFormat == EImageFormat::RGBA_UByte || BaseFormat == EImageFormat::BGRA_UByte)
 		{
 			// \todo: pass swizzle template argument if BGRA_UBYTE, not yet supported.
  			BufferLayerCombineColour<BLEND_FUNC, 4>(ResultImage, BaseImage, Color);
@@ -2389,22 +2389,22 @@ namespace mu
 
 		const EImageFormat BaseFormat = BaseImage->GetFormat();
 
-		if (MaskImage->GetFormat() != EImageFormat::IF_L_UBYTE)
+		if (MaskImage->GetFormat() != EImageFormat::L_UByte)
 		{
 			checkf(false, TEXT("Unsupported mask format."));
 
 			BufferLayerCombineColour<BLEND_FUNC, 1>(ResultImage, BaseImage, Color);
 		}
 
-		if (BaseFormat == EImageFormat::IF_L_UBYTE)
+		if (BaseFormat == EImageFormat::L_UByte)
 		{
 			BufferLayerCombineColour<BLEND_FUNC_MASKED, 1>(ResultImage, BaseImage, MaskImage, Color);
 		}
-		else if (BaseFormat == EImageFormat::IF_RGB_UBYTE)
+		else if (BaseFormat == EImageFormat::RGB_UByte)
 		{
 			BufferLayerCombineColour<BLEND_FUNC_MASKED, 3>(ResultImage, BaseImage, MaskImage, Color);
 		}
-		else if (BaseFormat == EImageFormat::IF_RGBA_UBYTE || BaseFormat == EImageFormat::IF_BGRA_UBYTE)
+		else if (BaseFormat == EImageFormat::RGBA_UByte || BaseFormat == EImageFormat::BGRA_UByte)
 		{
 			// \todo: pass swizzle template argument if BGRA_UBYTE, not yet supported.
  			BufferLayerCombineColour<BLEND_FUNC_MASKED, 4>(ResultImage, BaseImage, MaskImage, Color);
@@ -2485,17 +2485,17 @@ namespace mu
 		const EImageFormat BaseFormat = BaseImage->GetFormat();
 		const EImageFormat BlendFormat = BlendedImage->GetFormat();
 
-		if (BaseFormat == EImageFormat::IF_L_UBYTE )
+		if (BaseFormat == EImageFormat::L_UByte )
 		{
-			if (BlendFormat == EImageFormat::IF_L_UBYTE )
+			if (BlendFormat == EImageFormat::L_UByte )
 			{
 				BufferLayerCombineFunctor<1, 1>(ResultImage, BaseImage, BlendedImage, Forward<ImageCombineFn>(ImageCombine));
 			}
-			else if (BlendFormat == EImageFormat::IF_RGB_UBYTE )
+			else if (BlendFormat == EImageFormat::RGB_UByte )
 			{
 				BufferLayerCombineFunctor<1, 3>(ResultImage, BaseImage, BlendedImage, Forward<ImageCombineFn>(ImageCombine));
 			}
-			else if (BlendFormat == EImageFormat::IF_RGBA_UBYTE || BlendFormat == EImageFormat::IF_BGRA_UBYTE)
+			else if (BlendFormat == EImageFormat::RGBA_UByte || BlendFormat == EImageFormat::BGRA_UByte)
 			{
 				BufferLayerCombineFunctor<1, 4>(ResultImage, BaseImage, BlendedImage, Forward<ImageCombineFn>(ImageCombine));
 			}
@@ -2504,17 +2504,17 @@ namespace mu
 				checkf(false, TEXT("Unsupported format."));
 			}
 		}
-		else if (BaseFormat == EImageFormat::IF_RGB_UBYTE)
+		else if (BaseFormat == EImageFormat::RGB_UByte)
 		{
-			if (BlendFormat == EImageFormat::IF_L_UBYTE )
+			if (BlendFormat == EImageFormat::L_UByte )
 			{
 				BufferLayerCombineFunctor<3, 1>(ResultImage, BaseImage, BlendedImage, Forward<ImageCombineFn>(ImageCombine));
 			}
-			else if (BlendFormat == EImageFormat::IF_RGB_UBYTE )
+			else if (BlendFormat == EImageFormat::RGB_UByte )
 			{
 				BufferLayerCombineFunctor<3, 3>(ResultImage, BaseImage, BlendedImage, Forward<ImageCombineFn>(ImageCombine));
 			}
-			else if (BlendFormat == EImageFormat::IF_RGBA_UBYTE || BlendFormat == EImageFormat::IF_BGRA_UBYTE )
+			else if (BlendFormat == EImageFormat::RGBA_UByte || BlendFormat == EImageFormat::BGRA_UByte )
 			{
 				BufferLayerCombineFunctor<3, 4>(ResultImage, BaseImage, BlendedImage, Forward<ImageCombineFn>(ImageCombine));
 			}
@@ -2523,17 +2523,17 @@ namespace mu
 				checkf(false, TEXT("Unsupported format."));
 			}
 		}
-		else if (BaseFormat == EImageFormat::IF_RGBA_UBYTE || BaseFormat == EImageFormat::IF_BGRA_UBYTE)
+		else if (BaseFormat == EImageFormat::RGBA_UByte || BaseFormat == EImageFormat::BGRA_UByte)
 		{
-			if (BlendFormat == EImageFormat::IF_L_UBYTE )
+			if (BlendFormat == EImageFormat::L_UByte )
 			{
 				BufferLayerCombineFunctor<4, 1>(ResultImage, BaseImage, BlendedImage, Forward<ImageCombineFn>(ImageCombine));
 			}
-			else if (BlendFormat == EImageFormat::IF_RGB_UBYTE )
+			else if (BlendFormat == EImageFormat::RGB_UByte )
 			{
 				BufferLayerCombineFunctor<4, 3>(ResultImage, BaseImage, BlendedImage, Forward<ImageCombineFn>(ImageCombine));
 			}
-			else if (BlendFormat == EImageFormat::IF_RGBA_UBYTE || BlendFormat == EImageFormat::IF_BGRA_UBYTE )
+			else if (BlendFormat == EImageFormat::RGBA_UByte || BlendFormat == EImageFormat::BGRA_UByte )
 			{
 				BufferLayerCombineFunctor<4, 4>(ResultImage, BaseImage, BlendedImage, Forward<ImageCombineFn>(ImageCombine));
 			}
diff --git a/Engine/Plugins/Mutable/Source/MutableRuntime/Private/MuR/OpImageLuminance.h b/Engine/Plugins/Mutable/Source/MutableRuntime/Private/MuR/OpImageLuminance.h
--- a/Engine/Plugins/Mutable/Source/MutableRuntime/Private/MuR/OpImageLuminance.h
+++ b/Engine/Plugins/Mutable/Source/MutableRuntime/Private/MuR/OpImageLuminance.h
@@ -12,7 +12,7 @@ namespace mu
 
 	inline void ImageLuminance(Image* DestImage, const Image* AImage)
 	{
-		check(DestImage && AImage && DestImage->GetFormat() == EImageFormat::IF_L_UBYTE);
+		check(DestImage && AImage && DestImage->GetFormat() == EImageFormat::L_UByte);
 
 		const int32 BytesPerElem = GetImageFormatData(AImage->GetFormat()).BytesPerBlock;
 		constexpr int32 NumBatchElems = 1 << 14;
@@ -32,7 +32,7 @@ namespace mu
 
             switch (AImage->GetFormat())
             {
-            case EImageFormat::IF_RGB_UBYTE:
+            case EImageFormat::RGB_UByte:
             {
                 for (int32 I = 0; I < NumElems; ++I)
                 {
@@ -42,7 +42,7 @@ namespace mu
                 break;
             }
 
-            case EImageFormat::IF_RGBA_UBYTE:
+            case EImageFormat::RGBA_UByte:
             {
                 for (int32 I = 0; I < NumElems; ++I)
                 {
@@ -52,7 +52,7 @@ namespace mu
                 break;
             }
 
-            case EImageFormat::IF_BGRA_UBYTE:
+            case EImageFormat::BGRA_UByte:
             {
                 for (int32 I = 0; I < NumElems; ++I)
                 {
