[mutable] Cleanup MutableRuntime API moving to use more unreal types and naming standards.

diff --git a/Engine/Plugins/Mutable/Source/MutableRuntime/Private/MuR/OpImageMipmap.cpp b/Engine/Plugins/Mutable/Source/MutableRuntime/Private/MuR/OpImageMipmap.cpp
--- a/Engine/Plugins/Mutable/Source/MutableRuntime/Private/MuR/OpImageMipmap.cpp
+++ b/Engine/Plugins/Mutable/Source/MutableRuntime/Private/MuR/OpImageMipmap.cpp
@@ -548,7 +548,7 @@ namespace OpImageMipmapInternal
 					Scratch.Uncompressed = CreateImage(
 						(uint16)UncompressedSize[0], (uint16)UncompressedSize[1],
 						1,
-						EImageFormat::IF_RGBA_UBYTE, EInitializationType::NotInitialized);
+						EImageFormat::RGBA_UByte, EInitializationType::NotInitialized);
 				}
 
 				FIntVector2 UncompressedMipsSize = DestImage->CalculateMipSize(StartLevel + 1);
@@ -556,7 +556,7 @@ namespace OpImageMipmapInternal
 				Scratch.UncompressedMips = CreateImage(
 					(uint16)UncompressedMipsSize[0], (uint16)UncompressedMipsSize[1],
 					FMath::Max(1, LevelCount - StartLevel - 1),
-					EImageFormat::IF_RGBA_UBYTE, EInitializationType::NotInitialized);
+					EImageFormat::RGBA_UByte, EInitializationType::NotInitialized);
 
 				// Compress the mipmapped image
 			//	Scratch.CompressedMips = CreateImage(
@@ -573,7 +573,7 @@ namespace OpImageMipmapInternal
 				Scratch.Uncompressed = CreateImage(
 					(uint16)UncompressedSize[0], (uint16)UncompressedSize[1],
 					1,
-					EImageFormat::IF_L_UBYTE, EInitializationType::NotInitialized);
+					EImageFormat::L_UByte, EInitializationType::NotInitialized);
 
 
 				FIntVector2 UncompressedMipsSize = DestImage->CalculateMipSize(StartLevel + 1);
@@ -581,7 +581,7 @@ namespace OpImageMipmapInternal
 				Scratch.UncompressedMips = CreateImage(
 					(uint16)UncompressedMipsSize[0], (uint16)UncompressedMipsSize[1],
 					FMath::Max(1, LevelCount - StartLevel - 1),
-					EImageFormat::IF_L_UBYTE, EInitializationType::NotInitialized);
+					EImageFormat::L_UByte, EInitializationType::NotInitialized);
 
 
 				// Compress the mipmapped image
@@ -659,21 +659,21 @@ namespace OpImageMipmapInternal
 
 			switch (DestFormat)
 			{
-			case EImageFormat::IF_L_UBYTE:
+			case EImageFormat::L_UByte:
 			{
 				constexpr int32 PixelSize = 1;
 				GenerateMipmapsBlockCompressedLODRange<PixelSize>(
 						StartLevel, 0, UncompressedNumLODs, BaseImage, UncompressedImage, Settings);
 				break;
 			}
-			case EImageFormat::IF_RGB_UBYTE:
+			case EImageFormat::RGB_UByte:
 			{
 				constexpr int32 PixelSize = 3;
 				GenerateMipmapsBlockCompressedLODRange<PixelSize>(
 						StartLevel, 0, UncompressedNumLODs, BaseImage, UncompressedImage, Settings);
 				break;
 			}
-			case EImageFormat::IF_RGBA_UBYTE:
+			case EImageFormat::RGBA_UByte:
 			{
 				constexpr int32 PixelSize = 4;
 				GenerateMipmapsBlockCompressedLODRange<PixelSize>(
@@ -722,18 +722,18 @@ namespace OpImageMipmapInternal
 
 			switch (BaseImage->GetFormat())
 			{
-			case EImageFormat::IF_L_UBYTE:
+			case EImageFormat::L_UByte:
 			{
 				GenerateMipmapUint8LODRange<1>(StartLevel, DestLODBegin, DestLODEnd, BaseImage, DestImage, Settings);
 				break;
 			}
-			case EImageFormat::IF_RGB_UBYTE:
+			case EImageFormat::RGB_UByte:
 			{
 				GenerateMipmapUint8LODRange<3>(StartLevel, DestLODBegin, DestLODEnd, BaseImage, DestImage, Settings);
 				break;
 			}
-			case EImageFormat::IF_BGRA_UBYTE:
-			case EImageFormat::IF_RGBA_UBYTE:
+			case EImageFormat::BGRA_UByte:
+			case EImageFormat::RGBA_UByte:
 			{
 				GenerateMipmapUint8LODRange<4>(StartLevel, DestLODBegin, DestLODEnd, BaseImage, DestImage, Settings);
 				break;
@@ -771,16 +771,16 @@ namespace OpImageMipmapInternal
 
 		switch (InBase->GetFormat())
 		{
-		case EImageFormat::IF_L_UBYTE:
+		case EImageFormat::L_UByte:
 			GenerateMipmapUint8LODRange<1>(0, 1, LevelCount, InBase, InBase, InSettings);
 			break;
 
-		case EImageFormat::IF_RGB_UBYTE:
+		case EImageFormat::RGB_UByte:
 			GenerateMipmapUint8LODRange<3>(0, 1, LevelCount, InBase, InBase, InSettings);
 			break;
 
-		case EImageFormat::IF_BGRA_UBYTE:
-		case EImageFormat::IF_RGBA_UBYTE:
+		case EImageFormat::BGRA_UByte:
+		case EImageFormat::RGBA_UByte:
 			GenerateMipmapUint8LODRange<4>(0, 1, LevelCount, InBase, InBase, InSettings);
 			break;
 
diff --git a/Engine/Plugins/Mutable/Source/MutableRuntime/Private/MuR/OpImagePixelFormat.cpp b/Engine/Plugins/Mutable/Source/MutableRuntime/Private/MuR/OpImagePixelFormat.cpp
--- a/Engine/Plugins/Mutable/Source/MutableRuntime/Private/MuR/OpImagePixelFormat.cpp
+++ b/Engine/Plugins/Mutable/Source/MutableRuntime/Private/MuR/OpImagePixelFormat.cpp
@@ -205,11 +205,11 @@ void FImageOperator::ImagePixelFormat(bool& bOutSuccess, int32 CompressionQualit
 	{
 		switch (Result->GetFormat())
 		{
-		case EImageFormat::IF_L_UBYTE:
+		case EImageFormat::L_UByte:
 		{
 			switch (Base->GetFormat())
 			{
-			case EImageFormat::IF_L_UBYTE:
+			case EImageFormat::L_UByte:
 			{
 				ImagePixelFormatInternal::ProcessBatchedHelper(Result, Base, 1, 1, ResultLODBegin, BaseLODBegin, NumLODs,
 				[](uint8* DestBuf, const uint8* BaseBuf, int32 NumElems)
@@ -219,7 +219,7 @@ void FImageOperator::ImagePixelFormat(bool& bOutSuccess, int32 CompressionQualit
 
 				break;
 			}
-			case EImageFormat::IF_L_UBYTE_RLE:
+			case EImageFormat::L_UByteRLE:
 			{
 				if (bFormatAllLODs)
 				{
@@ -233,7 +233,7 @@ void FImageOperator::ImagePixelFormat(bool& bOutSuccess, int32 CompressionQualit
 				}
 				break;
 			}
-			case EImageFormat::IF_L_UBIT_RLE:
+			case EImageFormat::L_UBitRLE:
 			{
 				if(bFormatAllLODs)
 				{
@@ -247,7 +247,7 @@ void FImageOperator::ImagePixelFormat(bool& bOutSuccess, int32 CompressionQualit
 				}
 				break;
 			}
-			case EImageFormat::IF_RGB_UBYTE:
+			case EImageFormat::RGB_UByte:
 			{
 				ImagePixelFormatInternal::ProcessBatchedHelper(Result, Base, 1, 3, ResultLODBegin, BaseLODBegin, NumLODs,
 				[](uint8* DestBuf, const uint8* BaseBuf, int32 NumElems)
@@ -261,10 +261,10 @@ void FImageOperator::ImagePixelFormat(bool& bOutSuccess, int32 CompressionQualit
 
 				break;
 			}
-			case EImageFormat::IF_RGB_UBYTE_RLE:
+			case EImageFormat::RGB_UByteRLE:
 			{
 				Ptr<Image> TempBase =
-						ImagePixelFormat(CompressionQuality, Base, EImageFormat::IF_RGB_UBYTE);
+						ImagePixelFormat(CompressionQuality, Base, EImageFormat::RGB_UByte);
 
 				ImagePixelFormatInternal::ProcessBatchedHelper(Result, TempBase.get(), 1, 3, ResultLODBegin, BaseLODBegin, NumLODs,
 				[](uint8* DestBuf, const uint8* BaseBuf, int32 NumElems)
@@ -280,7 +280,7 @@ void FImageOperator::ImagePixelFormat(bool& bOutSuccess, int32 CompressionQualit
 
 				break;
 			}
-			case EImageFormat::IF_RGBA_UBYTE:
+			case EImageFormat::RGBA_UByte:
 			{
 				ImagePixelFormatInternal::ProcessBatchedHelper(Result, Base, 1, 4, ResultLODBegin, BaseLODBegin, NumLODs,
 				[](uint8* DestBuf, const uint8* BaseBuf, int32 NumElems)
@@ -294,9 +294,9 @@ void FImageOperator::ImagePixelFormat(bool& bOutSuccess, int32 CompressionQualit
 
 				break;
 			}
-			case EImageFormat::IF_RGBA_UBYTE_RLE:
+			case EImageFormat::RGBA_UByteRLE:
 			{
-				Ptr<Image> TempBase = ImagePixelFormat(CompressionQuality, Base, EImageFormat::IF_RGBA_UBYTE);
+				Ptr<Image> TempBase = ImagePixelFormat(CompressionQuality, Base, EImageFormat::RGBA_UByte);
 				
 				ImagePixelFormatInternal::ProcessBatchedHelper(Result, TempBase.get(), 1, 4, ResultLODBegin, BaseLODBegin, NumLODs,
 				[](uint8* DestBuf, const uint8* BaseBuf, int32 NumElems)
@@ -312,7 +312,7 @@ void FImageOperator::ImagePixelFormat(bool& bOutSuccess, int32 CompressionQualit
 
 				break;
 			}
-			case EImageFormat::IF_BGRA_UBYTE:
+			case EImageFormat::BGRA_UByte:
 			{
 				ImagePixelFormatInternal::ProcessBatchedHelper(Result, Base, 1, 4, ResultLODBegin, BaseLODBegin, NumLODs,
 				[](uint8* DestBuf, const uint8* BaseBuf, int32 NumElems)
@@ -326,11 +326,11 @@ void FImageOperator::ImagePixelFormat(bool& bOutSuccess, int32 CompressionQualit
 
 				break;
 			}
-			case EImageFormat::IF_BC1:
-			case EImageFormat::IF_BC2:
-			case EImageFormat::IF_BC3:
+			case EImageFormat::BC1:
+			case EImageFormat::BC2:
+			case EImageFormat::BC3:
 			{
-				Ptr<Image> TempBase = ImagePixelFormat(CompressionQuality, Base, EImageFormat::IF_RGB_UBYTE);
+				Ptr<Image> TempBase = ImagePixelFormat(CompressionQuality, Base, EImageFormat::RGB_UByte);
 				
 				ImagePixelFormatInternal::ProcessBatchedHelper(Result, TempBase.get(), 1, 3, ResultLODBegin, BaseLODBegin, NumLODs,
 				[](uint8* DestBuf, const uint8* BaseBuf, int32 NumElems)
@@ -346,7 +346,7 @@ void FImageOperator::ImagePixelFormat(bool& bOutSuccess, int32 CompressionQualit
 
 				break;
 			}
-			case EImageFormat::IF_BC4:
+			case EImageFormat::BC4:
 			{
 				ImagePixelFormatInternal::BlockDecompressionLODRangeHelper(Result, Base, ResultLODBegin, BaseLODBegin, NumLODs, miro::BC4_to_L);
 				break;
@@ -359,12 +359,12 @@ void FImageOperator::ImagePixelFormat(bool& bOutSuccess, int32 CompressionQualit
 			break;
 		}
 
-		case EImageFormat::IF_L_UBYTE_RLE:
+		case EImageFormat::L_UByteRLE:
 		{
 			switch (Base->GetFormat())
 			{
 
-			case EImageFormat::IF_L_UBYTE:
+			case EImageFormat::L_UByte:
 			{
 				check(bFormatAllLODs);
 
@@ -388,19 +388,19 @@ void FImageOperator::ImagePixelFormat(bool& bOutSuccess, int32 CompressionQualit
 //				if (bOutSuccess)
 //				{
 //					// Verify
-//					Ptr<Image> Test = ImagePixelFormat(CompressionQuality, Result, EImageFormat::IF_L_UBYTE);
+//					Ptr<Image> Test = ImagePixelFormat(CompressionQuality, Result, EImageFormat::L_UByte);
 //					check(Test->DataStorage == Base->DataStorage);
 //				}
 //#endif
 				break;
 			}
 
-			case EImageFormat::IF_RGB_UBYTE:
-			case EImageFormat::IF_RGBA_UBYTE:
-			case EImageFormat::IF_BGRA_UBYTE:
+			case EImageFormat::RGB_UByte:
+			case EImageFormat::RGBA_UByte:
+			case EImageFormat::BGRA_UByte:
 			{
 				check(bFormatAllLODs);
-				Ptr<Image> TempBase = ImagePixelFormat(CompressionQuality, Base, EImageFormat::IF_L_UBYTE, -1);
+				Ptr<Image> TempBase = ImagePixelFormat(CompressionQuality, Base, EImageFormat::L_UByte, -1);
  
 				// Try to compress
 				if (Result->DataStorage.IsEmpty())
@@ -422,7 +422,7 @@ void FImageOperator::ImagePixelFormat(bool& bOutSuccess, int32 CompressionQualit
 //				if (bOutSuccess)
 //				{
 //					// Verify
-//					Ptr<Image> Test = ImagePixelFormat(CompressionQuality, Result, EImageFormat::IF_L_UBYTE);
+//					Ptr<Image> Test = ImagePixelFormat(CompressionQuality, Result, EImageFormat::L_UByte);
 //					check(Test->DataStorage == TempBase->DataStorage);
 //				}
 //#endif
@@ -440,12 +440,12 @@ void FImageOperator::ImagePixelFormat(bool& bOutSuccess, int32 CompressionQualit
 			break;
 		}
 
-		case EImageFormat::IF_L_UBIT_RLE:
+		case EImageFormat::L_UBitRLE:
 		{
 			switch (Base->GetFormat())
 			{
 
-			case EImageFormat::IF_L_UBYTE:
+			case EImageFormat::L_UByte:
 			{
 				check(bFormatAllLODs);
 
@@ -469,7 +469,7 @@ void FImageOperator::ImagePixelFormat(bool& bOutSuccess, int32 CompressionQualit
 //				if (bOutSuccess)
 //				{
 //					// Verify
-//					Ptr<Image> Test = ImagePixelFormat(CompressionQuality, Result, EImageFormat::IF_L_UBYTE);
+//					Ptr<Image> Test = ImagePixelFormat(CompressionQuality, Result, EImageFormat::L_UByte);
 //					check(Test->DataStorage == Base->DataStorage);
 //				}
 //#endif
@@ -477,19 +477,19 @@ void FImageOperator::ImagePixelFormat(bool& bOutSuccess, int32 CompressionQualit
 				break;
 			}
 
-			case EImageFormat::IF_RGB_UBYTE:
-			case EImageFormat::IF_RGBA_UBYTE:
-			case EImageFormat::IF_BGRA_UBYTE:
+			case EImageFormat::RGB_UByte:
+			case EImageFormat::RGBA_UByte:
+			case EImageFormat::BGRA_UByte:
 			{
 				check(bFormatAllLODs);
-				Ptr<Image> TempBase = ImagePixelFormat(CompressionQuality, Base, EImageFormat::IF_L_UBYTE, -1);
+				Ptr<Image> TempBase = ImagePixelFormat(CompressionQuality, Base, EImageFormat::L_UByte, -1);
 				CompressRLE_L1(TempBase.get(), Result);
 
 //#ifdef MUTABLE_DEBUG_RLE
 //				if (bOutSuccess)
 //				{
 //					// Verify
-//					Ptr<Image> Test = ImagePixelFormat(CompressionQuality, Result, EImageFormat::IF_L_UBYTE);
+//					Ptr<Image> Test = ImagePixelFormat(CompressionQuality, Result, EImageFormat::L_UByte);
 //					check(Test->DataStorage == TempBase->DataStorage);
 //				}
 //#endif
@@ -506,12 +506,12 @@ void FImageOperator::ImagePixelFormat(bool& bOutSuccess, int32 CompressionQualit
 			break;
 		}
 
-		case EImageFormat::IF_RGB_UBYTE:
+		case EImageFormat::RGB_UByte:
 		{
 			switch (Base->GetFormat())
 			{
 
-			case EImageFormat::IF_L_UBYTE:
+			case EImageFormat::L_UByte:
 			{
 				ImagePixelFormatInternal::ProcessBatchedHelper(Result, Base, 3, 1, ResultLODBegin, BaseLODBegin, NumLODs,
 				[](uint8* DestBuf, const uint8* BaseBuf, int32 NumElems)
@@ -526,7 +526,7 @@ void FImageOperator::ImagePixelFormat(bool& bOutSuccess, int32 CompressionQualit
 
 				break;
 			}
-			case EImageFormat::IF_RGB_UBYTE:
+			case EImageFormat::RGB_UByte:
 			{
 				ImagePixelFormatInternal::ProcessBatchedHelper(Result, Base, 3, 3, ResultLODBegin, BaseLODBegin, NumLODs,
 				[](uint8* DestBuf, const uint8* BaseBuf, int32 NumElems)
@@ -536,7 +536,7 @@ void FImageOperator::ImagePixelFormat(bool& bOutSuccess, int32 CompressionQualit
 
 				break;
 			}
-			case EImageFormat::IF_RGBA_UBYTE:
+			case EImageFormat::RGBA_UByte:
 			{
 				ImagePixelFormatInternal::ProcessBatchedHelper(Result, Base, 3, 4, ResultLODBegin, BaseLODBegin, NumLODs,
 				[](uint8* DestBuf, const uint8* BaseBuf, int32 NumElems)
@@ -551,7 +551,7 @@ void FImageOperator::ImagePixelFormat(bool& bOutSuccess, int32 CompressionQualit
 
 				break;
 			}
-			case EImageFormat::IF_BGRA_UBYTE:
+			case EImageFormat::BGRA_UByte:
 			{
 				ImagePixelFormatInternal::ProcessBatchedHelper(Result, Base, 3, 4, ResultLODBegin, BaseLODBegin, NumLODs,
 				[](uint8* DestBuf, const uint8* BaseBuf, int32 NumElems)
@@ -566,16 +566,16 @@ void FImageOperator::ImagePixelFormat(bool& bOutSuccess, int32 CompressionQualit
 
 				break;
 			}
-			case EImageFormat::IF_RGB_UBYTE_RLE:
+			case EImageFormat::RGB_UByteRLE:
 			{
 				UncompressRLE_RGB(Base, Result);
 				break;
 			}
-			case EImageFormat::IF_RGBA_UBYTE_RLE:
+			case EImageFormat::RGBA_UByteRLE:
 			{
 				check(bFormatAllLODs);
 
-				Ptr<Image> TempBase = ImagePixelFormat(CompressionQuality, Base, EImageFormat::IF_RGBA_UBYTE, -1);
+				Ptr<Image> TempBase = ImagePixelFormat(CompressionQuality, Base, EImageFormat::RGBA_UByte, -1);
 
 				ImagePixelFormatInternal::ProcessBatchedHelper(Result, TempBase.get(), 3, 4, ResultLODBegin, BaseLODBegin, NumLODs,
 				[](uint8* DestBuf, const uint8* BaseBuf, int32 NumElems)
@@ -591,121 +591,121 @@ void FImageOperator::ImagePixelFormat(bool& bOutSuccess, int32 CompressionQualit
 				ReleaseImage(TempBase);
 				break;
 			}
-			case EImageFormat::IF_BC1:
+			case EImageFormat::BC1:
 			{
 				ImagePixelFormatInternal::BlockDecompressionLODRangeHelper(
 						Result, Base, ResultLODBegin, BaseLODBegin, NumLODs, miro::BC1_to_RGB);
 				break;
 			}
-			case EImageFormat::IF_BC2:
+			case EImageFormat::BC2:
 			{
 				ImagePixelFormatInternal::BlockDecompressionLODRangeHelper(
 						Result, Base, ResultLODBegin, BaseLODBegin, NumLODs, miro::BC2_to_RGB);
 				break;
 			}
-			case EImageFormat::IF_BC3:
+			case EImageFormat::BC3:
 			{
 				ImagePixelFormatInternal::BlockDecompressionLODRangeHelper(
 						Result, Base, ResultLODBegin, BaseLODBegin, NumLODs, miro::BC3_to_RGB);
 				break;
 			}
-			case EImageFormat::IF_BC4:
+			case EImageFormat::BC4:
 			{
 				ImagePixelFormatInternal::BlockDecompressionLODRangeHelper(
 						Result, Base, ResultLODBegin, BaseLODBegin, NumLODs, miro::BC4_to_RGB);
 				break;
 			}
-			case EImageFormat::IF_BC5:
+			case EImageFormat::BC5:
 			{
 				ImagePixelFormatInternal::BlockDecompressionLODRangeHelper(
 						Result, Base, ResultLODBegin, BaseLODBegin, NumLODs, miro::BC5_to_RGB);
 				break;
 			}
-			case EImageFormat::IF_ASTC_4x4_RGB_LDR:
+			case EImageFormat::ASTC_4x4_RGB_LDR:
 			{
 				ImagePixelFormatInternal::BlockDecompressionLODRangeHelper(
 						Result, Base, ResultLODBegin, BaseLODBegin, NumLODs, miro::ASTC4x4RGBL_to_RGB);
 				break;
 			}
-			case EImageFormat::IF_ASTC_4x4_RGBA_LDR:
+			case EImageFormat::ASTC_4x4_RGBA_LDR:
 			{
 				ImagePixelFormatInternal::BlockDecompressionLODRangeHelper(
 						Result, Base, ResultLODBegin, BaseLODBegin, NumLODs, miro::ASTC4x4RGBAL_to_RGB);
 				break;
 			}
-			case EImageFormat::IF_ASTC_4x4_RG_LDR:
+			case EImageFormat::ASTC_4x4_RG_LDR:
 			{
 				ImagePixelFormatInternal::BlockDecompressionLODRangeHelper(
 						Result, Base, ResultLODBegin, BaseLODBegin, NumLODs, miro::ASTC4x4RGL_to_RGB);
 				break;
 			}
-			case EImageFormat::IF_ASTC_6x6_RGB_LDR:
+			case EImageFormat::ASTC_6x6_RGB_LDR:
 			{
 				ImagePixelFormatInternal::BlockDecompressionLODRangeHelper(
 						Result, Base, ResultLODBegin, BaseLODBegin, NumLODs, miro::ASTC6x6RGBL_to_RGB);
 				break;
 			}
-			case EImageFormat::IF_ASTC_6x6_RGBA_LDR:
+			case EImageFormat::ASTC_6x6_RGBA_LDR:
 			{
 				ImagePixelFormatInternal::BlockDecompressionLODRangeHelper(
 						Result, Base, ResultLODBegin, BaseLODBegin, NumLODs, miro::ASTC6x6RGBAL_to_RGB);
 				break;
 			}
-			case EImageFormat::IF_ASTC_6x6_RG_LDR:
+			case EImageFormat::ASTC_6x6_RG_LDR:
 			{
 				ImagePixelFormatInternal::BlockDecompressionLODRangeHelper(
 						Result, Base, ResultLODBegin, BaseLODBegin, NumLODs, miro::ASTC6x6RGL_to_RGB);
 				break;
 			}
-			case EImageFormat::IF_ASTC_8x8_RGB_LDR:
+			case EImageFormat::ASTC_8x8_RGB_LDR:
 			{
 				ImagePixelFormatInternal::BlockDecompressionLODRangeHelper(
 						Result, Base, ResultLODBegin, BaseLODBegin, NumLODs, miro::ASTC8x8RGBL_to_RGB);
 				break;
 			}
-			case EImageFormat::IF_ASTC_8x8_RGBA_LDR:
+			case EImageFormat::ASTC_8x8_RGBA_LDR:
 			{
 				ImagePixelFormatInternal::BlockDecompressionLODRangeHelper(
 						Result, Base, ResultLODBegin, BaseLODBegin, NumLODs, miro::ASTC8x8RGBAL_to_RGB);
 				break;
 			}
-			case EImageFormat::IF_ASTC_8x8_RG_LDR:
+			case EImageFormat::ASTC_8x8_RG_LDR:
 			{
 				ImagePixelFormatInternal::BlockDecompressionLODRangeHelper(
 						Result, Base, ResultLODBegin, BaseLODBegin, NumLODs, miro::ASTC8x8RGL_to_RGB);
 				break;
 			}
-			case EImageFormat::IF_ASTC_10x10_RGB_LDR:
+			case EImageFormat::ASTC_10x10_RGB_LDR:
 			{
 				ImagePixelFormatInternal::BlockDecompressionLODRangeHelper(
 						Result, Base, ResultLODBegin, BaseLODBegin, NumLODs, miro::ASTC10x10RGBL_to_RGB);
 				break;
 			}
-			case EImageFormat::IF_ASTC_10x10_RGBA_LDR:
+			case EImageFormat::ASTC_10x10_RGBA_LDR:
 			{
 				ImagePixelFormatInternal::BlockDecompressionLODRangeHelper(
 						Result, Base, ResultLODBegin, BaseLODBegin, NumLODs, miro::ASTC10x10RGBAL_to_RGB);
 				break;
 			}
-			case EImageFormat::IF_ASTC_10x10_RG_LDR:
+			case EImageFormat::ASTC_10x10_RG_LDR:
 			{
 				ImagePixelFormatInternal::BlockDecompressionLODRangeHelper(
 						Result, Base, ResultLODBegin, BaseLODBegin, NumLODs, miro::ASTC10x10RGL_to_RGB);
 				break;
 			}
-			case EImageFormat::IF_ASTC_12x12_RGB_LDR:
+			case EImageFormat::ASTC_12x12_RGB_LDR:
 			{
 				ImagePixelFormatInternal::BlockDecompressionLODRangeHelper(
 						Result, Base, ResultLODBegin, BaseLODBegin, NumLODs, miro::ASTC12x12RGBL_to_RGB);
 				break;
 			}
-			case EImageFormat::IF_ASTC_12x12_RGBA_LDR:
+			case EImageFormat::ASTC_12x12_RGBA_LDR:
 			{
 				ImagePixelFormatInternal::BlockDecompressionLODRangeHelper(
 						Result, Base, ResultLODBegin, BaseLODBegin, NumLODs, miro::ASTC12x12RGBAL_to_RGB);
 				break;
 			}
-			case EImageFormat::IF_ASTC_12x12_RG_LDR:
+			case EImageFormat::ASTC_12x12_RG_LDR:
 			{
 				ImagePixelFormatInternal::BlockDecompressionLODRangeHelper(
 						Result, Base, ResultLODBegin, BaseLODBegin, NumLODs, miro::ASTC12x12RGL_to_RGB);
@@ -720,13 +720,13 @@ void FImageOperator::ImagePixelFormat(bool& bOutSuccess, int32 CompressionQualit
 			break;
 		}
 
-		case EImageFormat::IF_BGRA_UBYTE:
+		case EImageFormat::BGRA_UByte:
 		{
 			// TODO: Optimise
-			Result->DataStorage.ImageFormat = EImageFormat::IF_RGBA_UBYTE;
+			Result->DataStorage.ImageFormat = EImageFormat::RGBA_UByte;
 			ImagePixelFormat(bOutSuccess, CompressionQuality, Result, Base, ResultLODBegin, BaseLODBegin, NumLODs);
 			check(bOutSuccess);
-			Result->DataStorage.ImageFormat = EImageFormat::IF_BGRA_UBYTE;
+			Result->DataStorage.ImageFormat = EImageFormat::BGRA_UByte;
 
 			ImagePixelFormatInternal::ProcessBatchedHelper(Result, 4, ResultLODBegin, NumLODs,
 			[](uint8* DestBuf, int32 NumElems)
@@ -742,18 +742,18 @@ void FImageOperator::ImagePixelFormat(bool& bOutSuccess, int32 CompressionQualit
 			break;
 		}
 
-		case EImageFormat::IF_RGBA_UBYTE:
+		case EImageFormat::RGBA_UByte:
 		{
 			switch (Base->GetFormat())
 			{
 
-			case EImageFormat::IF_RGBA_UBYTE_RLE:
+			case EImageFormat::RGBA_UByteRLE:
 			{
 				check(bFormatAllLODs);
 				UncompressRLE_RGBA(Base, Result);
 				break;
 			}
-			case EImageFormat::IF_L_UBYTE:
+			case EImageFormat::L_UByte:
 			{
 				ImagePixelFormatInternal::ProcessBatchedHelper(Result, Base, 4, 1, ResultLODBegin, BaseLODBegin, NumLODs,
 				[](uint8* DestBuf, const uint8* BaseBuf, int32 NumElems)
@@ -769,7 +769,7 @@ void FImageOperator::ImagePixelFormat(bool& bOutSuccess, int32 CompressionQualit
 
 				break;
 			}
-			case EImageFormat::IF_RGB_UBYTE:
+			case EImageFormat::RGB_UByte:
 			{
 				ImagePixelFormatInternal::ProcessBatchedHelper(Result, Base, 4, 3, ResultLODBegin, BaseLODBegin, NumLODs,
 				[](uint8* DestBuf, const uint8* BaseBuf, int32 NumElems)
@@ -785,7 +785,7 @@ void FImageOperator::ImagePixelFormat(bool& bOutSuccess, int32 CompressionQualit
 
 				break;
 			}
-			case EImageFormat::IF_RGBA_UBYTE:
+			case EImageFormat::RGBA_UByte:
 			{
 				ImagePixelFormatInternal::ProcessBatchedHelper(Result, Base, 4, 4, ResultLODBegin, BaseLODBegin, NumLODs,
 				[](uint8* DestBuf, const uint8* BaseBuf, int32 NumElems)
@@ -795,7 +795,7 @@ void FImageOperator::ImagePixelFormat(bool& bOutSuccess, int32 CompressionQualit
 
 				break;
 			}
-			case EImageFormat::IF_BGRA_UBYTE:
+			case EImageFormat::BGRA_UByte:
 			{
 				ImagePixelFormatInternal::ProcessBatchedHelper(Result, Base, 4, 4, ResultLODBegin, BaseLODBegin, NumLODs,
 				[](uint8* DestBuf, const uint8* BaseBuf, int32 NumElems)
@@ -811,130 +811,130 @@ void FImageOperator::ImagePixelFormat(bool& bOutSuccess, int32 CompressionQualit
 
 				break;
 			}
-			case EImageFormat::IF_BC1:
+			case EImageFormat::BC1:
 			{
 				ImagePixelFormatInternal::BlockDecompressionLODRangeHelper(
 						Result, Base, ResultLODBegin, BaseLODBegin, NumLODs, miro::BC1_to_RGBA);
 				break;
 			}
-			case EImageFormat::IF_BC2:
+			case EImageFormat::BC2:
 			{
 				ImagePixelFormatInternal::BlockDecompressionLODRangeHelper(
 						Result, Base, ResultLODBegin, BaseLODBegin, NumLODs, miro::BC2_to_RGBA);
 				break;
 			}
-			case EImageFormat::IF_BC3:
+			case EImageFormat::BC3:
 			{
 				ImagePixelFormatInternal::BlockDecompressionLODRangeHelper(
 						Result, Base, ResultLODBegin, BaseLODBegin, NumLODs, miro::BC3_to_RGBA);
 				break;
 			}
-			case EImageFormat::IF_BC4:
+			case EImageFormat::BC4:
 			{
 				ImagePixelFormatInternal::BlockDecompressionLODRangeHelper(
 						Result, Base, ResultLODBegin, BaseLODBegin, NumLODs, miro::BC4_to_RGBA);
 				break;
 			}
-			case EImageFormat::IF_BC5:
+			case EImageFormat::BC5:
 			{
 				ImagePixelFormatInternal::BlockDecompressionLODRangeHelper(
 						Result, Base, ResultLODBegin, BaseLODBegin, NumLODs, miro::BC5_to_RGBA);
 				break;
 			}
-			case EImageFormat::IF_ASTC_4x4_RGB_LDR:
+			case EImageFormat::ASTC_4x4_RGB_LDR:
 			{
 				ImagePixelFormatInternal::BlockDecompressionLODRangeHelper(
 						Result, Base, ResultLODBegin, BaseLODBegin, NumLODs, miro::ASTC4x4RGBL_to_RGBA);
 				break;
 			}
-			case EImageFormat::IF_ASTC_4x4_RGBA_LDR:
+			case EImageFormat::ASTC_4x4_RGBA_LDR:
 			{
 				ImagePixelFormatInternal::BlockDecompressionLODRangeHelper(
 						Result, Base, ResultLODBegin, BaseLODBegin, NumLODs, miro::ASTC4x4RGBAL_to_RGBA);
 				break;
 			}
-			case EImageFormat::IF_ASTC_4x4_RG_LDR:
+			case EImageFormat::ASTC_4x4_RG_LDR:
 			{
 				ImagePixelFormatInternal::BlockDecompressionLODRangeHelper(
 						Result, Base, ResultLODBegin, BaseLODBegin, NumLODs, miro::ASTC4x4RGL_to_RGBA);
 				break;
 			}
-			case EImageFormat::IF_ASTC_6x6_RGB_LDR:
+			case EImageFormat::ASTC_6x6_RGB_LDR:
 			{
 				ImagePixelFormatInternal::BlockDecompressionLODRangeHelper(
 						Result, Base, ResultLODBegin, BaseLODBegin, NumLODs, miro::ASTC6x6RGBL_to_RGBA);
 				break;
 			}
-			case EImageFormat::IF_ASTC_6x6_RGBA_LDR:
+			case EImageFormat::ASTC_6x6_RGBA_LDR:
 			{
 				ImagePixelFormatInternal::BlockDecompressionLODRangeHelper(
 						Result, Base, ResultLODBegin, BaseLODBegin, NumLODs, miro::ASTC6x6RGBAL_to_RGBA);
 				break;
 			}
-			case EImageFormat::IF_ASTC_6x6_RG_LDR:
+			case EImageFormat::ASTC_6x6_RG_LDR:
 			{
 				ImagePixelFormatInternal::BlockDecompressionLODRangeHelper(
 						Result, Base, ResultLODBegin, BaseLODBegin, NumLODs, miro::ASTC6x6RGL_to_RGBA);
 				break;
 			}
-			case EImageFormat::IF_ASTC_8x8_RGB_LDR:
+			case EImageFormat::ASTC_8x8_RGB_LDR:
 			{
 				ImagePixelFormatInternal::BlockDecompressionLODRangeHelper(
 						Result, Base, ResultLODBegin, BaseLODBegin, NumLODs, miro::ASTC8x8RGBL_to_RGBA);
 				break;
 			}
-			case EImageFormat::IF_ASTC_8x8_RGBA_LDR:
+			case EImageFormat::ASTC_8x8_RGBA_LDR:
 			{
 				ImagePixelFormatInternal::BlockDecompressionLODRangeHelper(
 						Result, Base, ResultLODBegin, BaseLODBegin, NumLODs, miro::ASTC8x8RGBAL_to_RGBA);
 				break;
 			}
-			case EImageFormat::IF_ASTC_8x8_RG_LDR:
+			case EImageFormat::ASTC_8x8_RG_LDR:
 			{
 				ImagePixelFormatInternal::BlockDecompressionLODRangeHelper(
 						Result, Base, ResultLODBegin, BaseLODBegin, NumLODs, miro::ASTC8x8RGL_to_RGBA);
 				break;
 			}
-			case EImageFormat::IF_ASTC_10x10_RGB_LDR:
+			case EImageFormat::ASTC_10x10_RGB_LDR:
 			{
 				ImagePixelFormatInternal::BlockDecompressionLODRangeHelper(
 						Result, Base, ResultLODBegin, BaseLODBegin, NumLODs, miro::ASTC10x10RGBL_to_RGBA);
 				break;
 			}
-			case EImageFormat::IF_ASTC_10x10_RGBA_LDR:
+			case EImageFormat::ASTC_10x10_RGBA_LDR:
 			{
 				ImagePixelFormatInternal::BlockDecompressionLODRangeHelper(
 						Result, Base, ResultLODBegin, BaseLODBegin, NumLODs, miro::ASTC10x10RGBAL_to_RGBA);
 				break;
 			}
-			case EImageFormat::IF_ASTC_10x10_RG_LDR:
+			case EImageFormat::ASTC_10x10_RG_LDR:
 			{
 				ImagePixelFormatInternal::BlockDecompressionLODRangeHelper(
 						Result, Base, ResultLODBegin, BaseLODBegin, NumLODs, miro::ASTC10x10RGL_to_RGBA);
 				break;
 			}
-			case EImageFormat::IF_ASTC_12x12_RGB_LDR:
+			case EImageFormat::ASTC_12x12_RGB_LDR:
 			{
 				ImagePixelFormatInternal::BlockDecompressionLODRangeHelper(
 						Result, Base, ResultLODBegin, BaseLODBegin, NumLODs, miro::ASTC12x12RGBL_to_RGBA);
 				break;
 			}
-			case EImageFormat::IF_ASTC_12x12_RGBA_LDR:
+			case EImageFormat::ASTC_12x12_RGBA_LDR:
 			{
 				ImagePixelFormatInternal::BlockDecompressionLODRangeHelper(
 						Result, Base, ResultLODBegin, BaseLODBegin, NumLODs, miro::ASTC12x12RGBAL_to_RGBA);
 				break;
 			}
-			case EImageFormat::IF_ASTC_12x12_RG_LDR:
+			case EImageFormat::ASTC_12x12_RG_LDR:
 			{
 				ImagePixelFormatInternal::BlockDecompressionLODRangeHelper(
 						Result, Base, ResultLODBegin, BaseLODBegin, NumLODs, miro::ASTC12x12RGL_to_RGBA);
 				break;
 			}
-			case EImageFormat::IF_RGB_UBYTE_RLE:
+			case EImageFormat::RGB_UByteRLE:
 			{
 				check(bFormatAllLODs);
-				Ptr<Image> TempBase = ImagePixelFormat(CompressionQuality, Base, EImageFormat::IF_RGB_UBYTE);
+				Ptr<Image> TempBase = ImagePixelFormat(CompressionQuality, Base, EImageFormat::RGB_UByte);
 
 				ImagePixelFormatInternal::ProcessBatchedHelper(Result, TempBase.get(), 4, 3, ResultLODBegin, BaseLODBegin, NumLODs,
 				[](uint8* DestBuf, const uint8* BaseBuf, int32 NumElems)
@@ -951,10 +951,10 @@ void FImageOperator::ImagePixelFormat(bool& bOutSuccess, int32 CompressionQualit
 				ReleaseImage(TempBase);
 				break;
 			}
-			case EImageFormat::IF_L_UBYTE_RLE:
+			case EImageFormat::L_UByteRLE:
 			{
 				check(bFormatAllLODs);
-				Ptr<Image> TempBase = ImagePixelFormat(CompressionQuality, Base, EImageFormat::IF_L_UBYTE);
+				Ptr<Image> TempBase = ImagePixelFormat(CompressionQuality, Base, EImageFormat::L_UByte);
 
 				ImagePixelFormatInternal::ProcessBatchedHelper(Result, TempBase.get(), 4, 1, ResultLODBegin, BaseLODBegin, NumLODs,
 				[](uint8* DestBuf, const uint8* BaseBuf, int32 NumElems)
@@ -980,41 +980,41 @@ void FImageOperator::ImagePixelFormat(bool& bOutSuccess, int32 CompressionQualit
 			break;
 		}
 
-		case EImageFormat::IF_RGBA_UBYTE_RLE:
+		case EImageFormat::RGBA_UByteRLE:
 		{
 			switch (Base->GetFormat())
 			{
 
-			case EImageFormat::IF_RGBA_UBYTE:
+			case EImageFormat::RGBA_UByte:
 			{
 				check(bFormatAllLODs);
 				CompressRLE_RGBA(Base, Result);
 
 				break;
 			}
-			case EImageFormat::IF_RGB_UBYTE:
+			case EImageFormat::RGB_UByte:
 			{
 				check(bFormatAllLODs);
 
 				// \todo: optimise
-				Ptr<Image> TempBase = ImagePixelFormat(CompressionQuality, Base, EImageFormat::IF_RGBA_UBYTE);
+				Ptr<Image> TempBase = ImagePixelFormat(CompressionQuality, Base, EImageFormat::RGBA_UByte);
 				CompressRLE_RGBA(TempBase.get(), Result);
 
 				// Test
-				//Ptr<Image> Test = ImagePixelFormat(CompressionQuality, Result, EImageFormat::IF_RGBA_UBYTE); 
+				//Ptr<Image> Test = ImagePixelFormat(CompressionQuality, Result, EImageFormat::RGBA_UByte); 
 				//check(Test->DataStorage == Base->DataStorage);
 
 				ReleaseImage(TempBase);
 
 				break;
 			}
-			case EImageFormat::IF_RGB_UBYTE_RLE:
+			case EImageFormat::RGB_UByteRLE:
 			{
 				check(bFormatAllLODs);
 
 				// \todo: optimise
-				Ptr<Image> Temp1 = ImagePixelFormat(CompressionQuality, Base, EImageFormat::IF_RGB_UBYTE);
-				Ptr<Image> Temp2 = ImagePixelFormat(CompressionQuality, Temp1.get(), EImageFormat::IF_RGBA_UBYTE);
+				Ptr<Image> Temp1 = ImagePixelFormat(CompressionQuality, Base, EImageFormat::RGB_UByte);
+				Ptr<Image> Temp2 = ImagePixelFormat(CompressionQuality, Temp1.get(), EImageFormat::RGBA_UByte);
 				ReleaseImage(Temp1);
 				CompressRLE_RGBA(Temp2.get(), Result);
 				ReleaseImage(Temp2);
@@ -1030,29 +1030,29 @@ void FImageOperator::ImagePixelFormat(bool& bOutSuccess, int32 CompressionQualit
 			break;
 		}
 
-		case EImageFormat::IF_RGB_UBYTE_RLE:
+		case EImageFormat::RGB_UByteRLE:
 		{
 			switch (Base->GetFormat())
 			{
 
-			case EImageFormat::IF_RGB_UBYTE:
+			case EImageFormat::RGB_UByte:
 			{
 				check(bFormatAllLODs);
 				CompressRLE_RGB(Base, Result);
 
 				// Test
-				// Ptr<Image> Test = ImagePixelFormat(CompressionQuality, Result, EImageFormat::IF_RGB_UBYTE); 
+				// Ptr<Image> Test = ImagePixelFormat(CompressionQuality, Result, EImageFormat::RGB_UByte); 
 				// check(Test->DataStorage == Base->DataStorage);
 				break;
 			}
 
-			case EImageFormat::IF_RGBA_UBYTE_RLE:
+			case EImageFormat::RGBA_UByteRLE:
 			{
 				check(bFormatAllLODs);
 
 				// \todo: optimise
-				Ptr<Image> Temp1 = ImagePixelFormat(CompressionQuality, Base, EImageFormat::IF_RGBA_UBYTE);
-				Ptr<Image> Temp2 = ImagePixelFormat(CompressionQuality, Temp1.get(), EImageFormat::IF_RGB_UBYTE);
+				Ptr<Image> Temp1 = ImagePixelFormat(CompressionQuality, Base, EImageFormat::RGBA_UByte);
+				Ptr<Image> Temp2 = ImagePixelFormat(CompressionQuality, Temp1.get(), EImageFormat::RGB_UByte);
 				ReleaseImage(Temp1);
 				CompressRLE_RGB(Temp2.get(), Result);
 				ReleaseImage(Temp2);
@@ -1068,34 +1068,34 @@ void FImageOperator::ImagePixelFormat(bool& bOutSuccess, int32 CompressionQualit
 			break;
 		}
 
-		case EImageFormat::IF_BC1:
+		case EImageFormat::BC1:
 		{
 			switch (Base->GetFormat())
 			{
-			case EImageFormat::IF_RGB_UBYTE:
+			case EImageFormat::RGB_UByte:
 			{
 				ImagePixelFormatInternal::BlockCompressionLODRangeHelper(
 						Result, Base, ResultLODBegin, BaseLODBegin, NumLODs, CompressionQuality, miro::RGB_to_BC1);
 				break;
 			}
-			case EImageFormat::IF_RGBA_UBYTE:
+			case EImageFormat::RGBA_UByte:
 			{
 				ImagePixelFormatInternal::BlockCompressionLODRangeHelper(
 						Result, Base, ResultLODBegin, BaseLODBegin, NumLODs, CompressionQuality, miro::RGBA_to_BC1);
 				break;
 			}
-			case EImageFormat::IF_L_UBYTE:
+			case EImageFormat::L_UByte:
 			{
-				Ptr<Image> TempBase = ImagePixelFormat(CompressionQuality, Base, EImageFormat::IF_RGB_UBYTE);
+				Ptr<Image> TempBase = ImagePixelFormat(CompressionQuality, Base, EImageFormat::RGB_UByte);
 				ImagePixelFormatInternal::BlockCompressionLODRangeHelper(
 						Result, TempBase.get(), ResultLODBegin, BaseLODBegin, NumLODs, CompressionQuality, miro::RGB_to_BC1);
 				ReleaseImage(TempBase);
 
 				break;
 			}
-			case EImageFormat::IF_BC3:
+			case EImageFormat::BC3:
 			{
-				Ptr<Image> TempBase = ImagePixelFormat(CompressionQuality, Base, EImageFormat::IF_RGBA_UBYTE);
+				Ptr<Image> TempBase = ImagePixelFormat(CompressionQuality, Base, EImageFormat::RGBA_UByte);
 				ImagePixelFormatInternal::BlockCompressionLODRangeHelper(
 						Result, TempBase.get(), ResultLODBegin, BaseLODBegin, NumLODs, CompressionQuality, miro::RGBA_to_BC1);
 				ReleaseImage(TempBase);
@@ -1107,7 +1107,7 @@ void FImageOperator::ImagePixelFormat(bool& bOutSuccess, int32 CompressionQualit
 				UE_LOG(LogMutableCore, Log, TEXT("Image format conversion not implemented. Expensive generic one used. "));
 
 				// Case not implemented, try a generic more expensive operation.
-				Ptr<Image> TempBase = ImagePixelFormat(CompressionQuality, Base, EImageFormat::IF_RGBA_UBYTE);
+				Ptr<Image> TempBase = ImagePixelFormat(CompressionQuality, Base, EImageFormat::RGBA_UByte);
 				ImagePixelFormatInternal::BlockCompressionLODRangeHelper(
 						Result, TempBase.get(), ResultLODBegin, BaseLODBegin, NumLODs, CompressionQuality, miro::RGBA_to_BC1);
 				ReleaseImage(TempBase);
@@ -1118,11 +1118,11 @@ void FImageOperator::ImagePixelFormat(bool& bOutSuccess, int32 CompressionQualit
 			break;
 		}
 
-		case EImageFormat::IF_BC2:
+		case EImageFormat::BC2:
 		{
 			switch (Base->GetFormat())
 			{
-			case EImageFormat::IF_RGBA_UBYTE:
+			case EImageFormat::RGBA_UByte:
 			{
 				ImagePixelFormatInternal::BlockCompressionLODRangeHelper(
 						Result, Base, ResultLODBegin, BaseLODBegin, NumLODs, CompressionQuality, miro::RGBA_to_BC2);
@@ -1133,7 +1133,7 @@ void FImageOperator::ImagePixelFormat(bool& bOutSuccess, int32 CompressionQualit
 				UE_LOG(LogMutableCore, Log, TEXT("Image format conversion not implemented. Expensive generic one used."));
 
 				// Case not implemented, try a generic more expensive operatino.
-				Ptr<Image> TempBase = ImagePixelFormat(CompressionQuality, Base, EImageFormat::IF_RGBA_UBYTE);
+				Ptr<Image> TempBase = ImagePixelFormat(CompressionQuality, Base, EImageFormat::RGBA_UByte);
 				ImagePixelFormatInternal::BlockCompressionLODRangeHelper(
 						Result, TempBase.get(), ResultLODBegin, BaseLODBegin, NumLODs, CompressionQuality, miro::RGBA_to_BC2);
 				ReleaseImage(TempBase);
@@ -1144,23 +1144,23 @@ void FImageOperator::ImagePixelFormat(bool& bOutSuccess, int32 CompressionQualit
 			break;
 		}
 
-		case EImageFormat::IF_BC3:
+		case EImageFormat::BC3:
 		{
 			switch (Base->GetFormat())
 			{
-			case EImageFormat::IF_RGBA_UBYTE:
+			case EImageFormat::RGBA_UByte:
 			{
 				ImagePixelFormatInternal::BlockCompressionLODRangeHelper(
 						Result, Base, ResultLODBegin, BaseLODBegin, NumLODs, CompressionQuality, miro::RGBA_to_BC3);
 				break;
 			}
-			case EImageFormat::IF_RGB_UBYTE:
+			case EImageFormat::RGB_UByte:
 			{
 				ImagePixelFormatInternal::BlockCompressionLODRangeHelper(
 						Result, Base, ResultLODBegin, BaseLODBegin, NumLODs, CompressionQuality, miro::RGB_to_BC3);
 				break;
 			}
-			case EImageFormat::IF_BC1:
+			case EImageFormat::BC1:
 			{
 				MUTABLE_CPUPROFILER_SCOPE(BC1toBC3);
 				ImagePixelFormatInternal::BlockCompressionLODRangeHelper(
@@ -1172,7 +1172,7 @@ void FImageOperator::ImagePixelFormat(bool& bOutSuccess, int32 CompressionQualit
 				UE_LOG(LogMutableCore, Log, TEXT("Image format conversion not implemented. Expensive generic one used."));
 
 				// Case not implemented, try a generic more expensive operation.
-				Ptr<Image> TempBase = ImagePixelFormat(CompressionQuality, Base, EImageFormat::IF_RGBA_UBYTE);
+				Ptr<Image> TempBase = ImagePixelFormat(CompressionQuality, Base, EImageFormat::RGBA_UByte);
 				ImagePixelFormatInternal::BlockCompressionLODRangeHelper(
 						Result, TempBase.get(), ResultLODBegin, BaseLODBegin, NumLODs, CompressionQuality, miro::RGBA_to_BC3);
 				ReleaseImage(TempBase);
@@ -1183,22 +1183,22 @@ void FImageOperator::ImagePixelFormat(bool& bOutSuccess, int32 CompressionQualit
 			break;
 		}
 
-		case EImageFormat::IF_BC4:
+		case EImageFormat::BC4:
 		{
 			switch (Base->GetFormat())
 			{
-			case EImageFormat::IF_L_UBYTE:
+			case EImageFormat::L_UByte:
 			{
 				ImagePixelFormatInternal::BlockCompressionLODRangeHelper(
 						Result, Base, ResultLODBegin, BaseLODBegin, NumLODs, CompressionQuality, miro::L_to_BC4);
 				break;
 			}
-			case EImageFormat::IF_RGB_UBYTE:
-			case EImageFormat::IF_RGBA_UBYTE:
-			case EImageFormat::IF_BGRA_UBYTE:
+			case EImageFormat::RGB_UByte:
+			case EImageFormat::RGBA_UByte:
+			case EImageFormat::BGRA_UByte:
 			default:
 			{	
-				Ptr<Image> TempBase = ImagePixelFormat(CompressionQuality, Base, EImageFormat::IF_L_UBYTE, -1);
+				Ptr<Image> TempBase = ImagePixelFormat(CompressionQuality, Base, EImageFormat::L_UByte, -1);
 				ImagePixelFormatInternal::BlockCompressionLODRangeHelper(
 						Result, TempBase.get(), ResultLODBegin, BaseLODBegin, NumLODs, CompressionQuality, miro::L_to_BC4);
 				ReleaseImage(TempBase);
@@ -1209,17 +1209,17 @@ void FImageOperator::ImagePixelFormat(bool& bOutSuccess, int32 CompressionQualit
 			break;
 		}
 
-		case EImageFormat::IF_BC5:
+		case EImageFormat::BC5:
 		{
 			switch (Base->GetFormat())
 			{
-			case EImageFormat::IF_RGBA_UBYTE:
+			case EImageFormat::RGBA_UByte:
 			{
 				ImagePixelFormatInternal::BlockCompressionLODRangeHelper(
 						Result, Base, ResultLODBegin, BaseLODBegin, NumLODs, CompressionQuality, miro::RGBA_to_BC5);
 				break;
 			}
-			case EImageFormat::IF_RGB_UBYTE:
+			case EImageFormat::RGB_UByte:
 			{
 				ImagePixelFormatInternal::BlockCompressionLODRangeHelper(
 						Result, Base, ResultLODBegin, BaseLODBegin, NumLODs, CompressionQuality, miro::RGB_to_BC5);
@@ -1229,7 +1229,7 @@ void FImageOperator::ImagePixelFormat(bool& bOutSuccess, int32 CompressionQualit
 			{
 				UE_LOG(LogMutableCore, Log, TEXT("Image format conversion not implemented. Expensive generic one used. "));
 
-				Ptr<Image> TempBase = ImagePixelFormat(CompressionQuality, Base, EImageFormat::IF_RGBA_UBYTE, -1);
+				Ptr<Image> TempBase = ImagePixelFormat(CompressionQuality, Base, EImageFormat::RGBA_UByte, -1);
 				ImagePixelFormatInternal::BlockCompressionLODRangeHelper(
 						Result, TempBase.get(), ResultLODBegin, BaseLODBegin, NumLODs, CompressionQuality, miro::RGBA_to_BC5);
 				ReleaseImage(TempBase);
@@ -1240,29 +1240,29 @@ void FImageOperator::ImagePixelFormat(bool& bOutSuccess, int32 CompressionQualit
 			break;
 		}
 
-		case EImageFormat::IF_ASTC_4x4_RGB_LDR:
+		case EImageFormat::ASTC_4x4_RGB_LDR:
 		{
 			switch (Base->GetFormat())
 			{
-			case EImageFormat::IF_RGBA_UBYTE:
+			case EImageFormat::RGBA_UByte:
 			{
 				ImagePixelFormatInternal::BlockCompressionLODRangeHelper(
 						Result, Base, ResultLODBegin, BaseLODBegin, NumLODs, CompressionQuality, miro::RGBA_to_ASTC4x4RGBL);
 				break;
 			}
-			case EImageFormat::IF_RGB_UBYTE:
+			case EImageFormat::RGB_UByte:
 			{
 				ImagePixelFormatInternal::BlockCompressionLODRangeHelper(
 						Result, Base, ResultLODBegin, BaseLODBegin, NumLODs, CompressionQuality, miro::RGB_to_ASTC4x4RGBL);
 				break;
 			}
-			case EImageFormat::IF_ASTC_4x4_RGBA_LDR:
+			case EImageFormat::ASTC_4x4_RGBA_LDR:
 			{
 				ImagePixelFormatInternal::BlockDecompressionLODRangeHelper(
 						Result, Base, ResultLODBegin, BaseLODBegin, NumLODs, miro::ASTC4x4RGBAL_to_ASTC4x4RGBL);
 				break;
 			}
-			case EImageFormat::IF_L_UBYTE:
+			case EImageFormat::L_UByte:
 			{
 				ImagePixelFormatInternal::BlockCompressionLODRangeHelper(
 						Result, Base, ResultLODBegin, BaseLODBegin, NumLODs, CompressionQuality, miro::L_to_ASTC4x4RGBL);
@@ -1272,7 +1272,7 @@ void FImageOperator::ImagePixelFormat(bool& bOutSuccess, int32 CompressionQualit
 			{
 				UE_LOG(LogMutableCore, Log, TEXT("Image format conversion not implemented. Expensive generic one used. "));
 
-				Ptr<Image> TempBase = ImagePixelFormat(CompressionQuality, Base, EImageFormat::IF_RGBA_UBYTE, -1);
+				Ptr<Image> TempBase = ImagePixelFormat(CompressionQuality, Base, EImageFormat::RGBA_UByte, -1);
 				ImagePixelFormatInternal::BlockCompressionLODRangeHelper(
 						Result, TempBase.get(), ResultLODBegin, BaseLODBegin, NumLODs, CompressionQuality, miro::RGBA_to_ASTC4x4RGBL);
 				ReleaseImage(TempBase);
@@ -1283,23 +1283,23 @@ void FImageOperator::ImagePixelFormat(bool& bOutSuccess, int32 CompressionQualit
 			break;
 		}
 
-		case EImageFormat::IF_ASTC_4x4_RGBA_LDR:
+		case EImageFormat::ASTC_4x4_RGBA_LDR:
 		{
 			switch (Base->GetFormat())
 			{
-			case EImageFormat::IF_RGBA_UBYTE:
+			case EImageFormat::RGBA_UByte:
 			{
 				ImagePixelFormatInternal::BlockCompressionLODRangeHelper(
 						Result, Base, ResultLODBegin, BaseLODBegin, NumLODs, CompressionQuality, miro::RGBA_to_ASTC4x4RGBAL);
 				break;
 			}
-			case EImageFormat::IF_RGB_UBYTE:
+			case EImageFormat::RGB_UByte:
 			{
 				ImagePixelFormatInternal::BlockCompressionLODRangeHelper(
 						Result, Base, ResultLODBegin, BaseLODBegin, NumLODs, CompressionQuality, miro::RGB_to_ASTC4x4RGBAL);
 				break;
 			}
-			case EImageFormat::IF_ASTC_4x4_RGB_LDR:
+			case EImageFormat::ASTC_4x4_RGB_LDR:
 			{
 				ImagePixelFormatInternal::BlockDecompressionLODRangeHelper(
 						Result, Base, ResultLODBegin, BaseLODBegin, NumLODs, miro::ASTC4x4RGBL_to_ASTC4x4RGBAL);
@@ -1309,7 +1309,7 @@ void FImageOperator::ImagePixelFormat(bool& bOutSuccess, int32 CompressionQualit
 			{
 				UE_LOG(LogMutableCore, Log, TEXT("Image format conversion not implemented. Expensive generic one used. "));
 
-				Ptr<Image> TempBase = ImagePixelFormat(CompressionQuality, Base, EImageFormat::IF_RGBA_UBYTE, -1);
+				Ptr<Image> TempBase = ImagePixelFormat(CompressionQuality, Base, EImageFormat::RGBA_UByte, -1);
 				ImagePixelFormatInternal::BlockCompressionLODRangeHelper(
 						Result, TempBase.get(), ResultLODBegin, BaseLODBegin, NumLODs, CompressionQuality, miro::RGBA_to_ASTC4x4RGBAL);
 				ReleaseImage(TempBase);
@@ -1320,23 +1320,23 @@ void FImageOperator::ImagePixelFormat(bool& bOutSuccess, int32 CompressionQualit
 			break;
 		}
 
-		case EImageFormat::IF_ASTC_4x4_RG_LDR:
+		case EImageFormat::ASTC_4x4_RG_LDR:
 		{
 			switch (Base->GetFormat())
 			{
-			case EImageFormat::IF_RGBA_UBYTE:
+			case EImageFormat::RGBA_UByte:
 			{
 				ImagePixelFormatInternal::BlockCompressionLODRangeHelper(
 						Result, Base, ResultLODBegin, BaseLODBegin, NumLODs, CompressionQuality, miro::RGBA_to_ASTC4x4RGL);
 				break;
 			}
-			case EImageFormat::IF_RGB_UBYTE:
+			case EImageFormat::RGB_UByte:
 			{
 				ImagePixelFormatInternal::BlockCompressionLODRangeHelper(
 						Result, Base, ResultLODBegin, BaseLODBegin, NumLODs, CompressionQuality, miro::RGB_to_ASTC4x4RGL);
 				break;
 			}
-			case EImageFormat::IF_ASTC_4x4_RG_LDR:
+			case EImageFormat::ASTC_4x4_RG_LDR:
 			{
 				// Hack that actually works because of block size.
 				ImagePixelFormatInternal::BlockDecompressionLODRangeHelper(
@@ -1347,7 +1347,7 @@ void FImageOperator::ImagePixelFormat(bool& bOutSuccess, int32 CompressionQualit
 			{
 				UE_LOG(LogMutableCore, Log, TEXT("Image format conversion not implemented. Expensive generic one used. "));
 
-				Ptr<Image> TempBase = ImagePixelFormat(CompressionQuality, Base, EImageFormat::IF_RGBA_UBYTE, -1);
+				Ptr<Image> TempBase = ImagePixelFormat(CompressionQuality, Base, EImageFormat::RGBA_UByte, -1);
 				ImagePixelFormatInternal::BlockCompressionLODRangeHelper(
 						Result, TempBase.get(), ResultLODBegin, BaseLODBegin, NumLODs, CompressionQuality, miro::RGBA_to_ASTC4x4RGL);
 				ReleaseImage(TempBase);
@@ -1358,29 +1358,29 @@ void FImageOperator::ImagePixelFormat(bool& bOutSuccess, int32 CompressionQualit
 			break;
 		}
 
-		case EImageFormat::IF_ASTC_6x6_RGB_LDR:
+		case EImageFormat::ASTC_6x6_RGB_LDR:
 		{
 			switch (Base->GetFormat())
 			{
-			case EImageFormat::IF_RGBA_UBYTE:
+			case EImageFormat::RGBA_UByte:
 			{
 				ImagePixelFormatInternal::BlockCompressionLODRangeHelper(
 						Result, Base, ResultLODBegin, BaseLODBegin, NumLODs, CompressionQuality, miro::RGBA_to_ASTC6x6RGBL);
 				break;
 			}
-			case EImageFormat::IF_RGB_UBYTE:
+			case EImageFormat::RGB_UByte:
 			{
 				ImagePixelFormatInternal::BlockCompressionLODRangeHelper(
 						Result, Base, ResultLODBegin, BaseLODBegin, NumLODs, CompressionQuality, miro::RGB_to_ASTC6x6RGBL);
 				break;
 			}
-			//case EImageFormat::IF_ASTC_6x6_RGBA_LDR:
+			//case EImageFormat::ASTC_6x6_RGBA_LDR:
 			//{
 			//  ImagePixelFormatInternal::BlockCompressionLODRangeHelper(
 			//  		Result, Base, ResultLODBegin, BaseLODBegin, NumLODs, CompressionQuality, miro::ASTC6x6RGBAL_to_ASTC6x6RGBL);
 			//	break;
 			//}
-			//case EImageFormat::IF_L_UBYTE:
+			//case EImageFormat::L_UByte:
 			//{
 			//  ImagePixelFormatInternal::BlockCompressionLODRangeHelper(
 			//  		Result, Base, ResultLODBegin, BaseLODBegin, NumLODs, CompressionQuality, miro::L_to_ASTC6x6RGBL);
@@ -1391,7 +1391,7 @@ void FImageOperator::ImagePixelFormat(bool& bOutSuccess, int32 CompressionQualit
 			{
 				UE_LOG(LogMutableCore, Log, TEXT("Image format conversion not implemented. Expensive generic one used. "));
 
-				Ptr<Image> TempBase = ImagePixelFormat(CompressionQuality, Base, EImageFormat::IF_RGBA_UBYTE, -1);
+				Ptr<Image> TempBase = ImagePixelFormat(CompressionQuality, Base, EImageFormat::RGBA_UByte, -1);
 				ImagePixelFormatInternal::BlockCompressionLODRangeHelper(
 						Result, TempBase.get(), ResultLODBegin, BaseLODBegin, NumLODs, CompressionQuality, miro::RGBA_to_ASTC6x6RGBL);
 				ReleaseImage(TempBase);
@@ -1402,23 +1402,23 @@ void FImageOperator::ImagePixelFormat(bool& bOutSuccess, int32 CompressionQualit
 			break;
 		}
 
-		case EImageFormat::IF_ASTC_6x6_RGBA_LDR:
+		case EImageFormat::ASTC_6x6_RGBA_LDR:
 		{
 			switch (Base->GetFormat())
 			{
-			case EImageFormat::IF_RGBA_UBYTE:
+			case EImageFormat::RGBA_UByte:
 			{
 				ImagePixelFormatInternal::BlockCompressionLODRangeHelper(
 						Result, Base, ResultLODBegin, BaseLODBegin, NumLODs, CompressionQuality, miro::RGBA_to_ASTC6x6RGBAL);
 				break;
 			}
-			case EImageFormat::IF_RGB_UBYTE:
+			case EImageFormat::RGB_UByte:
 			{
 				ImagePixelFormatInternal::BlockCompressionLODRangeHelper(
 						Result, Base, ResultLODBegin, BaseLODBegin, NumLODs, CompressionQuality, miro::RGB_to_ASTC6x6RGBAL);
 				break;
 			}
-			//case EImageFormat::IF_ASTC_6x6_RGB_LDR:
+			//case EImageFormat::ASTC_6x6_RGB_LDR:
 			//{
 			//	ImagePixelFormatInternal::BlockCompressionLODRangeHelper(
 			//			Result, Base, ResultLODBegin, BaseLODBegin, NumLODs, CompressionQuality, miro::ASTC6x6RGBL_to_ASTC6x6RGBAL);
@@ -1428,7 +1428,7 @@ void FImageOperator::ImagePixelFormat(bool& bOutSuccess, int32 CompressionQualit
 			{
 				UE_LOG(LogMutableCore, Log, TEXT("Image format conversion not implemented. Expensive generic one used. "));
 
-				Ptr<Image> TempBase = ImagePixelFormat(CompressionQuality, Base, EImageFormat::IF_RGBA_UBYTE, -1);
+				Ptr<Image> TempBase = ImagePixelFormat(CompressionQuality, Base, EImageFormat::RGBA_UByte, -1);
 				ImagePixelFormatInternal::BlockCompressionLODRangeHelper(
 						Result, TempBase.get(), ResultLODBegin, BaseLODBegin, NumLODs, CompressionQuality, miro::RGBA_to_ASTC6x6RGBAL);
 				ReleaseImage(TempBase);
@@ -1439,24 +1439,24 @@ void FImageOperator::ImagePixelFormat(bool& bOutSuccess, int32 CompressionQualit
 			break;
 		}
 
-		case EImageFormat::IF_ASTC_6x6_RG_LDR:
+		case EImageFormat::ASTC_6x6_RG_LDR:
 		{
 			switch (Base->GetFormat())
 			{
-			case EImageFormat::IF_RGBA_UBYTE:
+			case EImageFormat::RGBA_UByte:
 			{
 				ImagePixelFormatInternal::BlockCompressionLODRangeHelper(
 						Result, Base, ResultLODBegin, BaseLODBegin, NumLODs, CompressionQuality, miro::RGBA_to_ASTC6x6RGL);
 				break;
 			}
 
-			case EImageFormat::IF_RGB_UBYTE:
+			case EImageFormat::RGB_UByte:
 			{
 				ImagePixelFormatInternal::BlockCompressionLODRangeHelper(
 						Result, Base, ResultLODBegin, BaseLODBegin, NumLODs, CompressionQuality, miro::RGB_to_ASTC6x6RGL);
 				break;
 			}
-			//case EImageFormat::IF_ASTC_6x6_RG_LDR:
+			//case EImageFormat::ASTC_6x6_RG_LDR:
 			//{
 			//	// Hack that actually works because of block size.
 			//	ImagePixelFormatInternal::BlockCompressionLODRangeHelper(
@@ -1467,7 +1467,7 @@ void FImageOperator::ImagePixelFormat(bool& bOutSuccess, int32 CompressionQualit
 			{
 				UE_LOG(LogMutableCore, Log, TEXT("Image format conversion not implemented. Expensive generic one used. "));
 
-				Ptr<Image> TempBase = ImagePixelFormat(CompressionQuality, Base, EImageFormat::IF_RGBA_UBYTE, -1);
+				Ptr<Image> TempBase = ImagePixelFormat(CompressionQuality, Base, EImageFormat::RGBA_UByte, -1);
 				ImagePixelFormatInternal::BlockCompressionLODRangeHelper(
 						Result, TempBase.get(), ResultLODBegin, BaseLODBegin, NumLODs, CompressionQuality, miro::RGBA_to_ASTC6x6RGL);
 				ReleaseImage(TempBase);
@@ -1478,29 +1478,29 @@ void FImageOperator::ImagePixelFormat(bool& bOutSuccess, int32 CompressionQualit
 			break;
 		}
 
-		case EImageFormat::IF_ASTC_8x8_RGB_LDR:
+		case EImageFormat::ASTC_8x8_RGB_LDR:
 		{
 			switch (Base->GetFormat())
 			{
-			case EImageFormat::IF_RGBA_UBYTE:
+			case EImageFormat::RGBA_UByte:
 			{
 				ImagePixelFormatInternal::BlockCompressionLODRangeHelper(
 						Result, Base, ResultLODBegin, BaseLODBegin, NumLODs, CompressionQuality, miro::RGBA_to_ASTC8x8RGBL);
 				break;
 			}
-			case EImageFormat::IF_RGB_UBYTE:
+			case EImageFormat::RGB_UByte:
 			{
 				ImagePixelFormatInternal::BlockCompressionLODRangeHelper(
 						Result, Base, ResultLODBegin, BaseLODBegin, NumLODs, CompressionQuality, miro::RGB_to_ASTC8x8RGBL);
 				break;
 			}
-			//case EImageFormat::IF_ASTC_8x8_RGBA_LDR:
+			//case EImageFormat::ASTC_8x8_RGBA_LDR:
 			//{
 			//	ImagePixelFormatInternal::BlockCompressionLODRangeHelper(
 			//			Result, Base, ResultLODBegin, BaseLODBegin, NumLODs, CompressionQuality, miro::ASTC8x8RGBAL_to_ASTC8x8RGBL);
 			//	break;
 			//}
-			//case EImageFormat::IF_L_UBYTE:
+			//case EImageFormat::L_UByte:
 			//{
 			//	ImagePixelFormatInternal::BlockCompressionLODRangeHelper(
 			//			Result, Base, ResultLODBegin, BaseLODBegin, NumLODs, CompressionQuality, miro::L_to_ASTC8x8RGBL);
@@ -1510,7 +1510,7 @@ void FImageOperator::ImagePixelFormat(bool& bOutSuccess, int32 CompressionQualit
 			{
 				UE_LOG(LogMutableCore, Log, TEXT("Image format conversion not implemented. Expensive generic one used. "));
 
-				Ptr<Image> TempBase = ImagePixelFormat(CompressionQuality, Base, EImageFormat::IF_RGBA_UBYTE, -1);
+				Ptr<Image> TempBase = ImagePixelFormat(CompressionQuality, Base, EImageFormat::RGBA_UByte, -1);
 				ImagePixelFormatInternal::BlockCompressionLODRangeHelper(
 						Result, TempBase.get(), ResultLODBegin, BaseLODBegin, NumLODs, CompressionQuality, miro::RGBA_to_ASTC8x8RGBL);
 				ReleaseImage(TempBase);
@@ -1521,23 +1521,23 @@ void FImageOperator::ImagePixelFormat(bool& bOutSuccess, int32 CompressionQualit
 			break;
 		}
 
-		case EImageFormat::IF_ASTC_8x8_RGBA_LDR:
+		case EImageFormat::ASTC_8x8_RGBA_LDR:
 		{
 			switch (Base->GetFormat())
 			{
-			case EImageFormat::IF_RGBA_UBYTE:
+			case EImageFormat::RGBA_UByte:
 			{
 				ImagePixelFormatInternal::BlockCompressionLODRangeHelper(
 						Result, Base, ResultLODBegin, BaseLODBegin, NumLODs, CompressionQuality, miro::RGBA_to_ASTC8x8RGBAL);
 				break;
 			}
-			case EImageFormat::IF_RGB_UBYTE:
+			case EImageFormat::RGB_UByte:
 			{
 				ImagePixelFormatInternal::BlockCompressionLODRangeHelper(
 						Result, Base, ResultLODBegin, BaseLODBegin, NumLODs, CompressionQuality, miro::RGB_to_ASTC8x8RGBAL);
 				break;
 			}
-			//case EImageFormat::IF_ASTC_8x8_RGB_LDR:
+			//case EImageFormat::ASTC_8x8_RGB_LDR:
 			//{
 			//	ImagePixelFormatInternal::BlockCompressionLODRangeHelper(
 			//			Result, Base, ResultLODBegin, BaseLODBegin, NumLODs, CompressionQuality, miro::ASTC8x8RGBL_to_ASTC8x8RGBAL);
@@ -1547,7 +1547,7 @@ void FImageOperator::ImagePixelFormat(bool& bOutSuccess, int32 CompressionQualit
 			{
 				UE_LOG(LogMutableCore, Log, TEXT("Image format conversion not implemented. Expensive generic one used. "));
 
-				Ptr<Image> TempBase = ImagePixelFormat(CompressionQuality, Base, EImageFormat::IF_RGBA_UBYTE, -1);
+				Ptr<Image> TempBase = ImagePixelFormat(CompressionQuality, Base, EImageFormat::RGBA_UByte, -1);
 				ImagePixelFormatInternal::BlockCompressionLODRangeHelper(
 						Result, TempBase.get(), ResultLODBegin, BaseLODBegin, NumLODs, CompressionQuality, miro::RGBA_to_ASTC8x8RGBAL);
 				ReleaseImage(TempBase);
@@ -1558,17 +1558,17 @@ void FImageOperator::ImagePixelFormat(bool& bOutSuccess, int32 CompressionQualit
 			break;
 		}
 
-		case EImageFormat::IF_ASTC_8x8_RG_LDR:
+		case EImageFormat::ASTC_8x8_RG_LDR:
 		{
 			switch (Base->GetFormat())
 			{
-			case EImageFormat::IF_RGBA_UBYTE:
+			case EImageFormat::RGBA_UByte:
 			{
 				ImagePixelFormatInternal::BlockCompressionLODRangeHelper(
 						Result, Base, ResultLODBegin, BaseLODBegin, NumLODs, CompressionQuality, miro::RGBA_to_ASTC8x8RGL);
 				break;
 			}
-			case EImageFormat::IF_RGB_UBYTE:
+			case EImageFormat::RGB_UByte:
 			{
 				ImagePixelFormatInternal::BlockCompressionLODRangeHelper(
 						Result, Base, ResultLODBegin, BaseLODBegin, NumLODs, CompressionQuality, miro::RGB_to_ASTC8x8RGL);
@@ -1578,7 +1578,7 @@ void FImageOperator::ImagePixelFormat(bool& bOutSuccess, int32 CompressionQualit
 			{
 				UE_LOG(LogMutableCore, Log, TEXT("Image format conversion not implemented. Expensive generic one used. "));
 
-				Ptr<Image> TempBase = ImagePixelFormat(CompressionQuality, Base, EImageFormat::IF_RGBA_UBYTE, -1);
+				Ptr<Image> TempBase = ImagePixelFormat(CompressionQuality, Base, EImageFormat::RGBA_UByte, -1);
 				ImagePixelFormatInternal::BlockCompressionLODRangeHelper(
 						Result, TempBase.get(), ResultLODBegin, BaseLODBegin, NumLODs, CompressionQuality, miro::RGBA_to_ASTC8x8RGL);
 				ReleaseImage(TempBase);
diff --git a/Engine/Plugins/Mutable/Source/MutableRuntime/Private/MuR/OpImageProject.cpp b/Engine/Plugins/Mutable/Source/MutableRuntime/Private/MuR/OpImageProject.cpp
--- a/Engine/Plugins/Mutable/Source/MutableRuntime/Private/MuR/OpImageProject.cpp
+++ b/Engine/Plugins/Mutable/Source/MutableRuntime/Private/MuR/OpImageProject.cpp
@@ -781,7 +781,7 @@ namespace Private
 		EMeshBufferFormat LayoutBlockType = pMesh->VertexBuffers.Buffers[1].Channels[0].Format;
 		switch (LayoutBlockType)
 		{
-		case MBF_UINT64:
+		case EMeshBufferFormat::UInt64:
 		{
 			const uint64* LayoutBlockIds = reinterpret_cast<const uint64*>(pMesh->GetVertexBuffers().GetBufferData(1));
 			for (int32 V = 0; V < VertexCount; ++V)
@@ -796,7 +796,7 @@ namespace Private
 			break;
 		}
 
-		case MBF_UINT16:
+		case EMeshBufferFormat::UInt16:
 		{
 			const uint16* LayoutBlockIds = reinterpret_cast<const uint16*>(pMesh->GetVertexBuffers().GetBufferData(1));
 			for (int32 V = 0; V < VertexCount; ++V)
@@ -942,7 +942,7 @@ void ImageRasterProjectedPlanar(const Mesh* pMesh, Image* pImage,
 
     check(!pMask || pMask->GetSizeX() == pImage->GetSizeX());
     check(!pMask || pMask->GetSizeY() == pImage->GetSizeY());
-    check(!pMask || pMask->GetFormat()==EImageFormat::IF_L_UBYTE);
+    check(!pMask || pMask->GetFormat()==EImageFormat::L_UByte);
 
 	MUTABLE_CPUPROFILER_SCOPE(ImageProject);
 
@@ -954,15 +954,15 @@ void ImageRasterProjectedPlanar(const Mesh* pMesh, Image* pImage,
 		ProcessorFeatures |= EPixelProcessorFeatures::VectorizedImpl;
 	}
 
-	if (pSource->GetFormat() == EImageFormat::IF_L_UBYTE)
+	if (pSource->GetFormat() == EImageFormat::L_UByte)
 	{
 		ProcessorFeatures |= EPixelProcessorFeatures::FormatL;
 	}
-	else if (pSource->GetFormat() == EImageFormat::IF_RGB_UBYTE)
+	else if (pSource->GetFormat() == EImageFormat::RGB_UByte)
 	{
 		ProcessorFeatures |= EPixelProcessorFeatures::FormatRGB;
 	}
-	else if (pSource->GetFormat() == EImageFormat::IF_RGBA_UBYTE || pSource->GetFormat() == EImageFormat::IF_BGRA_UBYTE)
+	else if (pSource->GetFormat() == EImageFormat::RGBA_UByte || pSource->GetFormat() == EImageFormat::BGRA_UByte)
 	{
 		ProcessorFeatures |= EPixelProcessorFeatures::FormatRGBA;
 	}
@@ -1157,7 +1157,7 @@ void ImageRasterProjectedWrapping( const Mesh* pMesh, Image* pImage,
     
 	check(!pMask || pMask->GetSizeX() == pImage->GetSizeX());
     check(!pMask || pMask->GetSizeY() == pImage->GetSizeY());
-    check(!pMask || pMask->GetFormat()==EImageFormat::IF_L_UBYTE);
+    check(!pMask || pMask->GetFormat()==EImageFormat::L_UByte);
 
 	MUTABLE_CPUPROFILER_SCOPE(ImageProjectWrapping);
 
@@ -1168,15 +1168,15 @@ void ImageRasterProjectedWrapping( const Mesh* pMesh, Image* pImage,
 		ProcessorFeatures |= EPixelProcessorFeatures::VectorizedImpl;
 	}
 
-	if (pSource->GetFormat() == EImageFormat::IF_L_UBYTE)
+	if (pSource->GetFormat() == EImageFormat::L_UByte)
 	{
 		ProcessorFeatures |= EPixelProcessorFeatures::FormatL;
 	}
-	else if (pSource->GetFormat() == EImageFormat::IF_RGB_UBYTE)
+	else if (pSource->GetFormat() == EImageFormat::RGB_UByte)
 	{
 		ProcessorFeatures |= EPixelProcessorFeatures::FormatRGB;
 	}
-	else if (pSource->GetFormat() == EImageFormat::IF_RGBA_UBYTE || pSource->GetFormat() == EImageFormat::IF_BGRA_UBYTE)
+	else if (pSource->GetFormat() == EImageFormat::RGBA_UByte || pSource->GetFormat() == EImageFormat::BGRA_UByte)
 	{
 		ProcessorFeatures |= EPixelProcessorFeatures::FormatRGBA;
 	}
@@ -1371,7 +1371,7 @@ void ImageRasterProjectedCylindrical( const Mesh* pMesh, Image* pImage,
 
     check(!pMask || pMask->GetSizeX() == pImage->GetSizeX());
     check(!pMask || pMask->GetSizeY() == pImage->GetSizeY());
-    check(!pMask || pMask->GetFormat()==EImageFormat::IF_L_UBYTE);
+    check(!pMask || pMask->GetFormat()==EImageFormat::L_UByte);
 
 	MUTABLE_CPUPROFILER_SCOPE(ImageProjectCylindrical);
 
@@ -1382,15 +1382,15 @@ void ImageRasterProjectedCylindrical( const Mesh* pMesh, Image* pImage,
 		ProcessorFeatures |= EPixelProcessorFeatures::VectorizedImpl;
 	}
 
-	if (pSource->GetFormat() == EImageFormat::IF_L_UBYTE)
+	if (pSource->GetFormat() == EImageFormat::L_UByte)
 	{
 		ProcessorFeatures |= EPixelProcessorFeatures::FormatL;
 	}
-	else if (pSource->GetFormat() == EImageFormat::IF_RGB_UBYTE)
+	else if (pSource->GetFormat() == EImageFormat::RGB_UByte)
 	{
 		ProcessorFeatures |= EPixelProcessorFeatures::FormatRGB;
 	}
-	else if (pSource->GetFormat() == EImageFormat::IF_RGBA_UBYTE || pSource->GetFormat() == EImageFormat::IF_BGRA_UBYTE)
+	else if (pSource->GetFormat() == EImageFormat::RGBA_UByte || pSource->GetFormat() == EImageFormat::BGRA_UByte)
 	{
 		ProcessorFeatures |= EPixelProcessorFeatures::FormatRGBA;
 	}
@@ -1572,9 +1572,9 @@ float ComputeProjectedFootprintBestMip(
 	const Mesh* pMesh, const FProjector& Projector, const FVector2f& TargetSize, const FVector2f& SourceSize)
 {
 	// Compute projected mesh footprint on the traget image and extract a optimal mip for source image.
-	const bool bIsPlanarProjection = Projector.type == PROJECTOR_TYPE::PLANAR;
-	const bool bIsCylindicalProjection = Projector.type == PROJECTOR_TYPE::CYLINDRICAL;
-	const bool bIsWrappingProjection = Projector.type == PROJECTOR_TYPE::WRAPPING;
+	const bool bIsPlanarProjection = Projector.type == EProjectorType::Planar;
+	const bool bIsCylindicalProjection = Projector.type == EProjectorType::Cylindrical;
+	const bool bIsWrappingProjection = Projector.type == EProjectorType::Wrapping;
 
 	check(pMesh->GetIndexBuffers().GetElementSize(0) == sizeof(uint32));
 	const uint32* IndicesPtr = reinterpret_cast<const uint32*>(pMesh->GetIndexBuffers().GetBufferData(0));
@@ -1723,7 +1723,7 @@ inline void MeshCreateCollapsedVertexMap( const Mesh* pMesh,
 	CollapsedVerticesMap.Reserve(NumVertices);
 
 	// Get Vertices
-	mu::UntypedMeshBufferIteratorConst ItPosition = mu::UntypedMeshBufferIteratorConst(pMesh->GetVertexBuffers(), mu::MBS_POSITION);
+	mu::UntypedMeshBufferIteratorConst ItPosition = mu::UntypedMeshBufferIteratorConst(pMesh->GetVertexBuffers(), mu::EMeshBufferSemantic::Position);
 	FVector3f* VertexData = OutVertices.GetData();
 	
 	for (int32 VertexIndex = 0; VertexIndex < NumVertices; ++VertexIndex)
@@ -2878,14 +2878,14 @@ void MeshProject_Optimised_Wrapping(const Mesh* pMesh,
 					{
 						switch (LayoutBlockType)
 						{
-						case MBF_UINT64:
+						case EMeshBufferFormat::UInt64:
 						{
 							const uint64* LayoutBlockIds = reinterpret_cast<const uint64*>(pMesh->GetVertexBuffers().GetBufferData(1));
 							pResultLayoutBlockIds[currentVertex] = LayoutBlockIds[i];
 							break;
 						}
 
-						case MBF_UINT16:
+						case EMeshBufferFormat::UInt16:
 						{
 							const uint16* LayoutBlockIds = reinterpret_cast<const uint16*>(pMesh->GetVertexBuffers().GetBufferData(1));
 							uint64 BlockId = (uint64(pMesh->MeshIDPrefix) << 32) | uint64(LayoutBlockIds[i]);
@@ -2995,8 +2995,8 @@ void MeshProject_Optimised(Mesh* Result, const Mesh* pMesh, const FProjector& pr
     switch (projector.type)
     {
 
-    case PROJECTOR_TYPE::PLANAR:
-    case PROJECTOR_TYPE::CYLINDRICAL:
+    case EProjectorType::Planar:
+    case EProjectorType::Cylindrical:
     {
         CreateMeshOptimisedForProjection(Result, layout);
         Result->GetVertexBuffers().SetElementCount(vertexCount);
@@ -3013,7 +3013,7 @@ void MeshProject_Optimised(Mesh* Result, const Mesh* pMesh, const FProjector& pr
 		const uint32* pIndices = reinterpret_cast<const uint32*>(pMesh->GetIndexBuffers().GetBufferData(0));
         int faceCount = pMesh->GetFaceCount();
 
-        if (projector.type==PROJECTOR_TYPE::PLANAR)
+        if (projector.type==EProjectorType::Planar)
         {
             MeshProject_Optimised_Planar(pVertices, vertexCount,
                                           pIndices, faceCount,
@@ -3038,7 +3038,7 @@ void MeshProject_Optimised(Mesh* Result, const Mesh* pMesh, const FProjector& pr
         break;
     }
 
-    case PROJECTOR_TYPE::WRAPPING:
+    case EProjectorType::Wrapping:
 	{
         CreateMeshOptimisedForWrappingProjection(Result, layout);
 
@@ -3116,9 +3116,9 @@ void CreateMeshOptimisedForProjection(Mesh* Result, int32 layout)
     Result->GetVertexBuffers().SetBufferCount( 1 );
     Result->GetIndexBuffers().SetBufferCount( 1 );
 
-    EMeshBufferSemantic semantics[3] =	{ MBS_TEXCOORDS,	MBS_POSITION,	MBS_NORMAL };
+    EMeshBufferSemantic semantics[3] =	{ EMeshBufferSemantic::TexCoords,	EMeshBufferSemantic::Position,	EMeshBufferSemantic::Normal };
     int32 semanticIndices[3] =			{ 0,				0,				0 };
-    EMeshBufferFormat formats[3] =		{ MBF_FLOAT32,		MBF_FLOAT32,	MBF_FLOAT32 };
+    EMeshBufferFormat formats[3] =		{ EMeshBufferFormat::Float32,		EMeshBufferFormat::Float32,	EMeshBufferFormat::Float32 };
     int32 componentCounts[3] =			{ 2,				3,				3 };
     int32 offsets[3] =					{ 0,				8,				20 };
     semanticIndices[0] = layout;
@@ -3128,9 +3128,9 @@ void CreateMeshOptimisedForProjection(Mesh* Result, int32 layout)
               formats, componentCounts,
               offsets );
 
-    EMeshBufferSemantic isemantics[1] =	{ MBS_VERTEXINDEX };
+    EMeshBufferSemantic isemantics[1] =	{ EMeshBufferSemantic::VertexIndex };
     int32 isemanticIndices[1] =			{ 0 };
-    EMeshBufferFormat iformats[1] =		{ MBF_UINT32 };
+    EMeshBufferFormat iformats[1] =		{ EMeshBufferFormat::UInt32 };
     int32 icomponentCounts[1] =			{ 1 };
     int32 ioffsets[1] =					{ 0 };
     Result->GetIndexBuffers().SetBuffer
@@ -3147,25 +3147,25 @@ void CreateMeshOptimisedForWrappingProjection(Mesh* Result, int32 LayoutIndex)
     Result->GetVertexBuffers().SetBufferCount( 2 );
     Result->GetIndexBuffers().SetBufferCount( 1 );
 
-    EMeshBufferSemantic semantics[3] =	{ MBS_TEXCOORDS,	MBS_POSITION,	MBS_NORMAL };
+    EMeshBufferSemantic semantics[3] =	{ EMeshBufferSemantic::TexCoords,	EMeshBufferSemantic::Position,	EMeshBufferSemantic::Normal };
     int32 semanticIndices[3] =			{ 0,				0,				0 };
-    EMeshBufferFormat formats[3] =		{ MBF_FLOAT32,		MBF_FLOAT32,	MBF_FLOAT32 };
+    EMeshBufferFormat formats[3] =		{ EMeshBufferFormat::Float32,		EMeshBufferFormat::Float32,	EMeshBufferFormat::Float32 };
     int32 componentCounts[3] =			{ 2,				3,				3 };
     int32 offsets[3] =					{ 0,				8,				20 };
     semanticIndices[0] = LayoutIndex;
     Result->GetVertexBuffers().SetBuffer( 0, 32, 3, semantics, semanticIndices, formats, componentCounts, offsets );
 
-	EMeshBufferSemantic LayoutSemantics[1] = { MBS_LAYOUTBLOCK };
+	EMeshBufferSemantic LayoutSemantics[1] = { EMeshBufferSemantic::LayoutBlock };
 	int32 LayoutSemanticIndices[1] = { 0 };
-	EMeshBufferFormat LayoutFormats[1] = { MBF_UINT64 };
+	EMeshBufferFormat LayoutFormats[1] = { EMeshBufferFormat::UInt64 };
 	int32 LayoutComponentCounts[1] = { 1 };
 	int32 LayoutOffsets[1] = { 0 };
 	LayoutSemanticIndices[0] = LayoutIndex;
 	Result->GetVertexBuffers().SetBuffer(1, sizeof(uint64), 1, LayoutSemantics, LayoutSemanticIndices, LayoutFormats, LayoutComponentCounts, LayoutOffsets);
 
-    EMeshBufferSemantic isemantics[1] =	{ MBS_VERTEXINDEX };
+    EMeshBufferSemantic isemantics[1] =	{ EMeshBufferSemantic::VertexIndex };
     int32 isemanticIndices[1] =			{ 0 };
-    EMeshBufferFormat iformats[1] =		{ MBF_UINT32 };
+    EMeshBufferFormat iformats[1] =		{ EMeshBufferFormat::UInt32 };
     int32 icomponentCounts[1] =			{ 1 };
     int32 ioffsets[1] =					{ 0 };
     Result->GetIndexBuffers().SetBuffer( 0, 4, 1, isemantics, isemanticIndices, iformats, icomponentCounts, ioffsets );
diff --git a/Engine/Plugins/Mutable/Source/MutableRuntime/Private/MuR/OpImageRasterMesh.h b/Engine/Plugins/Mutable/Source/MutableRuntime/Private/MuR/OpImageRasterMesh.h
--- a/Engine/Plugins/Mutable/Source/MutableRuntime/Private/MuR/OpImageRasterMesh.h
+++ b/Engine/Plugins/Mutable/Source/MutableRuntime/Private/MuR/OpImageRasterMesh.h
@@ -37,7 +37,7 @@ namespace mu
 			return;
 		}
 
-		check( pImage->GetFormat()== EImageFormat::IF_L_UBYTE );
+		check( pImage->GetFormat()== EImageFormat::L_UByte );
 
 		int32 sizeX = pImage->GetSizeX();
 		int32 sizeY = pImage->GetSizeY();
@@ -47,7 +47,7 @@ namespace mu
 		TArray< RasterVertex<1> > vertices;
 		vertices.SetNumZeroed(vertexCount);
 
-		UntypedMeshBufferIteratorConst texIt( pMesh->GetVertexBuffers(), MBS_TEXCOORDS, 0 );
+		UntypedMeshBufferIteratorConst texIt( pMesh->GetVertexBuffers(), EMeshBufferSemantic::TexCoords, 0 );
 		if (!texIt.ptr())
 		{
 			ensure(false);
@@ -57,8 +57,8 @@ namespace mu
 		for ( int32 v=0; v<vertexCount; ++v )
 		{
             float uv[2] = {0.0f,0.0f};
-			ConvertData( 0, uv, MBF_FLOAT32, texIt.ptr(), texIt.GetFormat() );
-			ConvertData( 1, uv, MBF_FLOAT32, texIt.ptr(), texIt.GetFormat() );
+			ConvertData( 0, uv, EMeshBufferFormat::Float32, texIt.ptr(), texIt.GetFormat() );
+			ConvertData( 1, uv, EMeshBufferFormat::Float32, texIt.ptr(), texIt.GetFormat() );
 
 			bool bUseCropping = UncroppedSize[0] > 0;
 			if (bUseCropping)
@@ -79,17 +79,17 @@ namespace mu
 		TArray<int32> indices;
 		indices.SetNumZeroed(faceCount * 3);
 
-		UntypedMeshBufferIteratorConst indIt( pMesh->GetIndexBuffers(), MBS_VERTEXINDEX, 0 );
+		UntypedMeshBufferIteratorConst indIt( pMesh->GetIndexBuffers(), EMeshBufferSemantic::VertexIndex, 0 );
 		for ( int32 i=0; i<faceCount*3; ++i )
 		{
             uint32_t index=0;
-			ConvertData( 0, &index, MBF_UINT32, indIt.ptr(), indIt.GetFormat() );
+			ConvertData( 0, &index, EMeshBufferFormat::UInt32, indIt.ptr(), indIt.GetFormat() );
 
 			indices[i] = index;
 			++indIt;
 		}
 
-        UntypedMeshBufferIteratorConst bloIt( pMesh->GetVertexBuffers(), MBS_LAYOUTBLOCK, LayoutIndex );
+        UntypedMeshBufferIteratorConst bloIt( pMesh->GetVertexBuffers(), EMeshBufferSemantic::LayoutBlock, LayoutIndex );
 
         if (BlockId== FLayoutBlock::InvalidBlockId || bloIt.GetElementSize()==0 )
 		{
@@ -124,7 +124,7 @@ namespace mu
 			TArray<uint64> VertexBlockIds;
 			VertexBlockIds.SetNumZeroed(vertexCount);
 
-			if (bloIt.GetFormat() == MBF_UINT16)
+			if (bloIt.GetFormat() == EMeshBufferFormat::UInt16)
 			{
 				// Relative blocks.
 				const uint16* SourceIds = reinterpret_cast<const uint16*>(bloIt.ptr());
@@ -135,7 +135,7 @@ namespace mu
 					VertexBlockIds[i] = Id;
 				}
 			}
-			else if (bloIt.GetFormat() == MBF_UINT64)
+			else if (bloIt.GetFormat() == EMeshBufferFormat::UInt64)
 			{
 				// Absolute blocks.
 				const uint64* SourceIds = reinterpret_cast<const uint64*>(bloIt.ptr());
diff --git a/Engine/Plugins/Mutable/Source/MutableRuntime/Private/MuR/OpImageResize.cpp b/Engine/Plugins/Mutable/Source/MutableRuntime/Private/MuR/OpImageResize.cpp
--- a/Engine/Plugins/Mutable/Source/MutableRuntime/Private/MuR/OpImageResize.cpp
+++ b/Engine/Plugins/Mutable/Source/MutableRuntime/Private/MuR/OpImageResize.cpp
@@ -84,20 +84,20 @@ namespace mu
 		switch (Base->GetFormat())
 		{
 
-		case EImageFormat::IF_L_UBYTE:
+		case EImageFormat::L_UByte:
 		{
 			ImageMagnifyX<1>(Dest, Base);
 			break;
 		}
 
-		case EImageFormat::IF_RGB_UBYTE:
+		case EImageFormat::RGB_UByte:
 		{
 			ImageMagnifyX<3>(Dest, Base);
 			break;
 		}
 
-		case EImageFormat::IF_BGRA_UBYTE:
-		case EImageFormat::IF_RGBA_UBYTE:
+		case EImageFormat::BGRA_UByte:
+		case EImageFormat::RGBA_UByte:
 		{
 			ImageMagnifyX<4>(Dest, Base);
 			break;
@@ -326,7 +326,7 @@ namespace mu
 		switch (Base->GetFormat())
 		{
 
-		case EImageFormat::IF_L_UBYTE:
+		case EImageFormat::L_UByte:
 		{
 			if (2 * Dest->GetSizeX() == Base->GetSizeX())
 			{
@@ -341,7 +341,7 @@ namespace mu
 			break;
 		}
 
-		case EImageFormat::IF_RGB_UBYTE:
+		case EImageFormat::RGB_UByte:
 		{
 			if (2 * Dest->GetSizeX() == Base->GetSizeX())
 			{
@@ -356,8 +356,8 @@ namespace mu
 			break;
 		}
 
-		case EImageFormat::IF_BGRA_UBYTE:
-		case EImageFormat::IF_RGBA_UBYTE:
+		case EImageFormat::BGRA_UByte:
+		case EImageFormat::RGBA_UByte:
 		{
 			if (2 * Dest->GetSizeX() == Base->GetSizeX())
 			{
@@ -477,20 +477,20 @@ namespace mu
 		switch (Base->GetFormat())
 		{
 
-		case EImageFormat::IF_L_UBYTE:
+		case EImageFormat::L_UByte:
 		{
 			ImageMagnifyY<1>(Dest, Base);
 			break;
 		}
 
-		case EImageFormat::IF_RGB_UBYTE:
+		case EImageFormat::RGB_UByte:
 		{
 			ImageMagnifyY<3>(Dest, Base);
 			break;
 		}
 
-		case EImageFormat::IF_RGBA_UBYTE:
-		case EImageFormat::IF_BGRA_UBYTE:
+		case EImageFormat::RGBA_UByte:
+		case EImageFormat::BGRA_UByte:
 		{
 			ImageMagnifyY<4>(Dest, Base);
 			break;
@@ -647,7 +647,7 @@ namespace mu
 		switch (Base->GetFormat())
 		{
 
-		case EImageFormat::IF_L_UBYTE:
+		case EImageFormat::L_UByte:
 		{
 			if (2 * Dest->GetSizeY() == Base->GetSizeY())
 			{
@@ -662,7 +662,7 @@ namespace mu
 			break;
 		}
 
-		case EImageFormat::IF_RGB_UBYTE:
+		case EImageFormat::RGB_UByte:
 		{
 			if (2 * Dest->GetSizeY() == Base->GetSizeY())
 			{
@@ -677,8 +677,8 @@ namespace mu
 			break;
 		}
 
-		case EImageFormat::IF_RGBA_UBYTE:
-		case EImageFormat::IF_BGRA_UBYTE:
+		case EImageFormat::RGBA_UByte:
+		case EImageFormat::BGRA_UByte:
 		{
 			if (2 * Dest->GetSizeY() == Base->GetSizeY())
 			{
@@ -933,7 +933,7 @@ namespace
 			// Uncompressed format, resize in one go directly to Dest.
 			switch(SrcFormat)
 			{
-			case EImageFormat::IF_L_UBYTE:
+			case EImageFormat::L_UByte:
 			{
 				constexpr int32 NumChannels = 1;
 				ImageResizeLinearDownByLessThanTwoSubImage<NumChannels, bUseVectorImpl>(
@@ -942,7 +942,7 @@ namespace
 				break;
 			}
 
-			case EImageFormat::IF_RGB_UBYTE:
+			case EImageFormat::RGB_UByte:
 			{
 				constexpr int32 NumChannels = 3;
 				ImageResizeLinearDownByLessThanTwoSubImage<NumChannels, bUseVectorImpl>(
@@ -951,8 +951,8 @@ namespace
 				break;
 			}
 
-			case EImageFormat::IF_RGBA_UBYTE:
-			case EImageFormat::IF_BGRA_UBYTE:
+			case EImageFormat::RGBA_UByte:
+			case EImageFormat::BGRA_UByte:
 			{
 				constexpr int32 NumChannels = 4;
 				ImageResizeLinearDownByLessThanTwoSubImage<NumChannels, bUseVectorImpl>(
diff --git a/Engine/Plugins/Mutable/Source/MutableRuntime/Private/MuR/OpImageSaturate.h b/Engine/Plugins/Mutable/Source/MutableRuntime/Private/MuR/OpImageSaturate.h
--- a/Engine/Plugins/Mutable/Source/MutableRuntime/Private/MuR/OpImageSaturate.h
+++ b/Engine/Plugins/Mutable/Source/MutableRuntime/Private/MuR/OpImageSaturate.h
@@ -38,7 +38,7 @@ namespace mu
 
 			switch (AImage->GetFormat())
 			{
-			case EImageFormat::IF_RGB_UBYTE:
+			case EImageFormat::RGB_UByte:
 			{
 				if constexpr (!bUseVectorImpl)
 				{
@@ -94,7 +94,7 @@ namespace mu
 				break;
 			}
 
-			case EImageFormat::IF_RGBA_UBYTE:
+			case EImageFormat::RGBA_UByte:
 			{
 				if constexpr (!bUseVectorImpl)
 				{
@@ -151,7 +151,7 @@ namespace mu
 				break;
 			}
 
-			case EImageFormat::IF_BGRA_UBYTE:
+			case EImageFormat::BGRA_UByte:
 			{
 				if constexpr (!bUseVectorImpl)
 				{
diff --git a/Engine/Plugins/Mutable/Source/MutableRuntime/Private/MuR/OpImageSwizzle.cpp b/Engine/Plugins/Mutable/Source/MutableRuntime/Private/MuR/OpImageSwizzle.cpp
--- a/Engine/Plugins/Mutable/Source/MutableRuntime/Private/MuR/OpImageSwizzle.cpp
+++ b/Engine/Plugins/Mutable/Source/MutableRuntime/Private/MuR/OpImageSwizzle.cpp
@@ -25,16 +25,16 @@ namespace mu
 
 		switch (Format)
 		{
-		case EImageFormat::IF_L_UBYTE:
+		case EImageFormat::L_UByte:
 			NumDestChannels = 1;
 			break;
 
-		case EImageFormat::IF_RGB_UBYTE:
+		case EImageFormat::RGB_UByte:
 			NumDestChannels = 3;
 			break;
 
-        case EImageFormat::IF_RGBA_UBYTE:
-        case EImageFormat::IF_BGRA_UBYTE:
+        case EImageFormat::RGBA_UByte:
+        case EImageFormat::BGRA_UByte:
 			NumDestChannels = 4;
 			break;
 
@@ -46,7 +46,7 @@ namespace mu
 		{
 			const Image* Src = Sources[Channel].get();
 			
-			const int32 DestChannel = Format == EImageFormat::IF_BGRA_UBYTE && Channel < 3
+			const int32 DestChannel = Format == EImageFormat::BGRA_UByte && Channel < 3
 								    ? FMath::Abs(Channel - 2)
 									: Channel;
 
@@ -59,7 +59,7 @@ namespace mu
 				EImageFormat SrcFormat = Sources[Channel]->GetFormat();
 				switch (SrcFormat)
 				{
-				case EImageFormat::IF_L_UBYTE:
+				case EImageFormat::L_UByte:
 				{
 					if (Channels[Channel] < 1)
 					{
@@ -98,7 +98,7 @@ namespace mu
 					}
 					break;
 				}
-				case EImageFormat::IF_RGB_UBYTE:
+				case EImageFormat::RGB_UByte:
 				{
 					if (Channels[Channel] < 3)
 					{
@@ -137,12 +137,12 @@ namespace mu
 					}
 					break;
 				}
-				case EImageFormat::IF_RGBA_UBYTE:
-				case EImageFormat::IF_BGRA_UBYTE:
+				case EImageFormat::RGBA_UByte:
+				case EImageFormat::BGRA_UByte:
 				{
 					if (Channels[Channel] < 4)
 					{
-						const int32 SrcChannel = SrcFormat == EImageFormat::IF_BGRA_UBYTE && Channels[Channel] < 3
+						const int32 SrcChannel = SrcFormat == EImageFormat::BGRA_UByte && Channels[Channel] < 3
 											   ? FMath::Abs(int32(Channels[Channel]) - 2)
 											   : Channels[Channel];
 
diff --git a/Engine/Plugins/Mutable/Source/MutableRuntime/Private/MuR/OpImageTransform.h b/Engine/Plugins/Mutable/Source/MutableRuntime/Private/MuR/OpImageTransform.h
--- a/Engine/Plugins/Mutable/Source/MutableRuntime/Private/MuR/OpImageTransform.h
+++ b/Engine/Plugins/Mutable/Source/MutableRuntime/Private/MuR/OpImageTransform.h
@@ -67,7 +67,7 @@ namespace mu
 		switch (pDestImage->GetFormat())
 		{
 		
-		case EImageFormat::IF_L_UBYTE:
+		case EImageFormat::L_UByte:
 		{
 
 			if (bUseVectorImplCvar)
@@ -111,7 +111,7 @@ namespace mu
 			break;
 		}
 		
-		case EImageFormat::IF_RGB_UBYTE:
+		case EImageFormat::RGB_UByte:
 		{
 			if (bUseVectorImplCvar)
 			{
@@ -155,8 +155,8 @@ namespace mu
 			break;
 		}
 
-		case EImageFormat::IF_BGRA_UBYTE:
-		case EImageFormat::IF_RGBA_UBYTE:
+		case EImageFormat::BGRA_UByte:
+		case EImageFormat::RGBA_UByte:
 		{
 			if (bUseVectorImplCvar)
 			{
diff --git a/Engine/Plugins/Mutable/Source/MutableRuntime/Private/MuR/OpLayoutRemoveBlocks.h b/Engine/Plugins/Mutable/Source/MutableRuntime/Private/MuR/OpLayoutRemoveBlocks.h
--- a/Engine/Plugins/Mutable/Source/MutableRuntime/Private/MuR/OpLayoutRemoveBlocks.h
+++ b/Engine/Plugins/Mutable/Source/MutableRuntime/Private/MuR/OpLayoutRemoveBlocks.h
@@ -21,8 +21,8 @@ namespace mu
 		const Layout* Source = InMesh->GetLayout(InLayoutIndex);
 		Ptr<Layout> pResult;
 
-		UntypedMeshBufferIteratorConst itBlocks(InMesh->GetVertexBuffers(), MBS_LAYOUTBLOCK, InLayoutIndex);
-		if (itBlocks.GetFormat() == MBF_UINT16)
+		UntypedMeshBufferIteratorConst itBlocks(InMesh->GetVertexBuffers(), EMeshBufferSemantic::LayoutBlock, InLayoutIndex);
+		if (itBlocks.GetFormat() == EMeshBufferFormat::UInt16)
 		{
 			// Relative blocks.
 
@@ -67,7 +67,7 @@ namespace mu
 			pResult->SetBlockCount(DestBlockIndex);
 
 		}
-		else if (itBlocks.GetFormat() == MBF_UINT64)
+		else if (itBlocks.GetFormat() == EMeshBufferFormat::UInt64)
 		{
 			// Absolute blocks.
 
@@ -100,7 +100,7 @@ namespace mu
 			pResult->SetBlockCount(DestBlockIndex);
 
 		}
-		else if (itBlocks.GetFormat() == MBF_NONE)
+		else if (itBlocks.GetFormat() == EMeshBufferFormat::None)
 		{
 			// This seems to happen.
 			// May this happen when entire meshes are removed?
diff --git a/Engine/Plugins/Mutable/Source/MutableRuntime/Private/MuR/OpMeshApplyLayout.h b/Engine/Plugins/Mutable/Source/MutableRuntime/Private/MuR/OpMeshApplyLayout.h
--- a/Engine/Plugins/Mutable/Source/MutableRuntime/Private/MuR/OpMeshApplyLayout.h
+++ b/Engine/Plugins/Mutable/Source/MutableRuntime/Private/MuR/OpMeshApplyLayout.h
@@ -18,11 +18,11 @@ namespace mu
 
         int32 buffer = -1;
 		int32 channel = -1;
-		Applied->GetVertexBuffers().FindChannel( MBS_TEXCOORDS, TexCoordsIndex, &buffer, &channel );
+		Applied->GetVertexBuffers().FindChannel( EMeshBufferSemantic::TexCoords, TexCoordsIndex, &buffer, &channel );
 
         int32 layoutBuffer = -1;
         int32 layoutChannel = -1;
-        Applied->GetVertexBuffers().FindChannel( MBS_LAYOUTBLOCK, TexCoordsIndex, &layoutBuffer, &layoutChannel );
+        Applied->GetVertexBuffers().FindChannel( EMeshBufferSemantic::LayoutBlock, TexCoordsIndex, &layoutBuffer, &layoutChannel );
 
 		if (buffer < 0 || layoutBuffer < 0)
 		{
@@ -36,7 +36,7 @@ namespace mu
 		int32 components;
 		int32 offset;
 		Applied->GetVertexBuffers().GetChannel( buffer, channel, &semantic, &semanticIndex, &format, &components, &offset );
-		check( semantic == MBS_TEXCOORDS );
+		check( semantic == EMeshBufferSemantic::TexCoords );
 
         uint8* pData = Applied->GetVertexBuffers().GetBufferData( buffer );
 		int32 elemSize = Applied->GetVertexBuffers().GetElementSize( buffer );
@@ -70,7 +70,7 @@ namespace mu
         check( Applied->GetVertexBuffers().Buffers[layoutBuffer].Channels[layoutChannel].ComponentCount==1 );
 
         const uint16* InLayoutData = reinterpret_cast<const uint16*>( Applied->GetVertexBuffers().GetBufferData( layoutBuffer ) );
-		UntypedMeshBufferIterator ItLayoutData(Applied->GetVertexBuffers(), MBS_LAYOUTBLOCK, TexCoordsIndex);
+		UntypedMeshBufferIterator ItLayoutData(Applied->GetVertexBuffers(), EMeshBufferSemantic::LayoutBlock, TexCoordsIndex);
 
 		// In some corner case involving automatic LODs and remove meshes behaving differently among them we may need to remove vertices that don't 
 		// have any block in the current layout. Track them here.
@@ -80,14 +80,14 @@ namespace mu
 		for ( int32 v=0; v<Applied->GetVertexBuffers().GetElementCount(); ++v )
 		{
 			uint64 BlockId = 0;
-			if (ItLayoutData.GetFormat() == MBF_UINT16)
+			if (ItLayoutData.GetFormat() == EMeshBufferFormat::UInt16)
 			{
 				// Relative blocks.
 				const uint16* SourceIds = reinterpret_cast<const uint16*>(ItLayoutData.ptr());
 				BlockId = *SourceIds;
 				BlockId = BlockId | (uint64(Applied->MeshIDPrefix) << 32);
 			}
-			else if (ItLayoutData.GetFormat() == MBF_UINT64)
+			else if (ItLayoutData.GetFormat() == EMeshBufferFormat::UInt64)
 			{
 				// Absolute blocks.
 				const uint64* SourceIds = reinterpret_cast<const uint64*>(ItLayoutData.ptr());
@@ -114,7 +114,7 @@ namespace mu
 
 			switch( format )
 			{
-			case MBF_FLOAT32:
+			case EMeshBufferFormat::Float32:
 			{
 				FVector2f uv;
 				FMemory::Memcpy( &uv, pVertices, sizeof(float)*2 );
@@ -124,7 +124,7 @@ namespace mu
 				break;
 			}
 
-			case MBF_FLOAT16:
+			case EMeshBufferFormat::Float16:
 			{
 				// TODO: Optimise
 				FFloat16* pUV = reinterpret_cast<FFloat16*>( pVertices );
@@ -139,10 +139,10 @@ namespace mu
 				break;
 			}
 
-			case MBF_NUINT32:
-			case MBF_NINT32:
-			case MBF_UINT32:
-			case MBF_INT32:
+			case EMeshBufferFormat::NUInt32:
+			case EMeshBufferFormat::NInt32:
+			case EMeshBufferFormat::UInt32:
+			case EMeshBufferFormat::Int32:
 			{
 				// TODO: Optimise
                 uint32* pUV = reinterpret_cast<uint32*>( pVertices );
@@ -157,10 +157,10 @@ namespace mu
 				break;
 			}
 
-			case MBF_NUINT16:
-			case MBF_NINT16:
-			case MBF_UINT16:
-			case MBF_INT16:
+			case EMeshBufferFormat::NUInt16:
+			case EMeshBufferFormat::NInt16:
+			case EMeshBufferFormat::UInt16:
+			case EMeshBufferFormat::Int16:
 			{
 				// TODO: Optimise
                 uint16* pUV = reinterpret_cast<uint16*>( pVertices );
@@ -176,10 +176,10 @@ namespace mu
 				break;
 			}
 
-			case MBF_NUINT8:
-			case MBF_NINT8:
-			case MBF_UINT8:
-			case MBF_INT8:
+			case EMeshBufferFormat::NUInt8:
+			case EMeshBufferFormat::NInt8:
+			case EMeshBufferFormat::UInt8:
+			case EMeshBufferFormat::Int8:
 			{
 				// TODO: Optimise
                 uint8* pUV = reinterpret_cast<uint8*>( pVertices );
diff --git a/Engine/Plugins/Mutable/Source/MutableRuntime/Private/MuR/OpMeshApplyPose.h b/Engine/Plugins/Mutable/Source/MutableRuntime/Private/MuR/OpMeshApplyPose.h
--- a/Engine/Plugins/Mutable/Source/MutableRuntime/Private/MuR/OpMeshApplyPose.h
+++ b/Engine/Plugins/Mutable/Source/MutableRuntime/Private/MuR/OpMeshApplyPose.h
@@ -83,7 +83,7 @@ namespace mu
 		}
 
         // Get pointers to vertex position data
-		MeshBufferIteratorConst<MBF_FLOAT32, float, 3> SourcePositionIterBegin(BaseMesh->VertexBuffers, MBS_POSITION, 0);
+		MeshBufferIteratorConst<EMeshBufferFormat::Float32, float, 3> SourcePositionIterBegin(BaseMesh->VertexBuffers, EMeshBufferSemantic::Position, 0);
         if (!SourcePositionIterBegin.ptr())
         {
             // Formats not implemented
@@ -93,8 +93,8 @@ namespace mu
         }
  
 		// Get pointers to skinning data
-		UntypedMeshBufferIteratorConst BoneIndicesIterBegin(BaseMesh->VertexBuffers, MBS_BONEINDICES, 0);
-		UntypedMeshBufferIteratorConst BoneWeightsIterBegin(BaseMesh->VertexBuffers, MBS_BONEWEIGHTS, 0);
+		UntypedMeshBufferIteratorConst BoneIndicesIterBegin(BaseMesh->VertexBuffers, EMeshBufferSemantic::BoneIndices, 0);
+		UntypedMeshBufferIteratorConst BoneWeightsIterBegin(BaseMesh->VertexBuffers, EMeshBufferSemantic::BoneWeights, 0);
         if (!BoneIndicesIterBegin.ptr() || !BoneWeightsIterBegin.ptr())
         {
             // No skinning data
@@ -134,16 +134,16 @@ namespace mu
 			Result->SetBonePose(BoneIndex, BoneName, ModelRefTransform * PoseTransform, UsageFlags);
 		}
 
-		MeshBufferIterator<MBF_FLOAT32, float, 3> TargetPositionIterBegin(Result->VertexBuffers, MBS_POSITION, 0);
+		MeshBufferIterator<EMeshBufferFormat::Float32, float, 3> TargetPositionIterBegin(Result->VertexBuffers, EMeshBufferSemantic::Position, 0);
 		
 		// Tangent frame buffers are optional.
-		UntypedMeshBufferIteratorConst SourceNormalIterBegin(BaseMesh->VertexBuffers, MBS_NORMAL, 0);
-		UntypedMeshBufferIteratorConst SourceTangentIterBegin(BaseMesh->VertexBuffers, MBS_TANGENT, 0);
-		UntypedMeshBufferIteratorConst SourceBiNormalIterBegin(BaseMesh->VertexBuffers, MBS_BINORMAL, 0);
+		UntypedMeshBufferIteratorConst SourceNormalIterBegin(BaseMesh->VertexBuffers, EMeshBufferSemantic::Normal, 0);
+		UntypedMeshBufferIteratorConst SourceTangentIterBegin(BaseMesh->VertexBuffers, EMeshBufferSemantic::Tangent, 0);
+		UntypedMeshBufferIteratorConst SourceBiNormalIterBegin(BaseMesh->VertexBuffers, EMeshBufferSemantic::Binormal, 0);
 
-		UntypedMeshBufferIterator TargetNormalIterBegin(Result->VertexBuffers, MBS_NORMAL, 0);
-		UntypedMeshBufferIterator TargetTangentIterBegin(Result->VertexBuffers, MBS_TANGENT, 0);
-		UntypedMeshBufferIterator TargetBiNormalIterBegin(Result->VertexBuffers, MBS_BINORMAL, 0);
+		UntypedMeshBufferIterator TargetNormalIterBegin(Result->VertexBuffers, EMeshBufferSemantic::Normal, 0);
+		UntypedMeshBufferIterator TargetTangentIterBegin(Result->VertexBuffers, EMeshBufferSemantic::Tangent, 0);
+		UntypedMeshBufferIterator TargetBiNormalIterBegin(Result->VertexBuffers, EMeshBufferSemantic::Binormal, 0);
 		
 		check(TargetPositionIterBegin.ptr());
 
@@ -179,7 +179,7 @@ namespace mu
 			const EMeshBufferFormat WeightsFormat = BoneWeightsIterBegin.GetFormat(); 
 			const EMeshBufferFormat BoneIndexFormat = BoneIndicesIterBegin.GetFormat(); 
 			const EMeshBufferFormat NormalFormat = TargetNormalIterBegin.ptr() 
-					? TargetNormalIterBegin.GetFormat() : EMeshBufferFormat::MBF_NONE;	
+					? TargetNormalIterBegin.GetFormat() : EMeshBufferFormat::None;	
 
 			for (int32 VertexIndex = BatchBeginVertIndex; VertexIndex < BatchEndVertIndex; ++VertexIndex)
 			{
@@ -189,7 +189,7 @@ namespace mu
 				const uint8* VertexBoneWeightData = (BoneWeightsIterBegin + VertexIndex).ptr();
 				for (int32 WeightIndex = 0; WeightIndex < WeightCount; ++WeightIndex)
 				{
-					ConvertData(WeightIndex, Weights, MBF_FLOAT32, VertexBoneWeightData, WeightsFormat);
+					ConvertData(WeightIndex, Weights, EMeshBufferFormat::Float32, VertexBoneWeightData, WeightsFormat);
 					TotalWeight += Weights[WeightIndex];
 				}
 
@@ -197,7 +197,7 @@ namespace mu
 				const uint8* VertexBoneIndexData = (BoneIndicesIterBegin + VertexIndex).ptr();
 				for (int32 WeightIndex = 0; WeightIndex < WeightCount; ++WeightIndex)
 				{
-					ConvertData(WeightIndex, BoneIndices, MBF_UINT32, VertexBoneIndexData, BoneIndexFormat);
+					ConvertData(WeightIndex, BoneIndices, EMeshBufferFormat::UInt32, VertexBoneIndexData, BoneIndexFormat);
 				}
 		
 				FVector3f SourcePosition = (SourcePositionIterBegin + VertexIndex).GetAsVec3f();
diff --git a/Engine/Plugins/Mutable/Source/MutableRuntime/Private/MuR/OpMeshBind.h b/Engine/Plugins/Mutable/Source/MutableRuntime/Private/MuR/OpMeshBind.h
--- a/Engine/Plugins/Mutable/Source/MutableRuntime/Private/MuR/OpMeshBind.h
+++ b/Engine/Plugins/Mutable/Source/MutableRuntime/Private/MuR/OpMeshBind.h
@@ -136,8 +136,8 @@ namespace mu
 	{
 		constexpr static int32 ElementSize = sizeof(FReshapeVertexBindingData);
 		constexpr static int32 Channels = 4;
-		constexpr static EMeshBufferSemantic Semantics[Channels] = { MBS_BARYCENTRICCOORDS, MBS_DISTANCE, MBS_TRIANGLEINDEX, MBS_OTHER };
-		constexpr static EMeshBufferFormat Formats[Channels] = { MBF_FLOAT32, MBF_FLOAT32, MBF_INT32, MBF_FLOAT32 };
+		constexpr static EMeshBufferSemantic Semantics[Channels] = { EMeshBufferSemantic::BarycentricCoords, EMeshBufferSemantic::Distance, EMeshBufferSemantic::TriangleIndex, EMeshBufferSemantic::Other };
+		constexpr static EMeshBufferFormat Formats[Channels] = { EMeshBufferFormat::Float32, EMeshBufferFormat::Float32, EMeshBufferFormat::Int32, EMeshBufferFormat::Float32 };
 		constexpr static int32 Components[Channels] = { 2, 3, 1, 2+3+3+1 };
 		constexpr static int32 Offsets[Channels] = { 0, 4*2, (2+3)*4, (2+3+1)*4};
 
@@ -163,8 +163,8 @@ namespace mu
 	{
 		constexpr static int32 ElementSize = sizeof(FReshapeVertexBindingData);
 		constexpr static int32 Channels = 4;
-		constexpr static EMeshBufferSemantic Semantics[Channels] = { MBS_BARYCENTRICCOORDS, MBS_DISTANCE, MBS_TRIANGLEINDEX, MBS_OTHER };
-		constexpr static EMeshBufferFormat Formats[Channels] = { MBF_FLOAT32, MBF_FLOAT32, MBF_INT32, MBF_FLOAT32 };
+		constexpr static EMeshBufferSemantic Semantics[Channels] = { EMeshBufferSemantic::BarycentricCoords, EMeshBufferSemantic::Distance, EMeshBufferSemantic::TriangleIndex, EMeshBufferSemantic::Other };
+		constexpr static EMeshBufferFormat Formats[Channels] = { EMeshBufferFormat::Float32, EMeshBufferFormat::Float32, EMeshBufferFormat::Int32, EMeshBufferFormat::Float32 };
 		constexpr static int32 Components[Channels] = { 2, 3, 1, 1 };
 		constexpr static int32 Offsets[Channels] = { 0, 8, 20, 24 };
 
@@ -180,8 +180,8 @@ namespace mu
 	{
 		constexpr static int32 ElementSize = sizeof(int32);
 		constexpr static int32 Channels = 1;
-		constexpr static EMeshBufferSemantic Semantics[Channels] = { MBS_OTHER };
-		constexpr static EMeshBufferFormat Formats[Channels] = { MBF_INT32 };
+		constexpr static EMeshBufferSemantic Semantics[Channels] = { EMeshBufferSemantic::Other };
+		constexpr static EMeshBufferFormat Formats[Channels] = { EMeshBufferFormat::Int32 };
 		constexpr static int32 Components[Channels] = { 1 };
 		constexpr static int32 Offsets[Channels] = { 0 };
 		constexpr static int32 SemanticIndices[Channels] = { 0 };
@@ -210,8 +210,8 @@ namespace mu
 	struct FClipDeformVertexBindingDataBufferDescriptor
 	{
 		constexpr static int32 Channels = 3;
-		constexpr static EMeshBufferSemantic Semantics[Channels] = { MBS_BARYCENTRICCOORDS, MBS_TRIANGLEINDEX, MBS_OTHER };
-		constexpr static EMeshBufferFormat Formats[Channels] = { MBF_FLOAT32, MBF_INT32, MBF_FLOAT32 };
+		constexpr static EMeshBufferSemantic Semantics[Channels] = { EMeshBufferSemantic::BarycentricCoords, EMeshBufferSemantic::TriangleIndex, EMeshBufferSemantic::Other };
+		constexpr static EMeshBufferFormat Formats[Channels] = { EMeshBufferFormat::Float32, EMeshBufferFormat::Int32, EMeshBufferFormat::Float32 };
 		constexpr static int32 Components[Channels] = { 2, 1, 1 };
 		constexpr static int32 Offsets[Channels] = { 0, 8, 12 };
 
@@ -235,7 +235,7 @@ namespace mu
 
 	inline float GetVertexMaskWeight(const UntypedMeshBufferIteratorConst& ColorIter, const FMeshBindColorChannelUsageMasks& ChannelUsages)
 	{
-		check(ColorIter.GetFormat() == MBF_UINT8 || ColorIter.GetFormat() == MBF_NUINT8);
+		check(ColorIter.GetFormat() == EMeshBufferFormat::UInt8 || ColorIter.GetFormat() == EMeshBufferFormat::NUInt8);
 		check(ColorIter.GetComponents() == 4);
 		check(ColorIter.ptr());
 
@@ -245,7 +245,7 @@ namespace mu
 
 	inline uint32 GetVertexClusterId(const UntypedMeshBufferIteratorConst& ColorIter, const FMeshBindColorChannelUsageMasks& ChannelUsages)
 	{	
-		check(ColorIter.GetFormat() == MBF_UINT8 || ColorIter.GetFormat() == MBF_NUINT8);
+		check(ColorIter.GetFormat() == EMeshBufferFormat::UInt8 || ColorIter.GetFormat() == EMeshBufferFormat::NUInt8);
 		check(ColorIter.GetComponents() == 4);
 		check(ColorIter.ptr());
 
@@ -371,7 +371,7 @@ namespace mu
 
 		OutNonRigidClusterIdx = -1;
 		
-		const UntypedMeshBufferIteratorConst ItColorBase(Mesh->GetVertexBuffers(), MBS_COLOUR);
+		const UntypedMeshBufferIteratorConst ItColorBase(Mesh->GetVertexBuffers(), EMeshBufferSemantic::Color);
 		if (ItColorBase.ptr())
 		{
 			const int32 VertexCount = Mesh->GetVertexCount();
@@ -424,7 +424,7 @@ namespace mu
 			const Mesh* BaseMesh, FShapeMeshTree& ShapeMeshTree,
 			const TArray<int32>& Cluster, FReshapeVertexBindingData& OutBindingData, float BindTolerance)
 	{
-		const UntypedMeshBufferIteratorConst ItPositionBase(BaseMesh->GetVertexBuffers(), MBS_POSITION);
+		const UntypedMeshBufferIteratorConst ItPositionBase(BaseMesh->GetVertexBuffers(), EMeshBufferSemantic::Position);
 
 		FBox3f ClusterBoundingBox(ForceInit);
 		for (const int32& V : Cluster)
@@ -728,13 +728,13 @@ namespace mu
 
 			BindData.SetNum(MeshVertexCount);
 
-			const UntypedMeshBufferIteratorConst ItPositionBase(BaseMesh->GetVertexBuffers(), MBS_POSITION);
-			const UntypedMeshBufferIteratorConst ItNormalBase(BaseMesh->GetVertexBuffers(), MBS_NORMAL);
+			const UntypedMeshBufferIteratorConst ItPositionBase(BaseMesh->GetVertexBuffers(), EMeshBufferSemantic::Position);
+			const UntypedMeshBufferIteratorConst ItNormalBase(BaseMesh->GetVertexBuffers(), EMeshBufferSemantic::Normal);
 
 			// Disable vertex color reads if the color is not used for mask weights.
 			const UntypedMeshBufferIteratorConst ItColorBase = ColorUsageMasks.MaskWeight == 0 
 				? UntypedMeshBufferIteratorConst()
-				: UntypedMeshBufferIteratorConst(BaseMesh->GetVertexBuffers(), MBS_COLOUR);
+				: UntypedMeshBufferIteratorConst(BaseMesh->GetVertexBuffers(), EMeshBufferSemantic::Color);
 
 			// Special case for non rigid parts
 
@@ -812,12 +812,12 @@ namespace mu
 
 		TArrayView<const FVector3f> VerticesView;
 		{
-			const UntypedMeshBufferIteratorConst PositionBegin(InOutMesh.GetVertexBuffers(), MBS_POSITION);
+			const UntypedMeshBufferIteratorConst PositionBegin(InOutMesh.GetVertexBuffers(), EMeshBufferSemantic::Position);
 			const int32 NumVertices = InOutMesh.GetVertexBuffers().GetElementCount();
 			
 			const bool bIsCompatibleBuffer =
 				PositionBegin.GetElementSize() == sizeof(FVector3f) &&
-				PositionBegin.GetFormat() == MBF_FLOAT32 &&
+				PositionBegin.GetFormat() == EMeshBufferFormat::Float32 &&
 				PositionBegin.GetComponents() == 3;
 			
 			const bool bIsAlignmentGood = reinterpret_cast<UPTRINT>(PositionBegin.ptr()) % alignof(FVector3f) == 0;
@@ -841,12 +841,12 @@ namespace mu
 
 		TArrayView<const uint32> IndicesView;
 		{
-			const UntypedMeshBufferIteratorConst IndicesBegin(InOutMesh.GetIndexBuffers(), MBS_VERTEXINDEX);
+			const UntypedMeshBufferIteratorConst IndicesBegin(InOutMesh.GetIndexBuffers(), EMeshBufferSemantic::VertexIndex);
 			const int32 NumIndices = InOutMesh.GetIndexBuffers().GetElementCount();
 			
 			const bool bIsCompatibleBuffer =
 				IndicesBegin.GetElementSize() == sizeof(uint32) &&
-				IndicesBegin.GetFormat() == MBF_UINT32 &&
+				IndicesBegin.GetFormat() == EMeshBufferFormat::UInt32 &&
 				IndicesBegin.GetComponents() == 1;
 
 			const bool bIsAlignmentGood = reinterpret_cast<UPTRINT>(IndicesBegin.ptr()) % alignof(uint32) == 0;
@@ -1053,8 +1053,8 @@ namespace mu
 			ShapeMeshDescriptor.Normals.SetNum(ShapeVertexCount);
 
 			// \TODO: Simple but inefficient
-			const UntypedMeshBufferIteratorConst ItPosition(ShapeMesh->GetVertexBuffers(), MBS_POSITION);
-			const UntypedMeshBufferIteratorConst ItNormal(ShapeMesh->GetVertexBuffers(), MBS_NORMAL);
+			const UntypedMeshBufferIteratorConst ItPosition(ShapeMesh->GetVertexBuffers(), EMeshBufferSemantic::Position);
+			const UntypedMeshBufferIteratorConst ItNormal(ShapeMesh->GetVertexBuffers(), EMeshBufferSemantic::Normal);
 			for (int32 ShapeVertexIndex = 0; ShapeVertexIndex < ShapeVertexCount; ++ShapeVertexIndex)
 			{
 				FVector3f Position = (ItPosition + ShapeVertexIndex).GetAsVec3f();
@@ -1070,7 +1070,7 @@ namespace mu
 			MUTABLE_CPUPROFILER_SCOPE(GenerateTrianglesQueryData);	
 			ShapeMeshDescriptor.Triangles.SetNum(ShapeTriangleCount);
 			// \TODO: Simple but inefficient
-			const UntypedMeshBufferIteratorConst ItIndices(ShapeMesh->GetIndexBuffers(), MBS_VERTEXINDEX);
+			const UntypedMeshBufferIteratorConst ItIndices(ShapeMesh->GetIndexBuffers(), EMeshBufferSemantic::VertexIndex);
 			for (int32 TriangleIndex = 0; TriangleIndex < ShapeTriangleCount; ++TriangleIndex)
 			{
 				UE::Geometry::FIndex3i Triangle;
@@ -1154,8 +1154,8 @@ namespace mu
 			SkeletonBuffer.SetBuffer(0, sizeof(FReshapePointBindingData), BufDesc.Channels, BufDesc.Semantics, BufDesc.SemanticIndices, BufDesc.Formats, BufDesc.Components);
 
 			// Bone indices buffer
-			EMeshBufferSemantic BoneSemantics[1] = { MBS_OTHER };
-			EMeshBufferFormat BoneFormats[1] = { MBF_INT32 };
+			EMeshBufferSemantic BoneSemantics[1] = { EMeshBufferSemantic::Other };
+			EMeshBufferFormat BoneFormats[1] = { EMeshBufferFormat::Int32 };
 			int32 BoneSemanticIndices[1] = { 0 };
 			int32 BoneComponents[1] = { 1 };
 			int32 BoneOffsets[1] = { 0 };
@@ -1391,8 +1391,8 @@ namespace mu
 
 			BindData.SetNum( MeshVertexCount );
 
-			UntypedMeshBufferIteratorConst ItPositionBase(BaseMesh->GetVertexBuffers(), MBS_POSITION);
-			UntypedMeshBufferIteratorConst ItNormalBase(BaseMesh->GetVertexBuffers(), MBS_NORMAL);
+			UntypedMeshBufferIteratorConst ItPositionBase(BaseMesh->GetVertexBuffers(), EMeshBufferSemantic::Position);
+			UntypedMeshBufferIteratorConst ItNormalBase(BaseMesh->GetVertexBuffers(), EMeshBufferSemantic::Normal);
 
 			EShapeBindingMethod ActualBindingMethod = BindingMethod == EShapeBindingMethod::ClipDeformNormalProject && !ItNormalBase.ptr() 
 					? EShapeBindingMethod::ClipDeformClosestToSurface 
@@ -1477,8 +1477,8 @@ namespace mu
 			ShapeMeshDescriptor.Normals.SetNum(ShapeVertexCount);
 
 			// \TODO: Simple but inefficient
-			UntypedMeshBufferIteratorConst ItPosition(ShapeMesh->GetVertexBuffers(), MBS_POSITION);
-			UntypedMeshBufferIteratorConst ItNormal(ShapeMesh->GetVertexBuffers(), MBS_NORMAL);
+			UntypedMeshBufferIteratorConst ItPosition(ShapeMesh->GetVertexBuffers(), EMeshBufferSemantic::Position);
+			UntypedMeshBufferIteratorConst ItNormal(ShapeMesh->GetVertexBuffers(), EMeshBufferSemantic::Normal);
 			for (int32 ShapeVertexIndex = 0; ShapeVertexIndex < ShapeVertexCount; ++ShapeVertexIndex)
 			{
 				FVector3f Position = (ItPosition + ShapeVertexIndex).GetAsVec3f();
@@ -1494,7 +1494,7 @@ namespace mu
 			MUTABLE_CPUPROFILER_SCOPE(GenerateTrianglesQueryData);	
 			ShapeMeshDescriptor.Triangles.SetNum(ShapeTriangleCount);
 			// \TODO: Simple but inefficient
-			const UntypedMeshBufferIteratorConst ItIndices(ShapeMesh->GetIndexBuffers(), MBS_VERTEXINDEX);
+			const UntypedMeshBufferIteratorConst ItIndices(ShapeMesh->GetIndexBuffers(), EMeshBufferSemantic::VertexIndex);
 			for (int32 TriangleIndex = 0; TriangleIndex < ShapeTriangleCount; ++TriangleIndex)
 			{
 				UE::Geometry::FIndex3i Triangle;
