[mutable] Cleanup MutableRuntime API moving to use more unreal types and naming standards.

diff --git a/Engine/Plugins/Mutable/Source/MutableRuntime/Private/MuR/OpMeshClipDeform.h b/Engine/Plugins/Mutable/Source/MutableRuntime/Private/MuR/OpMeshClipDeform.h
--- a/Engine/Plugins/Mutable/Source/MutableRuntime/Private/MuR/OpMeshClipDeform.h
+++ b/Engine/Plugins/Mutable/Source/MutableRuntime/Private/MuR/OpMeshClipDeform.h
@@ -117,7 +117,7 @@ namespace mu
 		int BarycentricDataBuffer = 0;
 		int BarycentricDataChannel = 0;
 		const FMeshBufferSet& VB = BaseMesh->GetVertexBuffers();
-		VB.FindChannel(MBS_BARYCENTRICCOORDS, 0, &BarycentricDataBuffer, &BarycentricDataChannel);
+		VB.FindChannel(EMeshBufferSemantic::BarycentricCoords, 0, &BarycentricDataBuffer, &BarycentricDataChannel);
 		
 		// Copy buffers skipping binding data.
 		FMeshBufferSet& ResultBuffers = Result->GetVertexBuffers();
@@ -168,12 +168,12 @@ namespace mu
 			ShapeNormals.SetNum(ShapeVertexCount);
 			ShapeWeights.SetNum(ShapeVertexCount);
 
-			UntypedMeshBufferIteratorConst ItPositionBegin(ShapeMesh->GetVertexBuffers(), MBS_POSITION);
-			UntypedMeshBufferIteratorConst ItNormalBegin(ShapeMesh->GetVertexBuffers(), MBS_NORMAL);
-			UntypedMeshBufferIteratorConst ItUvsBegin(ShapeMesh->GetVertexBuffers(), MBS_TEXCOORDS);
+			UntypedMeshBufferIteratorConst ItPositionBegin(ShapeMesh->GetVertexBuffers(), EMeshBufferSemantic::Position);
+			UntypedMeshBufferIteratorConst ItNormalBegin(ShapeMesh->GetVertexBuffers(), EMeshBufferSemantic::Normal);
+			UntypedMeshBufferIteratorConst ItUvsBegin(ShapeMesh->GetVertexBuffers(), EMeshBufferSemantic::TexCoords);
 
 			const bool bIsPositionBufferCompatible = 
-					ItPositionBegin.GetFormat() == MBF_FLOAT32 && 
+					ItPositionBegin.GetFormat() == EMeshBufferFormat::Float32 && 
 					ItPositionBegin.GetElementSize() == sizeof(FVector3f);
 			if (bIsPositionBufferCompatible)
 			{
@@ -194,7 +194,7 @@ namespace mu
 			}
 
 			const bool bIsNormalBufferCompatible = 
-					ItNormalBegin.GetFormat() == MBF_FLOAT32 && 
+					ItNormalBegin.GetFormat() == EMeshBufferFormat::Float32 && 
 					ItNormalBegin.GetElementSize() == sizeof(FVector3f);
 			if (bIsNormalBufferCompatible)
 			{
@@ -232,12 +232,12 @@ namespace mu
 
 			ShapeTriangles.SetNum(ShapeTriangleCount);
 
-			const UntypedMeshBufferIteratorConst ItIndicesBegin(ShapeMesh->GetIndexBuffers(), MBS_VERTEXINDEX);
+			const UntypedMeshBufferIteratorConst ItIndicesBegin(ShapeMesh->GetIndexBuffers(), EMeshBufferSemantic::VertexIndex);
 
 			check(ShapeMesh->GetIndexCount() % 3 == 0);
 	
 			const bool bIsIndexBufferCompatible = 
-					(ItIndicesBegin.GetFormat() == MBF_INT32 || ItIndicesBegin.GetFormat() == MBF_UINT32) &&
+					(ItIndicesBegin.GetFormat() == EMeshBufferFormat::Int32 || ItIndicesBegin.GetFormat() == EMeshBufferFormat::UInt32) &&
 					ItIndicesBegin.GetElementSize() == int32(sizeof(int32));
 
 			if (bIsIndexBufferCompatible)
@@ -277,9 +277,9 @@ namespace mu
 		{
 			MUTABLE_CPUPROFILER_SCOPE(UpdateClipDeformVertices);
 
-			UntypedMeshBufferIterator ItPosition(Result->GetVertexBuffers(), MBS_POSITION);
-			UntypedMeshBufferIterator ItNormal(Result->GetVertexBuffers(), MBS_NORMAL);
-			UntypedMeshBufferIterator ItTangent(Result->GetVertexBuffers(), MBS_TANGENT);
+			UntypedMeshBufferIterator ItPosition(Result->GetVertexBuffers(), EMeshBufferSemantic::Position);
+			UntypedMeshBufferIterator ItNormal(Result->GetVertexBuffers(), EMeshBufferSemantic::Normal);
+			UntypedMeshBufferIterator ItTangent(Result->GetVertexBuffers(), EMeshBufferSemantic::Tangent);
 
 			for (int32 MeshVertexIndex = 0; MeshVertexIndex < MeshVertexCount; ++MeshVertexIndex)
 			{
diff --git a/Engine/Plugins/Mutable/Source/MutableRuntime/Private/MuR/OpMeshClipMorphPlane.h b/Engine/Plugins/Mutable/Source/MutableRuntime/Private/MuR/OpMeshClipMorphPlane.h
--- a/Engine/Plugins/Mutable/Source/MutableRuntime/Private/MuR/OpMeshClipMorphPlane.h
+++ b/Engine/Plugins/Mutable/Source/MutableRuntime/Private/MuR/OpMeshClipMorphPlane.h
@@ -151,9 +151,9 @@ namespace mu
 				for (int32 ChannelIndex = 0; ChannelIndex < pBase->GetVertexBuffers().GetBufferChannelCount(BufferIndex); ++ChannelIndex)
 				{
 					// Get info about the destination channel
-					EMeshBufferSemantic Semantic = MBS_NONE;
+					EMeshBufferSemantic Semantic = EMeshBufferSemantic::None;
 					int32 SemanticIndex = 0;
-					EMeshBufferFormat Format = MBF_NONE;
+					EMeshBufferFormat Format = EMeshBufferFormat::None;
 					int32 Components = 0;
 					int32 Offset = 0;
 					pBase->GetVertexBuffers().GetChannel(BufferIndex, ChannelIndex, &Semantic, &SemanticIndex, &Format, &Components, &Offset);
@@ -162,14 +162,14 @@ namespace mu
 
 					//int32 ResultOffset = FirstSize + Offset;
 
-					if (Semantic == MBS_BONEINDICES)
+					if (Semantic == EMeshBufferSemantic::BoneIndices)
 					{
 						int32 NumVertices = pBase->GetVertexCount();
 						{
 							switch (Format)
 							{
-							case MBF_INT8:
-							case MBF_UINT8:
+							case EMeshBufferFormat::Int8:
+							case EMeshBufferFormat::UInt8:
 							{
 								for (int32 i = 0; i < NumVertices; ++i)
 								{
@@ -182,8 +182,8 @@ namespace mu
 								break;
 							}
 
-							case MBF_INT16:
-							case MBF_UINT16:
+							case EMeshBufferFormat::Int16:
+							case EMeshBufferFormat::UInt16:
 							{
 								for (int32 i = 0; i < NumVertices; ++i)
 								{
@@ -197,8 +197,8 @@ namespace mu
 								break;
 							}
 
-							case MBF_INT32:
-							case MBF_UINT32:
+							case EMeshBufferFormat::Int32:
+							case EMeshBufferFormat::UInt32:
 							{
 								for (int32 i = 0; i < NumVertices; ++i)
 								{
@@ -217,15 +217,15 @@ namespace mu
 							}
 						}
 					}
-					else if (Semantic == MBS_BONEWEIGHTS)
+					else if (Semantic == EMeshBufferSemantic::BoneWeights)
 					{
 						const int32 NumVertices = pBase->GetVertexCount();
 						{
 							switch (Format)
 							{
-							case MBF_INT8:
-							case MBF_UINT8:
-							case MBF_NUINT8:
+							case EMeshBufferFormat::Int8:
+							case EMeshBufferFormat::UInt8:
+							case EMeshBufferFormat::NUInt8:
 							{
 								for (int32 i = 0; i < NumVertices; ++i)
 								{
@@ -239,9 +239,9 @@ namespace mu
 								break;
 							}
 
-							case MBF_INT16:
-							case MBF_UINT16:
-							case MBF_NUINT16:
+							case EMeshBufferFormat::Int16:
+							case EMeshBufferFormat::UInt16:
+							case EMeshBufferFormat::NUInt16:
 							{
 								for (int32 i = 0; i < NumVertices; ++i)
 								{
@@ -255,9 +255,9 @@ namespace mu
 								break;
 							}
 
-							case MBF_INT32:
-							case MBF_UINT32:
-							case MBF_NUINT32:
+							case EMeshBufferFormat::Int32:
+							case EMeshBufferFormat::UInt32:
+							case EMeshBufferFormat::NUInt32:
 							{
 								for (int32 i = 0; i < NumVertices; ++i)
 								{
@@ -271,7 +271,7 @@ namespace mu
 								break;
 							}
 
-                            case MBF_FLOAT32:
+                            case EMeshBufferFormat::Float32:
                             {
 								for (int32 i = 0; i < NumVertices; ++i)
 								{
@@ -302,7 +302,7 @@ namespace mu
 		
 		// Positions can  be assumed to be in a FVector3f struct but changing the iterator to a untyped one should 
 		// work as well.
-		const MeshBufferIterator<MBF_FLOAT32, float, 3> PositionIterBegin(Result->GetVertexBuffers(), MBS_POSITION);
+		const MeshBufferIterator<EMeshBufferFormat::Float32, float, 3> PositionIterBegin(Result->GetVertexBuffers(), EMeshBufferSemantic::Position);
 		for (int32 VertexIndex = 0; VertexIndex < NumVertices; ++VertexIndex)
 		{	
 			FVector3f Position = (PositionIterBegin + VertexIndex).GetAsVec3f();
diff --git a/Engine/Plugins/Mutable/Source/MutableRuntime/Private/MuR/OpMeshClipWithMesh.cpp b/Engine/Plugins/Mutable/Source/MutableRuntime/Private/MuR/OpMeshClipWithMesh.cpp
--- a/Engine/Plugins/Mutable/Source/MutableRuntime/Private/MuR/OpMeshClipWithMesh.cpp
+++ b/Engine/Plugins/Mutable/Source/MutableRuntime/Private/MuR/OpMeshClipWithMesh.cpp
@@ -41,7 +41,7 @@ namespace mu { namespace
 		OutCollapsedVertices.Init(INDEX_NONE, NumVertices);
 
 		// Get Vertices
-		mu::UntypedMeshBufferIteratorConst ItPosition = mu::UntypedMeshBufferIteratorConst(pMesh->GetVertexBuffers(), mu::MBS_POSITION);
+		mu::UntypedMeshBufferIteratorConst ItPosition = mu::UntypedMeshBufferIteratorConst(pMesh->GetVertexBuffers(), mu::EMeshBufferSemantic::Position);
 		FVector3f* VertexData = OutVertices.GetData();
 
 		for (int32 VertexIndex = 0; VertexIndex < NumVertices; ++VertexIndex)
@@ -97,7 +97,7 @@ namespace mu { namespace
   //      // Acumulate edges
   //      TMap< TPair<int,int>, int > faceCountPerEdge;
 
-  //      UntypedMeshBufferIteratorConst itClipMesh(pMesh->GetIndexBuffers(), MBS_VERTEXINDEX);
+  //      UntypedMeshBufferIteratorConst itClipMesh(pMesh->GetIndexBuffers(), EMeshBufferSemantic::VertexIndex);
   //      for (int f = 0; f < fcount; ++f)
   //      {
   //          int face[3];
@@ -155,7 +155,7 @@ namespace mu { namespace
         // Mark used vertices
 		TArray<uint8> used;
 		used.SetNumZeroed(pMesh->GetVertexCount());
-        UntypedMeshBufferIteratorConst iti(pMesh->GetIndexBuffers(), MBS_VERTEXINDEX);
+        UntypedMeshBufferIteratorConst iti(pMesh->GetIndexBuffers(), EMeshBufferSemantic::VertexIndex);
         int IndexCount = pMesh->GetIndexCount();
         for (int i = 0; i < IndexCount; ++i)
         {
@@ -204,8 +204,8 @@ namespace mu { namespace
         pMesh->GetVertexBuffers().SetElementCount(totalNewVertices);
 
         // Update indices
-        UntypedMeshBufferIteratorConst ito(pMesh->GetIndexBuffers(), MBS_VERTEXINDEX);
-        if (ito.GetFormat() == MBF_UINT32)
+        UntypedMeshBufferIteratorConst ito(pMesh->GetIndexBuffers(), EMeshBufferSemantic::VertexIndex);
+        if (ito.GetFormat() == EMeshBufferFormat::UInt32)
         {
             for (int i = 0; i < IndexCount; ++i)
             {
@@ -216,7 +216,7 @@ namespace mu { namespace
                 ++ito;
             }
         }
-        else if (ito.GetFormat() == MBF_UINT16)
+        else if (ito.GetFormat() == EMeshBufferFormat::UInt16)
         {
             for (int i = 0; i < IndexCount; ++i)
             {
@@ -340,7 +340,7 @@ namespace mu { namespace
 			for (int32 c = 0; c < Channels.Num(); ++c)
 			{
 				EMeshBufferSemantic Sem = Channels[c].Semantic;
-				if (Sem != MBS_TEXCOORDS)
+				if (Sem != EMeshBufferSemantic::TexCoords)
 				{
 					continue;
 				}
@@ -391,7 +391,7 @@ namespace mu { namespace
 			for (int32 c = 0; c < Channels.Num(); ++c)
 			{
 				EMeshBufferSemantic Sem = Channels[c].Semantic;
-				if (Sem != MBS_TEXCOORDS)
+				if (Sem != EMeshBufferSemantic::TexCoords)
 				{
 					continue;
 				}
@@ -449,7 +449,7 @@ namespace mu { namespace
 		Result->GetVertexBuffers().SetBufferCount(1);
 		Result->MeshIDPrefix = Base->MeshIDPrefix;
 
-		EMeshBufferSemantic Semantic = MBS_VERTEXINDEX;
+		EMeshBufferSemantic Semantic = EMeshBufferSemantic::VertexIndex;
 		int32 SemanticIndex = 0;
 		int32 Components = 1;
 		int32 Offsets = 0;
@@ -457,7 +457,7 @@ namespace mu { namespace
 		bool bMakeRelativeIds = !Base->AreVertexIdsExplicit();
 		if (bMakeRelativeIds)
 		{
-			EMeshBufferFormat Format = MBF_UINT32;
+			EMeshBufferFormat Format = EMeshBufferFormat::UInt32;
 			Result->GetVertexBuffers().SetBuffer( 0, sizeof(uint32), 1, &Semantic, &SemanticIndex, &Format, &Components, &Offsets );
 
 			// Fill the buffer
@@ -476,7 +476,7 @@ namespace mu { namespace
 		}
 		else
 		{
-			EMeshBufferFormat Format = MBF_UINT64;
+			EMeshBufferFormat Format = EMeshBufferFormat::UInt64;
 			Result->GetVertexBuffers().SetBuffer( 0, sizeof(uint64), 1, &Semantic, &SemanticIndex, &Format, &Components, &Offsets );
 
 			// Fill the buffer
@@ -537,7 +537,7 @@ namespace mu
         MeshCreateCollapsedVertexMap( pClipMesh, CollapsedVertexMap, Vertices );
 
         // Create cache of the faces
-        UntypedMeshBufferIteratorConst ItClipMesh(pClipMesh->GetIndexBuffers(), MBS_VERTEXINDEX);
+        UntypedMeshBufferIteratorConst ItClipMesh(pClipMesh->GetIndexBuffers(), EMeshBufferSemantic::VertexIndex);
 
         for ( int32 F = 0; F < FCount; ++F )
         {
@@ -618,13 +618,13 @@ namespace mu
 
                 switch ( Sem )
                 {
-                case MBS_POSITION:
+                case EMeshBufferSemantic::Position:
                     for (uint32 V = 0; V < DestVertexCount; ++V)
                     {
                         FVector3f Vertex(0.0f, 0.0f, 0.0f);
                         for (int32 Offset = 0; Offset < 3; ++Offset)
                         {
-                            ConvertData(Offset, &Vertex[0], MBF_FLOAT32, It.ptr(), It.GetFormat());
+                            ConvertData(Offset, &Vertex[0], EMeshBufferFormat::Float32, It.ptr(), It.GetFormat());
                         }
 						
 						const FVector2f ClipBBoxMin = FVector2f( ClipMeshBoundingBox.Min.X, ClipMeshBoundingBox.Min.Y);
@@ -703,11 +703,11 @@ namespace mu
         MeshClipMeshClassifyVertices(VertexInClipMesh, pBase, pClipMesh);
 
         // Now remove all the faces from the result mesh that have all the vertices outside the clip volume
-        UntypedMeshBufferIteratorConst ItBase(Result->GetIndexBuffers(), MBS_VERTEXINDEX);
-        UntypedMeshBufferIterator ItDest(Result->GetIndexBuffers(), MBS_VERTEXINDEX);
+        UntypedMeshBufferIteratorConst ItBase(Result->GetIndexBuffers(), EMeshBufferSemantic::VertexIndex);
+        UntypedMeshBufferIterator ItDest(Result->GetIndexBuffers(), EMeshBufferSemantic::VertexIndex);
         int32 AFaceCount = Result->GetFaceCount();
 
-        UntypedMeshBufferIteratorConst Ito(Result->GetIndexBuffers(), MBS_VERTEXINDEX);
+        UntypedMeshBufferIteratorConst Ito(Result->GetIndexBuffers(), EMeshBufferSemantic::VertexIndex);
         for (int32 F = 0; F < AFaceCount; ++F)
         {
 			uint32 OV[3] = {0, 0, 0};
@@ -760,7 +760,7 @@ namespace mu
         TArray<uint8> VertexWithFaceNotClipped;
 		VertexWithFaceNotClipped.AddZeroed( VertexInClipMesh.Num() );
 
-        UntypedMeshBufferIteratorConst Ito(pBase->GetIndexBuffers(), MBS_VERTEXINDEX);
+        UntypedMeshBufferIteratorConst Ito(pBase->GetIndexBuffers(), EMeshBufferSemantic::VertexIndex);
         int32 AFaceCount = pBase->GetFaceCount();
         for (int32 F = 0; F < AFaceCount; ++F)
         {
@@ -805,7 +805,7 @@ namespace mu
 		TArray<uint8> VertexWithFaceNotClipped;
 		VertexWithFaceNotClipped.AddZeroed(VertexClipped.Num());
 
-		UntypedMeshBufferIteratorConst Ito(Base->GetIndexBuffers(), MBS_VERTEXINDEX);
+		UntypedMeshBufferIteratorConst Ito(Base->GetIndexBuffers(), EMeshBufferSemantic::VertexIndex);
 		int32 FaceCount = Base->GetFaceCount();
 		for (int32 F = 0; F < FaceCount; ++F)
 		{
@@ -850,7 +850,7 @@ namespace mu
 		TArray<uint8> VertexWithFaceNotClipped;
 		VertexWithFaceNotClipped.AddZeroed(VertexClipped.Num());
 
-		UntypedMeshBufferIteratorConst Ito(Base->GetIndexBuffers(), MBS_VERTEXINDEX);
+		UntypedMeshBufferIteratorConst Ito(Base->GetIndexBuffers(), EMeshBufferSemantic::VertexIndex);
 		int32 FaceCount = Base->GetFaceCount();
 		for (int32 F = 0; F < FaceCount; ++F)
 		{
@@ -899,7 +899,7 @@ namespace mu
         box< FVector3f > aabbox;
         if ( fragmentFaceCount > 0 )
         {
-            MeshBufferIteratorConst<MBF_FLOAT32,float,3> itp( pFragment->GetVertexBuffers(), MBS_POSITION );
+            MeshBufferIteratorConst<EMeshBufferFormat::Float32,float,3> itp( pFragment->GetVertexBuffers(), EMeshBufferSemantic::Position );
 
             aabbox.min = itp.GetAsVec3f();
             ++itp;
@@ -923,8 +923,8 @@ namespace mu
         float bucketSize = aabbox.size[AXIS] / NUM_BUCKETS;
 
         float bucketThreshold = ( 4 * tolerance ) / bucketSize;
-        UntypedMeshBufferIteratorConst iti( pFragment->GetIndexBuffers(), MBS_VERTEXINDEX );
-        MeshBufferIteratorConst<MBF_FLOAT32,float,3> itp( pFragment->GetVertexBuffers(), MBS_POSITION );
+        UntypedMeshBufferIteratorConst iti( pFragment->GetIndexBuffers(), EMeshBufferSemantic::VertexIndex );
+        MeshBufferIteratorConst<EMeshBufferFormat::Float32,float,3> itp( pFragment->GetVertexBuffers(), EMeshBufferSemantic::Position );
         for ( int tf=0; tf<fragmentFaceCount; tf++ )
         {
             uint32 index0 = iti.GetAsUINT32(); ++iti;
@@ -957,9 +957,9 @@ namespace mu
 		TArray<uint8> faceClipped;
 		faceClipped.SetNumZeroed(sourceFaceCount);
 
-        UntypedMeshBufferIteratorConst ito( pBase->GetIndexBuffers(), MBS_VERTEXINDEX );
-        MeshBufferIteratorConst<MBF_FLOAT32,float,3> itop( pBase->GetVertexBuffers(), MBS_POSITION );
-        UntypedMeshBufferIteratorConst itti( pFragment->GetIndexBuffers(), MBS_VERTEXINDEX );
+        UntypedMeshBufferIteratorConst ito( pBase->GetIndexBuffers(), EMeshBufferSemantic::VertexIndex );
+        MeshBufferIteratorConst<EMeshBufferFormat::Float32,float,3> itop( pBase->GetVertexBuffers(), EMeshBufferSemantic::Position );
+        UntypedMeshBufferIteratorConst itti( pFragment->GetIndexBuffers(), EMeshBufferSemantic::VertexIndex );
         for ( int f=0; f<sourceFaceCount; ++f )
         {
             bool hasFace = false;
@@ -1001,7 +1001,7 @@ namespace mu
 		TArray<uint8> vertex_with_face_not_clipped;
 		vertex_with_face_not_clipped.SetNumZeroed(sourceVertexCount);
 
-        UntypedMeshBufferIteratorConst itoi(pBase->GetIndexBuffers(), MBS_VERTEXINDEX);
+        UntypedMeshBufferIteratorConst itoi(pBase->GetIndexBuffers(), EMeshBufferSemantic::VertexIndex);
         int aFaceCount = pBase->GetFaceCount();
         for (int f = 0; f < aFaceCount; ++f)
         {
diff --git a/Engine/Plugins/Mutable/Source/MutableRuntime/Private/MuR/OpMeshComputeNormals.h b/Engine/Plugins/Mutable/Source/MutableRuntime/Private/MuR/OpMeshComputeNormals.h
--- a/Engine/Plugins/Mutable/Source/MutableRuntime/Private/MuR/OpMeshComputeNormals.h
+++ b/Engine/Plugins/Mutable/Source/MutableRuntime/Private/MuR/OpMeshComputeNormals.h
@@ -48,10 +48,10 @@ namespace mu
 	{
 		MUTABLE_CPUPROFILER_SCOPE(ComputeMeshNormals);
 
-		const UntypedMeshBufferIteratorConst BaseIndicesIter = UntypedMeshBufferIteratorConst(DestMesh.GetIndexBuffers(), MBS_VERTEXINDEX, 0);
+		const UntypedMeshBufferIteratorConst BaseIndicesIter = UntypedMeshBufferIteratorConst(DestMesh.GetIndexBuffers(), EMeshBufferSemantic::VertexIndex, 0);
 		const int32 NumIndices = DestMesh.GetIndexBuffers().GetElementCount();
 
-		const UntypedMeshBufferIteratorConst BasePositionIter = UntypedMeshBufferIteratorConst(DestMesh.GetVertexBuffers(), MBS_POSITION, 0);
+		const UntypedMeshBufferIteratorConst BasePositionIter = UntypedMeshBufferIteratorConst(DestMesh.GetVertexBuffers(), EMeshBufferSemantic::Position, 0);
 		const int32 NumVertices = DestMesh.GetVertexBuffers().GetElementCount();
 
 		check(BaseIndicesIter.ptr());
@@ -86,11 +86,11 @@ namespace mu
 			NormalsAccumulation[VertexIndex] = NormalsAccumulation[VertexIndex].GetSafeNormal();
 		}
 
-		const UntypedMeshBufferIterator NormalIter = UntypedMeshBufferIterator(DestMesh.GetVertexBuffers(), MBS_NORMAL, 0);
+		const UntypedMeshBufferIterator NormalIter = UntypedMeshBufferIterator(DestMesh.GetVertexBuffers(), EMeshBufferSemantic::Normal, 0);
 		check(NormalIter.ptr());
 
-		const UntypedMeshBufferIterator TangentIter  = UntypedMeshBufferIterator(DestMesh.GetVertexBuffers(), MBS_TANGENT, 0);
-		const UntypedMeshBufferIterator BiNormalIter = UntypedMeshBufferIterator(DestMesh.GetVertexBuffers(), MBS_BINORMAL, 0);
+		const UntypedMeshBufferIterator TangentIter  = UntypedMeshBufferIterator(DestMesh.GetVertexBuffers(), EMeshBufferSemantic::Tangent, 0);
+		const UntypedMeshBufferIterator BiNormalIter = UntypedMeshBufferIterator(DestMesh.GetVertexBuffers(), EMeshBufferSemantic::Binormal, 0);
 
         const EMeshBufferFormat NormalFormat   = NormalIter.GetFormat();
         const EMeshBufferFormat TangentFormat  = TangentIter.GetFormat();
@@ -101,7 +101,7 @@ namespace mu
         const int32 BiNormalComps = BiNormalIter.GetComponents();
 
 		// When normal is packed, binormal channel is not expected. It is not a big deal if it's there but we would be doing extra unused work in that case. 
-		ensureAlways(!(NormalFormat == MBF_PACKEDDIR8_W_TANGENTSIGN || NormalFormat == MBF_PACKEDDIRS8_W_TANGENTSIGN) || !BiNormalIter.ptr());
+		ensureAlways(!(NormalFormat == EMeshBufferFormat::PackedDir8_W_TangentSign || NormalFormat == EMeshBufferFormat::PackedDirS8_W_TangentSign) || !BiNormalIter.ptr());
 		for (int32 VertexIndex = 0; VertexIndex < NumVertices; ++VertexIndex)
 		{
 			const FVector3f Normal = NormalsAccumulation[VertexIndex];
@@ -129,7 +129,7 @@ namespace mu
 				uint8 * const NormalElemPtr = (NormalIter + VertexIndex).ptr();
 				for (int32 C = 0; C < NormalComps && C < 3; ++C)
 				{
-					ConvertData(C, NormalElemPtr, NormalFormat, &Normal, MBF_FLOAT32);
+					ConvertData(C, NormalElemPtr, NormalFormat, &Normal, EMeshBufferFormat::Float32);
 				}
 			}
 
@@ -138,7 +138,7 @@ namespace mu
 				uint8 * const TangentElemPtr = (TangentIter + VertexIndex).ptr();
 				for (int32 C = 0; C < TangentComps && C < 3; ++C)
 				{
-					ConvertData(C, TangentElemPtr, TangentFormat, &Tangent, MBF_FLOAT32);
+					ConvertData(C, TangentElemPtr, TangentFormat, &Tangent, EMeshBufferFormat::Float32);
 				}
 			}
 			
@@ -147,7 +147,7 @@ namespace mu
 				uint8 * const BiNormalElemPtr = (BiNormalIter + VertexIndex).ptr();
 				for (int32 C = 0; C < BiNormalComps && C < 3; ++C)
 				{
-					ConvertData(C, BiNormalElemPtr, BiNormalFormat, &BiNormal, MBF_FLOAT32);
+					ConvertData(C, BiNormalElemPtr, BiNormalFormat, &BiNormal, EMeshBufferFormat::Float32);
 				}
 			}
 		}
diff --git a/Engine/Plugins/Mutable/Source/MutableRuntime/Private/MuR/OpMeshDifference.h b/Engine/Plugins/Mutable/Source/MutableRuntime/Private/MuR/OpMeshDifference.h
--- a/Engine/Plugins/Mutable/Source/MutableRuntime/Private/MuR/OpMeshDifference.h
+++ b/Engine/Plugins/Mutable/Source/MutableRuntime/Private/MuR/OpMeshDifference.h
@@ -54,15 +54,15 @@ namespace mu
 			{
 				for ( int32 c=0; c<pBase->GetVertexBuffers().GetBufferChannelCount(vb); ++c )
 				{
-					EMeshBufferSemantic sem = MBS_NONE;
+					EMeshBufferSemantic sem = EMeshBufferSemantic::None;
 					int32 semIndex = 0;
 					pBase->GetVertexBuffers().GetChannel( vb, c, &sem, &semIndex, nullptr, nullptr, nullptr );			
-					if ( sem != MBS_VERTEXINDEX &&
-					     sem != MBS_BONEINDICES &&
-					     sem != MBS_BONEWEIGHTS &&
-					     sem != MBS_LAYOUTBLOCK &&
-					     sem != MBS_OTHER && 
-					     ( !ignoreTexCoords || sem!=MBS_TEXCOORDS ) )
+					if ( sem != EMeshBufferSemantic::VertexIndex &&
+					     sem != EMeshBufferSemantic::BoneIndices &&
+					     sem != EMeshBufferSemantic::BoneWeights &&
+					     sem != EMeshBufferSemantic::LayoutBlock &&
+					     sem != EMeshBufferSemantic::Other && 
+					     ( !ignoreTexCoords || sem!=EMeshBufferSemantic::TexCoords ) )
 					{
 						allSemantics.Add( sem );
 						allSemanticIndices.Add( semIndex );
@@ -133,7 +133,7 @@ namespace mu
 			{
 				semantic.Add(semantics[c]);
 				semanticIndex.Add(semanticIndices[c]);
-				format.Add(MBF_FLOAT32);
+				format.Add(EMeshBufferFormat::Float32);
 				components.Add(4);
 				offsets.Add(ElementSize);
 				ElementSize += 4 * sizeof(float);
@@ -143,19 +143,19 @@ namespace mu
 			Result->GetVertexBuffers().SetBuffer(BufferIndex, ElementSize, semantic.Num(), semantic.GetData(), semanticIndex.GetData(), format.GetData(), components.GetData(), offsets.GetData());
 
 			// The second buffer will contain the ids of the vertices to morph.
-			semantic = { MBS_VERTEXINDEX };
+			semantic = { EMeshBufferSemantic::VertexIndex };
 			semanticIndex = { 0 };
 			components = { 1 };
 			offsets = { 0 };
 			bool bGenerateRelativeIds = !pBase->AreVertexIdsExplicit();
 			if (bGenerateRelativeIds)
 			{
-				format = { MBF_UINT32 };
+				format = { EMeshBufferFormat::UInt32 };
 				ElementSize = sizeof(uint32);
 			}
 			else
 			{
-				format = { MBF_UINT64 };
+				format = { EMeshBufferFormat::UInt64 };
 				ElementSize = sizeof(uint64);
 			}
 
diff --git a/Engine/Plugins/Mutable/Source/MutableRuntime/Private/MuR/OpMeshExtractLayoutBlock.h b/Engine/Plugins/Mutable/Source/MutableRuntime/Private/MuR/OpMeshExtractLayoutBlock.h
--- a/Engine/Plugins/Mutable/Source/MutableRuntime/Private/MuR/OpMeshExtractLayoutBlock.h
+++ b/Engine/Plugins/Mutable/Source/MutableRuntime/Private/MuR/OpMeshExtractLayoutBlock.h
@@ -44,9 +44,9 @@ namespace mu
 			// Add a new buffer
 			int32 NewBuffer = Result->VertexBuffers.GetBufferCount();
 			Result->VertexBuffers.SetBufferCount(NewBuffer + 1);
-			EMeshBufferSemantic Semantic = MBS_VERTEXINDEX;
+			EMeshBufferSemantic Semantic = EMeshBufferSemantic::VertexIndex;
 			int32 SemanticIndex = 0;
-			EMeshBufferFormat Format = MBF_UINT32;
+			EMeshBufferFormat Format = EMeshBufferFormat::UInt32;
 			int32 Components = 1;
 			int32 Offset = 0;
 			Result->VertexBuffers.SetBuffer(NewBuffer, sizeof(uint32), 1, &Semantic, &SemanticIndex, &Format, &Components, &Offset);
@@ -63,11 +63,11 @@ namespace mu
 		TBitArray<> UsedSourceFaces;
 
 		UsedSourceFaces.SetNum(Source->GetFaceCount(), false);
-        UntypedMeshBufferIteratorConst itIndex(Source->GetIndexBuffers(), MBS_VERTEXINDEX);
-        UntypedMeshBufferIterator itResultIndex(Result->GetIndexBuffers(), MBS_VERTEXINDEX);
+        UntypedMeshBufferIteratorConst itIndex(Source->GetIndexBuffers(), EMeshBufferSemantic::VertexIndex);
+        UntypedMeshBufferIterator itResultIndex(Result->GetIndexBuffers(), EMeshBufferSemantic::VertexIndex);
         
         int32 IndexCount = 0;
-        if (itIndex.GetFormat() == MBF_UINT32)
+        if (itIndex.GetFormat() == EMeshBufferFormat::UInt32)
         {
             const uint32* pIndices = reinterpret_cast<const uint32*>(itIndex.ptr());
             uint32* pDestIndices = reinterpret_cast<uint32*>(itResultIndex.ptr());
@@ -84,7 +84,7 @@ namespace mu
                 }
             }
         }
-        else if (itIndex.GetFormat() == MBF_UINT16)
+        else if (itIndex.GetFormat() == EMeshBufferFormat::UInt16)
         {
             const uint16* pIndices = reinterpret_cast<const uint16*>(itIndex.ptr());
             uint16* pDestIndices = reinterpret_cast<uint16*>(itResultIndex.ptr());
@@ -134,9 +134,9 @@ namespace mu
 		// TODO: Optimise
 		Result->CopyFrom(*Source);
 
-		UntypedMeshBufferIteratorConst itBlocks(Source->GetVertexBuffers(), MBS_LAYOUTBLOCK, LayoutIndex);
+		UntypedMeshBufferIteratorConst itBlocks(Source->GetVertexBuffers(), EMeshBufferSemantic::LayoutBlock, LayoutIndex);
 
-        if (itBlocks.GetFormat() != MBF_NONE)
+        if (itBlocks.GetFormat() != EMeshBufferFormat::None)
         {
             int32 ResultVertices = 0;
 			TArray<int32> OldToNew;
@@ -144,7 +144,7 @@ namespace mu
 			TArray<int32> NewToOld;
             NewToOld.Reserve(Source->GetVertexCount());
 
-            if (itBlocks.GetFormat() == MBF_UINT16)
+            if (itBlocks.GetFormat() == EMeshBufferFormat::UInt16)
             {
                 const uint16* pBlocks = reinterpret_cast<const uint16*>(itBlocks.ptr());
                 for ( int32 i=0; i<Source->GetVertexCount(); ++i )
@@ -169,7 +169,7 @@ namespace mu
                     }
                 }
             }
-            else if (itBlocks.GetFormat() == MBF_UINT64)
+            else if (itBlocks.GetFormat() == EMeshBufferFormat::UInt64)
 			{
 				const uint64* pBlocks = reinterpret_cast<const uint64*>(itBlocks.ptr());
 				for (int32 i = 0; i < Source->GetVertexCount(); ++i)
@@ -211,9 +211,9 @@ namespace mu
 		// TODO: Optimise
 		Result->CopyFrom(*Source);
 
-		UntypedMeshBufferIteratorConst itBlocks(Source->GetVertexBuffers(), MBS_LAYOUTBLOCK, LayoutIndex);
+		UntypedMeshBufferIteratorConst itBlocks(Source->GetVertexBuffers(), EMeshBufferSemantic::LayoutBlock, LayoutIndex);
 
-		if (itBlocks.GetFormat() != MBF_NONE)
+		if (itBlocks.GetFormat() != EMeshBufferFormat::None)
 		{
 			int32 ResultVertices = 0;
 			TArray<int32> OldToNew;
@@ -221,7 +221,7 @@ namespace mu
 			TArray<int32> NewToOld;
 			NewToOld.Reserve(Source->GetVertexCount());
 
-			if (itBlocks.GetFormat() == MBF_UINT16)
+			if (itBlocks.GetFormat() == EMeshBufferFormat::UInt16)
 			{
 				const uint16* pBlocks = reinterpret_cast<const uint16*>(itBlocks.ptr());
 				for (int32 i = 0; i < Source->GetVertexCount(); ++i)
@@ -236,7 +236,7 @@ namespace mu
 					}
 				}
 			}
-			else if (itBlocks.GetFormat() == MBF_UINT64)
+			else if (itBlocks.GetFormat() == EMeshBufferFormat::UInt64)
 			{
 				const uint64* pBlocks = reinterpret_cast<const uint64*>(itBlocks.ptr());
 				for (int32 i = 0; i < Source->GetVertexCount(); ++i)
diff --git a/Engine/Plugins/Mutable/Source/MutableRuntime/Private/MuR/OpMeshFormat.cpp b/Engine/Plugins/Mutable/Source/MutableRuntime/Private/MuR/OpMeshFormat.cpp
--- a/Engine/Plugins/Mutable/Source/MutableRuntime/Private/MuR/OpMeshFormat.cpp
+++ b/Engine/Plugins/Mutable/Source/MutableRuntime/Private/MuR/OpMeshFormat.cpp
@@ -60,7 +60,7 @@ namespace mu
 
 			// Case 1: Special semantics that may be implicit or relative
 			//---------------------------------------------------------------------------------
-			if (bHasSpecialSemantics && ResultChannel.Semantic == MBS_VERTEXINDEX)
+			if (bHasSpecialSemantics && ResultChannel.Semantic == EMeshBufferSemantic::VertexIndex)
 			{
 				bool bHasVertexIndices = (SourceBufferIndex >=0 );
 				if (bHasVertexIndices)
@@ -80,10 +80,10 @@ namespace mu
 					{
 						// Relative vertex IDs
 						check(IDPrefix);
-						check(SourceChannel.Format == MBF_UINT32);
+						check(SourceChannel.Format == EMeshBufferFormat::UInt32);
 						const uint32* SourceData = reinterpret_cast<const uint32*>(SourceBuffer.Data.GetData());
 
-						check(ResultChannel.Format == MBF_UINT64);
+						check(ResultChannel.Format == EMeshBufferFormat::UInt64);
 						uint64* ResultData = reinterpret_cast<uint64*>(pResultBuf);
 
 						for (int32 i = 0; i < SourceElementCount; ++i)
@@ -100,7 +100,7 @@ namespace mu
 					// Implicit IDs
 					check(IDPrefix);
 
-					check(ResultChannel.Format == MBF_UINT64);
+					check(ResultChannel.Format == EMeshBufferFormat::UInt64);
 					uint64* ResultData = reinterpret_cast<uint64*>(pResultBuf);
 
 					for (int32 VertexIndex = 0; VertexIndex < SourceElementCount; ++VertexIndex)
@@ -114,7 +114,7 @@ namespace mu
 				continue;
 			}
 
-			else if (bHasSpecialSemantics && ResultChannel.Semantic == MBS_LAYOUTBLOCK)
+			else if (bHasSpecialSemantics && ResultChannel.Semantic == EMeshBufferSemantic::LayoutBlock)
 			{
 				bool bHasBlockIds = (SourceBufferIndex >= 0);
 
@@ -132,10 +132,10 @@ namespace mu
 					else
 					{
 						// Relative vertex IDs
-						check(SourceChannel.Format == MBF_UINT16);
+						check(SourceChannel.Format == EMeshBufferFormat::UInt16);
 						const uint16* SourceData = reinterpret_cast<const uint16*>(SourceBuffer.Data.GetData());
 
-						check(ResultChannel.Format == MBF_UINT64);
+						check(ResultChannel.Format == EMeshBufferFormat::UInt64);
 						uint64* ResultData = reinterpret_cast<uint64*>(pResultBuf);
 
 						for (int32 i = 0; i < SourceElementCount; ++i)
@@ -168,13 +168,13 @@ namespace mu
 
 				// If we have to add colour channels, we will add them as white, to be neutral.
 				// \todo: normal channels also should have special values.
-				if (ResultChannel.Semantic == MBS_COLOUR)
+				if (ResultChannel.Semantic == EMeshBufferSemantic::Color)
 				{
 					generated = true;
 
 					switch (ResultChannel.Format)
 					{
-					case MBF_FLOAT32:
+					case EMeshBufferFormat::Float32:
 					{
 						for (int32 v = 0; v < SourceElementCount; ++v)
 						{
@@ -188,7 +188,7 @@ namespace mu
 						break;
 					}
 
-					case MBF_NUINT8:
+					case EMeshBufferFormat::NUInt8:
 					{
 						for (int32 v = 0; v < SourceElementCount; ++v)
 						{
@@ -202,7 +202,7 @@ namespace mu
 						break;
 					}
 
-					case MBF_NUINT16:
+					case EMeshBufferFormat::NUInt16:
 					{
 						for (int32 v = 0; v < SourceElementCount; ++v)
 						{
@@ -268,9 +268,9 @@ namespace mu
 					{
 						FMemory::Memcpy(pResultBuf, pSourceBuf, resultChannelSize);
 					}
-					else if (ResultChannel.Format == MBF_PACKEDDIR8_W_TANGENTSIGN
+					else if (ResultChannel.Format == EMeshBufferFormat::PackedDir8_W_TangentSign
 						||
-						ResultChannel.Format == MBF_PACKEDDIRS8_W_TANGENTSIGN)
+						ResultChannel.Format == EMeshBufferFormat::PackedDirS8_W_TangentSign)
 					{
 						check(sourceComponents >= 3);
 						check(resultComponents == 4);
@@ -295,15 +295,15 @@ namespace mu
 
 						// Look for the full tangent space
 						int32 tanXBuf, tanXChan, tanYBuf, tanYChan, tanZBuf, tanZChan;
-						Source.FindChannel(MBS_TANGENT, ResultChannel.SemanticIndex, &tanXBuf, &tanXChan);
-						Source.FindChannel(MBS_BINORMAL, ResultChannel.SemanticIndex, &tanYBuf, &tanYChan);
-						Source.FindChannel(MBS_NORMAL, ResultChannel.SemanticIndex, &tanZBuf, &tanZChan);
+						Source.FindChannel(EMeshBufferSemantic::Tangent, ResultChannel.SemanticIndex, &tanXBuf, &tanXChan);
+						Source.FindChannel(EMeshBufferSemantic::Binormal, ResultChannel.SemanticIndex, &tanYBuf, &tanYChan);
+						Source.FindChannel(EMeshBufferSemantic::Normal, ResultChannel.SemanticIndex, &tanZBuf, &tanZChan);
 
 						if (tanXBuf >= 0 && tanYBuf >= 0 && tanZBuf >= 0)
 						{
-							UntypedMeshBufferIteratorConst xIt(Source, MBS_TANGENT, ResultChannel.SemanticIndex);
-							UntypedMeshBufferIteratorConst yIt(Source, MBS_BINORMAL, ResultChannel.SemanticIndex);
-							UntypedMeshBufferIteratorConst zIt(Source, MBS_NORMAL, ResultChannel.SemanticIndex);
+							UntypedMeshBufferIteratorConst xIt(Source, EMeshBufferSemantic::Tangent, ResultChannel.SemanticIndex);
+							UntypedMeshBufferIteratorConst yIt(Source, EMeshBufferSemantic::Binormal, ResultChannel.SemanticIndex);
+							UntypedMeshBufferIteratorConst zIt(Source, EMeshBufferSemantic::Normal, ResultChannel.SemanticIndex);
 
 							xIt += v;
 							yIt += v;
@@ -312,11 +312,11 @@ namespace mu
 							FMatrix44f Mat(xIt.GetAsVec3f(), yIt.GetAsVec3f(), zIt.GetAsVec3f(), FVector3f(0, 0, 0));
 
 							uint8 sign = 0;
-							if (ResultChannel.Format == MBF_PACKEDDIR8_W_TANGENTSIGN)
+							if (ResultChannel.Format == EMeshBufferFormat::PackedDir8_W_TangentSign)
 							{
 								sign = Mat.RotDeterminant() < 0 ? 0 : 255;
 							}
-							else if (ResultChannel.Format == MBF_PACKEDDIRS8_W_TANGENTSIGN)
+							else if (ResultChannel.Format == EMeshBufferFormat::PackedDirS8_W_TangentSign)
 							{
 								sign = Mat.RotDeterminant() < 0 ? -128 : 127;
 							}
@@ -357,9 +357,9 @@ namespace mu
 						// Extra step to normalise some semantics in some formats
 						// TODO: Make it optional, and add different normalisation types n, n^2
 						// TODO: Optimise
-						if (sourceSemantic == MBS_BONEWEIGHTS)
+						if (sourceSemantic == EMeshBufferSemantic::BoneWeights)
 						{
-							if (ResultChannel.Format == MBF_NUINT8)
+							if (ResultChannel.Format == EMeshBufferFormat::NUInt8)
 							{
 								uint8* pData = (uint8*)pResultBuf;
 								uint8 accum = 0;
@@ -370,7 +370,7 @@ namespace mu
 								pData[0] += 255 - accum;
 							}
 
-							else if (ResultChannel.Format == MBF_NUINT16)
+							else if (ResultChannel.Format == EMeshBufferFormat::NUInt16)
 							{
 								uint16* pData = (uint16*)pResultBuf;
 								uint16 accum = 0;
@@ -496,9 +496,9 @@ namespace mu
 						&sourceOffset
 					);
 
-					if (sourceSemantic == MBS_LAYOUTBLOCK
+					if (sourceSemantic == EMeshBufferSemantic::LayoutBlock
 						||
-						(bIsVertexBuffer && sourceSemantic == MBS_VERTEXINDEX))
+						(bIsVertexBuffer && sourceSemantic == EMeshBufferSemantic::VertexIndex))
 					{
 						// Add it if it wasn't already there, which could happen if it was included in the format mesh.
 						int32 AlreadyExistingBufIndex = -1;
@@ -571,15 +571,15 @@ namespace mu
 				{
 					const FMeshBufferChannel& Channel = VertexBuffers.Buffers[BufferIndex].Channels[ChannelIndex];
 
-					if (Channel.Semantic == MBS_BONEINDICES)
+					if (Channel.Semantic == EMeshBufferSemantic::BoneIndices)
 					{
 						int32 resultBuf = 0;
 						int32 resultChan = 0;
 						FMeshBufferSet& formBuffs = Result->GetVertexBuffers();
-						formBuffs.FindChannel(MBS_BONEINDICES, Channel.SemanticIndex, &resultBuf, &resultChan);
+						formBuffs.FindChannel(EMeshBufferSemantic::BoneIndices, Channel.SemanticIndex, &resultBuf, &resultChan);
 						if (resultBuf >= 0)
 						{
-							UntypedMeshBufferIteratorConst it(VertexBuffers, MBS_BONEINDICES, Channel.SemanticIndex);
+							UntypedMeshBufferIteratorConst it(VertexBuffers, EMeshBufferSemantic::BoneIndices, Channel.SemanticIndex);
 							int32 maxBoneIndex = 0;
 							for (int32 v = 0; v < VertexBuffers.GetElementCount(); ++v)
 							{
@@ -595,24 +595,24 @@ namespace mu
 							}
 
 							EMeshBufferFormat& format = formBuffs.Buffers[resultBuf].Channels[resultChan].Format;
-							if (maxBoneIndex > 0xffff && (format == MBF_UINT8 || format == MBF_UINT16))
+							if (maxBoneIndex > 0xffff && (format == EMeshBufferFormat::UInt8 || format == EMeshBufferFormat::UInt16))
 							{
-								format = MBF_UINT32;
+								format = EMeshBufferFormat::UInt32;
 								formBuffs.UpdateOffsets(resultBuf);
 							}
-							else if (maxBoneIndex > 0x7fff && (format == MBF_INT8 || format == MBF_INT16))
+							else if (maxBoneIndex > 0x7fff && (format == EMeshBufferFormat::Int8 || format == EMeshBufferFormat::Int16))
 							{
-								format = MBF_UINT32;
+								format = EMeshBufferFormat::UInt32;
 								formBuffs.UpdateOffsets(resultBuf);
 							}
-							else if (maxBoneIndex > 0xff && format == MBF_UINT8)
+							else if (maxBoneIndex > 0xff && format == EMeshBufferFormat::UInt8)
 							{
-								format = MBF_UINT16;
+								format = EMeshBufferFormat::UInt16;
 								formBuffs.UpdateOffsets(resultBuf);
 							}
-							else if (maxBoneIndex > 0x7f && format == MBF_INT8)
+							else if (maxBoneIndex > 0x7f && format == EMeshBufferFormat::Int8)
 							{
-								format = MBF_INT16;
+								format = EMeshBufferFormat::Int16;
 								formBuffs.UpdateOffsets(resultBuf);
 							}
 						}
@@ -684,7 +684,7 @@ namespace mu
 		// Reduce the number of influences if possible
 		constexpr int32 SemanticIndex = 0;
 		
-		UntypedMeshBufferIteratorConst WeightIt(VertexBuffers, MBS_BONEWEIGHTS, SemanticIndex);
+		UntypedMeshBufferIteratorConst WeightIt(VertexBuffers, EMeshBufferSemantic::BoneWeights, SemanticIndex);
 		if (WeightIt.ptr())
 		{
 			int32 BufferInfluences = WeightIt.GetComponents();
@@ -696,7 +696,7 @@ namespace mu
 
 				switch (WeightIt.GetFormat())
 				{
-				case MBF_NUINT8:
+				case EMeshBufferFormat::NUInt8:
 				{
 					const uint8* Data = reinterpret_cast<const uint8*>(WeightIt.ptr());
 					for (int32 InfluenceIndex = 0; InfluenceIndex < BufferInfluences; ++InfluenceIndex)
@@ -709,7 +709,7 @@ namespace mu
 					}
 					break;
 				}
-				case MBF_NUINT16:
+				case EMeshBufferFormat::NUInt16:
 				{
 					const uint16* Data = reinterpret_cast<const uint16*>(WeightIt.ptr());
 					for (int32 InfluenceIndex = 0; InfluenceIndex < BufferInfluences; ++InfluenceIndex)
@@ -751,7 +751,7 @@ namespace mu
 
 						if (Channel.SemanticIndex == SemanticIndex
 							&&
-							(Channel.Semantic == MBS_BONEWEIGHTS || Channel.Semantic == MBS_BONEINDICES)
+							(Channel.Semantic == EMeshBufferSemantic::BoneWeights || Channel.Semantic == EMeshBufferSemantic::BoneIndices)
 							)
 						{
 							Channel.ComponentCount = RealInfluences;
diff --git a/Engine/Plugins/Mutable/Source/MutableRuntime/Private/MuR/OpMeshGeometryOperation.h b/Engine/Plugins/Mutable/Source/MutableRuntime/Private/MuR/OpMeshGeometryOperation.h
--- a/Engine/Plugins/Mutable/Source/MutableRuntime/Private/MuR/OpMeshGeometryOperation.h
+++ b/Engine/Plugins/Mutable/Source/MutableRuntime/Private/MuR/OpMeshGeometryOperation.h
@@ -41,8 +41,8 @@ namespace mu
 			TriangleMutableToDyn.SetNum(TriangleCount);
 
 			// \TODO: Simple but inefficient
-			UntypedMeshBufferIteratorConst ItPosition(MutableMesh->GetVertexBuffers(), MBS_POSITION);
-			UntypedMeshBufferIteratorConst ItNormal(MutableMesh->GetVertexBuffers(), MBS_NORMAL);
+			UntypedMeshBufferIteratorConst ItPosition(MutableMesh->GetVertexBuffers(), EMeshBufferSemantic::Position);
+			UntypedMeshBufferIteratorConst ItNormal(MutableMesh->GetVertexBuffers(), EMeshBufferSemantic::Normal);
 			for (int32 VertexIndex = 0; VertexIndex < VertexCount; ++VertexIndex)
 			{
 				FVector3d Position = ItPosition.GetAsVec3d();
@@ -55,7 +55,7 @@ namespace mu
 				VertexMutableToDyn[VertexIndex] = DynId;
 			}
 
-			UntypedMeshBufferIteratorConst ItIndices(MutableMesh->GetIndexBuffers(), MBS_VERTEXINDEX);
+			UntypedMeshBufferIteratorConst ItIndices(MutableMesh->GetIndexBuffers(), EMeshBufferSemantic::VertexIndex);
 			for (int32 TriangleIndex = 0; TriangleIndex < TriangleCount; ++TriangleIndex)
 			{
 				UE::Geometry::FIndex3i Triangle;
@@ -83,8 +83,8 @@ namespace mu
 		MUTABLE_CPUPROFILER_SCOPE(DynamicMeshToMutable);
 
 		// \TODO: Simple but inefficient
-		UntypedMeshBufferIterator ItPosition(MutableMesh->GetVertexBuffers(), MBS_POSITION);
-		UntypedMeshBufferIterator ItNormal(MutableMesh->GetVertexBuffers(), MBS_NORMAL);
+		UntypedMeshBufferIterator ItPosition(MutableMesh->GetVertexBuffers(), EMeshBufferSemantic::Position);
+		UntypedMeshBufferIterator ItNormal(MutableMesh->GetVertexBuffers(), EMeshBufferSemantic::Normal);
 		
 		int32 VertexCount = MutableMesh->GetVertexCount();
 		for (int32 VertexIndex = 0; VertexIndex < VertexCount; ++VertexIndex)
diff --git a/Engine/Plugins/Mutable/Source/MutableRuntime/Private/MuR/OpMeshInterpolate.h b/Engine/Plugins/Mutable/Source/MutableRuntime/Private/MuR/OpMeshInterpolate.h
--- a/Engine/Plugins/Mutable/Source/MutableRuntime/Private/MuR/OpMeshInterpolate.h
+++ b/Engine/Plugins/Mutable/Source/MutableRuntime/Private/MuR/OpMeshInterpolate.h
@@ -27,22 +27,22 @@ namespace mu
 			// A mesh position
 			int aPosElementSize = 0;
             uint8* pAPosBuf = 0;
-			GetMeshBuf( pDest.get(), MBS_POSITION, MBF_FLOAT32, 3, pAPosBuf, aPosElementSize );
+			GetMeshBuf( pDest.get(), EMeshBufferSemantic::Position, EMeshBufferFormat::Float32, 3, pAPosBuf, aPosElementSize );
 
 			// A mesh normal
 			int aNorElementSize = 0;
             uint8* pANorBuf = 0;
-			GetMeshBuf( pDest.get(), MBS_NORMAL, MBF_FLOAT32, 3, pANorBuf, aNorElementSize );
+			GetMeshBuf( pDest.get(), EMeshBufferSemantic::Normal, EMeshBufferFormat::Float32, 3, pANorBuf, aNorElementSize );
 
 			// B mesh position
 			int bPosElementSize = 0;
             const uint8* pBPosBuf = 0;
-			GetMeshBuf( pB, MBS_POSITION, MBF_FLOAT32, 3, pBPosBuf, bPosElementSize );
+			GetMeshBuf( pB, EMeshBufferSemantic::Position, EMeshBufferFormat::Float32, 3, pBPosBuf, bPosElementSize );
 
 			// B mesh normal
 			int bNorElementSize = 0;
             const uint8* pBNorBuf = 0;
-			GetMeshBuf( pB, MBS_NORMAL, MBF_FLOAT32, 3, pBNorBuf, bNorElementSize );
+			GetMeshBuf( pB, EMeshBufferSemantic::Normal, EMeshBufferFormat::Float32, 3, pBNorBuf, bNorElementSize );
 
             uint32 vcount = pA->GetVertexBuffers().GetElementCount();
             for ( uint32 vi=0; vi<vcount; ++vi )
diff --git a/Engine/Plugins/Mutable/Source/MutableRuntime/Private/MuR/OpMeshMorph.h b/Engine/Plugins/Mutable/Source/MutableRuntime/Private/MuR/OpMeshMorph.h
--- a/Engine/Plugins/Mutable/Source/MutableRuntime/Private/MuR/OpMeshMorph.h
+++ b/Engine/Plugins/Mutable/Source/MutableRuntime/Private/MuR/OpMeshMorph.h
@@ -88,7 +88,7 @@ namespace mu
             const int32 BiNormalComps = BaseBiNormalIter.GetComponents();
 
 			// When normal is packed, binormal channel is not expected. It is not a big deal if it's there but we would be doing extra unused work in that case. 
-			ensureAlways(!(NormalFormat == MBF_PACKEDDIR8_W_TANGENTSIGN || NormalFormat == MBF_PACKEDDIRS8_W_TANGENTSIGN) || !BaseBiNormalIter.ptr());
+			ensureAlways(!(NormalFormat == EMeshBufferFormat::PackedDir8_W_TangentSign || NormalFormat == EMeshBufferFormat::PackedDirS8_W_TangentSign) || !BaseBiNormalIter.ptr());
 
             for (int32 VertexIndex = 0; VertexIndex < BaseNum; ++VertexIndex)
             {
@@ -120,7 +120,7 @@ namespace mu
 					// Leave the tangent basis sign untouched for packed normals formats.
 					for (int32 C = 0; C < NormalComps && C < 3; ++C)
 					{
-						ConvertData(C, NormalIter.ptr(), NormalFormat, &Normal, MBF_FLOAT32);
+						ConvertData(C, NormalIter.ptr(), NormalFormat, &Normal, EMeshBufferFormat::Float32);
 					}
 
 					// Tangent
@@ -135,7 +135,7 @@ namespace mu
 
 						for (int32 C = 0; C < TangentComps && C < 3; ++C)
 						{
-							ConvertData(C, TangentIter.ptr(), TangentFormat, &Tangent, MBF_FLOAT32);
+							ConvertData(C, TangentIter.ptr(), TangentFormat, &Tangent, EMeshBufferFormat::Float32);
 						}
 
 						// BiNormal
@@ -157,7 +157,7 @@ namespace mu
 
 							for (int32 C = 0; C < BiNormalComps && C < 3; ++C)
 							{
-								ConvertData(C, BiNormalIter.ptr(), BiNormalFormat, &BiNormal, MBF_FLOAT32);
+								ConvertData(C, BiNormalIter.ptr(), BiNormalFormat, &BiNormal, EMeshBufferFormat::Float32);
 							}
 						}
 					}
@@ -214,7 +214,7 @@ namespace mu
                         // Max 4 components
                         for (int32 Comp = 0; Comp < DestChannelComps && Comp < 4; ++Comp)
                         {
-                            ConvertData(Comp, ChannelBaseIter.ptr(), DestChannelFormat, &Value, MBF_FLOAT32);
+                            ConvertData(Comp, ChannelBaseIter.ptr(), DestChannelFormat, &Value, EMeshBufferFormat::Float32);
                         }
                     }
                 }
@@ -252,7 +252,7 @@ namespace mu
 			UntypedMeshBufferIteratorConst MinNormalChannelIter;
 			UntypedMeshBufferIteratorConst MaxNormalChannelIter;
 
-			const bool bBaseHasNormals = UntypedMeshBufferIteratorConst(BaseMesh->GetVertexBuffers(), MBS_NORMAL, 0).ptr() != nullptr;
+			const bool bBaseHasNormals = UntypedMeshBufferIteratorConst(BaseMesh->GetVertexBuffers(), EMeshBufferSemantic::Normal, 0).ptr() != nullptr;
 			for (int32 ChannelIndex = 0; ChannelIndex < ChannelsNum; ++ChannelIndex)
 			{
 				const FMeshBufferSet& MBSPriv = RefTarget->GetVertexBuffers();
@@ -260,7 +260,7 @@ namespace mu
 				EMeshBufferSemantic Sem = Channel.Semantic;
 				int32 SemIndex = Channel.SemanticIndex;
 			
-				if (Sem == MBS_NORMAL && bBaseHasNormals)
+				if (Sem == EMeshBufferSemantic::Normal && bBaseHasNormals)
 				{
 					BaseTangentFrameChannelsIters[2] = UntypedMeshBufferIterator(BaseMesh->GetVertexBuffers(), Sem, SemIndex);
 					if (MinNum > 0)
@@ -273,11 +273,11 @@ namespace mu
 						MaxNormalChannelIter = UntypedMeshBufferIteratorConst(MaxMesh->GetVertexBuffers(), Sem, SemIndex);
 					}
 				}
-				else if (Sem == MBS_TANGENT && bBaseHasNormals)
+				else if (Sem == EMeshBufferSemantic::Tangent && bBaseHasNormals)
 				{
 					BaseTangentFrameChannelsIters[1] = UntypedMeshBufferIterator(BaseMesh->GetVertexBuffers(), Sem, SemIndex);
 				}
-				else if (Sem == MBS_BINORMAL && bBaseHasNormals)
+				else if (Sem == EMeshBufferSemantic::Binormal && bBaseHasNormals)
 				{
 					BaseTangentFrameChannelsIters[0] = UntypedMeshBufferIterator(BaseMesh->GetVertexBuffers(), Sem, SemIndex);
 				}
diff --git a/Engine/Plugins/Mutable/Source/MutableRuntime/Private/MuR/OpMeshOptimizeSkinning.h b/Engine/Plugins/Mutable/Source/MutableRuntime/Private/MuR/OpMeshOptimizeSkinning.h
--- a/Engine/Plugins/Mutable/Source/MutableRuntime/Private/MuR/OpMeshOptimizeSkinning.h
+++ b/Engine/Plugins/Mutable/Source/MutableRuntime/Private/MuR/OpMeshOptimizeSkinning.h
@@ -42,7 +42,7 @@ namespace mu
 		bool bRequiresFormatChange = false;
 
 		// Desired format if MaxBoneMapIndex <= MAX_uint8
-		const EMeshBufferFormat DesiredBoneIndexFormat = mu::MBF_UINT8;
+		const EMeshBufferFormat DesiredBoneIndexFormat = mu::EMeshBufferFormat::UInt8;
 		
 		// Iterate all vertex buffers and check if the BoneIndices buffers have the desired format
 		const FMeshBufferSet& InMeshVertexBuffers = InMesh->GetVertexBuffers();
@@ -53,7 +53,7 @@ namespace mu
 			const int32 ChannelsCount = InMeshVertexBuffers.GetBufferChannelCount(VertexBufferIndex);
 			for (int32 ChannelIndex = 0; ChannelIndex < ChannelsCount; ++ChannelIndex)
 			{
-				if (Buffer.Channels[ChannelIndex].Semantic == MBS_BONEINDICES)
+				if (Buffer.Channels[ChannelIndex].Semantic == EMeshBufferSemantic::BoneIndices)
 				{
 					bRequiresFormatChange = Buffer.Channels[ChannelIndex].Format != DesiredBoneIndexFormat;
 					break;
@@ -70,7 +70,7 @@ namespace mu
 
 		MUTABLE_CPUPROFILER_SCOPE(MeshOptimizeSkinning_Format);
 
-		// Format Bone Indices. For some reason, the bone index format is MBF_UINT16 and it should be MBF_UINT8.
+		// Format Bone Indices. For some reason, the bone index format is EMeshBufferFormat::UInt16 and it should be EMeshBufferFormat::UInt8.
 
 		// TODO: Replace by MeshFormatInPlace once implemented
 		const int32 VertexBuffersCount = InMeshVertexBuffers.GetBufferCount();
@@ -99,7 +99,7 @@ namespace mu
 			TArray<int32> Components;
 			TArray<int32> Offsets;
 
-			EMeshBufferFormat SourceBoneIndexFormat = mu::MBF_NONE;
+			EMeshBufferFormat SourceBoneIndexFormat = mu::EMeshBufferFormat::None;
 
 			// Offset accumulator
 			int32 AuxOffset = 0;
@@ -110,10 +110,10 @@ namespace mu
 				const FMeshBufferChannel& Channel = SourceBuffer.Channels[ChannelIndex];
 				EMeshBufferFormat Format = Channel.Format;
 
-				if (Channel.Semantic == MBS_BONEINDICES)
+				if (Channel.Semantic == EMeshBufferSemantic::BoneIndices)
 				{
 					SourceBoneIndexFormat = Channel.Format;
-					Format = mu::MBF_UINT8;
+					Format = mu::EMeshBufferFormat::UInt8;
 					BoneIndexChannelIndex = ChannelIndex;
 				}
 
@@ -130,7 +130,7 @@ namespace mu
 			}
 
 			// Copy buffers
-			if (BoneIndexChannelIndex != INDEX_NONE && SourceBoneIndexFormat == mu::MBF_UINT16)
+			if (BoneIndexChannelIndex != INDEX_NONE && SourceBoneIndexFormat == mu::EMeshBufferFormat::UInt16)
 			{
 				// This buffer has a BoneIndex channel that needs to be manually fixed. 
 				VertexBuffers.SetBufferCount(BufferIndex + 1);
@@ -140,7 +140,7 @@ namespace mu
 				const uint8* SourceData = (const uint8*)SourceBuffer.Data.GetData();
 				const int32 SourceBoneIndexSize = GetMeshFormatData(SourceBoneIndexFormat).SizeInBytes;
 
-				const int32 BoneIndexSize = GetMeshFormatData(mu::MBF_UINT8).SizeInBytes;
+				const int32 BoneIndexSize = GetMeshFormatData(mu::EMeshBufferFormat::UInt8).SizeInBytes;
 				const int32 BoneIndexComponentCount = Components[BoneIndexChannelIndex];
 				const int32 BoneIndexChannelOffset = Offsets[BoneIndexChannelIndex];
 
@@ -185,8 +185,8 @@ namespace mu
 			}
 			else
 			{
-				// SourceBoneIndexFormat must be mu::MBF_UINT8 or none if the buffer doesn't have BoneIndices
-				check(SourceBoneIndexFormat == mu::MBF_NONE || SourceBoneIndexFormat == mu::MBF_UINT8);
+				// SourceBoneIndexFormat must be mu::EMeshBufferFormat::UInt8 or none if the buffer doesn't have BoneIndices
+				check(SourceBoneIndexFormat == mu::EMeshBufferFormat::None || SourceBoneIndexFormat == mu::EMeshBufferFormat::UInt8);
 
 				// Add buffers that don't require a fix up
 				VertexBuffers.AddBuffer(InMeshVertexBuffers, BufferIndex);
diff --git a/Engine/Plugins/Mutable/Source/MutableRuntime/Private/MuR/OpMeshRemove.cpp b/Engine/Plugins/Mutable/Source/MutableRuntime/Private/MuR/OpMeshRemove.cpp
--- a/Engine/Plugins/Mutable/Source/MutableRuntime/Private/MuR/OpMeshRemove.cpp
+++ b/Engine/Plugins/Mutable/Source/MutableRuntime/Private/MuR/OpMeshRemove.cpp
@@ -144,7 +144,7 @@ namespace mu
 
     //    // Map from source vertex index, to new vertex index for used vertices.
     //    // These are indices as in the index buffer, not the absolute vertex index as in the
-    //    // vertexbuffer MBS_VERTEXINDEX buffers.
+    //    // vertexbuffer EMeshBufferSemantic::VertexIndex buffers.
 	//	TArray<int32> UsedVertexMap;
     //    TBitArray<> UsedFaces;
 	//	UsedVertexMap.Init(-1, Result->GetVertexCount());
@@ -152,8 +152,8 @@ namespace mu
     //    {
     //        if ( Result->GetIndexBuffers().GetElementSize(0)==4 )
     //        {
-    //            MeshBufferIteratorConst<MBF_UINT32,uint32,1> itSource( Result->GetIndexBuffers(), MBS_VERTEXINDEX );
-    //            MeshBufferIterator<MBF_UINT32,uint32,1> itDest( Result->GetIndexBuffers(), MBS_VERTEXINDEX );
+    //            MeshBufferIteratorConst<EMeshBufferFormat::UInt32,uint32,1> itSource( Result->GetIndexBuffers(), EMeshBufferSemantic::VertexIndex );
+    //            MeshBufferIterator<EMeshBufferFormat::UInt32,uint32,1> itDest( Result->GetIndexBuffers(), EMeshBufferSemantic::VertexIndex );
 
     //            int32 IndexCount = Result->GetIndexCount();
 	//			check(IndexCount%3==0);
@@ -202,8 +202,8 @@ namespace mu
 
     //        else if ( Result->GetIndexBuffers().GetElementSize(0)==2 )
     //        {
-    //            MeshBufferIteratorConst<MBF_UINT16,uint16,1> itSource( Result->GetIndexBuffers(), MBS_VERTEXINDEX );
-    //            MeshBufferIterator<MBF_UINT16,uint16,1> itDest( Result->GetIndexBuffers(), MBS_VERTEXINDEX );
+    //            MeshBufferIteratorConst<EMeshBufferFormat::UInt16,uint16,1> itSource( Result->GetIndexBuffers(), EMeshBufferSemantic::VertexIndex );
+    //            MeshBufferIterator<EMeshBufferFormat::UInt16,uint16,1> itDest( Result->GetIndexBuffers(), EMeshBufferSemantic::VertexIndex );
 
     //            int32 indexCount = Result->GetIndexCount();
     //            for ( int32 f=0; f<indexCount/3; ++f )
@@ -341,7 +341,7 @@ namespace mu
     { 
         MUTABLE_CPUPROFILER_SCOPE(MeshRemoveVerticesWithCullSet);
 		
-        const UntypedMeshBufferIterator IndicesBegin(Result->GetIndexBuffers(), MBS_VERTEXINDEX);
+        const UntypedMeshBufferIterator IndicesBegin(Result->GetIndexBuffers(), EMeshBufferSemantic::VertexIndex);
 		
         const int32 NumFaces = Result->GetFaceCount();
         const int32 NumVertices = Result->GetVertexCount();
diff --git a/Engine/Plugins/Mutable/Source/MutableRuntime/Private/MuR/OpMeshReshape.h b/Engine/Plugins/Mutable/Source/MutableRuntime/Private/MuR/OpMeshReshape.h
--- a/Engine/Plugins/Mutable/Source/MutableRuntime/Private/MuR/OpMeshReshape.h
+++ b/Engine/Plugins/Mutable/Source/MutableRuntime/Private/MuR/OpMeshReshape.h
@@ -374,10 +374,10 @@ namespace mu
 		check(Mesh);
 		check(Mesh->GetVertexCount() == BindingData.Num());
 		
-		const UntypedMeshBufferIterator PositionIter(Mesh->GetVertexBuffers(), MBS_POSITION);
-		const UntypedMeshBufferIterator NormalIter(Mesh->GetVertexBuffers(), MBS_NORMAL);
-		const UntypedMeshBufferIterator TangentIter(Mesh->GetVertexBuffers(), MBS_TANGENT);
-		const UntypedMeshBufferIterator BiNormalIter(Mesh->GetVertexBuffers(), MBS_BINORMAL);
+		const UntypedMeshBufferIterator PositionIter(Mesh->GetVertexBuffers(), EMeshBufferSemantic::Position);
+		const UntypedMeshBufferIterator NormalIter(Mesh->GetVertexBuffers(), EMeshBufferSemantic::Normal);
+		const UntypedMeshBufferIterator TangentIter(Mesh->GetVertexBuffers(), EMeshBufferSemantic::Tangent);
+		const UntypedMeshBufferIterator BiNormalIter(Mesh->GetVertexBuffers(), EMeshBufferSemantic::Binormal);
 
         const EMeshBufferFormat NormalFormat   = NormalIter.GetFormat();
         const EMeshBufferFormat TangentFormat  = TangentIter.GetFormat();
@@ -444,7 +444,7 @@ namespace mu
 					uint8 * const NormalElemPtr = (NormalIter + VertexIndex).ptr();
 					for (int32 C = 0; C < NormalComps && C < 3; ++C)
 					{
-						ConvertData(C, NormalElemPtr, NormalFormat, &Normal, MBF_FLOAT32);
+						ConvertData(C, NormalElemPtr, NormalFormat, &Normal, EMeshBufferFormat::Float32);
 					}
 
 					if (TangentIter.ptr())
@@ -452,7 +452,7 @@ namespace mu
 						uint8 * const TangentElemPtr = (TangentIter + VertexIndex).ptr();
 						for (int32 C = 0; C < TangentComps && C < 3; ++C)
 						{
-							ConvertData(C, TangentElemPtr, TangentFormat, &Tangent, MBF_FLOAT32);
+							ConvertData(C, TangentElemPtr, TangentFormat, &Tangent, EMeshBufferFormat::Float32);
 						}
 					}
 					
@@ -461,7 +461,7 @@ namespace mu
 						uint8 * const BiNormalElemPtr = (BiNormalIter + VertexIndex).ptr();
 						for (int32 C = 0; C < BiNormalComps && C < 3; ++C)
 						{
-							ConvertData(C, BiNormalElemPtr, BiNormalFormat, &BiNormal, MBF_FLOAT32);
+							ConvertData(C, BiNormalElemPtr, BiNormalFormat, &BiNormal, EMeshBufferFormat::Float32);
 						}
 					}
 				}
@@ -746,7 +746,7 @@ namespace mu
 		int32 BarycentricDataBuffer = 0;
 		int32 BarycentricDataChannel = 0;
 		const FMeshBufferSet& VB = BaseMesh->GetVertexBuffers();
-		VB.FindChannel(MBS_BARYCENTRICCOORDS, BindingDataIndex, &BarycentricDataBuffer, &BarycentricDataChannel);
+		VB.FindChannel(EMeshBufferSemantic::BarycentricCoords, BindingDataIndex, &BarycentricDataBuffer, &BarycentricDataChannel);
 		
 		Ptr<Mesh> TempMesh = nullptr;
 		Mesh* VerticesReshapeMesh = Result;
@@ -816,8 +816,8 @@ namespace mu
 			ShapeDescriptor.Normals.SetNum(ShapeVertexCount);
 
 			// \TODO: Simple but inefficient
-			const UntypedMeshBufferIteratorConst ItPosition(ShapeMesh->GetVertexBuffers(), MBS_POSITION);
-			const UntypedMeshBufferIteratorConst ItNormal(ShapeMesh->GetVertexBuffers(), MBS_NORMAL);
+			const UntypedMeshBufferIteratorConst ItPosition(ShapeMesh->GetVertexBuffers(), EMeshBufferSemantic::Position);
+			const UntypedMeshBufferIteratorConst ItNormal(ShapeMesh->GetVertexBuffers(), EMeshBufferSemantic::Normal);
 			for (int32 ShapeVertexIndex = 0; ShapeVertexIndex < ShapeVertexCount; ++ShapeVertexIndex)
 			{
 				FVector3f Position = (ItPosition + ShapeVertexIndex).GetAsVec3f();
@@ -836,7 +836,7 @@ namespace mu
 			ShapeDescriptor.Triangles.SetNum(ShapeTriangleCount);
 
 			// \TODO: Simple but inefficient
-			const UntypedMeshBufferIteratorConst ItIndices(ShapeMesh->GetIndexBuffers(), MBS_VERTEXINDEX);
+			const UntypedMeshBufferIteratorConst ItIndices(ShapeMesh->GetIndexBuffers(), EMeshBufferSemantic::VertexIndex);
 			for (int32 TriangleIndex = 0; TriangleIndex < ShapeTriangleCount; ++TriangleIndex)
 			{
 				UE::Geometry::FIndex3i Triangle;
@@ -895,7 +895,7 @@ namespace mu
 			BarycentricDataBuffer = -1;
 			if (SkeletonBindBuffer)
 			{
-				SkeletonBindBuffer->FindChannel(MBS_BARYCENTRICCOORDS, BindingDataIndex, &BarycentricDataBuffer, &BarycentricDataChannel);
+				SkeletonBindBuffer->FindChannel(EMeshBufferSemantic::BarycentricCoords, BindingDataIndex, &BarycentricDataBuffer, &BarycentricDataChannel);
 			}
 
 			if (SkeletonBindBuffer && BarycentricDataBuffer >= 0)
@@ -940,7 +940,7 @@ namespace mu
 			BarycentricDataBuffer = -1;
 			if (FoundPhysicsBindBuffer)
 			{
-				FoundPhysicsBindBuffer->Value.FindChannel(MBS_BARYCENTRICCOORDS, BindingDataIndex, &BarycentricDataBuffer, &BarycentricDataChannel);
+				FoundPhysicsBindBuffer->Value.FindChannel(EMeshBufferSemantic::BarycentricCoords, BindingDataIndex, &BarycentricDataBuffer, &BarycentricDataChannel);
 			}
 			
 			const bool bAllNeededBuffersFound =
diff --git a/Engine/Plugins/Mutable/Source/MutableRuntime/Private/MuR/OpMeshSmoothing.h b/Engine/Plugins/Mutable/Source/MutableRuntime/Private/MuR/OpMeshSmoothing.h
--- a/Engine/Plugins/Mutable/Source/MutableRuntime/Private/MuR/OpMeshSmoothing.h
+++ b/Engine/Plugins/Mutable/Source/MutableRuntime/Private/MuR/OpMeshSmoothing.h
@@ -495,8 +495,8 @@ namespace mu
 				reinterpret_cast<const int32*>(MeshUniqueVertexMapBuffer->Value.GetBufferData(0)),
 				MeshUniqueVertexMapBuffer->Value.GetElementCount());
 
-		const UntypedMeshBufferIteratorConst SrcPositionBegin(SrcMesh.GetVertexBuffers(), MBS_POSITION);
-		const UntypedMeshBufferIterator      DstPositionBegin(DstMesh.GetVertexBuffers(), MBS_POSITION);
+		const UntypedMeshBufferIteratorConst SrcPositionBegin(SrcMesh.GetVertexBuffers(), EMeshBufferSemantic::Position);
+		const UntypedMeshBufferIterator      DstPositionBegin(DstMesh.GetVertexBuffers(), EMeshBufferSemantic::Position);
 
 		const int32 NumVertices = DstMesh.GetVertexCount();
  
diff --git a/Engine/Plugins/Mutable/Source/MutableRuntime/Private/MuR/OpMeshTransform.h b/Engine/Plugins/Mutable/Source/MutableRuntime/Private/MuR/OpMeshTransform.h
--- a/Engine/Plugins/Mutable/Source/MutableRuntime/Private/MuR/OpMeshTransform.h
+++ b/Engine/Plugins/Mutable/Source/MutableRuntime/Private/MuR/OpMeshTransform.h
@@ -43,42 +43,42 @@ namespace mu
 
                 switch ( sem )
                 {
-                case MBS_POSITION:
+                case EMeshBufferSemantic::Position:
                     for ( uint32_t v=0; v<vcount; ++v )
                     {
                         FVector4f value( 0.0f, 0.0f, 0.0f, 1.0f );
                         for( int i=0; i<it.GetComponents(); ++i )
                         {
-                            ConvertData( i, &value[0], MBF_FLOAT32, it.ptr(), it.GetFormat() );
+                            ConvertData( i, &value[0], EMeshBufferFormat::Float32, it.ptr(), it.GetFormat() );
                         }
 
                         value = transform.TransformFVector4( value );
 
                         for( int i=0; i<it.GetComponents(); ++i )
                         {
-                            ConvertData( i, it.ptr(), it.GetFormat(), &value[0], MBF_FLOAT32 );
+                            ConvertData( i, it.ptr(), it.GetFormat(), &value[0], EMeshBufferFormat::Float32 );
                         }
 
                         ++it;
                     }
                     break;
 
-                case MBS_NORMAL:
-                case MBS_TANGENT:
-                case MBS_BINORMAL:
+                case EMeshBufferSemantic::Normal:
+                case EMeshBufferSemantic::Tangent:
+                case EMeshBufferSemantic::Binormal:
                     for ( uint32_t v=0; v<vcount; ++v )
                     {
 						FVector4f value( 0.0f, 0.0f, 0.0f, 1.0f );
                         for( int i=0; i<it.GetComponents(); ++i )
                         {
-                            ConvertData( i, &value[0], MBF_FLOAT32, it.ptr(), it.GetFormat() );
+                            ConvertData( i, &value[0], EMeshBufferFormat::Float32, it.ptr(), it.GetFormat() );
                         }
 
                         value = transformIT.TransformFVector4(value);
 
                         for( int i=0; i<it.GetComponents(); ++i )
                         {
-                            ConvertData( i, it.ptr(), it.GetFormat(), &value[0], MBF_FLOAT32 );
+                            ConvertData( i, it.ptr(), it.GetFormat(), &value[0], EMeshBufferFormat::Float32 );
                         }
 
                         ++it;
diff --git a/Engine/Plugins/Mutable/Source/MutableRuntime/Private/MuR/OpMeshTransformWithMesh.cpp b/Engine/Plugins/Mutable/Source/MutableRuntime/Private/MuR/OpMeshTransformWithMesh.cpp
--- a/Engine/Plugins/Mutable/Source/MutableRuntime/Private/MuR/OpMeshTransformWithMesh.cpp
+++ b/Engine/Plugins/Mutable/Source/MutableRuntime/Private/MuR/OpMeshTransformWithMesh.cpp
@@ -46,7 +46,7 @@ void MeshTransformWithMesh(Mesh* Result, const Mesh* SourceMesh, const Mesh* Bou
 
             switch ( sem )
             {
-            case MBS_POSITION:
+            case EMeshBufferSemantic::Position:
                 for ( uint32_t v=0; v<VCount; ++v )
                 {
                 	if (!BoundingMesh || VertexInBoundaryMesh[v])
@@ -54,23 +54,23 @@ void MeshTransformWithMesh(Mesh* Result, const Mesh* SourceMesh, const Mesh* Bou
                 		FVector4f value( 0.0f, 0.0f, 0.0f, 1.0f );
                 		for( int i=0; i<it.GetComponents(); ++i )
                 		{
-                			ConvertData( i, &value[0], MBF_FLOAT32, it.ptr(), it.GetFormat() );
+                			ConvertData( i, &value[0], EMeshBufferFormat::Float32, it.ptr(), it.GetFormat() );
                 		}
 
                 		value = Transform.TransformFVector4( value );
 
                 		for( int i=0; i<it.GetComponents(); ++i )
                 		{
-                			ConvertData( i, it.ptr(), it.GetFormat(), &value[0], MBF_FLOAT32 );
+                			ConvertData( i, it.ptr(), it.GetFormat(), &value[0], EMeshBufferFormat::Float32 );
                 		}
                 	}
                 	++it;
                 }
                 break;
 
-            case MBS_NORMAL:
-            case MBS_TANGENT:
-            case MBS_BINORMAL:
+            case EMeshBufferSemantic::Normal:
+            case EMeshBufferSemantic::Tangent:
+            case EMeshBufferSemantic::Binormal:
                 for ( uint32_t v=0; v<VCount; ++v )
                 {
                 	if (!BoundingMesh || VertexInBoundaryMesh[v])
@@ -78,14 +78,14 @@ void MeshTransformWithMesh(Mesh* Result, const Mesh* SourceMesh, const Mesh* Bou
                 		FVector4f value( 0.0f, 0.0f, 0.0f, 1.0f );
 	                    for( int i=0; i<it.GetComponents(); ++i )
 	                    {
-	                        ConvertData( i, &value[0], MBF_FLOAT32, it.ptr(), it.GetFormat() );
+	                        ConvertData( i, &value[0], EMeshBufferFormat::Float32, it.ptr(), it.GetFormat() );
 	                    }
 
 	                    value = TransformInvT.TransformFVector4(value);
 
 	                    for( int i=0; i<it.GetComponents(); ++i )
 	                    {
-	                        ConvertData( i, it.ptr(), it.GetFormat(), &value[0], MBF_FLOAT32 );
+	                        ConvertData( i, it.ptr(), it.GetFormat(), &value[0], EMeshBufferFormat::Float32 );
 	                    }
                     }
 					++it;
