[mutable] Cleanup MutableRuntime API moving to use more unreal types and naming standards.

diff --git a/Engine/Plugins/Mutable/Source/MutableRuntime/Private/MuR/Parameters.cpp b/Engine/Plugins/Mutable/Source/MutableRuntime/Private/MuR/Parameters.cpp
--- a/Engine/Plugins/Mutable/Source/MutableRuntime/Private/MuR/Parameters.cpp
+++ b/Engine/Plugins/Mutable/Source/MutableRuntime/Private/MuR/Parameters.cpp
@@ -17,8 +17,8 @@
 
 namespace mu
 {
-    MUTABLE_IMPLEMENT_ENUM_SERIALISABLE(PARAMETER_TYPE)              
-    MUTABLE_IMPLEMENT_ENUM_SERIALISABLE(PROJECTOR_TYPE)
+    MUTABLE_IMPLEMENT_ENUM_SERIALISABLE(EParameterType)              
+    MUTABLE_IMPLEMENT_ENUM_SERIALISABLE(EProjectorType)
 
 	
 	//---------------------------------------------------------------------------------------------
@@ -55,111 +55,6 @@ namespace mu
 	}
 
 
-    //---------------------------------------------------------------------------------------------
-    //void Parameters::SerialisePortable( const Parameters* params, OutputArchive& arch )
-    //{
-    //    const int32_t ver = 0;
-    //    arch << ver;
-
-    //    uint32 valueCount = params->GetPrivate()->m_values.Num();
-    //    arch << valueCount;
-
-    //    for (size_t p=0; p<valueCount; ++p)
-    //    {
-    //        const FParameterDesc& desc = params->GetPrivate()->m_pModel->GetPrivate()->m_program.m_parameters[p];
-    //        arch << desc.m_name;
-    //        arch << desc.m_uid;
-    //        arch << desc.m_type;
-    //        arch << params->GetPrivate()->m_values[p];
-    //    }
-    //}
-
-
-    //---------------------------------------------------------------------------------------------
-  //  ParametersPtr Parameters::UnserialisePortable( InputArchive& arch, TSharedPtr<const Model> pModel )
-  //  {
-		//LLM_SCOPE_BYNAME(TEXT("MutableRuntime"));
-
-  //      ParametersPtr pResult = Model::NewParameters(pModel);
-  //      size_t modelParameters = pModel->GetPrivate()->m_program.m_parameters.Num();
-
-  //      int32_t ver;
-  //      arch >> ver;
-  //      check( ver == 0 );
-
-  //      uint32_t valueCount = 0;
-  //      arch >> valueCount;
-
-  //      for (uint32_t p=0; p<valueCount; ++p)
-  //      {
-  //          // Read the parameter data.
-  //          std::string name;
-  //          arch >> name;
-
-  //          std::string uid;
-  //          arch >> uid;
-
-  //          PARAMETER_TYPE type;
-  //          arch >> type;
-
-  //          PARAMETER_VALUE value;
-  //          arch >> value;
-
-  //          int modelParam = -1;
-
-  //          // UIDs have higher priority
-  //          if ( !uid.empty() )
-  //          {
-  //              for (size_t mp=0; mp<modelParameters; ++mp)
-  //              {
-  //                  const FParameterDesc& desc = pModel->GetPrivate()->m_program.m_parameters[mp];
-  //                  if (desc.m_uid==uid)
-  //                  {
-  //                      modelParam = int(mp);
-  //                      break;
-  //                  }
-  //              }
-  //          }
-
-  //          // Try with name+type
-  //          if (modelParam<0)
-  //          {
-  //              for (size_t mp=0; mp<modelParameters; ++mp)
-  //              {
-  //                  const FParameterDesc& desc = pModel->GetPrivate()->m_program.m_parameters[mp];
-  //                  if (desc.m_name==name && desc.m_type==type)
-  //                  {
-  //                      modelParam = int(mp);
-  //                      break;
-  //                  }
-  //              }
-  //          }
-
-  //          // Try with name only, and maybe adapt type
-  //          if (modelParam<0)
-  //          {
-  //              for (size_t mp=0; mp<modelParameters; ++mp)
-  //              {
-  //                  const FParameterDesc& desc = pModel->GetPrivate()->m_program.m_parameters[mp];
-  //                  if (desc.m_name==name)
-  //                  {
-  //                      modelParam = int(mp);
-  //                      break;
-  //                  }
-  //              }
-  //          }
-
-  //          if (modelParam>=0)
-  //          {
-  //              // We found something. \todo: convert the type if necessary?
-  //              pResult->GetPrivate()->m_values[modelParam] = value;
-  //          }
-  //      }
-
-  //      return pResult;
-  //  }
-
-
 	//---------------------------------------------------------------------------------------------
 	Parameters::Private* Parameters::GetPrivate() const
 	{
@@ -183,59 +78,59 @@ namespace mu
 
 
 	//---------------------------------------------------------------------------------------------
-	int Parameters::GetCount() const
+	int32 Parameters::GetCount() const
 	{
-		return (int)m_pD->m_values.Num();
+		return (int32)m_pD->m_values.Num();
 	}
 
 
 	//---------------------------------------------------------------------------------------------
-	const FString& Parameters::GetName( int index ) const
+	const FString& Parameters::GetName( int32 index ) const
 	{
 		const FProgram& program = m_pD->m_pModel->GetPrivate()->m_program;
-		check( index>=0 && index<(int)program.m_parameters.Num() );
+		check( index>=0 && index<(int32)program.m_parameters.Num() );
 
 		return program.m_parameters[index].m_name;
 	}
 
 
 	//---------------------------------------------------------------------------------------------
-	const FGuid& Parameters::GetUid( int index ) const
+	const FGuid& Parameters::GetUid( int32 index ) const
 	{
 		const FProgram& program = m_pD->m_pModel->GetPrivate()->m_program;
-		check( index>=0 && index<(int)program.m_parameters.Num() );
+		check( index>=0 && index<(int32)program.m_parameters.Num() );
 
 		return program.m_parameters[index].m_uid;
 	}
 
 
 	//---------------------------------------------------------------------------------------------
-	int Parameters::Find( const FString& strName ) const
+	int32 Parameters::Find( const FString& strName ) const
 	{
 		return m_pD->Find( strName );
 	}
 
 
 	//---------------------------------------------------------------------------------------------
-	PARAMETER_TYPE Parameters::GetType( int index ) const
+	EParameterType Parameters::GetType( int32 index ) const
 	{
 		const FProgram& program = m_pD->m_pModel->GetPrivate()->m_program;
-		check( index>=0 && index<(int)program.m_parameters.Num() );
+		check( index>=0 && index<(int32)program.m_parameters.Num() );
 
 		return program.m_parameters[index].m_type;
 	}
 
 
     //---------------------------------------------------------------------------------------------
-    Ptr<RangeIndex> Parameters::NewRangeIndex( int paramIndex ) const
+    Ptr<RangeIndex> Parameters::NewRangeIndex( int32 paramIndex ) const
     {
 		LLM_SCOPE_BYNAME(TEXT("MutableRuntime"));
 
         const FProgram& program = m_pD->m_pModel->GetPrivate()->m_program;
-        check( paramIndex>=0 && paramIndex<int(program.m_parameters.Num()) );
+        check( paramIndex>=0 && paramIndex<int32(program.m_parameters.Num()) );
         Ptr<RangeIndex> range;
 
-        if ( paramIndex>=0 && paramIndex<int(program.m_parameters.Num()) )
+        if ( paramIndex>=0 && paramIndex<int32(program.m_parameters.Num()) )
         {
             size_t rangeCount = program.m_parameters[paramIndex].m_ranges.Num();
             if ( rangeCount>0 )
@@ -252,32 +147,32 @@ namespace mu
 
 
     //---------------------------------------------------------------------------------------------
-    int Parameters::GetValueCount( int paramIndex ) const
+    int32 Parameters::GetValueCount( int32 paramIndex ) const
     {
-        if (paramIndex<0 || paramIndex >= int(m_pD->m_multiValues.Num()) )
+        if (paramIndex<0 || paramIndex >= int32(m_pD->m_multiValues.Num()) )
         {
             return 0;
         }
 
-        return int( m_pD->m_multiValues[paramIndex].Num() );
+        return int32( m_pD->m_multiValues[paramIndex].Num() );
     }
 
 
     //---------------------------------------------------------------------------------------------
-    Ptr<RangeIndex> Parameters::GetValueIndex( int paramIndex, int valueIndex ) const
+    Ptr<RangeIndex> Parameters::GetValueIndex( int32 paramIndex, int32 valueIndex ) const
     {
-        if (paramIndex<0 || paramIndex >= int(m_pD->m_multiValues.Num()) )
+        if (paramIndex<0 || paramIndex >= int32(m_pD->m_multiValues.Num()) )
         {
             return nullptr;
         }
 
-        if (valueIndex<0 || valueIndex >= int(m_pD->m_multiValues[paramIndex].Num()) )
+        if (valueIndex<0 || valueIndex >= int32(m_pD->m_multiValues[paramIndex].Num()) )
         {
             return nullptr;
         }
 
-		TMap< TArray<int32_t>, PARAMETER_VALUE >::TRangedForIterator it = m_pD->m_multiValues[paramIndex].begin();
-        for ( int i=0; i<valueIndex; ++i )
+		TMap< TArray<int32>, FParameterValue >::TRangedForIterator it = m_pD->m_multiValues[paramIndex].begin();
+        for ( int32 i=0; i<valueIndex; ++i )
         {
             ++it;
         }
@@ -290,9 +185,9 @@ namespace mu
 
 
     //---------------------------------------------------------------------------------------------
-    void Parameters::ClearAllValues( int paramIndex )
+    void Parameters::ClearAllValues( int32 paramIndex )
     {
-        if (paramIndex<0 || paramIndex >= int(m_pD->m_multiValues.Num()) )
+        if (paramIndex<0 || paramIndex >= int32(m_pD->m_multiValues.Num()) )
         {
             return;
         }
@@ -302,18 +197,18 @@ namespace mu
 
 
 	//---------------------------------------------------------------------------------------------
-    bool Parameters::GetBoolValue( int index,
+    bool Parameters::GetBoolValue( int32 index,
                                    const Ptr<const RangeIndex>& pos ) const
 	{
-		check( index>=0 && index<(int)m_pD->m_values.Num() );
-        check( GetType(index)==PARAMETER_TYPE::T_BOOL );
+		check( index>=0 && index<(int32)m_pD->m_values.Num() );
+        check( GetType(index)==EParameterType::Bool );
 
         // Early out in case of invalid parameters
         if ( index < 0
              ||
-             index >= (int)m_pD->m_values.Num()
+             index >= (int32)m_pD->m_values.Num()
              ||
-             GetType(index) != PARAMETER_TYPE::T_BOOL )
+             GetType(index) != EParameterType::Bool )
         {
             return false;
         }
@@ -322,40 +217,40 @@ namespace mu
         if (!pos)
         {
             // Return the single value
-            return m_pD->m_values[index].Get<ParamBoolType>();
+            return m_pD->m_values[index].Get<FParamBoolType>();
         }
 
         // Multivalue case
         check( pos->m_pD->m_parameter==index );
 
-        if ( index<int(m_pD->m_multiValues.Num()))
+        if ( index<int32(m_pD->m_multiValues.Num()))
         {
-            const TMap< TArray<int32_t>, PARAMETER_VALUE >& m = m_pD->m_multiValues[index];
-			const PARAMETER_VALUE* it = m.Find(pos->m_pD->m_values);
+            const TMap< TArray<int32>, FParameterValue >& m = m_pD->m_multiValues[index];
+			const FParameterValue* it = m.Find(pos->m_pD->m_values);
             if (it)
             {
-                return it->Get<ParamBoolType>();
+                return it->Get<FParamBoolType>();
             }
         }
 
         // Multivalue parameter, but no multivalue set. Return single value.
-        return m_pD->m_values[index].Get<ParamBoolType>();
+        return m_pD->m_values[index].Get<FParamBoolType>();
 	}
 
 
 	//---------------------------------------------------------------------------------------------
-    void Parameters::SetBoolValue( int index, bool value,
+    void Parameters::SetBoolValue( int32 index, bool value,
                                    const Ptr<const RangeIndex>& pos )
 	{
-		check( index>=0 && index<(int)m_pD->m_values.Num() );
-		check( GetType(index)== PARAMETER_TYPE::T_BOOL );
+		check( index>=0 && index<(int32)m_pD->m_values.Num() );
+		check( GetType(index)== EParameterType::Bool );
 
         // Early out in case of invalid parameters
         if ( index < 0
              ||
-             index >= (int)m_pD->m_values.Num()
+             index >= (int32)m_pD->m_values.Num()
              ||
-             GetType(index) != PARAMETER_TYPE::T_BOOL )
+             GetType(index) != EParameterType::Bool )
         {
             return;
         }
@@ -364,12 +259,12 @@ namespace mu
         if (!pos)
         {
             // Clear multivalue, if set.
-            if (index<int(m_pD->m_multiValues.Num()))
+            if (index<int32(m_pD->m_multiValues.Num()))
             {
                 m_pD->m_multiValues[index].Empty();
             }
 
-            m_pD->m_values[index].Set<ParamBoolType>(value);
+            m_pD->m_values[index].Set<FParamBoolType>(value);
         }
 
         // Multivalue case
@@ -377,53 +272,53 @@ namespace mu
         {
             check( pos->m_pD->m_parameter==index );
 
-            if ( index>=int(m_pD->m_multiValues.Num()))
+            if ( index>=int32(m_pD->m_multiValues.Num()))
             {
                 m_pD->m_multiValues.SetNum(index+1);
             }
 
-			TMap< TArray<int32_t>, PARAMETER_VALUE >& m = m_pD->m_multiValues[index];
-			PARAMETER_VALUE& it = m.FindOrAdd(pos->m_pD->m_values);
-            it.Set<ParamBoolType>(value);
+			TMap< TArray<int32>, FParameterValue >& m = m_pD->m_multiValues[index];
+			FParameterValue& it = m.FindOrAdd(pos->m_pD->m_values);
+            it.Set<FParamBoolType>(value);
         }
 	}
 
 
 	//---------------------------------------------------------------------------------------------
-	int Parameters::GetIntPossibleValueCount( int paramIndex ) const
+	int32 Parameters::GetIntPossibleValueCount( int32 paramIndex ) const
 	{
 		const FProgram& program = m_pD->m_pModel->GetPrivate()->m_program;
-		check( paramIndex>=0 && paramIndex<(int)program.m_parameters.Num() );
-		return (int)program.m_parameters[paramIndex].m_possibleValues.Num();
+		check( paramIndex>=0 && paramIndex<(int32)program.m_parameters.Num() );
+		return (int32)program.m_parameters[paramIndex].m_possibleValues.Num();
 	}
 
 
 	//---------------------------------------------------------------------------------------------
-	int Parameters::GetIntPossibleValue( int paramIndex, int valueIndex ) const
+	int32 Parameters::GetIntPossibleValue( int32 paramIndex, int32 valueIndex ) const
 	{
 		const FProgram& program = m_pD->m_pModel->GetPrivate()->m_program;
 		check( paramIndex>=0
-				&& paramIndex<(int)program.m_parameters.Num() );
+				&& paramIndex<(int32)program.m_parameters.Num() );
 		check( valueIndex>=0
-				&& valueIndex<(int)program.m_parameters[paramIndex].m_possibleValues.Num() );
+				&& valueIndex<(int32)program.m_parameters[paramIndex].m_possibleValues.Num() );
 
-		return (int)program.m_parameters[paramIndex].m_possibleValues[valueIndex].m_value;
+		return (int32)program.m_parameters[paramIndex].m_possibleValues[valueIndex].m_value;
 	}
 
 
 	//---------------------------------------------------------------------------------------------
-	int Parameters::GetIntValueIndex(int paramIndex, const FString& ValueName) const
+	int32 Parameters::GetIntValueIndex(int32 paramIndex, const FString& ValueName) const
 	{
 		const FProgram& program = m_pD->m_pModel->GetPrivate()->m_program;
 		check(paramIndex >= 0
-			&& paramIndex < (int)program.m_parameters.Num());
+			&& paramIndex < (int32)program.m_parameters.Num());
 
-		int result = -1;
+		int32 result = -1;
 		for (size_t v = 0; v < program.m_parameters[paramIndex].m_possibleValues.Num(); ++v)
 		{
 			if (program.m_parameters[paramIndex].m_possibleValues[v].m_name == ValueName)
 			{
-				result = (int)v;
+				result = (int32)v;
 				break;
 			}
 		}
@@ -432,17 +327,17 @@ namespace mu
 
 
 	//---------------------------------------------------------------------------------------------
-	int Parameters::GetIntValueIndex(int paramIndex, int32 Value) const
+	int32 Parameters::GetIntValueIndex(int32 paramIndex, int32 Value) const
 	{
 		const FProgram& program = m_pD->m_pModel->GetPrivate()->m_program;
 		check(paramIndex >= 0
-			&& paramIndex < (int)program.m_parameters.Num());
+			&& paramIndex < (int32)program.m_parameters.Num());
 
 		for (size_t v = 0; v < program.m_parameters[paramIndex].m_possibleValues.Num(); ++v)
 		{
 			if (program.m_parameters[paramIndex].m_possibleValues[v].m_value == Value)
 			{
-				return (int)v;
+				return (int32)v;
 			}
 		}
 		return -1;
@@ -450,30 +345,30 @@ namespace mu
 
 
 	//---------------------------------------------------------------------------------------------
-	const FString& Parameters::GetIntPossibleValueName( int paramIndex, int valueIndex ) const
+	const FString& Parameters::GetIntPossibleValueName( int32 paramIndex, int32 valueIndex ) const
 	{
 		const FProgram& program = m_pD->m_pModel->GetPrivate()->m_program;
 		check( paramIndex>=0
-				&& paramIndex<(int)program.m_parameters.Num() );
+				&& paramIndex<(int32)program.m_parameters.Num() );
 		check( valueIndex>=0
-				&& valueIndex<(int)program.m_parameters[paramIndex].m_possibleValues.Num() );
+				&& valueIndex<(int32)program.m_parameters[paramIndex].m_possibleValues.Num() );
 
 		return program.m_parameters[paramIndex].m_possibleValues[valueIndex].m_name;
 	}
 
 
 	//---------------------------------------------------------------------------------------------
-    int Parameters::GetIntValue( int index, const Ptr<const RangeIndex>& pos ) const
+    int32 Parameters::GetIntValue( int32 index, const Ptr<const RangeIndex>& pos ) const
 	{
-		check( index>=0 && index<(int)m_pD->m_values.Num() );
-		check( GetType(index)== PARAMETER_TYPE::T_INT );
+		check( index>=0 && index<(int32)m_pD->m_values.Num() );
+		check( GetType(index)== EParameterType::Int );
 
         // Early out in case of invalid parameters
         if ( index < 0
              ||
-             index >= (int)m_pD->m_values.Num()
+             index >= (int32)m_pD->m_values.Num()
              ||
-             GetType(index) != PARAMETER_TYPE::T_INT )
+             GetType(index) != EParameterType::Int )
         {
             return 0;
         }
@@ -482,40 +377,40 @@ namespace mu
         if (!pos)
         {
             // Return the single value
-            return m_pD->m_values[index].Get<ParamIntType>();
+            return m_pD->m_values[index].Get<FParamIntType>();
         }
 
         // Multivalue case
         check( pos->m_pD->m_parameter==index );
 
-        if ( index<int(m_pD->m_multiValues.Num()))
+        if ( index<int32(m_pD->m_multiValues.Num()))
         {
-            const TMap< TArray<int32_t>, PARAMETER_VALUE >& m = m_pD->m_multiValues[index];
-			const PARAMETER_VALUE* it = m.Find(pos->m_pD->m_values);
+            const TMap< TArray<int32>, FParameterValue >& m = m_pD->m_multiValues[index];
+			const FParameterValue* it = m.Find(pos->m_pD->m_values);
             if (it)
             {
-                return it->Get<ParamIntType>();
+                return it->Get<FParamIntType>();
             }
         }
 
         // Multivalue parameter, but no multivalue set. Return single value.
-        return m_pD->m_values[index].Get<ParamIntType>();
+        return m_pD->m_values[index].Get<FParamIntType>();
     }
 
 
 	//---------------------------------------------------------------------------------------------
-    void Parameters::SetIntValue( int index, int value,
+    void Parameters::SetIntValue( int32 index, int32 value,
                                   const Ptr<const RangeIndex>& pos )
 	{
-		check( index>=0 && index<(int)m_pD->m_values.Num() );
-		check( GetType(index)== PARAMETER_TYPE::T_INT );
+		check( index>=0 && index<(int32)m_pD->m_values.Num() );
+		check( GetType(index)== EParameterType::Int );
 
         // Early out in case of invalid parameters
         if ( index < 0
              ||
-             index >= (int)m_pD->m_values.Num()
+             index >= (int32)m_pD->m_values.Num()
              ||
-             GetType(index) != PARAMETER_TYPE::T_INT )
+             GetType(index) != EParameterType::Int )
         {
             return;
         }
@@ -524,12 +419,12 @@ namespace mu
         if (!pos)
         {
             // Clear multivalue, if set.
-            if (index<int(m_pD->m_multiValues.Num()))
+            if (index<int32(m_pD->m_multiValues.Num()))
             {
                 m_pD->m_multiValues[index].Empty();
             }
 
-            m_pD->m_values[index].Set<ParamIntType>(value);
+            m_pD->m_values[index].Set<FParamIntType>(value);
         }
 
         // Multivalue case
@@ -537,31 +432,31 @@ namespace mu
         {
             check( pos->m_pD->m_parameter==index );
 
-            if ( index>=int(m_pD->m_multiValues.Num()))
+            if ( index>=int32(m_pD->m_multiValues.Num()))
             {
                 m_pD->m_multiValues.SetNum(index+1);
             }
 
-			TMap< TArray<int32_t>, PARAMETER_VALUE >& m = m_pD->m_multiValues[index];
-			PARAMETER_VALUE& it = m.FindOrAdd(pos->m_pD->m_values);
-            it.Set<ParamIntType>(value);
+			TMap< TArray<int32>, FParameterValue >& m = m_pD->m_multiValues[index];
+			FParameterValue& it = m.FindOrAdd(pos->m_pD->m_values);
+            it.Set<FParamIntType>(value);
         }
     }
 
 
 	//---------------------------------------------------------------------------------------------
-    float Parameters::GetFloatValue( int index,
+    float Parameters::GetFloatValue( int32 index,
                                      const Ptr<const RangeIndex>& pos ) const
 	{
-		check( index>=0 && index<(int)m_pD->m_values.Num() );
-		check( GetType(index)== PARAMETER_TYPE::T_FLOAT );
+		check( index>=0 && index<(int32)m_pD->m_values.Num() );
+		check( GetType(index)== EParameterType::Float );
 
         // Early out in case of invalid parameters
         if ( index < 0
              ||
-             index >= (int)m_pD->m_values.Num()
+             index >= (int32)m_pD->m_values.Num()
              ||
-             GetType(index) != PARAMETER_TYPE::T_FLOAT )
+             GetType(index) != EParameterType::Float )
         {
             return 0.0f;
         }
@@ -570,40 +465,40 @@ namespace mu
         if (!pos)
         {
             // Return the single value
-            return m_pD->m_values[index].Get<ParamFloatType>();
+            return m_pD->m_values[index].Get<FParamFloatType>();
         }
 
         // Multivalue case
         check( pos->m_pD->m_parameter==index );
 
-        if ( index<int(m_pD->m_multiValues.Num()))
+        if ( index<int32(m_pD->m_multiValues.Num()))
         {
-            const TMap< TArray<int32_t>, PARAMETER_VALUE >& m = m_pD->m_multiValues[index];
-			const PARAMETER_VALUE* it = m.Find(pos->m_pD->m_values);
+            const TMap< TArray<int32>, FParameterValue >& m = m_pD->m_multiValues[index];
+			const FParameterValue* it = m.Find(pos->m_pD->m_values);
             if (it)
             {
-                return it->Get<ParamFloatType>();
+                return it->Get<FParamFloatType>();
             }
         }
 
         // Multivalue parameter, but no multivalue set. Return single value.
-        return m_pD->m_values[index].Get<ParamFloatType>();
+        return m_pD->m_values[index].Get<FParamFloatType>();
     }
 
 
 	//---------------------------------------------------------------------------------------------
-    void Parameters::SetFloatValue( int index, float value,
+    void Parameters::SetFloatValue( int32 index, float value,
                                     const Ptr<const RangeIndex>& pos )
 	{
-		check( index>=0 && index<(int)m_pD->m_values.Num() );
-		check( GetType(index)== PARAMETER_TYPE::T_FLOAT );
+		check( index>=0 && index<(int32)m_pD->m_values.Num() );
+		check( GetType(index)== EParameterType::Float );
 
         // Early out in case of invalid parameters
         if ( index < 0
              ||
-             index >= (int)m_pD->m_values.Num()
+             index >= (int32)m_pD->m_values.Num()
              ||
-             GetType(index) != PARAMETER_TYPE::T_FLOAT )
+             GetType(index) != EParameterType::Float )
         {
             return;
         }
@@ -612,12 +507,12 @@ namespace mu
         if (!pos)
         {
             // Clear multivalue, if set.
-            if (index<int(m_pD->m_multiValues.Num()))
+            if (index<int32(m_pD->m_multiValues.Num()))
             {
                 m_pD->m_multiValues[index].Empty();
             }
 
-            m_pD->m_values[index].Set<ParamFloatType>(value);
+            m_pD->m_values[index].Set<FParamFloatType>(value);
         }
 
         // Multivalue case
@@ -625,31 +520,30 @@ namespace mu
         {
             check( pos->m_pD->m_parameter==index );
 
-            if ( index>=int(m_pD->m_multiValues.Num()))
+            if ( index>=int32(m_pD->m_multiValues.Num()))
             {
                 m_pD->m_multiValues.SetNum(index+1);
             }
 
-			TMap< TArray<int32_t>, PARAMETER_VALUE >& m = m_pD->m_multiValues[index];
-			PARAMETER_VALUE& it = m.FindOrAdd(pos->m_pD->m_values);
-            it.Set<ParamFloatType>(value);
+			TMap< TArray<int32>, FParameterValue >& m = m_pD->m_multiValues[index];
+			FParameterValue& it = m.FindOrAdd(pos->m_pD->m_values);
+            it.Set<FParamFloatType>(value);
         }
     }
 
 
     //---------------------------------------------------------------------------------------------
-    void Parameters::GetColourValue( int index, float* pR, float* pG, float* pB, float* pA,
-                                     const Ptr<const RangeIndex>& pos ) const
+    void Parameters::GetColourValue( int32 index, FVector4f& OutValue, const Ptr<const RangeIndex>& pos ) const
     {
-        check( index>=0 && index<(int)m_pD->m_values.Num() );
-        check( GetType(index)== PARAMETER_TYPE::T_COLOUR );
+        check( index>=0 && index<(int32)m_pD->m_values.Num() );
+        check( GetType(index)== EParameterType::Color );
 
         // Early out in case of invalid parameters
         if ( index < 0
              ||
-             index >= (int)m_pD->m_values.Num()
+             index >= (int32)m_pD->m_values.Num()
              ||
-             GetType(index) != PARAMETER_TYPE::T_COLOUR )
+             GetType(index) != EParameterType::Color )
         {
             return;
         }
@@ -657,101 +551,42 @@ namespace mu
         // Single value case
         if (!pos)
         {
-            // Return the single value
-			if (pR)
-			{
-				*pR = m_pD->m_values[index].Get<ParamColorType>()[0];
-			}
-
-			if (pG)
-			{
-				*pG = m_pD->m_values[index].Get<ParamColorType>()[1];
-			}
-
-			if (pB)
-			{
-				*pB = m_pD->m_values[index].Get<ParamColorType>()[2];
-			}
-
-			if (pA)
-			{
-				*pA = m_pD->m_values[index].Get<ParamColorType>()[3];
-			}
-
+			OutValue = m_pD->m_values[index].Get<FParamColorType>();
             return;
         }
 
         // Multivalue case
         check( pos->m_pD->m_parameter==index );
 
-        if ( index<int(m_pD->m_multiValues.Num()))
+        if ( index<int32(m_pD->m_multiValues.Num()))
         {
-            const TMap< TArray<int32_t>, PARAMETER_VALUE >& m = m_pD->m_multiValues[index];
-			const PARAMETER_VALUE* it = m.Find(pos->m_pD->m_values);
+            const TMap< TArray<int32>, FParameterValue >& m = m_pD->m_multiValues[index];
+			const FParameterValue* it = m.Find(pos->m_pD->m_values);
             if (it)
             {
-				if (pR)
-				{
-					*pR = it->Get<ParamColorType>()[0];
-				}
-
-				if (pG)
-				{
-					*pG = it->Get<ParamColorType>()[1];
-				}
-
-				if (pB)
-				{
-					*pB = it->Get<ParamColorType>()[2];
-				}
-
-				if (pA)
-				{
-					*pA = it->Get<ParamColorType>()[3];
-				}
-
+				OutValue = it->Get<FParamColorType>();
                 return;
             }
         }
 
         // Multivalue parameter, but no multivalue set. Return single value.
-		if (pR)
-		{
-			*pR = m_pD->m_values[index].Get<ParamColorType>()[0];
-		}
-
-		if (pG)
-		{
-			*pG = m_pD->m_values[index].Get<ParamColorType>()[1];
-		}
-
-		if (pB)
-		{
-			*pB = m_pD->m_values[index].Get<ParamColorType>()[2];
-		}
-
-		if (pA)
-		{
-			*pA = m_pD->m_values[index].Get<ParamColorType>()[3];
-		}
-
+		OutValue = m_pD->m_values[index].Get<FParamColorType>();
         return;
     }
 
 
     //---------------------------------------------------------------------------------------------
-    void Parameters::SetColourValue( int index, float R, float G, float B, float A,
-                                     const Ptr<const RangeIndex>& pos )
+    void Parameters::SetColourValue( int32 index, FVector4f InValue, const Ptr<const RangeIndex>& pos )
     {
-        check( index>=0 && index<(int)m_pD->m_values.Num() );
-        check( GetType(index)== PARAMETER_TYPE::T_COLOUR );
+        check( index>=0 && index<(int32)m_pD->m_values.Num() );
+        check( GetType(index)== EParameterType::Color );
 
         // Early out in case of invalid parameters
         if ( index < 0
              ||
-             index >= (int)m_pD->m_values.Num()
+             index >= (int32)m_pD->m_values.Num()
              ||
-             GetType(index) != PARAMETER_TYPE::T_COLOUR )
+             GetType(index) != EParameterType::Color )
         {
             return;
         }
@@ -760,18 +595,12 @@ namespace mu
         if (!pos)
         {
             // Clear multivalue, if set.
-            if (index<int(m_pD->m_multiValues.Num()))
+            if (index<int32(m_pD->m_multiValues.Num()))
             {
                 m_pD->m_multiValues[index].Empty();
             }
 
-        	ParamColorType Value;
-            Value[0] = R;
-            Value[1] = G;
-            Value[2] = B;
-			Value[3] = A;
-        	
-        	m_pD->m_values[index].Set<ParamColorType>(Value);
+        	m_pD->m_values[index].Set<FParamColorType>(InValue);
         }
 
         // Multivalue case
@@ -779,37 +608,31 @@ namespace mu
         {
             check( pos->m_pD->m_parameter==index );
 
-            if ( index>=int(m_pD->m_multiValues.Num()))
+            if ( index>=int32(m_pD->m_multiValues.Num()))
             {
                 m_pD->m_multiValues.SetNum(index+1);
             }
 
-			TMap< TArray<int32_t>, PARAMETER_VALUE >& m = m_pD->m_multiValues[index];
-			PARAMETER_VALUE& it = m.FindOrAdd(pos->m_pD->m_values);
-
-        	ParamColorType Value;
-            Value[0] = R;
-            Value[1] = G;
-			Value[2] = B;
-			Value[3] = A;
+			TMap< TArray<int32>, FParameterValue >& m = m_pD->m_multiValues[index];
+			FParameterValue& it = m.FindOrAdd(pos->m_pD->m_values);
 
-        	it.Set<ParamColorType>(Value);
+        	it.Set<FParamColorType>(InValue);
         }
     }
 
 
     //---------------------------------------------------------------------------------------------
-	FName Parameters::GetImageValue( int index, const Ptr<const RangeIndex>& pos ) const
+	FName Parameters::GetImageValue( int32 index, const Ptr<const RangeIndex>& pos ) const
     {
-        check( index >= 0 && index < (int)m_pD->m_values.Num() );
-        check( GetType( index ) == PARAMETER_TYPE::T_IMAGE );
+        check( index >= 0 && index < (int32)m_pD->m_values.Num() );
+        check( GetType( index ) == EParameterType::Image );
 
 		// Early out in case of invalid parameters
         if ( index < 0
              ||
-             index >= (int)m_pD->m_values.Num()
+             index >= (int32)m_pD->m_values.Num()
              ||
-             GetType(index) != PARAMETER_TYPE::T_IMAGE )
+             GetType(index) != EParameterType::Image )
         {
             return {};
         }
@@ -817,38 +640,38 @@ namespace mu
 		// Single value case
 		if (!pos)
 		{
-			return m_pD->m_values[index].Get<ParamImageType>();
+			return m_pD->m_values[index].Get<FParamImageType>();
 		}
 
 		// Multivalue case
 		check(pos->m_pD->m_parameter == index);
 
-		if (index < int(m_pD->m_multiValues.Num()))
+		if (index < int32(m_pD->m_multiValues.Num()))
 		{
-			const TMap< TArray<int32_t>, PARAMETER_VALUE >& m = m_pD->m_multiValues[index];
-			const PARAMETER_VALUE* it = m.Find(pos->m_pD->m_values);
+			const TMap< TArray<int32>, FParameterValue >& m = m_pD->m_multiValues[index];
+			const FParameterValue* it = m.Find(pos->m_pD->m_values);
 			if (it)
 			{
-				return it->Get<ParamImageType>();
+				return it->Get<FParamImageType>();
 			}
 		}
 
 		// Multivalue parameter, but no multivalue set. Return single value.
-		return m_pD->m_values[index].Get<ParamImageType>();
+		return m_pD->m_values[index].Get<FParamImageType>();
 
     }
 
 
     //---------------------------------------------------------------------------------------------
-    void Parameters::SetImageValue( int index, FName id, const Ptr<const RangeIndex>& pos )
+    void Parameters::SetImageValue( int32 index, FName id, const Ptr<const RangeIndex>& pos )
     {
-        check( index >= 0 && index < (int)m_pD->m_values.Num() );
-        check( GetType( index ) == PARAMETER_TYPE::T_IMAGE );
+        check( index >= 0 && index < (int32)m_pD->m_values.Num() );
+        check( GetType( index ) == EParameterType::Image );
 
 		// Single value case
 		if (!pos)
 		{
-			m_pD->m_values[index].Set<ParamImageType>(id);
+			m_pD->m_values[index].Set<FParamImageType>(id);
 		}
 
 		// Multivalue case
@@ -856,27 +679,27 @@ namespace mu
 		{
 			check(pos->m_pD->m_parameter == index);
 
-			if (index >= int(m_pD->m_multiValues.Num()))
+			if (index >= int32(m_pD->m_multiValues.Num()))
 			{
 				m_pD->m_multiValues.SetNum(index + 1);
 			}
 
-			TMap< TArray<int32_t>, PARAMETER_VALUE >& m = m_pD->m_multiValues[index];
-			PARAMETER_VALUE& it = m.FindOrAdd(pos->m_pD->m_values);
-			it.Set<ParamImageType>(id);
+			TMap< TArray<int32>, FParameterValue >& m = m_pD->m_multiValues[index];
+			FParameterValue& it = m.FindOrAdd(pos->m_pD->m_values);
+			it.Set<FParamImageType>(id);
 		}
     }
 
 
     //---------------------------------------------------------------------------------------------
-    void Parameters::GetStringValue( int index, FString& OutValue, const Ptr<const RangeIndex>& pos ) const
+    void Parameters::GetStringValue( int32 index, FString& OutValue, const Ptr<const RangeIndex>& pos ) const
     {
-        check( index >= 0 && index < (int)m_pD->m_values.Num() );
-        check( GetType( index ) == PARAMETER_TYPE::T_STRING );
+        check( index >= 0 && index < (int32)m_pD->m_values.Num() );
+        check( GetType( index ) == EParameterType::String );
 
         // Early out in case of invalid parameters
-        if ( index < 0 || index >= (int)m_pD->m_values.Num() ||
-             GetType( index ) != PARAMETER_TYPE::T_STRING )
+        if ( index < 0 || index >= (int32)m_pD->m_values.Num() ||
+             GetType( index ) != EParameterType::String )
         {
             OutValue = TEXT("");
 			return;
@@ -886,41 +709,41 @@ namespace mu
         if ( !pos )
         {
             // Return the single value
-			OutValue = m_pD->m_values[index].Get<ParamStringType>();
+			OutValue = m_pD->m_values[index].Get<FParamStringType>();
             return;
         }
 
         // Multivalue case
         check( pos->m_pD->m_parameter == index );
 
-        if ( index < int( m_pD->m_multiValues.Num() ) )
+        if ( index < int32( m_pD->m_multiValues.Num() ) )
         {
-            const TMap< TArray<int32_t>, PARAMETER_VALUE >& m = m_pD->m_multiValues[index];
-			const PARAMETER_VALUE* it = m.Find( pos->m_pD->m_values );
+            const TMap< TArray<int32>, FParameterValue >& m = m_pD->m_multiValues[index];
+			const FParameterValue* it = m.Find( pos->m_pD->m_values );
             if ( it )
             {
-				OutValue = it->Get<ParamStringType>();
+				OutValue = it->Get<FParamStringType>();
                 return;
             }
         }
 
         // Multivalue parameter, but no multivalue set. Return single value.
-		OutValue = m_pD->m_values[index].Get<ParamStringType>();
+		OutValue = m_pD->m_values[index].Get<FParamStringType>();
         return;
     }
 
 
     //---------------------------------------------------------------------------------------------
-    void Parameters::SetStringValue( int index, const FString& Value, const Ptr<const RangeIndex>& pos )
+    void Parameters::SetStringValue( int32 index, const FString& Value, const Ptr<const RangeIndex>& pos )
     {
 		LLM_SCOPE_BYNAME(TEXT("MutableRuntime"));
 
-        check( index >= 0 && index < (int)m_pD->m_values.Num() );
-        check( GetType( index ) == PARAMETER_TYPE::T_FLOAT );
+        check( index >= 0 && index < (int32)m_pD->m_values.Num() );
+        check( GetType( index ) == EParameterType::Float );
 
         // Early out in case of invalid parameters
-        if ( index < 0 || index >= (int)m_pD->m_values.Num() ||
-             GetType( index ) != PARAMETER_TYPE::T_FLOAT )
+        if ( index < 0 || index >= (int32)m_pD->m_values.Num() ||
+             GetType( index ) != EParameterType::Float )
         {
             return;
         }
@@ -929,12 +752,12 @@ namespace mu
         if ( !pos )
         {
             // Clear multivalue, if set.
-            if ( index < int( m_pD->m_multiValues.Num() ) )
+            if ( index < int32( m_pD->m_multiValues.Num() ) )
             {
                 m_pD->m_multiValues[index].Empty();
             }
 
-        	m_pD->m_values[index].Set<ParamStringType>(Value);
+        	m_pD->m_values[index].Set<FParamStringType>(Value);
         }
 
         // Multivalue case
@@ -942,25 +765,25 @@ namespace mu
         {
             check( pos->m_pD->m_parameter == index );
 
-            if ( index >= int( m_pD->m_multiValues.Num() ) )
+            if ( index >= int32( m_pD->m_multiValues.Num() ) )
             {
                 m_pD->m_multiValues.SetNum( index + 1 );
             }
 
-			TMap< TArray<int32_t>, PARAMETER_VALUE >& m = m_pD->m_multiValues[index];
-			PARAMETER_VALUE& it = m.FindOrAdd( pos->m_pD->m_values );
-        	it.Set<ParamStringType>(Value);
+			TMap< TArray<int32>, FParameterValue >& m = m_pD->m_multiValues[index];
+			FParameterValue& it = m.FindOrAdd( pos->m_pD->m_values );
+        	it.Set<FParamStringType>(Value);
         }
     }
 
-    void Parameters::GetMatrixValue(int index, FMatrix44f& OutValue, const Ptr<const RangeIndex>& pos) const
+    void Parameters::GetMatrixValue(int32 index, FMatrix44f& OutValue, const Ptr<const RangeIndex>& pos) const
     {
-    	check( index >= 0 && index < (int)m_pD->m_values.Num() );
-    	check( GetType( index ) == PARAMETER_TYPE::T_MATRIX );
+    	check( index >= 0 && index < (int32)m_pD->m_values.Num() );
+    	check( GetType( index ) == EParameterType::Matrix );
 
     	// Early out in case of invalid parameters
-    	if ( index < 0 || index >= (int)m_pD->m_values.Num() ||
-			 GetType( index ) != PARAMETER_TYPE::T_MATRIX )
+    	if ( index < 0 || index >= (int32)m_pD->m_values.Num() ||
+			 GetType( index ) != EParameterType::Matrix )
     	{
     		OutValue = FMatrix44f::Identity;
     		return;
@@ -970,38 +793,38 @@ namespace mu
     	if ( !pos )
     	{
     		// Return the single value
-    		OutValue = m_pD->m_values[index].Get<ParamMatrixType>();
+    		OutValue = m_pD->m_values[index].Get<FParamMatrixType>();
     		return;
     	}
 
     	// Multivalue case
     	check( pos->m_pD->m_parameter == index );
 
-    	if ( index < int( m_pD->m_multiValues.Num() ) )
+    	if ( index < int32( m_pD->m_multiValues.Num() ) )
     	{
-    		const TMap< TArray<int32_t>, PARAMETER_VALUE >& m = m_pD->m_multiValues[index];
-    		const PARAMETER_VALUE* it = m.Find( pos->m_pD->m_values );
+    		const TMap< TArray<int32>, FParameterValue >& m = m_pD->m_multiValues[index];
+    		const FParameterValue* it = m.Find( pos->m_pD->m_values );
     		if ( it )
     		{
-    			OutValue = it->Get<ParamMatrixType>();
+    			OutValue = it->Get<FParamMatrixType>();
     			return;
     		}
     	}
 
     	// Multivalue parameter, but no multivalue set. Return single value.
-    	OutValue = m_pD->m_values[index].Get<ParamMatrixType>();
+    	OutValue = m_pD->m_values[index].Get<FParamMatrixType>();
     }
 
-    void Parameters::SetMatrixValue(int index, const FMatrix44f& Value, const Ptr<const RangeIndex>& pos)
+    void Parameters::SetMatrixValue(int32 index, const FMatrix44f& Value, const Ptr<const RangeIndex>& pos)
     {
     	LLM_SCOPE_BYNAME(TEXT("MutableRuntime"));
 
-    	check( index >= 0 && index < (int)m_pD->m_values.Num() );
-    	check( GetType( index ) == PARAMETER_TYPE::T_MATRIX );
+    	check( index >= 0 && index < (int32)m_pD->m_values.Num() );
+    	check( GetType( index ) == EParameterType::Matrix );
 
     	// Early out in case of invalid parameters
-    	if ( index < 0 || index >= (int)m_pD->m_values.Num() ||
-			 GetType( index ) != PARAMETER_TYPE::T_MATRIX )
+    	if ( index < 0 || index >= (int32)m_pD->m_values.Num() ||
+			 GetType( index ) != EParameterType::Matrix )
     	{
     		return;
     	}
@@ -1010,12 +833,12 @@ namespace mu
     	if ( !pos )
     	{
     		// Clear multivalue, if set.
-    		if ( index < int( m_pD->m_multiValues.Num() ) )
+    		if ( index < int32( m_pD->m_multiValues.Num() ) )
     		{
     			m_pD->m_multiValues[index].Empty();
     		}
 
-    		m_pD->m_values[index].Set<ParamMatrixType>(Value);
+    		m_pD->m_values[index].Set<FParamMatrixType>(Value);
     	}
 
     	// Multivalue case
@@ -1023,20 +846,20 @@ namespace mu
     	{
     		check( pos->m_pD->m_parameter == index );
 
-    		if ( index >= int( m_pD->m_multiValues.Num() ) )
+    		if ( index >= int32( m_pD->m_multiValues.Num() ) )
     		{
     			m_pD->m_multiValues.SetNum( index + 1 );
     		}
 
-    		TMap< TArray<int32_t>, PARAMETER_VALUE >& m = m_pD->m_multiValues[index];
-    		PARAMETER_VALUE& it = m.FindOrAdd( pos->m_pD->m_values );
-    		it.Set<ParamMatrixType>(Value);
+    		TMap< TArray<int32>, FParameterValue >& m = m_pD->m_multiValues[index];
+    		FParameterValue& it = m.FindOrAdd( pos->m_pD->m_values );
+    		it.Set<FParamMatrixType>(Value);
     	}
     }
 
 
     //---------------------------------------------------------------------------------------------
-    FProjector Parameters::Private::GetProjectorValue( int index, const Ptr<const RangeIndex>& pos ) const
+    FProjector Parameters::Private::GetProjectorValue( int32 index, const Ptr<const RangeIndex>& pos ) const
     {
 
         const FProgram& program = m_pModel->GetPrivate()->m_program;
@@ -1044,11 +867,11 @@ namespace mu
         // Early out in case of invalid parameters
         if ( index < 0
              ||
-             index >= (int)m_values.Num()
+             index >= (int32)m_values.Num()
 			 || 
-			 index >= (int)program.m_parameters.Num()
+			 index >= (int32)program.m_parameters.Num()
              ||
-             program.m_parameters[index].m_type != PARAMETER_TYPE::T_PROJECTOR )
+             program.m_parameters[index].m_type != EParameterType::Projector )
         {
 			check(false);
             return FProjector();
@@ -1060,7 +883,7 @@ namespace mu
         if (!pos)
         {
             // Return the single value
-            result = &m_values[index].Get<ParamProjectorType>();
+            result = &m_values[index].Get<FParamProjectorType>();
         }
 
         // Multivalue case
@@ -1070,10 +893,10 @@ namespace mu
 
             if ( index<m_multiValues.Num())
             {
-                const PARAMETER_VALUE* it = m_multiValues[index].Find(pos->m_pD->m_values);
+                const FParameterValue* it = m_multiValues[index].Find(pos->m_pD->m_values);
                 if (it)
                 {
-                    result = &it->Get<ParamProjectorType>();
+                    result = &it->Get<FParamProjectorType>();
                 }
             }
         }
@@ -1081,7 +904,7 @@ namespace mu
         // Multivalue parameter, but no multivalue set. Return single value.
         if (!result)
         {
-            result = &m_values[index].Get<ParamProjectorType>();
+            result = &m_values[index].Get<FParamProjectorType>();
         }
 
         return *result;
@@ -1089,8 +912,8 @@ namespace mu
 
 
 	//---------------------------------------------------------------------------------------------
-	void Parameters::GetProjectorValue( int ParameterIndex,
-		PROJECTOR_TYPE* OutType,
+	void Parameters::GetProjectorValue( int32 ParameterIndex,
+		EProjectorType* OutType,
 		FVector3f* OutPos,
 		FVector3f* OutDir,
 		FVector3f* OutUp,
@@ -1099,14 +922,14 @@ namespace mu
 		const Ptr<const RangeIndex>& pos ) const
 	{
 		check(ParameterIndex >=0 && ParameterIndex <m_pD->m_values.Num() );
-		check( GetType(ParameterIndex)== PARAMETER_TYPE::T_PROJECTOR );
+		check( GetType(ParameterIndex)== EParameterType::Projector );
 
         // Early out in case of invalid parameters
         if (ParameterIndex < 0
              ||
 			ParameterIndex >= m_pD->m_values.Num()
              ||
-             GetType(ParameterIndex) != PARAMETER_TYPE::T_PROJECTOR )
+             GetType(ParameterIndex) != EParameterType::Projector )
         {
             return;
         }
@@ -1126,7 +949,7 @@ namespace mu
 
 
 	//---------------------------------------------------------------------------------------------
-	void Parameters::SetProjectorValue( int ParameterIndex,
+	void Parameters::SetProjectorValue( int32 ParameterIndex,
 		const FVector3f& pos,
 		const FVector3f& dir,
 		const FVector3f& up,
@@ -1135,25 +958,25 @@ namespace mu
         const Ptr<const RangeIndex>& RangePosition)
 	{
 		check(ParameterIndex >=0 && ParameterIndex <m_pD->m_values.Num() );
-		check( GetType(ParameterIndex)== PARAMETER_TYPE::T_PROJECTOR );
+		check( GetType(ParameterIndex)== EParameterType::Projector );
 
         // Early out in case of invalid parameters
         if (ParameterIndex < 0
              ||
 			ParameterIndex >= m_pD->m_values.Num()
              ||
-             GetType(ParameterIndex) != PARAMETER_TYPE::T_PROJECTOR )
+             GetType(ParameterIndex) != EParameterType::Projector )
         {
             return;
         }
 
 	    // Parameters cannot change the projector type anymore
 		FProjector Value;
-        Value.type = PROJECTOR_TYPE::COUNT;
+        Value.type = EProjectorType::Count;
         if (m_pD->m_pModel)
         {
             const FProgram& program = m_pD->m_pModel->GetPrivate()->m_program;
-            const FProjector& Projector = program.m_parameters[ParameterIndex].m_defaultValue.Get<ParamProjectorType>();
+            const FProjector& Projector = program.m_parameters[ParameterIndex].m_defaultValue.Get<FParamProjectorType>();
             Value.type = Projector.type;
         }
 
@@ -1173,7 +996,7 @@ namespace mu
                 m_pD->m_multiValues[ParameterIndex].Empty();
             }
 
-            m_pD->m_values[ParameterIndex].Set<ParamProjectorType>(Value);
+            m_pD->m_values[ParameterIndex].Set<FParamProjectorType>(Value);
         }
 
         // Multivalue case
@@ -1186,17 +1009,17 @@ namespace mu
                 m_pD->m_multiValues.SetNum(ParameterIndex+1);
             }
 
-			TMap< TArray<int32_t>, PARAMETER_VALUE >& m = m_pD->m_multiValues[ParameterIndex];
-			PARAMETER_VALUE& it = m.FindOrAdd(RangePosition->m_pD->m_values);
-            it.Set<ParamProjectorType>(Value);
+			TMap< TArray<int32>, FParameterValue >& m = m_pD->m_multiValues[ParameterIndex];
+			FParameterValue& it = m.FindOrAdd(RangePosition->m_pD->m_values);
+            it.Set<FParamProjectorType>(Value);
         }
     }
 
 
     //---------------------------------------------------------------------------------------------
-    bool Parameters::HasSameValue( int thisParamIndex,
+    bool Parameters::HasSameValue( int32 thisParamIndex,
                                    const ParametersPtrConst& other,
-                                   int otherParamIndex ) const
+                                   int32 otherParamIndex ) const
     {
         if ( GetType(thisParamIndex) != other->GetType(otherParamIndex) )
         {
@@ -1209,14 +1032,14 @@ namespace mu
         }
 
         size_t thisNumMultiValues = 0;
-        bool thisHasMultiValues = int(m_pD->m_multiValues.Num()) > thisParamIndex;
+        bool thisHasMultiValues = int32(m_pD->m_multiValues.Num()) > thisParamIndex;
         if (thisHasMultiValues)
         {
             thisNumMultiValues = m_pD->m_multiValues[thisParamIndex].Num();
         }
 
         size_t otherNumMultiValues = 0;
-        bool otherHasMultiValues = int(other->m_pD->m_multiValues.Num()) > otherParamIndex;
+        bool otherHasMultiValues = int32(other->m_pD->m_multiValues.Num()) > otherParamIndex;
         if (otherHasMultiValues)
         {
             otherNumMultiValues = other->m_pD->m_multiValues[otherParamIndex].Num();
@@ -1245,9 +1068,9 @@ namespace mu
 	{
 		const FProgram& program = m_pModel->GetPrivate()->m_program;
 
-		int result = -1;
+		int32 result = -1;
 
-		for( int i=0; result<0 && i<(int)program.m_parameters.Num(); ++i )
+		for( int32 i=0; result<0 && i<(int32)program.m_parameters.Num(); ++i )
 		{
 			const FParameterDesc& p = program.m_parameters[i];
 
@@ -1287,40 +1110,40 @@ namespace mu
 
 
     //---------------------------------------------------------------------------------------------
-    int RangeIndex::GetRangeCount() const
+    int32 RangeIndex::GetRangeCount() const
     {
-        return int( m_pD->m_values.Num() );
+        return int32( m_pD->m_values.Num() );
     }
 
 
     //---------------------------------------------------------------------------------------------
-    const FString& RangeIndex::GetRangeName( int index ) const
+    const FString& RangeIndex::GetRangeName( int32 index ) const
     {
         check( index >= 0 && index < GetRangeCount() );
-        check( index < int(m_pD->m_pParameters->GetPrivate()->m_values.Num()) );
+        check( index < int32(m_pD->m_pParameters->GetPrivate()->m_values.Num()) );
         TSharedPtr<const Model> pModel = m_pD->m_pParameters->GetPrivate()->m_pModel;
         check(pModel);
-        check( index < int(pModel->GetPrivate()->m_program.m_ranges.Num()) );
+        check( index < int32(pModel->GetPrivate()->m_program.m_ranges.Num()) );
 
         return pModel->GetPrivate()->m_program.m_ranges[index].m_name;
     }
 
 
     //---------------------------------------------------------------------------------------------
-    const FString& RangeIndex::GetRangeUid( int index ) const
+    const FString& RangeIndex::GetRangeUid( int32 index ) const
     {
         check( index >= 0 && index < GetRangeCount() );
-        check( index < int(m_pD->m_pParameters->GetPrivate()->m_values.Num()) );
+        check( index < int32(m_pD->m_pParameters->GetPrivate()->m_values.Num()) );
 		TSharedPtr<const Model> pModel = m_pD->m_pParameters->GetPrivate()->m_pModel;
 		check(pModel);
-        check( index < int(pModel->GetPrivate()->m_program.m_ranges.Num()) );
+        check( index < int32(pModel->GetPrivate()->m_program.m_ranges.Num()) );
 
         return pModel->GetPrivate()->m_program.m_ranges[index].m_uid;
     }
 
 
     //---------------------------------------------------------------------------------------------
-    void RangeIndex::SetPosition( int index, int position )
+    void RangeIndex::SetPosition( int32 index, int32 position )
     {
         check( index >= 0 && index < GetRangeCount() );
         if ( index >= 0 && index < GetRangeCount() )
@@ -1331,7 +1154,7 @@ namespace mu
 
 
     //---------------------------------------------------------------------------------------------
-    int RangeIndex::GetPosition( int index ) const
+    int32 RangeIndex::GetPosition( int32 index ) const
     {
         check( index >= 0 && index < GetRangeCount() );
         if ( index >= 0 && index < GetRangeCount() )
diff --git a/Engine/Plugins/Mutable/Source/MutableRuntime/Private/MuR/System.cpp b/Engine/Plugins/Mutable/Source/MutableRuntime/Private/MuR/System.cpp
--- a/Engine/Plugins/Mutable/Source/MutableRuntime/Private/MuR/System.cpp
+++ b/Engine/Plugins/Mutable/Source/MutableRuntime/Private/MuR/System.cpp
@@ -357,7 +357,7 @@ namespace mu
 		// We always need to return something valid.
 		if (!pResult)
 		{
-			pResult = new mu::Image(16, 16, 1, EImageFormat::IF_RGBA_UBYTE, EInitializationType::Black);
+			pResult = new mu::Image(16, 16, 1, EImageFormat::RGBA_UByte, EInitializationType::Black);
 		}
 
 		m_pD->WorkingMemoryManager.CurrentInstanceCache = nullptr;
@@ -388,7 +388,7 @@ namespace mu
 			m_pD->WorkingMemoryManager.CurrentInstanceCache = nullptr;
 
 			return UE::Tasks::MakeCompletedTask<Ptr<const Image>>(
-					new mu::Image(16, 16, 1, EImageFormat::IF_RGBA_UBYTE, EInitializationType::Black));
+					new mu::Image(16, 16, 1, EImageFormat::RGBA_UByte, EInitializationType::Black));
 		}
 		
 		TSharedRef<CodeRunner> Runner = CodeRunner::Create(
@@ -411,7 +411,7 @@ namespace mu
 
 					if (!Result)
 					{
-						Result = new mu::Image(16, 16, 1, EImageFormat::IF_RGBA_UBYTE, EInitializationType::Black);
+						Result = new mu::Image(16, 16, 1, EImageFormat::RGBA_UByte, EInitializationType::Black);
 					}
 
 					SystemPrivate->WorkingMemoryManager.EndRunnerThread();
@@ -1690,12 +1690,12 @@ namespace mu
 
 
 	//---------------------------------------------------------------------------------------------
-	static void AddMultiValueKeys(FBitWriter& Blob, const TMap< TArray<int32>, PARAMETER_VALUE >& Multi)
+	static void AddMultiValueKeys(FBitWriter& Blob, const TMap< TArray<int32>, FParameterValue >& Multi)
 	{
 		uint16 Num = uint16(Multi.Num());
 		Blob.Serialize(&Num, 2);
 
-		for (const TPair< TArray<int32>, PARAMETER_VALUE >& v : Multi)
+		for (const TPair< TArray<int32>, FParameterValue >& v : Multi)
 		{
 			uint16 RangeNum = uint16(v.Key.Num());
 			Blob.Serialize(&RangeNum, 2);
@@ -1771,25 +1771,25 @@ namespace mu
 
 			switch (Program.m_parameters[ParamIndex].m_type)
 			{
-			case PARAMETER_TYPE::T_BOOL:
-				Blob.WriteBit(Params->GetPrivate()->m_values[ParamIndex].Get<ParamBoolType>() ? 1 : 0);
+			case EParameterType::Bool:
+				Blob.WriteBit(Params->GetPrivate()->m_values[ParamIndex].Get<FParamBoolType>() ? 1 : 0);
 
 				// Multi-values
 				if (Params->GetPrivate()->HasMultipleValues(ParamIndex))
 				{
-					const TMap< TArray<int32>, PARAMETER_VALUE >& Multi = Params->GetPrivate()->m_multiValues[ParamIndex];
+					const TMap< TArray<int32>, FParameterValue >& Multi = Params->GetPrivate()->m_multiValues[ParamIndex];
 					AddMultiValueKeys(Blob, Multi);
-					for (const TPair< TArray<int32>, PARAMETER_VALUE >& v : Multi)
+					for (const TPair< TArray<int32>, FParameterValue >& v : Multi)
 					{
-						Blob.WriteBit(v.Value.Get<ParamBoolType>() ? 1 : 0);
+						Blob.WriteBit(v.Value.Get<FParamBoolType>() ? 1 : 0);
 					}
 				}
 				break;
 
-			case PARAMETER_TYPE::T_INT:
+			case EParameterType::Int:
 			{
 				int32 MaxValue = ParamDescs[ParamIndex].m_possibleValues.Num();
-				int32 Value = Params->GetPrivate()->m_values[ParamIndex].Get<ParamIntType>();
+				int32 Value = Params->GetPrivate()->m_values[ParamIndex].Get<FParamIntType>();
 				if (MaxValue)
 				{
 					// It is an enum
@@ -1799,18 +1799,18 @@ namespace mu
 				else
 				{
 					// It may have any value
-					DataSize = sizeof(ParamIntType);
+					DataSize = sizeof(FParamIntType);
 					Blob.Serialize(&Value, DataSize);
 				}
 
 				// Multi-values
 				if (Params->GetPrivate()->HasMultipleValues(ParamIndex))
 				{
-					const TMap< TArray<int32>, PARAMETER_VALUE >& Multi = Params->GetPrivate()->m_multiValues[ParamIndex];
+					const TMap< TArray<int32>, FParameterValue >& Multi = Params->GetPrivate()->m_multiValues[ParamIndex];
 					AddMultiValueKeys(Blob, Multi);
-					for (const TPair< TArray<int32>, PARAMETER_VALUE >& v : Multi)
+					for (const TPair< TArray<int32>, FParameterValue >& v : Multi)
 					{
-						Value = v.Value.Get<ParamIntType>();
+						Value = v.Value.Get<FParamIntType>();
 						if (MaxValue)
 						{
 							// It is an enum
@@ -1820,7 +1820,7 @@ namespace mu
 						else
 						{
 							// It may have any value
-							DataSize = sizeof(ParamIntType);
+							DataSize = sizeof(FParamIntType);
 							Blob.Serialize(&Value, DataSize);
 						}
 					}
@@ -1828,42 +1828,42 @@ namespace mu
 				break;
 			}
 
-			case PARAMETER_TYPE::T_FLOAT:
-				DataSize = sizeof(ParamFloatType);
-				Blob.Serialize(&Params->GetPrivate()->m_values[ParamIndex].Get<ParamFloatType>(), DataSize);
+			case EParameterType::Float:
+				DataSize = sizeof(FParamFloatType);
+				Blob.Serialize(&Params->GetPrivate()->m_values[ParamIndex].Get<FParamFloatType>(), DataSize);
 
 				// Multi-values
 				if (Params->GetPrivate()->HasMultipleValues(ParamIndex))
 				{
-					const TMap< TArray<int32>, PARAMETER_VALUE >& Multi = Params->GetPrivate()->m_multiValues[ParamIndex];
+					const TMap< TArray<int32>, FParameterValue >& Multi = Params->GetPrivate()->m_multiValues[ParamIndex];
 					AddMultiValueKeys(Blob, Multi);
-					for (const TPair< TArray<int32>, PARAMETER_VALUE >& v : Multi)
+					for (const TPair< TArray<int32>, FParameterValue >& v : Multi)
 					{
-						Blob.Serialize((void*)&v.Value.Get<ParamFloatType>(), DataSize);
+						Blob.Serialize((void*)&v.Value.Get<FParamFloatType>(), DataSize);
 					}
 				}
 				break;
 
-			case PARAMETER_TYPE::T_COLOUR:
-				DataSize = sizeof(ParamColorType);
-				Blob.Serialize(&Params->GetPrivate()->m_values[ParamIndex].Get<ParamColorType>(), DataSize);
+			case EParameterType::Color:
+				DataSize = sizeof(FParamColorType);
+				Blob.Serialize(&Params->GetPrivate()->m_values[ParamIndex].Get<FParamColorType>(), DataSize);
 
 				// Multi-values
 				if (Params->GetPrivate()->HasMultipleValues(ParamIndex))
 				{
-					const TMap< TArray<int32>, PARAMETER_VALUE >& Multi = Params->GetPrivate()->m_multiValues[ParamIndex];
+					const TMap< TArray<int32>, FParameterValue >& Multi = Params->GetPrivate()->m_multiValues[ParamIndex];
 					AddMultiValueKeys(Blob, Multi);
-					for (const TPair< TArray<int32>, PARAMETER_VALUE >& v : Multi)
+					for (const TPair< TArray<int32>, FParameterValue >& v : Multi)
 					{
-						Blob.Serialize((void*)&v.Value.Get<ParamColorType>(), DataSize);
+						Blob.Serialize((void*)&v.Value.Get<FParamColorType>(), DataSize);
 					}
 				}
 				break;
 
-			case PARAMETER_TYPE::T_PROJECTOR:
+			case EParameterType::Projector:
 			{
 				FPackedNormal TempVec;
-				FProjector& Value = Params->GetPrivate()->m_values[ParamIndex].Get<ParamProjectorType>();
+				FProjector& Value = Params->GetPrivate()->m_values[ParamIndex].Get<FParamProjectorType>();
 				Blob.Serialize((void*)&Value.position, sizeof(FVector3f));
 				TempVec = FPackedNormal(Value.direction);
 				Blob.Serialize(&TempVec, sizeof(FPackedNormal));
@@ -1875,11 +1875,11 @@ namespace mu
 				// Multi-values
 				if (Params->GetPrivate()->HasMultipleValues(ParamIndex))
 				{
-					const TMap< TArray<int32>, PARAMETER_VALUE >& Multi = Params->GetPrivate()->m_multiValues[ParamIndex];
+					const TMap< TArray<int32>, FParameterValue >& Multi = Params->GetPrivate()->m_multiValues[ParamIndex];
 					AddMultiValueKeys(Blob, Multi);
-					for (const TPair< TArray<int32>, PARAMETER_VALUE >& v : Multi)
+					for (const TPair< TArray<int32>, FParameterValue >& v : Multi)
 					{
-						Value = v.Value.Get<ParamProjectorType>();
+						Value = v.Value.Get<FParamProjectorType>();
 						Blob.Serialize((void*)&Value.position, sizeof(FVector3f));
 						TempVec = FPackedNormal(Value.direction);
 						Blob.Serialize(&TempVec, sizeof(FPackedNormal));
@@ -1892,19 +1892,19 @@ namespace mu
 				break;
 			}
 
-			case PARAMETER_TYPE::T_IMAGE:
+			case EParameterType::Image:
 			{
 				DataSize = sizeof(FName);
-				Blob.Serialize((void*)&Params->GetPrivate()->m_values[ParamIndex].Get<ParamImageType>(), DataSize);
+				Blob.Serialize((void*)&Params->GetPrivate()->m_values[ParamIndex].Get<FParamImageType>(), DataSize);
 
 				// Multi-values
 				if (Params->GetPrivate()->HasMultipleValues(ParamIndex))
 				{
-					const TMap< TArray<int32>, PARAMETER_VALUE >& Multi = Params->GetPrivate()->m_multiValues[ParamIndex];
+					const TMap< TArray<int32>, FParameterValue >& Multi = Params->GetPrivate()->m_multiValues[ParamIndex];
 					AddMultiValueKeys(Blob, Multi);
-					for (const TPair< TArray<int32>, PARAMETER_VALUE >& v : Multi)
+					for (const TPair< TArray<int32>, FParameterValue >& v : Multi)
 					{
-						Blob.Serialize((void*)&v.Value.Get<ParamImageType>(), DataSize);
+						Blob.Serialize((void*)&v.Value.Get<FParamImageType>(), DataSize);
 					}
 				}
 				break;
diff --git a/Engine/Plugins/Mutable/Source/MutableRuntime/Public/MuR/Image.h b/Engine/Plugins/Mutable/Source/MutableRuntime/Public/MuR/Image.h
--- a/Engine/Plugins/Mutable/Source/MutableRuntime/Public/MuR/Image.h
+++ b/Engine/Plugins/Mutable/Source/MutableRuntime/Public/MuR/Image.h
@@ -357,36 +357,36 @@ namespace mu
 
 	MUTABLERUNTIME_API inline EImageFormat GetUncompressedFormat(EImageFormat Format)
 	{
-		check(Format < EImageFormat::IF_COUNT);
+		check(Format < EImageFormat::Count);
 
 		EImageFormat Result = Format;
 
 		switch (Result)
 		{
-		case EImageFormat::IF_L_UBIT_RLE: Result = EImageFormat::IF_L_UBYTE; break;
-		case EImageFormat::IF_L_UBYTE_RLE: Result = EImageFormat::IF_L_UBYTE; break;
-		case EImageFormat::IF_RGB_UBYTE_RLE: Result = EImageFormat::IF_RGB_UBYTE; break;
-        case EImageFormat::IF_RGBA_UBYTE_RLE: Result = EImageFormat::IF_RGBA_UBYTE; break;
-        case EImageFormat::IF_BC1: Result = EImageFormat::IF_RGBA_UBYTE; break;
-        case EImageFormat::IF_BC2: Result = EImageFormat::IF_RGBA_UBYTE; break;
-        case EImageFormat::IF_BC3: Result = EImageFormat::IF_RGBA_UBYTE; break;
-        case EImageFormat::IF_BC4: Result = EImageFormat::IF_L_UBYTE; break;
-        case EImageFormat::IF_BC5: Result = EImageFormat::IF_RGB_UBYTE; break;
-		case EImageFormat::IF_ASTC_4x4_RGB_LDR: Result = EImageFormat::IF_RGB_UBYTE; break;
-		case EImageFormat::IF_ASTC_4x4_RGBA_LDR: Result = EImageFormat::IF_RGBA_UBYTE; break;
-		case EImageFormat::IF_ASTC_4x4_RG_LDR: Result = EImageFormat::IF_RGB_UBYTE; break;
-		case EImageFormat::IF_ASTC_6x6_RGB_LDR: Result = EImageFormat::IF_RGB_UBYTE; break;
-		case EImageFormat::IF_ASTC_6x6_RGBA_LDR: Result = EImageFormat::IF_RGBA_UBYTE; break;
-		case EImageFormat::IF_ASTC_6x6_RG_LDR: Result = EImageFormat::IF_RGB_UBYTE; break;
-		case EImageFormat::IF_ASTC_8x8_RGB_LDR: Result = EImageFormat::IF_RGB_UBYTE; break;
-		case EImageFormat::IF_ASTC_8x8_RGBA_LDR: Result = EImageFormat::IF_RGBA_UBYTE; break;
-		case EImageFormat::IF_ASTC_8x8_RG_LDR: Result = EImageFormat::IF_RGB_UBYTE; break;
-		case EImageFormat::IF_ASTC_10x10_RGB_LDR: Result = EImageFormat::IF_RGB_UBYTE; break;
-		case EImageFormat::IF_ASTC_10x10_RGBA_LDR: Result = EImageFormat::IF_RGBA_UBYTE; break;
-		case EImageFormat::IF_ASTC_10x10_RG_LDR: Result = EImageFormat::IF_RGB_UBYTE; break;
-		case EImageFormat::IF_ASTC_12x12_RGB_LDR: Result = EImageFormat::IF_RGB_UBYTE; break;
-		case EImageFormat::IF_ASTC_12x12_RGBA_LDR: Result = EImageFormat::IF_RGBA_UBYTE; break;
-		case EImageFormat::IF_ASTC_12x12_RG_LDR: Result = EImageFormat::IF_RGB_UBYTE; break;
+		case EImageFormat::L_UBitRLE: Result = EImageFormat::L_UByte; break;
+		case EImageFormat::L_UByteRLE: Result = EImageFormat::L_UByte; break;
+		case EImageFormat::RGB_UByteRLE: Result = EImageFormat::RGB_UByte; break;
+        case EImageFormat::RGBA_UByteRLE: Result = EImageFormat::RGBA_UByte; break;
+        case EImageFormat::BC1: Result = EImageFormat::RGBA_UByte; break;
+        case EImageFormat::BC2: Result = EImageFormat::RGBA_UByte; break;
+        case EImageFormat::BC3: Result = EImageFormat::RGBA_UByte; break;
+        case EImageFormat::BC4: Result = EImageFormat::L_UByte; break;
+        case EImageFormat::BC5: Result = EImageFormat::RGB_UByte; break;
+		case EImageFormat::ASTC_4x4_RGB_LDR: Result = EImageFormat::RGB_UByte; break;
+		case EImageFormat::ASTC_4x4_RGBA_LDR: Result = EImageFormat::RGBA_UByte; break;
+		case EImageFormat::ASTC_4x4_RG_LDR: Result = EImageFormat::RGB_UByte; break;
+		case EImageFormat::ASTC_6x6_RGB_LDR: Result = EImageFormat::RGB_UByte; break;
+		case EImageFormat::ASTC_6x6_RGBA_LDR: Result = EImageFormat::RGBA_UByte; break;
+		case EImageFormat::ASTC_6x6_RG_LDR: Result = EImageFormat::RGB_UByte; break;
+		case EImageFormat::ASTC_8x8_RGB_LDR: Result = EImageFormat::RGB_UByte; break;
+		case EImageFormat::ASTC_8x8_RGBA_LDR: Result = EImageFormat::RGBA_UByte; break;
+		case EImageFormat::ASTC_8x8_RG_LDR: Result = EImageFormat::RGB_UByte; break;
+		case EImageFormat::ASTC_10x10_RGB_LDR: Result = EImageFormat::RGB_UByte; break;
+		case EImageFormat::ASTC_10x10_RGBA_LDR: Result = EImageFormat::RGBA_UByte; break;
+		case EImageFormat::ASTC_10x10_RG_LDR: Result = EImageFormat::RGB_UByte; break;
+		case EImageFormat::ASTC_12x12_RGB_LDR: Result = EImageFormat::RGB_UByte; break;
+		case EImageFormat::ASTC_12x12_RGBA_LDR: Result = EImageFormat::RGBA_UByte; break;
+		case EImageFormat::ASTC_12x12_RG_LDR: Result = EImageFormat::RGB_UByte; break;
 		default: break;
 		}
 
diff --git a/Engine/Plugins/Mutable/Source/MutableRuntime/Public/MuR/ImageDataStorage.h b/Engine/Plugins/Mutable/Source/MutableRuntime/Public/MuR/ImageDataStorage.h
--- a/Engine/Plugins/Mutable/Source/MutableRuntime/Public/MuR/ImageDataStorage.h
+++ b/Engine/Plugins/Mutable/Source/MutableRuntime/Public/MuR/ImageDataStorage.h
@@ -40,7 +40,7 @@ public:
 	TArray<FImageArray, TInlineAllocator<1>> Buffers;
 	
 	FImageSize ImageSize = FImageSize(0, 0);
-	EImageFormat ImageFormat = EImageFormat::IF_NONE;
+	EImageFormat ImageFormat = EImageFormat::None;
 	uint8 NumLODs = 0;
 
 	static constexpr int32 NumLODsInCompactedTail = 7;
diff --git a/Engine/Plugins/Mutable/Source/MutableRuntime/Public/MuR/ImageTypes.h b/Engine/Plugins/Mutable/Source/MutableRuntime/Public/MuR/ImageTypes.h
--- a/Engine/Plugins/Mutable/Source/MutableRuntime/Public/MuR/ImageTypes.h
+++ b/Engine/Plugins/Mutable/Source/MutableRuntime/Public/MuR/ImageTypes.h
@@ -17,55 +17,55 @@ namespace mu
 	//! \ingroup runtime
 	enum class EImageFormat : uint8
 	{
-		IF_NONE,
-		IF_RGB_UBYTE,
-		IF_RGBA_UBYTE,
-		IF_L_UBYTE,
+		None,
+		RGB_UByte,
+		RGBA_UByte,
+		L_UByte,
 
         //! Deprecated formats
-        IF_PVRTC2_DEPRECATED,
-        IF_PVRTC4_DEPRECATED,
-        IF_ETC1_DEPRECATED,
-        IF_ETC2_DEPRECATED,
+        _DEPRECATED_1,
+        _DEPRECATED_2,
+        _DEPRECATED_3,
+        _DEPRECATED_4,
 
-		IF_L_UBYTE_RLE,
-		IF_RGB_UBYTE_RLE,
-		IF_RGBA_UBYTE_RLE,
-		IF_L_UBIT_RLE,
+		L_UByteRLE,
+		RGB_UByteRLE,
+		RGBA_UByteRLE,
+		L_UBitRLE,
 
         //! Common S3TC formats
-        IF_BC1,
-        IF_BC2,
-        IF_BC3,
-        IF_BC4,
-        IF_BC5,
+        BC1,
+        BC2,
+        BC3,
+        BC4,
+        BC5,
 
         //! Not really supported yet
-        IF_BC6,
-        IF_BC7,
+        BC6,
+        BC7,
 
         //! Swizzled versions, engineers be damned.
-        IF_BGRA_UBYTE,
+        BGRA_UByte,
 
         //! The new standard
-        IF_ASTC_4x4_RGB_LDR,
-        IF_ASTC_4x4_RGBA_LDR,
-        IF_ASTC_4x4_RG_LDR,
-
-		IF_ASTC_8x8_RGB_LDR,
-		IF_ASTC_8x8_RGBA_LDR,
-		IF_ASTC_8x8_RG_LDR,
-		IF_ASTC_12x12_RGB_LDR,
-		IF_ASTC_12x12_RGBA_LDR,
-		IF_ASTC_12x12_RG_LDR,
-		IF_ASTC_6x6_RGB_LDR,
-		IF_ASTC_6x6_RGBA_LDR,
-		IF_ASTC_6x6_RG_LDR,
-		IF_ASTC_10x10_RGB_LDR,
-		IF_ASTC_10x10_RGBA_LDR,
-		IF_ASTC_10x10_RG_LDR,
-
-        IF_COUNT
+        ASTC_4x4_RGB_LDR,
+        ASTC_4x4_RGBA_LDR,
+        ASTC_4x4_RG_LDR,
+
+		ASTC_8x8_RGB_LDR,
+		ASTC_8x8_RGBA_LDR,
+		ASTC_8x8_RG_LDR,
+		ASTC_12x12_RGB_LDR,
+		ASTC_12x12_RGBA_LDR,
+		ASTC_12x12_RG_LDR,
+		ASTC_6x6_RGB_LDR,
+		ASTC_6x6_RGBA_LDR,
+		ASTC_6x6_RG_LDR,
+		ASTC_10x10_RGB_LDR,
+		ASTC_10x10_RGBA_LDR,
+		ASTC_10x10_RG_LDR,
+
+        Count
 	};
 	MUTABLE_DEFINE_ENUM_SERIALISABLE(EImageFormat);
 
@@ -85,7 +85,7 @@ namespace mu
 		FImageSize m_size = FImageSize(0, 0);
 
 		//!
-		EImageFormat m_format = EImageFormat::IF_NONE;
+		EImageFormat m_format = EImageFormat::None;
 
 		//! Levels of detail (mipmaps)
 		uint8 m_lods = 0;
diff --git a/Engine/Plugins/Mutable/Source/MutableRuntime/Public/MuR/MeshBufferSet.h b/Engine/Plugins/Mutable/Source/MutableRuntime/Public/MuR/MeshBufferSet.h
--- a/Engine/Plugins/Mutable/Source/MutableRuntime/Public/MuR/MeshBufferSet.h
+++ b/Engine/Plugins/Mutable/Source/MutableRuntime/Public/MuR/MeshBufferSet.h
@@ -20,61 +20,58 @@ namespace mu::MemoryCounters
 namespace mu
 {
 	/** Supported formats for the elements in mesh buffers. **/
-	typedef enum
+	enum class EMeshBufferFormat : uint32
 	{
+		None,
 
-		MBF_NONE,
-		MBF_FLOAT16,
-		MBF_FLOAT32,
+		Float16,
+		Float32,
 
-		MBF_UINT8,
-		MBF_UINT16,
-		MBF_UINT32,
-		MBF_INT8,
-		MBF_INT16,
-		MBF_INT32,
+		UInt8,
+		UInt16,
+		UInt32,
+		Int8,
+		Int16,
+		Int32,
 
 		/** Integers interpreted as being in the range 0.0f to 1.0f */
-		MBF_NUINT8,
-		MBF_NUINT16,
-		MBF_NUINT32,
+		NUInt8,
+		NUInt16,
+		NUInt32,
 
 		/** Integers interpreted as being in the range -1.0f to 1.0f */
-		MBF_NINT8,
-		MBF_NINT16,
-		MBF_NINT32,
+		NInt8,
+		NInt16,
+		NInt32,
 
         /** Packed 1 to -1 value using multiply+add (128 is almost zero). Use 8-bit unsigned ints. */
-        MBF_PACKEDDIR8,
+        PackedDir8,
 
         /** 
-		 * Same as MBF_PACKEDDIR8, with the w component replaced with the sign of the determinant
+		 * Same as EMeshBufferFormat::PackedDir8, with the w component replaced with the sign of the determinant
          * of the vertex basis to define the orientation of the tangent space in UE4 format.
          * Use 8-bit unsigned ints.
 		*/
-        MBF_PACKEDDIR8_W_TANGENTSIGN,
+        PackedDir8_W_TangentSign,
 
         /** Packed 1 to -1 value using multiply+add (128 is almost zero). Use 8-bit signed ints. */
-        MBF_PACKEDDIRS8,
+        PackedDirS8,
 
         /** 
-		 * Same as MBF_PACKEDDIRS8, with the w component replaced with the sign of the determinant
+		 * Same as EMeshBufferFormat::PackedDirS8, with the w component replaced with the sign of the determinant
          * of the vertex basis to define the orientation of the tangent space in UE4 format.
          * Use 8-bit signed ints.
 		 */
-        MBF_PACKEDDIRS8_W_TANGENTSIGN,
+        PackedDirS8_W_TangentSign,
 
-		MBF_FLOAT64,
-		MBF_UINT64,
-		MBF_INT64,
-		MBF_NUINT64,
-		MBF_NINT64,
+		Float64,
+		UInt64,
+		Int64,
+		NUInt64,
+		NInt64,
 
-		MBF_COUNT,
-
-		_MBF_FORCE32BITS = 0xFFFFFFFF
-
-	} EMeshBufferFormat;
+		Count,
+	};
 
 	/** */
 	struct FMeshBufferFormatData
@@ -90,56 +87,52 @@ namespace mu
 
 
 	/** Semantics of the mesh buffers */
-	typedef enum
+	enum class EMeshBufferSemantic : uint32
 	{
-
-		MBS_NONE,
+		None,
 
 		/** For index buffers, and mesh morphs */
-		MBS_VERTEXINDEX,
+		VertexIndex,
 
 		/** Standard vertex semantics */
-		MBS_POSITION,
-		MBS_NORMAL,
-		MBS_TANGENT,
-		MBS_BINORMAL,
-		MBS_TEXCOORDS,
-		MBS_COLOUR,
-		MBS_BONEWEIGHTS,
-		MBS_BONEINDICES,
+		Position,
+		Normal,
+		Tangent,
+		Binormal,
+		TexCoords,
+		Color,
+		BoneWeights,
+		BoneIndices,
 
 		/**
 		 * Internal semantic indicating what layout block each vertex belongs to.
 		 * It can be safely ignored if present in meshes returned by the system.
 		 * It will never be in the same buffer that other vertex semantics.
 		 */
-		MBS_LAYOUTBLOCK,
+		LayoutBlock,
 
-		MBS_CHART_DEPRECATED,
+		_DEPRECATED,
 
 		/** 
 		 * To let users define channels with semantics unknown to the system.
 		 * These channels will never be transformed, and the per-vertex or per-index data will be
 		 * simply copied.
 		 */
-		MBS_OTHER,
+		Other,
 
-        /** Sign to define the orientation of the tangent space. **/
-        MBS_TANGENTSIGN_DEPRECATED,
+        _DEPRECATED2,
 
 		/** Semantics usefule for mesh binding. */
-		MBS_TRIANGLEINDEX,
-		MBS_BARYCENTRICCOORDS,
-		MBS_DISTANCE,
+		TriangleIndex,
+		BarycentricCoords,
+		Distance,
 
 		/** Semantics useful for alternative skin weight profiles. */
-		MBS_ALTSKINWEIGHT,
+		AltSkinWeight,
 
 		/** Utility */
-		MBS_COUNT,
-
-        _MBS_FORCE32BITS = 0xFFFFFFFF
-	} EMeshBufferSemantic;
+		Count,
+	};
 
 
 	/** */
@@ -147,8 +140,8 @@ namespace mu
 	{
 		FMeshBufferChannel()
 		{
-			Semantic = MBS_NONE;
-			Format = MBF_NONE;
+			Semantic = EMeshBufferSemantic::None;
+			Format = EMeshBufferFormat::None;
 			SemanticIndex = 0;
 			Offset = 0;
 			ComponentCount = 0;
diff --git a/Engine/Plugins/Mutable/Source/MutableRuntime/Public/MuR/Model.h b/Engine/Plugins/Mutable/Source/MutableRuntime/Public/MuR/Model.h
--- a/Engine/Plugins/Mutable/Source/MutableRuntime/Public/MuR/Model.h
+++ b/Engine/Plugins/Mutable/Source/MutableRuntime/Public/MuR/Model.h
@@ -88,13 +88,13 @@ namespace mu
 		//! Create a set of new parameters of the model with the default values.
 		//! If old parameters are provided, they will be reused when possible instead of the
 		//! default values.
-        static ParametersPtr NewParameters( TSharedPtr<const Model> Model, const Parameters* pOld = nullptr );
+        static Ptr<Parameters> NewParameters( TSharedPtr<const Model> Model, const Parameters* pOld = nullptr );
 
 		/** Return true if the parameter is multi-dimensional */
 		bool IsParameterMultidimensional(int32 ParameterIndex) const;
 
 		//! Get the number of states in the model.
-		int GetStateCount() const;
+		int32 GetStateCount() const;
 
 		//! Get a state name by state index from 0 to GetStateCount-1
 		const FString& GetStateName( int32 StateIndex ) const;
@@ -103,11 +103,11 @@ namespace mu
 		int32 FindState( const FString& Name ) const;
 
 		//! Get the number of parameters available in a particular state.
-		int GetStateParameterCount( int stateIndex ) const;
+		int32 GetStateParameterCount( int32 stateIndex ) const;
 
 		//! Get the index of one of the parameters in the given state. The index refers to the
 		//! parameters in a Parameters object obtained from this model with NewParameters.
-		int GetStateParameterIndex( int stateIndex, int paramIndex ) const;
+		int32 GetStateParameterIndex( int32 stateIndex, int32 paramIndex ) const;
 
         //! Free memory used by streaming data that may be loaded again when needed.
         void UnloadExternalData();
@@ -131,7 +131,7 @@ namespace mu
 		//! \pre The parameter specified by index is a T_FLOAT.
         //! \param Index Index of the parameter from 0 to GetCount()-1
         //! \param R,G,B Pointers to values where every resulting colour channel will be stored
-    	void GetColourDefaultValue(int32 Index, float* R, float* G, float* B, float* A) const;
+    	void GetColourDefaultValue(int32 Index, FVector4f& OutValue) const;
 
 		//! Return the default value of a colour parameter.
 		//! \pre The parameter specified by index is a T_FLOAT.
@@ -148,7 +148,7 @@ namespace mu
         //!         of the projector will be stored. This controls the "roll" angle of the
         //!         projector.
         //! \param OutScale Pointer to the projector-space scaling of the projector.
-    	void GetProjectorDefaultValue(int32 Index, PROJECTOR_TYPE* OutProjectionType, FVector3f* OutPos,
+    	void GetProjectorDefaultValue(int32 Index, EProjectorType* OutProjectionType, FVector3f* OutPos,
     	 	FVector3f* OutDir, FVector3f* OutUp, FVector3f* OutScale, float* OutProjectionAngle) const;
 
         //! Return the default value of an image parameter.
diff --git a/Engine/Plugins/Mutable/Source/MutableRuntime/Public/MuR/MutableString.h b/Engine/Plugins/Mutable/Source/MutableRuntime/Public/MuR/MutableString.h
--- a/Engine/Plugins/Mutable/Source/MutableRuntime/Public/MuR/MutableString.h
+++ b/Engine/Plugins/Mutable/Source/MutableRuntime/Public/MuR/MutableString.h
@@ -12,8 +12,7 @@
 namespace mu
 {
 
-    //! \brief Container for text string data.
-	//! \ingroup runtime
+    /** Container for text string data. */
     class MUTABLERUNTIME_API String : public Resource
 	{
 	public:
@@ -27,10 +26,6 @@ namespace mu
 		//! Deep clone this string.
 		Ptr<String> Clone() const;
 
-		//! Serialisation
-		//static void Serialise( const String* p, OutputArchive& arch );
-		//static Ptr<String> StaticUnserialise( InputArchive& arch );
-
 		// Resource interface
 		int32 GetDataSize() const override;
 
@@ -51,34 +46,6 @@ namespace mu
 		//!
 		FString Value;
 
-		//!
-		//void Serialise(OutputArchive& arch) const
-		//{
-		//	uint32 ver = 1;
-		//	arch << ver;
-
-		//	arch << Value;
-		//}
-
-		////!
-		//void Unserialise(InputArchive& arch)
-		//{
-		//	uint32 ver;
-		//	arch >> ver;
-		//	check(ver <= 1);
-
-		//	if (ver == 0)
-		//	{
-		//		std::string Temp;
-		//		arch >> Temp;
-		//		Value = Temp.c_str();
-		//	}
-		//	else
-		//	{
-		//		arch >> Value;
-		//	}
-		//}
-
 		//!
 		inline bool operator==(const String& o) const
 		{
diff --git a/Engine/Plugins/Mutable/Source/MutableRuntime/Public/MuR/Parameters.h b/Engine/Plugins/Mutable/Source/MutableRuntime/Public/MuR/Parameters.h
--- a/Engine/Plugins/Mutable/Source/MutableRuntime/Public/MuR/Parameters.h
+++ b/Engine/Plugins/Mutable/Source/MutableRuntime/Public/MuR/Parameters.h
@@ -32,100 +32,91 @@ namespace mu
 	static const FVector4f DefaultMutableColorValue(-UE_BIG_NUMBER);
 
 
-    //! \brief Model parameter types.
-	//! \ingroup runtime
-    enum class PARAMETER_TYPE : uint32
+    /** Model parameter types. */
+    enum class EParameterType : uint32
     {
-        //! Undefined parameter type.
-        T_NONE,
+        /** Undefined parameter type. */
+        None,
 
-        //! Boolean parameter type (true or false)
-        T_BOOL,
+        /** Boolean parameter type (true or false) */
+        Bool,
 
-        //! Integer parameter type. It usually has a limited range of possible values that can be
-        //! queried in the Parameters object.
-        T_INT,
+        /** Integer parameter type. It usually has a limited range of possible values that can be queried in the Parameters object. */
+        Int,
 
-        //! Floating point value in the range of 0.0 to 1.0
-        T_FLOAT,
+        /** Floating point value in the range of 0.0 to 1.0 */
+        Float,
 
-        //! Floating point RGBA colour, with each channel ranging from 0.0 to 1.0
-        T_COLOUR,
+        /** Floating point RGBA colour, with each channel ranging from 0.0 to 1.0 */
+        Color,
 
-        //! 3D Projector type, defining a position, scale and orientation. Basically used for
-        //! projected decals.
-        T_PROJECTOR,
+        /** 3D Projector type, defining a position, scale and orientation.Basically used for projected decals. */
+        Projector,
 
-        //! An externally provided image.
-        T_IMAGE,
+        /** An externally provided image. */
+        Image,
 
-        //! A text string.
-        T_STRING,
+        /** A text string. */
+        String,
 
-    	//! A 4x4 matrix.
-    	T_MATRIX,
+    	/** A 4x4 matrix. */
+    	Matrix,
 
-        //! Utility enumeration value, not really a parameter type.
-        T_COUNT
+        /** Utility enumeration value, not really a parameter type. */
+        Count
 
     };
 
 
-    //! \brief Types of projectors.
-    //! \ingroup runtime
-    enum class PROJECTOR_TYPE : uint32
+    /** Types of 3D projectors. */
+    enum class EProjectorType : uint32
     {
-        //!
-        PLANAR,
+        /** Standard projector that uses an affine transform. */
+        Planar,
 
-        //!
-        CYLINDRICAL,
+        /** Projector that wraps the projected image around a cylinder*/
+        Cylindrical,
 
-		//!
-        WRAPPING,
+		/** Smart projector that tries to follow the projected surface geometry to minimize streching. */
+        Wrapping,
 
-        //! Utility enumeration value, not really a projector type.
-        COUNT
+        /** Utility enumeration value, not really a projector type. */
+        Count
     };
 
 
-    //! \brief Class used to set or read multi-dimensional parameter values.
-    //! If parameters have multiple values because of ranges, RangeIndex can be used to specify
-    //! what value is read or set in the methods of the Parameter class.
-    //! RangeIndex objects can be reused in multiple calls for the same parameter in the Parameter
-    //! class interface even after changing the position values.
-    //! \ingroup runtime
-    class MUTABLERUNTIME_API RangeIndex : public RefCounted
+    /** Type used to set or read multi - dimensional parameter values.
+     * If parameters have multiple values because of ranges, RangeIndex can be used to specify
+     * what value is read or set in the methods of the Parameter class.
+     * RangeIndex objects can be reused in multiple calls for the same parameter in the Parameter
+     * class interface even after changing the position values.
+	 */
+	class MUTABLERUNTIME_API RangeIndex : public RefCounted
     {
     public:
 
-        //-----------------------------------------------------------------------------------------
-        // Life cycle
-        //-----------------------------------------------------------------------------------------
-
-
         //-----------------------------------------------------------------------------------------
         // Own interface
         //-----------------------------------------------------------------------------------------
 
         //! Return the number of ranges (or dimensions) used by this index
-        int GetRangeCount() const;
+        int32 GetRangeCount() const;
 
         //! Return the name of a range.
         //! \param index Index of the range from 0 to GetRangeCount()-1
-        const FString& GetRangeName( int index ) const;
+        const FString& GetRangeName( int32 index ) const;
 
         //! Return the Guid of the parameter, resistant to parameter name changes.
         //! \param index Index of the parameter from 0 to GetRangeCount()-1
-        const FString& GetRangeUid( int index ) const;
+        const FString& GetRangeUid( int32 index ) const;
 
         //! Set the position in one of the dimensions of the index.
         //! \param index Index of the parameter from 0 to GetRangeCount()-1
-        void SetPosition( int index, int position );
+        void SetPosition( int32 index, int32 position );
 
         //! Return the position in one of the dimensions of the index.
         //! \param index Index of the parameter from 0 to GetRangeCount()-1
-        int GetPosition( int index ) const;
+        int32 GetPosition( int32 index ) const;
 
 
         //-----------------------------------------------------------------------------------------
@@ -151,14 +142,13 @@ namespace mu
     };
 
 
-    //! \brief  This class represents the parameters of a model including description, type and
-    //! value, and additional resources associated to the parameter.
-    //!
-    //! \warning Every object of this class holds a reference to the model that it was created from.
-    //! This implies that while any instance of a Parameters object is alive, its model will not
-    //! be freed.
-    //!
-	//! \ingroup runtime
+    /** This class represents the parameters of a model including description, type and
+    * value, and additional resources associated to the parameter.
+    *
+    * \warning Every object of this class holds a reference to the model that it was created from.
+    * This implies that while any instance of a Parameters object is alive, its model will not
+    * be freed.
+    */
 	class MUTABLERUNTIME_API Parameters : public RefCounted
 	{
 	public:
@@ -171,14 +161,6 @@ namespace mu
         static void Serialise( const Parameters* p, OutputArchive& arch );
         static ParametersPtr StaticUnserialise( InputArchive& arch );
 
-        //! Portable serialisation, which tries to keep values when:
-        //! - new parameters have been added
-        //! - parameters have been removed
-        //! - parameter types have changed
-        //! - parameter names have changed, and the application is using UIDs
-        //static void SerialisePortable( const Parameters*, OutputArchive& );
-        //static ParametersPtr UnserialisePortable( InputArchive&, TSharedPtr<const Model> );
-
 		//! Deep clone this object.
 		ParametersPtr Clone() const;
 
@@ -187,46 +169,46 @@ namespace mu
 		//-----------------------------------------------------------------------------------------
 
 		//! Return the number of parameters
-		int GetCount() const;
+		int32 GetCount() const;
 
 		//! Return the name of the parameter.
 		//! \param index Index of the parameter from 0 to GetCount()-1
-		const FString& GetName( int index ) const;
+		const FString& GetName( int32 index ) const;
 
 		//! Return the Guid of the parameter, resistant to parameter name changes.
 		//! \param index Index of the parameter from 0 to GetCount()-1
-		const FGuid& GetUid( int index ) const;
+		const FGuid& GetUid( int32 index ) const;
 
 		//! Find the parameter index by name.
 		//! It returns -1 if the parameter doesn't exist.
-		int Find( const FString& Name ) const;
+		int32 Find( const FString& Name ) const;
 
 		//! Return the type of the parameter.
 		//! \param index Index of the parameter from 0 to GetCount()-1
-		PARAMETER_TYPE GetType( int index ) const;
+		EParameterType GetType( int32 index ) const;
 
         //! Create a new RangeIndex object to use to access a multi-dimensional parameter.
         //! It will return nullptr if the parameter is not multidimensional.
-        Ptr<RangeIndex> NewRangeIndex( int paramIndex ) const;
+        Ptr<RangeIndex> NewRangeIndex( int32 paramIndex ) const;
 
         //! Return the number of values other than the default that have been set to a specific
         //! parameter.
-        int GetValueCount( int paramIndex ) const;
+        int32 GetValueCount( int32 paramIndex ) const;
 
         //! Return the RangeIndex of a value that has been set to a parameter.
         //! \param paramIndex Index of the parameter from 0 to GetCount()-1
         //! \param valueIndex Index of the value from 0 to GetValueCount()-1
-		Ptr<RangeIndex> GetValueIndex( int paramIndex, int valueIndex ) const;
+		Ptr<RangeIndex> GetValueIndex( int32 paramIndex, int32 valueIndex ) const;
 
         //! Remove all the multidimensional values for a parameter. The current non-dimensional
         //! value is kept.
-        void ClearAllValues( int paramIndex );
+        void ClearAllValues( int32 paramIndex );
 
 		//! Return the value of a boolean parameter.
 		//! \pre The parameter specified by index is a T_BOOL.
 		//! \param index Index of the parameter from 0 to GetCount()-1
         //! \param pos Only for multidimensional parameters: relevant position to get in the ranges
-        bool GetBoolValue( int index, const Ptr<const RangeIndex>& pos=nullptr ) const;
+        bool GetBoolValue( int32 index, const Ptr<const RangeIndex>& pos=nullptr ) const;
 
 		//! Set the value of a parameter of type boolean.
 		//! \pre The parameter specified by index is a T_BOOL.
@@ -234,67 +216,65 @@ namespace mu
 		//! \param value new value of the parameter
         //! \param pos optional parameter to set a specific value for the given multidimensional
         //! position. If null, the value is set for all possible positions.
-        void SetBoolValue( int index, bool value, const Ptr<const RangeIndex>& pos=nullptr );
+        void SetBoolValue( int32 index, bool value, const Ptr<const RangeIndex>& pos=nullptr );
 
 		//! Get the number of possible values for this integer parameter.
 		//! If the number is zero, it means any integer value is accepted.
-		int GetIntPossibleValueCount( int paramIndex ) const;
+		int32 GetIntPossibleValueCount( int32 paramIndex ) const;
 
 		//! Get the value of one of the possible values for this integer.
 		//! The paramIndex is in the range of 0 to GetIntPossibleValueCount()-1
-		int GetIntPossibleValue( int paramIndex, int valueIndex ) const;
+		int32 GetIntPossibleValue( int32 paramIndex, int32 valueIndex ) const;
 
         //! Get the name of one of the possible values for this integer.
         //! The paramIndex is in the range of 0 to GetIntPossibleValueCount()-1
-        const FString& GetIntPossibleValueName( int paramIndex, int valueIndex ) const;
+        const FString& GetIntPossibleValueName( int32 paramIndex, int32 valueIndex ) const;
 
         //! Get the index of the value of one of the possible values for this integer.
         //! The paramIndex is in the range of 0 to GetIntPossibleValueCount()-1
-		int GetIntValueIndex(int paramIndex, const FString& valueName) const;
+		int32 GetIntValueIndex(int32 paramIndex, const FString& valueName) const;
 
 		//! Get the index of the value of one of the possible values for this integer.
 		//! The paramIndex is in the range of 0 to GetIntPossibleValueCount()-1
-		int GetIntValueIndex(int paramIndex, int32 Value) const;
+		int32 GetIntValueIndex(int32 paramIndex, int32 Value) const;
 
 		//! Return the value of a integer parameter.
 		//! \pre The parameter specified by index is a T_INT.
 		//! \param index Index of the parameter from 0 to GetCount()-1
         //! \param pos Only for multidimensional parameters: relevant position to get in the ranges
-        int GetIntValue( int index, const Ptr<const RangeIndex>& pos=nullptr ) const;
+        int32 GetIntValue( int32 index, const Ptr<const RangeIndex>& pos=nullptr ) const;
 
 		//! If the parameter is of the integer type, set its value.
 		//! \param index Index of the parameter from 0 to GetCount()-1
 		//! \param value new value of the parameter. It must be in the possible values for this
 		//! parameter (see GetIntPossibleValue), or the method will leave the value unchanged.
         //! \param pos Only for multidimensional parameters: relevant position to set in the ranges
-        void SetIntValue( int index, int value, const Ptr<const RangeIndex>& pos=nullptr );
+        void SetIntValue( int32 index, int32 value, const Ptr<const RangeIndex>& pos=nullptr );
 
 		//! Return the value of a float parameter.
 		//! \pre The parameter specified by index is a T_FLOAT.
 		//! \param index Index of the parameter from 0 to GetCount()-1
         //! \param pos Only for multidimensional parameters: relevant position to get in the ranges
-        float GetFloatValue( int index, const Ptr<const RangeIndex>& pos=nullptr ) const;
+        float GetFloatValue( int32 index, const Ptr<const RangeIndex>& pos=nullptr ) const;
 
 		//! If the parameter is of the float type, set its value.
 		//! \param index Index of the parameter from 0 to GetCount()-1
 		//! \param value new value of the parameter
         //! \param pos Only for multidimensional parameters: relevant position to set in the ranges
-        void SetFloatValue( int index, float value, const Ptr<const RangeIndex>& pos=nullptr );
+        void SetFloatValue( int32 index, float value, const Ptr<const RangeIndex>& pos=nullptr );
 
 		//! Return the value of a colour parameter.
 		//! \pre The parameter specified by index is a T_FLOAT.
         //! \param index Index of the parameter from 0 to GetCount()-1
-        //! \param pR,pG,pB Pointers to values where every resulting colour channel will be stored
+        //! \param OutColor Values where every resulting colour channel will be stored
         //! \param pos Only for multidimensional parameters: relevant position to get in the ranges
-        void GetColourValue( int index, float* pR, float* pG, float* pB, float* pA,
-                             const Ptr<const RangeIndex>& pos=nullptr ) const;
+        void GetColourValue( int32 index, FVector4f& OutColor, const Ptr<const RangeIndex>& pos=nullptr ) const;
 
 		//! If the parameter is of the colour type, set its value.
 		//! \param index Index of the parameter from 0 to GetCount()-1
         //! \param r,g,b new value of the parameter
         //! \param pos Only for multidimensional parameters: relevant position to set in the ranges
-        void SetColourValue( int index, float R, float G, float B, float A,
-                             const Ptr<const RangeIndex>& pos=nullptr );
+        void SetColourValue( int32 index, FVector4f Color, const Ptr<const RangeIndex>& pos=nullptr );
 
 		//! Return the value of a projector parameter, as a 4x4 matrix. The matrix is supposed to be
 		//! a linear transform in column-major.
@@ -307,8 +287,8 @@ namespace mu
         //!         projector.
         //! \param OutScale Pointer to the projector-space scaling of the projector.
         //! \param RangePosition Only for multidimensional parameters: relevant position to get in the ranges
-        void GetProjectorValue( int ParameterIndex,
-                                PROJECTOR_TYPE* OutProjectionType,
+        void GetProjectorValue( int32 ParameterIndex,
+                                EProjectorType* OutProjectionType,
 								FVector3f* OutPos,
 								FVector3f* OutDir,
 								FVector3f* OutUp,
@@ -325,7 +305,7 @@ namespace mu
         //! \param projectionAngle [only for Cylindrical projectors], the angle in radians of the
         //! projection area on the cylinder surface.
         //! \param RangePosition Only for multidimensional parameters: relevant position to set in the ranges
-        void SetProjectorValue( int ParameterIndex,
+        void SetProjectorValue( int32 ParameterIndex,
 								const FVector3f& pos,
 								const FVector3f& dir,
 								const FVector3f& up,
@@ -338,41 +318,41 @@ namespace mu
         //! \param index Index of the parameter from 0 to GetCount()-1
 		//! \param pos Only for multidimensional parameters: relevant position to set in the ranges
 	    //! \return The externalId specified when setting the image value (\see SetImageValue)
-		FName GetImageValue( int index, const Ptr<const RangeIndex>& pos = nullptr) const;
+		FName GetImageValue( int32 index, const Ptr<const RangeIndex>& pos = nullptr) const;
 
         //! If the parameter is of the image type, set its value.
         //! \param index Index of the parameter from 0 to GetCount()-1
         //! \param externalId Application-specific id used to identify this image during replication.
 		//! \param pos Only for multidimensional parameters: relevant position to set in the ranges
-		void SetImageValue( int index, FName externalId, const Ptr<const RangeIndex>& pos = nullptr);
+		void SetImageValue( int32 index, FName externalId, const Ptr<const RangeIndex>& pos = nullptr);
 
         //! Return the value of a float parameter.
         //! \pre The parameter specified by index is a T_FLOAT.
         //! \param index Index of the parameter from 0 to GetCount()-1
         //! \param pos Only for multidimensional parameters: relevant position to get in the ranges
-        void GetStringValue( int index, FString& OutValue, const Ptr<const RangeIndex>& pos = nullptr) const;
+        void GetStringValue( int32 index, FString& OutValue, const Ptr<const RangeIndex>& pos = nullptr) const;
 
         //! If the parameter is of the float type, set its value.
         //! \param index Index of the parameter from 0 to GetCount()-1
         //! \param value new value of the parameter
         //! \param pos Only for multidimensional parameters: relevant position to set in the ranges
-        void SetStringValue( int index, const FString& Value, const Ptr<const RangeIndex>& pos = nullptr );
+        void SetStringValue( int32 index, const FString& Value, const Ptr<const RangeIndex>& pos = nullptr );
 
 		//! Return the value of a matrix parameter.
 		//! \pre The parameter specified by index is a T_MATRIX.
 		//! \param index Index of the parameter from 0 to GetCount()-1
 		//! \param pos Only for multidimensional parameters: relevant position to get in the ranges
-		void GetMatrixValue( int index, FMatrix44f& OutValue, const Ptr<const RangeIndex>& pos = nullptr) const;
+		void GetMatrixValue( int32 index, FMatrix44f& OutValue, const Ptr<const RangeIndex>& pos = nullptr) const;
 
 		//! If the parameter is of the matrix type, set its value.
 		//! \param index Index of the parameter from 0 to GetCount()-1
 		//! \param value new value of the parameter
 		//! \param pos Only for multidimensional parameters: relevant position to set in the ranges
-		void SetMatrixValue( int index, const FMatrix44f& Value, const Ptr<const RangeIndex>& pos = nullptr );
+		void SetMatrixValue( int32 index, const FMatrix44f& Value, const Ptr<const RangeIndex>& pos = nullptr );
 
         //! Utility method to compare the values of a specific parameter with the values of another
         //! Parameters object. It returns false if type or values are different.
-        bool HasSameValue( int thisParamIndex, const ParametersPtrConst& other, int otherParamIndex ) const;
+        bool HasSameValue( int32 thisParamIndex, const ParametersPtrConst& other, int32 otherParamIndex ) const;
 
 		//-----------------------------------------------------------------------------------------
 		// Interface pattern
diff --git a/Engine/Plugins/Mutable/Source/MutableRuntime/Public/MuR/Ptr.h b/Engine/Plugins/Mutable/Source/MutableRuntime/Public/MuR/Ptr.h
--- a/Engine/Plugins/Mutable/Source/MutableRuntime/Public/MuR/Ptr.h
+++ b/Engine/Plugins/Mutable/Source/MutableRuntime/Public/MuR/Ptr.h
@@ -2,6 +2,8 @@
 
 #pragma once
 
+#include "HAL/Platform.h"
+
 // HEADER_UNIT_SKIP - Included through RefCounted
 
 namespace mu
@@ -121,7 +123,7 @@ namespace mu
 		}
 
 
-        inline explicit operator bool() const
+		FORCEINLINE explicit operator bool() const
 		{
             return m_ptr != nullptr;
 		}
diff --git a/Engine/Plugins/Mutable/Source/MutableRuntime/Public/MuR/RefCounted.h b/Engine/Plugins/Mutable/Source/MutableRuntime/Public/MuR/RefCounted.h
--- a/Engine/Plugins/Mutable/Source/MutableRuntime/Public/MuR/RefCounted.h
+++ b/Engine/Plugins/Mutable/Source/MutableRuntime/Public/MuR/RefCounted.h
@@ -24,12 +24,12 @@ namespace mu
 	{
 	public:
 
-		inline void IncRef() const
+		FORCEINLINE void IncRef() const
 		{
 			std::atomic_fetch_add_explicit( &m_refCount, 1, std::memory_order_relaxed );
 		}
 
-		inline void DecRef() const
+		FORCEINLINE void DecRef() const
 		{
 			LLM_SCOPE_BYNAME(TEXT("MutableRuntime"));
 
@@ -40,7 +40,7 @@ namespace mu
 			}
 		}
 
-		inline bool IsUnique() const
+		FORCEINLINE bool IsUnique() const
 		{
 			return m_refCount.load(std::memory_order_relaxed)==1;
 		}
@@ -57,7 +57,7 @@ namespace mu
 			m_refCount = 0;
 		}
 
-		inline virtual ~RefCounted() = default;
+		FORCEINLINE virtual ~RefCounted() = default;
 
 	private:
 
@@ -66,13 +66,13 @@ namespace mu
 	};
 
 
-	inline void mutable_ptr_add_ref(const RefCounted* p)
+	FORCEINLINE void mutable_ptr_add_ref(const RefCounted* p)
 	{
 		if (p) p->IncRef();
 	}
 
 
-	inline void mutable_ptr_release(const RefCounted* p)
+	FORCEINLINE void mutable_ptr_release(const RefCounted* p)
 	{
 		if (p)
 		{
