[mutable] Cleanup MutableRuntime API moving to use more unreal types and naming standards.

diff --git a/Engine/Plugins/Mutable/Source/MutableRuntime/Public/MuR/Serialisation.h b/Engine/Plugins/Mutable/Source/MutableRuntime/Public/MuR/Serialisation.h
--- a/Engine/Plugins/Mutable/Source/MutableRuntime/Public/MuR/Serialisation.h
+++ b/Engine/Plugins/Mutable/Source/MutableRuntime/Public/MuR/Serialisation.h
@@ -91,7 +91,7 @@ namespace mu
 
         //! Identifier of reading data operations sent to this interface.
 		//! Negative values indicate an error.
-        typedef int32 OPERATION_ID;
+        typedef int32 FOperationID;
 
 		//! \brief Start a data request operation.
 		//! \param Model.
@@ -104,21 +104,21 @@ namespace mu
 		//! \return a previously unused identifier, now used for this operation, that can be used in
 		//! calls to the other methods of this interface. If the return value is negative it indicates
 		//! an unrecoverable error.
-		virtual OPERATION_ID BeginReadBlock(const mu::Model*, uint32 BlockKey, void* Buffer, uint64 size, TFunction<void(bool bSuccess)>* CompletionCallback = nullptr) = 0;
+		virtual FOperationID BeginReadBlock(const mu::Model*, uint32 BlockKey, void* Buffer, uint64 size, TFunction<void(bool bSuccess)>* CompletionCallback = nullptr) = 0;
 
         //! Check if a data request operation has been completed.
         //! This is a weak check than *may* return true if the given operation has completed, but
         //! it is not mandatory. It is used as a hint by the System to optimise its opertaions.
         //! There is no guarantee that this method will ever be called, and it is safe to always
         //! return false.
-        virtual bool IsReadCompleted( OPERATION_ID ) = 0;
+        virtual bool IsReadCompleted(FOperationID) = 0;
 
         /** Complete a data request operation.This method has to block until a data request issued
         * with BeginReadBlock has been completed. After returning from this call, the ID cannot be used
         * any more to identify the same operation and becomes free.
 		* \return true if the data was loaded successfully.
 		*/
-        virtual bool EndRead( OPERATION_ID ) = 0;
+        virtual bool EndRead(FOperationID) = 0;
 
     };
 
diff --git a/Engine/Plugins/Mutable/Source/MutableTools/Internal/MuT/NodeImageFormatPrivate.h b/Engine/Plugins/Mutable/Source/MutableTools/Internal/MuT/NodeImageFormatPrivate.h
--- a/Engine/Plugins/Mutable/Source/MutableTools/Internal/MuT/NodeImageFormatPrivate.h
+++ b/Engine/Plugins/Mutable/Source/MutableTools/Internal/MuT/NodeImageFormatPrivate.h
@@ -17,8 +17,8 @@ namespace mu
 
 		static FNodeType s_type;
 
-		EImageFormat m_format = EImageFormat::IF_NONE;
-		EImageFormat m_formatIfAlpha = EImageFormat::IF_NONE;
+		EImageFormat m_format = EImageFormat::None;
+		EImageFormat m_formatIfAlpha = EImageFormat::None;
         NodeImagePtr m_source;
 
 	};
diff --git a/Engine/Plugins/Mutable/Source/MutableTools/Internal/MuT/NodeImagePlainColourPrivate.h b/Engine/Plugins/Mutable/Source/MutableTools/Internal/MuT/NodeImagePlainColourPrivate.h
--- a/Engine/Plugins/Mutable/Source/MutableTools/Internal/MuT/NodeImagePlainColourPrivate.h
+++ b/Engine/Plugins/Mutable/Source/MutableTools/Internal/MuT/NodeImagePlainColourPrivate.h
@@ -25,7 +25,7 @@ namespace mu
 
 		Ptr<NodeColour> m_pColour;
 		int32 m_sizeX, m_sizeY;
-		EImageFormat Format = EImageFormat::IF_RGBA_UBYTE;
+		EImageFormat Format = EImageFormat::RGBA_UByte;
 
 	};
 
diff --git a/Engine/Plugins/Mutable/Source/MutableTools/Internal/MuT/NodeProjectorPrivate.h b/Engine/Plugins/Mutable/Source/MutableTools/Internal/MuT/NodeProjectorPrivate.h
--- a/Engine/Plugins/Mutable/Source/MutableTools/Internal/MuT/NodeProjectorPrivate.h
+++ b/Engine/Plugins/Mutable/Source/MutableTools/Internal/MuT/NodeProjectorPrivate.h
@@ -25,7 +25,7 @@ namespace mu
 
 		static FNodeType s_type;
 
-        PROJECTOR_TYPE m_type = PROJECTOR_TYPE::PLANAR;
+        EProjectorType m_type = EProjectorType::Planar;
 		FVector3f m_position;
 		FVector3f m_direction;
 		FVector3f m_up;
diff --git a/Engine/Plugins/Mutable/Source/MutableTools/Private/MuT/AST.cpp b/Engine/Plugins/Mutable/Source/MutableTools/Private/MuT/AST.cpp
--- a/Engine/Plugins/Mutable/Source/MutableTools/Private/MuT/AST.cpp
+++ b/Engine/Plugins/Mutable/Source/MutableTools/Private/MuT/AST.cpp
@@ -1102,7 +1102,7 @@ FImageDesc ASTOpFixed::GetImageDesc( bool returnBestOption, FGetImageDescContext
 
     case OP_TYPE::IM_LUMINANCE:
         res = GetImageDesc( op.args.ImageLuminance.base,returnBestOption, context );
-        res.m_format = EImageFormat::IF_L_UBYTE;
+        res.m_format = EImageFormat::L_UByte;
         break;
 
     case OP_TYPE::IM_INTERPOLATE:
@@ -1114,7 +1114,7 @@ FImageDesc ASTOpFixed::GetImageDesc( bool returnBestOption, FGetImageDescContext
         res.m_size[0] = op.args.ImagePlainColour.size[0];
         res.m_size[1] = op.args.ImagePlainColour.size[1];
 		res.m_lods = op.args.ImagePlainColour.LODs;
-        check( res.m_format != EImageFormat::IF_NONE );
+        check( res.m_format != EImageFormat::None );
         break;
 
     case OP_TYPE::IM_RESIZE:
@@ -1144,7 +1144,7 @@ FImageDesc ASTOpFixed::GetImageDesc( bool returnBestOption, FGetImageDescContext
     case OP_TYPE::IM_GRADIENT:
         res.m_size[0] = op.args.ImageGradient.size[0];
         res.m_size[1] = op.args.ImageGradient.size[1];
-        res.m_format = EImageFormat::IF_RGB_UBYTE;
+        res.m_format = EImageFormat::RGB_UByte;
         break;
 
     case OP_TYPE::IM_BLANKLAYOUT:
@@ -1156,7 +1156,7 @@ FImageDesc ASTOpFixed::GetImageDesc( bool returnBestOption, FGetImageDescContext
 
     case OP_TYPE::IM_BINARISE:
         res = GetImageDesc( op.args.ImageBinarise.base, returnBestOption, context );
-        res.m_format = EImageFormat::IF_L_UBYTE;
+        res.m_format = EImageFormat::L_UByte;
         break;
 
     case OP_TYPE::IM_DISPLACE:
diff --git a/Engine/Plugins/Mutable/Source/MutableTools/Private/MuT/ASTOpConditional.cpp b/Engine/Plugins/Mutable/Source/MutableTools/Private/MuT/ASTOpConditional.cpp
--- a/Engine/Plugins/Mutable/Source/MutableTools/Private/MuT/ASTOpConditional.cpp
+++ b/Engine/Plugins/Mutable/Source/MutableTools/Private/MuT/ASTOpConditional.cpp
@@ -172,7 +172,7 @@ namespace mu
 			}
 			else
 			{
-				res.m_format = (yesDesc.m_format == noDesc.m_format) ? yesDesc.m_format : EImageFormat::IF_NONE;
+				res.m_format = (yesDesc.m_format == noDesc.m_format) ? yesDesc.m_format : EImageFormat::None;
 				res.m_lods = (yesDesc.m_lods == noDesc.m_lods) ? yesDesc.m_lods : 0;
 				res.m_size = (yesDesc.m_size == noDesc.m_size) ? yesDesc.m_size : FImageSize(0, 0);
 			}
diff --git a/Engine/Plugins/Mutable/Source/MutableTools/Private/MuT/ASTOpImageCompose.cpp b/Engine/Plugins/Mutable/Source/MutableTools/Private/MuT/ASTOpImageCompose.cpp
--- a/Engine/Plugins/Mutable/Source/MutableTools/Private/MuT/ASTOpImageCompose.cpp
+++ b/Engine/Plugins/Mutable/Source/MutableTools/Private/MuT/ASTOpImageCompose.cpp
@@ -264,11 +264,11 @@ mu::Ptr<ASTOp> ASTOpImageCompose::OptimiseSemantic(const FModelOptimizationOptio
 			// \todo isn't this a common operation?
 			if (baseDesc.m_format != blockDesc.m_format
 				&&
-				baseDesc.m_format != EImageFormat::IF_NONE)
+				baseDesc.m_format != EImageFormat::None)
 			{
 				mu::Ptr<ASTOpImagePixelFormat> reformat = new ASTOpImagePixelFormat;
 				reformat->Format = baseDesc.m_format;
-				reformat->FormatIfAlpha = EImageFormat::IF_NONE;
+				reformat->FormatIfAlpha = EImageFormat::None;
 				reformat->Source = at;
 				at = reformat;
 			}
diff --git a/Engine/Plugins/Mutable/Source/MutableTools/Private/MuT/ASTOpImageLayer.cpp b/Engine/Plugins/Mutable/Source/MutableTools/Private/MuT/ASTOpImageLayer.cpp
--- a/Engine/Plugins/Mutable/Source/MutableTools/Private/MuT/ASTOpImageLayer.cpp
+++ b/Engine/Plugins/Mutable/Source/MutableTools/Private/MuT/ASTOpImageLayer.cpp
@@ -433,7 +433,7 @@ namespace mu
 			if (Pass>0 && !at && blendAt && blendAt->GetOpType() == OP_TYPE::IM_SWIZZLE && Flags==0)
 			{
 				const ASTOpImageSwizzle* TypedBlend = static_cast<const ASTOpImageSwizzle*>(blendAt.get());
-				if (TypedBlend->Format==EImageFormat::IF_L_UBYTE
+				if (TypedBlend->Format==EImageFormat::L_UByte
 					&&
 					TypedBlend->SourceChannels[0]==3)
 				{
@@ -441,7 +441,7 @@ namespace mu
 					nop->Flags = Flags | OP::ImageLayerArgs::FLAGS::F_BLENDED_RGB_FROM_ALPHA;
 					nop->blend = TypedBlend->Sources[0].child();
 					auto Desc = nop->blend->GetImageDesc(true);
-					check(Desc.m_format==EImageFormat::IF_RGBA_UBYTE);
+					check(Desc.m_format==EImageFormat::RGBA_UByte);
 					at = nop;					
 				}
 			}
diff --git a/Engine/Plugins/Mutable/Source/MutableTools/Private/MuT/ASTOpImagePixelFormat.cpp b/Engine/Plugins/Mutable/Source/MutableTools/Private/MuT/ASTOpImagePixelFormat.cpp
--- a/Engine/Plugins/Mutable/Source/MutableTools/Private/MuT/ASTOpImagePixelFormat.cpp
+++ b/Engine/Plugins/Mutable/Source/MutableTools/Private/MuT/ASTOpImagePixelFormat.cpp
@@ -195,7 +195,7 @@ namespace mu
 			auto nop = mu::Clone<ASTOpFixed>(sourceAt);
 
 			EImageFormat layoutFormat = (EImageFormat)nop->op.args.ImageBlankLayout.format;
-			if (FormatIfAlpha != EImageFormat::IF_NONE
+			if (FormatIfAlpha != EImageFormat::None
 				&&
 				GetImageFormatData(layoutFormat).Channels > 3)
 			{
@@ -213,7 +213,7 @@ namespace mu
 			auto nop = mu::Clone<ASTOpFixed>(sourceAt);
 
 			EImageFormat layoutFormat = (EImageFormat)nop->op.args.ImagePlainColour.format;
-			if (FormatIfAlpha != EImageFormat::IF_NONE
+			if (FormatIfAlpha != EImageFormat::None
 				&&
 				GetImageFormatData(layoutFormat).Channels > 3)
 			{
@@ -265,7 +265,7 @@ namespace mu
 			res = Source.child()->GetImageDesc(returnBestOption, context);
 		}
 
-		if (FormatIfAlpha != EImageFormat::IF_NONE
+		if (FormatIfAlpha != EImageFormat::None
 			&&
 			GetImageFormatData(res.m_format).Channels > 3)
 		{
@@ -275,7 +275,7 @@ namespace mu
 		{
 			res.m_format = Format;
 		}
-		check(res.m_format != EImageFormat::IF_NONE);
+		check(res.m_format != EImageFormat::None);
 
 
 		// Cache the result
diff --git a/Engine/Plugins/Mutable/Source/MutableTools/Private/MuT/ASTOpImagePixelFormat.h b/Engine/Plugins/Mutable/Source/MutableTools/Private/MuT/ASTOpImagePixelFormat.h
--- a/Engine/Plugins/Mutable/Source/MutableTools/Private/MuT/ASTOpImagePixelFormat.h
+++ b/Engine/Plugins/Mutable/Source/MutableTools/Private/MuT/ASTOpImagePixelFormat.h
@@ -18,8 +18,8 @@ namespace mu
 	public:
 
 		ASTChild Source;
-		EImageFormat Format = EImageFormat::IF_NONE;
-		EImageFormat FormatIfAlpha = EImageFormat::IF_NONE;
+		EImageFormat Format = EImageFormat::None;
+		EImageFormat FormatIfAlpha = EImageFormat::None;
 
 	public:
 
diff --git a/Engine/Plugins/Mutable/Source/MutableTools/Private/MuT/ASTOpImageSwizzle.cpp b/Engine/Plugins/Mutable/Source/MutableTools/Private/MuT/ASTOpImageSwizzle.cpp
--- a/Engine/Plugins/Mutable/Source/MutableTools/Private/MuT/ASTOpImageSwizzle.cpp
+++ b/Engine/Plugins/Mutable/Source/MutableTools/Private/MuT/ASTOpImageSwizzle.cpp
@@ -157,7 +157,7 @@ namespace mu
 				if (formatSource)
 				{
 					FImageDesc desc = formatSource->GetImageDesc();
-					if (desc.m_format != EImageFormat::IF_NONE && !IsCompressedFormat(desc.m_format))
+					if (desc.m_format != EImageFormat::None && !IsCompressedFormat(desc.m_format))
 					{
 						if (!sat)
 						{
@@ -271,7 +271,7 @@ namespace mu
 			// a special multilayer blend mode. This happens often because of higher level group projector nodes.
 			if (!at
 				&&
-				Format == EImageFormat::IF_RGBA_UBYTE
+				Format == EImageFormat::RGBA_UByte
 				&&
 				Sources[0] == Sources[1] && Sources[0] == Sources[2]
 				&&
@@ -315,7 +315,7 @@ namespace mu
 
 					if ( NewMultiLayer->mask.child() == AlphaMultiLayer->blend.child()
 						&&
-						NewBlended->Format==EImageFormat::IF_RGBA_UBYTE
+						NewBlended->Format==EImageFormat::RGBA_UByte
 						)
 					{
 						// Additional optimization is possible here.
@@ -330,7 +330,7 @@ namespace mu
 			// If we have RGB being the same IM_LAYER, and alpha a compatible IM_LAYER we can optimize with a special layer blend mode.
 			if (!at
 				&&
-				Format == EImageFormat::IF_RGBA_UBYTE
+				Format == EImageFormat::RGBA_UByte
 				&&
 				Sources[0] == Sources[1] && (Sources[0] == Sources[2] || !Sources[2])
 				&&
@@ -372,7 +372,7 @@ namespace mu
 
 					if (NewLayer->mask.child() == AlphaLayer->blend.child()
 						&&
-						NewBlended->Format == EImageFormat::IF_RGBA_UBYTE
+						NewBlended->Format == EImageFormat::RGBA_UByte
 						)
 					{
 						// Additional optimization is possible here.
@@ -969,22 +969,22 @@ namespace mu
 				if (OldFormat->Source->GetOpType() == OP_TYPE::IM_SWIZZLE
 					&&
 					(
-						OldFormat->Format == EImageFormat::IF_RGB_UBYTE
+						OldFormat->Format == EImageFormat::RGB_UByte
 						||
-						OldFormat->Format == EImageFormat::IF_RGBA_UBYTE
+						OldFormat->Format == EImageFormat::RGBA_UByte
 						) 
 					)
 				{
 					const ASTOpImageSwizzle* OldChildSwizzle = static_cast<const ASTOpImageSwizzle*>(OldFormat->Source.child().get());
-					if (OldChildSwizzle->Format == EImageFormat::IF_L_UBYTE)
+					if (OldChildSwizzle->Format == EImageFormat::L_UByte)
 					{
 						Ptr<ASTOpImageSwizzle> NewBaseSwizzle = new ASTOpImageSwizzle;
-						NewBaseSwizzle->Format = EImageFormat::IF_L_UBYTE;
+						NewBaseSwizzle->Format = EImageFormat::L_UByte;
 						NewBaseSwizzle->Sources[0] = OldLayer->base.child();
 						NewBaseSwizzle->SourceChannels[0] = SourceChannels[3];
 
 						Ptr<ASTOpImageSwizzle> NewBlendSwizzle = new ASTOpImageSwizzle;
-						NewBlendSwizzle->Format = EImageFormat::IF_L_UBYTE;
+						NewBlendSwizzle->Format = EImageFormat::L_UByte;
 						NewBlendSwizzle->Sources[0] = OldLayer->blend.child();
 						NewBlendSwizzle->SourceChannels[0] = SourceChannels[3];
 
@@ -1042,7 +1042,7 @@ namespace mu
 		{
 			res = Sources[FirstValidSourceIndex]->GetImageDesc(returnBestOption, context);
 			res.m_format = Format;
-			check(res.m_format != EImageFormat::IF_NONE);
+			check(res.m_format != EImageFormat::None);
 		}
 
 		// Cache the result
diff --git a/Engine/Plugins/Mutable/Source/MutableTools/Private/MuT/ASTOpImageSwizzle.h b/Engine/Plugins/Mutable/Source/MutableTools/Private/MuT/ASTOpImageSwizzle.h
--- a/Engine/Plugins/Mutable/Source/MutableTools/Private/MuT/ASTOpImageSwizzle.h
+++ b/Engine/Plugins/Mutable/Source/MutableTools/Private/MuT/ASTOpImageSwizzle.h
@@ -21,7 +21,7 @@ namespace mu
 
 		uint8 SourceChannels[MUTABLE_OP_MAX_SWIZZLE_CHANNELS] = { 0,0,0,0 };
 
-		EImageFormat Format = EImageFormat::IF_NONE;
+		EImageFormat Format = EImageFormat::None;
 
 	public:
 
diff --git a/Engine/Plugins/Mutable/Source/MutableTools/Private/MuT/ASTOpSwitch.cpp b/Engine/Plugins/Mutable/Source/MutableTools/Private/MuT/ASTOpSwitch.cpp
--- a/Engine/Plugins/Mutable/Source/MutableTools/Private/MuT/ASTOpSwitch.cpp
+++ b/Engine/Plugins/Mutable/Source/MutableTools/Private/MuT/ASTOpSwitch.cpp
@@ -266,7 +266,7 @@ namespace mu
 		// Format will be the most generic from all cases.
 		if (!bSameFormat && !bReturnBestOption)
 		{
-			Result.m_format = EImageFormat::IF_NONE;
+			Result.m_format = EImageFormat::None;
 		}
 
 		if (!bSameSize && !bReturnBestOption)
diff --git a/Engine/Plugins/Mutable/Source/MutableTools/Private/MuT/CodeGenerator.cpp b/Engine/Plugins/Mutable/Source/MutableTools/Private/MuT/CodeGenerator.cpp
--- a/Engine/Plugins/Mutable/Source/MutableTools/Private/MuT/CodeGenerator.cpp
+++ b/Engine/Plugins/Mutable/Source/MutableTools/Private/MuT/CodeGenerator.cpp
@@ -320,8 +320,8 @@ namespace mu
             param.m_name = CacheKey.Table->GetName();
         }
         
-		param.m_type = PARAMETER_TYPE::T_INT;
-        param.m_defaultValue.Set<ParamIntType>(0);
+		param.m_type = EParameterType::Int;
+        param.m_defaultValue.Set<FParamIntType>(0);
 
 		// Add the possible values
 		{
@@ -343,7 +343,7 @@ namespace mu
 				nullValue.m_value = -1;
 				nullValue.m_name = "None";
 				param.m_possibleValues.Add(nullValue);
-				param.m_defaultValue.Set<ParamIntType>(nullValue.m_value);
+				param.m_defaultValue.Set<FParamIntType>(nullValue.m_value);
 			}
 
 			// Add every row
@@ -365,13 +365,13 @@ namespace mu
 				// Set the first row as the default one (if there is none option)
 				if (RowIndex == 0 && !bAddNoneOption)
 				{
-					param.m_defaultValue.Set<ParamIntType>(value.m_value);
+					param.m_defaultValue.Set<FParamIntType>(value.m_value);
 				}
 				
 				// Set the selected row as default (if exists)
 				if (value.m_name == DefaultRowName)
 				{
-					param.m_defaultValue.Set<ParamIntType>(value.m_value);
+					param.m_defaultValue.Set<FParamIntType>(value.m_value);
 				}
             }
         }
@@ -454,7 +454,7 @@ namespace mu
 			}
 			else
 			{
-				BlendOp = GenerateMissingImageCode(TEXT("Patch top image"), EImageFormat::IF_RGB_UBYTE, nullptr, ImageOptions);
+				BlendOp = GenerateMissingImageCode(TEXT("Patch top image"), EImageFormat::RGB_UByte, nullptr, ImageOptions);
 			}
 			BlendOp = GenerateImageFormat(BlendOp, InBlockOp->GetImageDesc().m_format);
 			BlendOp = GenerateImageSize(BlendOp, ImageOptions.RectSize);
@@ -492,7 +492,7 @@ namespace mu
 			{
 				MaskOp = RectConstantOp;
 			}
-			MaskOp = GenerateImageFormat(MaskOp, EImageFormat::IF_L_UBYTE);
+			MaskOp = GenerateImageFormat(MaskOp, EImageFormat::L_UByte);
 			MaskOp = GenerateImageSize(MaskOp, ImageOptions.RectSize);
 			LayerOp->mask = MaskOp;
 
@@ -817,7 +817,7 @@ namespace mu
 				FInt32Point BlockSize = BlockRectInPixels.Size();
 				if (!PatchMask)
 				{
-					PatchMask = new mu::Image(BlockSize[0], BlockSize[1], 1, mu::EImageFormat::IF_L_UBYTE, mu::EInitializationType::Black);
+					PatchMask = new mu::Image(BlockSize[0], BlockSize[1], 1, mu::EImageFormat::L_UByte, mu::EInitializationType::Black);
 				}
 
 				uint8* Pixels = PatchMask->GetMipData(0);
@@ -1236,7 +1236,7 @@ namespace mu
 							fop->SourceChannels[1] = swizzleNode->GetPrivate()->m_sourceChannels[1];
 							fop->SourceChannels[2] = swizzleNode->GetPrivate()->m_sourceChannels[2];
 							fop->SourceChannels[3] = swizzleNode->GetPrivate()->m_sourceChannels[3];
-							check(fop->Format != EImageFormat::IF_NONE);
+							check(fop->Format != EImageFormat::None);
 							imageAd = fop;
 						}
 
@@ -1258,7 +1258,7 @@ namespace mu
 							fop->Format = formatNode->GetPrivate()->m_format;
 							fop->FormatIfAlpha = formatNode->GetPrivate()->m_formatIfAlpha;
 							fop->Source = imageAd;
-							check(fop->Format != EImageFormat::IF_NONE);
+							check(fop->Format != EImageFormat::None);
 							imageAd = fop;
 						}
 
@@ -1298,7 +1298,7 @@ namespace mu
 							if (formatNode)
 							{
 								LayoutBlockDesc.FinalFormat = formatNode->GetPrivate()->m_formatIfAlpha;
-								if (LayoutBlockDesc.FinalFormat == EImageFormat::IF_NONE)
+								if (LayoutBlockDesc.FinalFormat == EImageFormat::None)
 								{
 									LayoutBlockDesc.FinalFormat = formatNode->GetPrivate()->m_format;
 								}
@@ -1408,13 +1408,13 @@ namespace mu
 									fop->Sources[ChannelIndex] = imageAd;
 									fop->SourceChannels[ChannelIndex] = swizzleNode->GetPrivate()->m_sourceChannels[ChannelIndex];
 								}
-								check(fop->Format != EImageFormat::IF_NONE);
+								check(fop->Format != EImageFormat::None);
 								imageAd = fop;
 							}
 
-							// Apply mipmap and format if necessary, skip if format is IF_NONE (possibly because a block was skipped above)
+							// Apply mipmap and format if necessary, skip if format is None (possibly because a block was skipped above)
 							bool bNeedsMips =
-								(mipmapNode && LayoutBlockDesc.FinalFormat != EImageFormat::IF_NONE)
+								(mipmapNode && LayoutBlockDesc.FinalFormat != EImageFormat::None)
 								||
 								LayoutBlockDesc.bBlocksHaveMips;
 
@@ -1461,7 +1461,7 @@ namespace mu
 								fop->Format = formatNode->GetPrivate()->m_format;
 								fop->FormatIfAlpha = formatNode->GetPrivate()->m_formatIfAlpha;
 								fop->Source = imageAd;
-								check(fop->Format != EImageFormat::IF_NONE);
+								check(fop->Format != EImageFormat::None);
 								imageAd = fop;
 							}
 
@@ -2328,7 +2328,7 @@ namespace mu
 					GenerateImage(ClipOptions, ClipMaskResult, TypedClipNode->ClipMask);
 
 					// It could be IF_L_UBIT, but since this should be optimized out at compile time, leave the most cpu efficient.
-					op->MaskImage = GenerateImageFormat(ClipMaskResult.op, mu::EImageFormat::IF_L_UBYTE);
+					op->MaskImage = GenerateImageFormat(ClipMaskResult.op, mu::EImageFormat::L_UByte);
 
 					if (!op->MaskImage)
 					{
@@ -2621,7 +2621,7 @@ namespace mu
 			Out.BlockPixelsY = FMath::Max(1, BlockDesc.m_size[1] / LayoutCellSize[1]);
 			Out.bBlocksHaveMips = BlockDesc.m_lods > 1;
 
-			if (Out.FinalFormat==EImageFormat::IF_NONE)
+			if (Out.FinalFormat==EImageFormat::None)
 			{
 				Out.FinalFormat = BlockDesc.m_format;
 			}
diff --git a/Engine/Plugins/Mutable/Source/MutableTools/Private/MuT/CodeGenerator.h b/Engine/Plugins/Mutable/Source/MutableTools/Private/MuT/CodeGenerator.h
--- a/Engine/Plugins/Mutable/Source/MutableTools/Private/MuT/CodeGenerator.h
+++ b/Engine/Plugins/Mutable/Source/MutableTools/Private/MuT/CodeGenerator.h
@@ -249,7 +249,7 @@ namespace mu
 		TMap<uint64,TArray<FGeneratedConstantMesh>> GeneratedConstantMeshes;
 
         /** List of image resources for every image formats that have been generated so far as palceholders for missing images. */
-        Ptr<Image> MissingImage[size_t(EImageFormat::IF_COUNT)];
+        Ptr<Image> MissingImage[size_t(EImageFormat::Count)];
 
 
         //! List of already used vertex ID groups that must be unique.
@@ -680,7 +680,7 @@ namespace mu
         struct FProjectorGenerationResult
         {
             Ptr<ASTOp> op;
-            PROJECTOR_TYPE type;
+            EProjectorType type;
         };
 
         typedef TMap<FGeneratedCacheKey,FProjectorGenerationResult> FGeneratedProjectorsMap;
@@ -814,7 +814,7 @@ namespace mu
 
 		struct FLayoutBlockDesc
 		{
-			EImageFormat FinalFormat = EImageFormat::IF_NONE;
+			EImageFormat FinalFormat = EImageFormat::None;
 			int32 BlockPixelsX = 0;
 			int32 BlockPixelsY = 0;
 			bool bBlocksHaveMips = false;
diff --git a/Engine/Plugins/Mutable/Source/MutableTools/Private/MuT/CodeGenerator_Bool.cpp b/Engine/Plugins/Mutable/Source/MutableTools/Private/MuT/CodeGenerator_Bool.cpp
--- a/Engine/Plugins/Mutable/Source/MutableTools/Private/MuT/CodeGenerator_Bool.cpp
+++ b/Engine/Plugins/Mutable/Source/MutableTools/Private/MuT/CodeGenerator_Bool.cpp
@@ -94,8 +94,8 @@ namespace mu
 			param.m_name = node.m_name;
 			const TCHAR* CStr = ToCStr(node.m_uid);
 			param.m_uid.ImportTextItem(CStr, 0, nullptr, nullptr);
-			param.m_type = PARAMETER_TYPE::T_BOOL;
-			param.m_defaultValue.Set<ParamBoolType>(node.m_defaultValue);
+			param.m_type = EParameterType::Bool;
+			param.m_defaultValue.Set<FParamBoolType>(node.m_defaultValue);
 
 			op = new ASTOpParameter();
 			op->type = OP_TYPE::BO_PARAMETER;
diff --git a/Engine/Plugins/Mutable/Source/MutableTools/Private/MuT/CodeGenerator_Color.cpp b/Engine/Plugins/Mutable/Source/MutableTools/Private/MuT/CodeGenerator_Color.cpp
--- a/Engine/Plugins/Mutable/Source/MutableTools/Private/MuT/CodeGenerator_Color.cpp
+++ b/Engine/Plugins/Mutable/Source/MutableTools/Private/MuT/CodeGenerator_Color.cpp
@@ -147,15 +147,15 @@ namespace mu
 			param.m_name = node.Name;
 			const TCHAR* CStr = ToCStr(node.Uid);
 			param.m_uid.ImportTextItem(CStr, 0, nullptr, nullptr);
-			param.m_type = PARAMETER_TYPE::T_COLOUR;
+			param.m_type = EParameterType::Color;
 
-			ParamColorType Value;
+			FParamColorType Value;
 			Value[0] =  node.DefaultValue[0];
 			Value[1] = node.DefaultValue[1];
 			Value[2] = node.DefaultValue[2];
 			Value[3] = node.DefaultValue[3];
 
-			param.m_defaultValue.Set<ParamColorType>(Value);
+			param.m_defaultValue.Set<FParamColorType>(Value);
 			
 			op = new ASTOpParameter();
 			op->type = OP_TYPE::CO_PARAMETER;
@@ -317,9 +317,9 @@ namespace mu
 		else
 		{
 			// This argument is required
-			base = GenerateMissingImageCode(TEXT("Sample image"), EImageFormat::IF_RGB_UBYTE, Typed->GetMessageContext(), ImageOptions);
+			base = GenerateMissingImageCode(TEXT("Sample image"), EImageFormat::RGB_UByte, Typed->GetMessageContext(), ImageOptions);
 		}
-		base = GenerateImageFormat(base, EImageFormat::IF_RGB_UBYTE);
+		base = GenerateImageFormat(base, EImageFormat::RGB_UByte);
 		op->SetChild(op->op.args.ColourSampleImage.image, base);
 
 		FScalarGenerationResult ChildResult;
diff --git a/Engine/Plugins/Mutable/Source/MutableTools/Private/MuT/CodeGenerator_FirstPass.cpp b/Engine/Plugins/Mutable/Source/MutableTools/Private/MuT/CodeGenerator_FirstPass.cpp
--- a/Engine/Plugins/Mutable/Source/MutableTools/Private/MuT/CodeGenerator_FirstPass.cpp
+++ b/Engine/Plugins/Mutable/Source/MutableTools/Private/MuT/CodeGenerator_FirstPass.cpp
@@ -624,8 +624,8 @@ namespace mu
             op->parameter.m_name = Private->Name;
 			const TCHAR* CStr = ToCStr(Private->Uid);
 			op->parameter.m_uid.ImportTextItem(CStr, 0, nullptr, nullptr);
-            op->parameter.m_type = PARAMETER_TYPE::T_INT;
-            op->parameter.m_defaultValue.Set<ParamIntType>(Private->DefaultValue);
+            op->parameter.m_type = EParameterType::Int;
+            op->parameter.m_defaultValue.Set<FParamIntType>(Private->DefaultValue);
 
             if (Private->m_type==NodeObjectGroup::CS_ONE_OR_NONE )
             {
@@ -667,8 +667,8 @@ namespace mu
 							op->parameter.m_name = pChildNode->GetName();
 							const TCHAR* CStr = ToCStr(pChildNode->GetUid());
 							op->parameter.m_uid.ImportTextItem(CStr, 0, nullptr, nullptr);
-							op->parameter.m_type = PARAMETER_TYPE::T_BOOL;
-							op->parameter.m_defaultValue.Set<ParamBoolType>(false);
+							op->parameter.m_type = EParameterType::Bool;
+							op->parameter.m_defaultValue.Set<FParamBoolType>(false);
 
 							ParameterNodes.Add(pChildNode, op);
 
diff --git a/Engine/Plugins/Mutable/Source/MutableTools/Private/MuT/CodeGenerator_Image.cpp b/Engine/Plugins/Mutable/Source/MutableTools/Private/MuT/CodeGenerator_Image.cpp
--- a/Engine/Plugins/Mutable/Source/MutableTools/Private/MuT/CodeGenerator_Image.cpp
+++ b/Engine/Plugins/Mutable/Source/MutableTools/Private/MuT/CodeGenerator_Image.cpp
@@ -179,7 +179,7 @@ namespace mu
         if (!pImage)
         {
             // This data is required
-            pImage = GenerateMissingImage(EImageFormat::IF_RGB_UBYTE );
+            pImage = GenerateMissingImage(EImageFormat::RGB_UByte );
 
             // Log an error message
             ErrorLog->GetPrivate()->Add( "Constant image not set.", ELMT_WARNING, InNode->GetMessageContext());
@@ -200,7 +200,7 @@ namespace mu
 			// channel if not used.
 			// TODO: review this, probably the reference descriptor generation needs to check for alpha channels 
 			// as well.
-			ReferenceOp->ImageDesc = FImageDesc(pImage->GetSize(), EImageFormat::IF_NONE, pImage->GetLODCount());
+			ReferenceOp->ImageDesc = FImageDesc(pImage->GetSize(), EImageFormat::None, pImage->GetLODCount());
 			Result.op = ReferenceOp;
 		}
 		else
@@ -286,8 +286,8 @@ namespace mu
 			op->parameter.m_name = node.m_name;
         	const TCHAR* CStr = ToCStr(node.m_uid);
         	op->parameter.m_uid.ImportTextItem(CStr, 0, nullptr, nullptr);
-        	op->parameter.m_type = PARAMETER_TYPE::T_IMAGE;
-        	op->parameter.m_defaultValue.Set<ParamImageType>(node.m_defaultValue);
+        	op->parameter.m_type = EParameterType::Image;
+        	op->parameter.m_defaultValue.Set<FParamImageType>(node.m_defaultValue);
 
 			// Generate the code for the ranges
 			for (int32 a = 0; a < node.m_ranges.Num(); ++a)
@@ -330,7 +330,7 @@ namespace mu
         else
         {
             // This argument is required
-            base = GenerateMissingImageCode(TEXT("Image Layer base"), EImageFormat::IF_RGB_UBYTE, InNode->GetMessageContext(), Options);
+            base = GenerateMissingImageCode(TEXT("Image Layer base"), EImageFormat::RGB_UByte, InNode->GetMessageContext(), Options);
         }
 
 		UE::Math::TIntVector2<int32> TargetSize = Options.RectSize;
@@ -357,7 +357,7 @@ namespace mu
 			GenerateImage(MaskOptions, MaskResult, node.m_pMask);
 			mask = MaskResult.op;
 
-            mask = GenerateImageFormat( mask, EImageFormat::IF_L_UBYTE );
+            mask = GenerateImageFormat( mask, EImageFormat::L_UByte );
             mask = GenerateImageSize( mask, TargetSize);
         }
         op->mask = mask;
@@ -406,7 +406,7 @@ namespace mu
         else
         {
             // This argument is required
-            base = GenerateMissingImageCode(TEXT("Layer base image"), EImageFormat::IF_RGB_UBYTE, InNode->GetMessageContext(), Options );
+            base = GenerateMissingImageCode(TEXT("Layer base image"), EImageFormat::RGB_UByte, InNode->GetMessageContext(), Options );
         }
 
 		UE::Math::TIntVector2<int32> TargetSize = Options.RectSize;
@@ -418,7 +418,7 @@ namespace mu
 			TargetSize = UE::Math::TIntVector2<int32>(BaseDesc.m_size);
 		}
 
-        base = GenerateImageFormat( base, EImageFormat::IF_RGB_UBYTE );
+        base = GenerateImageFormat( base, EImageFormat::RGB_UByte );
         base = GenerateImageSize( base, TargetSize);
         op->base = base;
 
@@ -430,7 +430,7 @@ namespace mu
 			GenerateImage(Options, MaskResult, node.m_pMask);
 			mask = MaskResult.op;
 			
-			mask = GenerateImageFormat( mask, EImageFormat::IF_L_UBYTE );
+			mask = GenerateImageFormat( mask, EImageFormat::L_UByte );
             mask = GenerateImageSize( mask, TargetSize);
         }
         op->mask = mask;
@@ -476,7 +476,7 @@ namespace mu
         else
         {
             // This argument is required
-            base = GenerateMissingImageCode(TEXT("Image MultiLayer base"), EImageFormat::IF_RGB_UBYTE, InNode->GetMessageContext(), Options);
+            base = GenerateMissingImageCode(TEXT("Image MultiLayer base"), EImageFormat::RGB_UByte, InNode->GetMessageContext(), Options);
         }
 
 		FImageDesc BaseDesc = base->GetImageDesc(true);
@@ -501,7 +501,7 @@ namespace mu
 			GenerateImage(Options, MaskResult, node.m_pMask);
 			mask = MaskResult.op;
 
-			mask = GenerateImageFormat( mask, EImageFormat::IF_L_UBYTE );
+			mask = GenerateImageFormat( mask, EImageFormat::L_UByte );
             mask = GenerateImageSize( mask, TargetSize);
         }
         op->mask = mask;
@@ -561,7 +561,7 @@ namespace mu
         else
         {
             // This argument is required
-            base = GenerateMissingImageCode(TEXT("Image Composite Base"), EImageFormat::IF_RGB_UBYTE, InNode->GetMessageContext(), Options);
+            base = GenerateMissingImageCode(TEXT("Image Composite Base"), EImageFormat::RGB_UByte, InNode->GetMessageContext(), Options);
         }
 
 		FImageDesc BaseDesc = base->GetImageDesc(true);
@@ -586,12 +586,12 @@ namespace mu
 			GenerateImage(Options, BaseResult, node.m_pNormal);
 			normal = BaseResult.op;
 
-            normal = GenerateImageFormat( normal, EImageFormat::IF_RGB_UBYTE );
+            normal = GenerateImageFormat( normal, EImageFormat::RGB_UByte );
         }
 		else
 		{
             // This argument is required
-            normal = GenerateMissingImageCode(TEXT("Image Composite Normal"), EImageFormat::IF_RGB_UBYTE, InNode->GetMessageContext(), Options);
+            normal = GenerateMissingImageCode(TEXT("Image Composite Normal"), EImageFormat::RGB_UByte, InNode->GetMessageContext(), Options);
 		}
 
 		normal = GenerateImageSize(normal, TargetSize);
@@ -677,7 +677,7 @@ namespace mu
         else
         {
             // This argument is required
-            Base = GenerateMissingImageCode(TEXT("Image Transform Base"), EImageFormat::IF_RGB_UBYTE, InNode->GetMessageContext(), NewOptions);
+            Base = GenerateMissingImageCode(TEXT("Image Transform Base"), EImageFormat::RGB_UByte, InNode->GetMessageContext(), NewOptions);
         }
 		
 		FImageDesc BaseDesc = Base->GetImageDesc();
@@ -735,7 +735,7 @@ namespace mu
 				}
 
                 // TODO: Support other formats
-                target = GenerateImageFormat( target, EImageFormat::IF_RGB_UBYTE );
+                target = GenerateImageFormat( target, EImageFormat::RGB_UByte );
                 target = GenerateImageSize( target, FinalRectSize);
 
                 op->SetChild( op->op.args.ImageInterpolate.targets[numTargets], target);
@@ -746,7 +746,7 @@ namespace mu
         // At least one target is required
         if (!op->op.args.ImageInterpolate.targets[0])
         {
-            Ptr<ASTOp> target = GenerateMissingImageCode(TEXT("First interpolation image"), EImageFormat::IF_RGB_UBYTE, InNode->GetMessageContext(), Options);
+            Ptr<ASTOp> target = GenerateMissingImageCode(TEXT("First interpolation image"), EImageFormat::RGB_UByte, InNode->GetMessageContext(), Options);
             target = GenerateImageSize( target, Options.RectSize);
             op->SetChild( op->op.args.ImageInterpolate.targets[0], target);
         }
@@ -762,9 +762,9 @@ namespace mu
 
 		//MUTABLE_CPUPROFILER_SCOPE(NodeImageSwizzle);
 
-		if (node.m_format == EImageFormat::IF_NONE)
+		if (node.m_format == EImageFormat::None)
 		{
-			Result.op = GenerateMissingImageCode(TEXT("Make Texture channel."), EImageFormat::IF_L_UBYTE, InNode->GetMessageContext(), Options);
+			Result.op = GenerateMissingImageCode(TEXT("Make Texture channel."), EImageFormat::L_UByte, InNode->GetMessageContext(), Options);
 			return;
 		}
 
@@ -773,35 +773,35 @@ namespace mu
         Ptr<ASTOpImageSwizzle> SwizzleOp = new ASTOpImageSwizzle();
 
 		// Format
-		EImageFormat CompressedFormat = EImageFormat::IF_NONE;
+		EImageFormat CompressedFormat = EImageFormat::None;
 
 		switch (node.m_format)
 		{
-        case EImageFormat::IF_BC1:
-        case EImageFormat::IF_ASTC_4x4_RGB_LDR:
+        case EImageFormat::BC1:
+        case EImageFormat::ASTC_4x4_RGB_LDR:
             CompressedFormat = node.m_format;
-            SwizzleOp->Format = (node.m_sources.IsValidIndex(3) && node.m_sources[3]) ? EImageFormat::IF_RGBA_UBYTE : EImageFormat::IF_RGB_UBYTE;
+            SwizzleOp->Format = (node.m_sources.IsValidIndex(3) && node.m_sources[3]) ? EImageFormat::RGBA_UByte : EImageFormat::RGB_UByte;
 			break;
 
-		case EImageFormat::IF_BC2:
-		case EImageFormat::IF_BC3:
-		case EImageFormat::IF_BC6:
-        case EImageFormat::IF_BC7:
-        case EImageFormat::IF_ASTC_4x4_RGBA_LDR:
+		case EImageFormat::BC2:
+		case EImageFormat::BC3:
+		case EImageFormat::BC6:
+        case EImageFormat::BC7:
+        case EImageFormat::ASTC_4x4_RGBA_LDR:
             CompressedFormat = node.m_format;
-            SwizzleOp->Format = EImageFormat::IF_RGBA_UBYTE;
+            SwizzleOp->Format = EImageFormat::RGBA_UByte;
 			break;
 
-		case EImageFormat::IF_BC4:
+		case EImageFormat::BC4:
 			CompressedFormat = node.m_format;
-            SwizzleOp->Format = EImageFormat::IF_L_UBYTE;
+            SwizzleOp->Format = EImageFormat::L_UByte;
 			break;
 
-		case EImageFormat::IF_BC5:
-        case EImageFormat::IF_ASTC_4x4_RG_LDR:
+		case EImageFormat::BC5:
+        case EImageFormat::ASTC_4x4_RG_LDR:
             CompressedFormat = node.m_format;
 			// TODO: Should be RG
-            SwizzleOp->Format = EImageFormat::IF_RGB_UBYTE;
+            SwizzleOp->Format = EImageFormat::RGB_UByte;
 			break;
 
 		default:
@@ -810,7 +810,7 @@ namespace mu
 
 		}
 
-		check(node.m_format != EImageFormat::IF_NONE);
+		check(node.m_format != EImageFormat::None);
 
 		// Source images and channels
 		check(node.m_sources.Num() == node.m_sourceChannels.Num());
@@ -834,7 +834,7 @@ namespace mu
 				if (!Source)
 				{
 					// TODO: Warn?
-					Source = GenerateMissingImageCode(TEXT("Swizzle channel"), EImageFormat::IF_L_UBYTE, InNode->GetMessageContext(), Options);
+					Source = GenerateMissingImageCode(TEXT("Swizzle channel"), EImageFormat::L_UByte, InNode->GetMessageContext(), Options);
 				}
 
                 Ptr<ASTOp> SizedSource;
@@ -857,14 +857,14 @@ namespace mu
 
         if (FirstValidSourceIndex < 0)
 		{
-            Ptr<ASTOp> Source = GenerateMissingImageCode(TEXT("First swizzle image"), EImageFormat::IF_RGBA_UBYTE, InNode->GetMessageContext(), Options);
+            Ptr<ASTOp> Source = GenerateMissingImageCode(TEXT("First swizzle image"), EImageFormat::RGBA_UByte, InNode->GetMessageContext(), Options);
             SwizzleOp->Sources[0] = Source;
 
 		}
 
         Ptr<ASTOp> ResultOp = SwizzleOp;
 
-		if (CompressedFormat != EImageFormat::IF_NONE)
+		if (CompressedFormat != EImageFormat::None)
 		{
             Ptr<ASTOpImagePixelFormat> FormatOp = new ASTOpImagePixelFormat();
             FormatOp->Source = ResultOp;
@@ -883,7 +883,7 @@ namespace mu
 
 		MUTABLE_CPUPROFILER_SCOPE(NodeImageFormat);
 
-        check(node.m_format != EImageFormat::IF_NONE);
+        check(node.m_format != EImageFormat::None);
 
         Ptr<ASTOpImagePixelFormat> fop = new ASTOpImagePixelFormat();
         fop->Format = node.m_format;
@@ -892,7 +892,7 @@ namespace mu
 		// Source is required
 		if (!node.m_source)
 		{
-            fop->Source = GenerateMissingImageCode(TEXT("Source image for format."), EImageFormat::IF_RGBA_UBYTE, InNode->GetMessageContext(), Options);
+            fop->Source = GenerateMissingImageCode(TEXT("Source image for format."), EImageFormat::RGBA_UByte, InNode->GetMessageContext(), Options);
 		}
 		else
 		{
@@ -925,7 +925,7 @@ namespace mu
         else
         {
             // This argument is required
-            base = GenerateMissingImageCode(TEXT("Saturate image"), EImageFormat::IF_RGB_UBYTE, InNode->GetMessageContext(), Options);
+            base = GenerateMissingImageCode(TEXT("Saturate image"), EImageFormat::RGB_UByte, InNode->GetMessageContext(), Options);
         }
 		
         base = GenerateImageFormat(base, GetRGBOrRGBAFormat(base->GetImageDesc().m_format));
@@ -967,9 +967,9 @@ namespace mu
         else
         {
             // This argument is required
-            base = GenerateMissingImageCode(TEXT("Image luminance"), EImageFormat::IF_RGB_UBYTE, InNode->GetMessageContext(), Options);
+            base = GenerateMissingImageCode(TEXT("Image luminance"), EImageFormat::RGB_UByte, InNode->GetMessageContext(), Options);
         }
-        base = GenerateImageFormat( base, EImageFormat::IF_RGB_UBYTE );
+        base = GenerateImageFormat( base, EImageFormat::RGB_UByte );
         op->SetChild( op->op.args.ImageLuminance.base, base);
 
         Result.op = op;
@@ -995,7 +995,7 @@ namespace mu
         else
         {
             // This argument is required
-            base = GenerateMissingImageCode(TEXT("Colourmap base image"), EImageFormat::IF_RGB_UBYTE, InNode->GetMessageContext(), Options);
+            base = GenerateMissingImageCode(TEXT("Colourmap base image"), EImageFormat::RGB_UByte, InNode->GetMessageContext(), Options);
         }
         base = GenerateImageSize( base, Options.RectSize);
         op->SetChild( op->op.args.ImageColourMap.base, base);
@@ -1014,7 +1014,7 @@ namespace mu
             // TODO: Special operation code without mask
             mask = GeneratePlainImageCode(FVector4f( 1,1,1,1 ), Options);
         }
-        mask = GenerateImageFormat( mask, EImageFormat::IF_L_UBYTE );
+        mask = GenerateImageFormat( mask, EImageFormat::L_UByte );
         mask = GenerateImageSize( mask, Options.RectSize);
         op->SetChild( op->op.args.ImageColourMap.mask, mask);
 
@@ -1031,12 +1031,12 @@ namespace mu
 			FImageGenerationResult BaseResult;
 			GenerateImage(MapOptions, BaseResult, node.m_pMap);
 			MapImageOp = BaseResult.op;
-			MapImageOp = GenerateImageFormat(MapImageOp, EImageFormat::IF_RGB_UBYTE);
+			MapImageOp = GenerateImageFormat(MapImageOp, EImageFormat::RGB_UByte);
 		}
         else
         {
             // This argument is required
-			MapImageOp = GenerateMissingImageCode(TEXT("Map image"), EImageFormat::IF_RGB_UBYTE, InNode->GetMessageContext(), Options);
+			MapImageOp = GenerateMissingImageCode(TEXT("Map image"), EImageFormat::RGB_UByte, InNode->GetMessageContext(), Options);
         }
         op->SetChild( op->op.args.ImageColourMap.map, MapImageOp);
 
@@ -1104,9 +1104,9 @@ namespace mu
         else
         {
             // This argument is required
-            a = GenerateMissingImageCode(TEXT("Image Binarise Base"), EImageFormat::IF_RGB_UBYTE, InNode->GetMessageContext(), Options);
+            a = GenerateMissingImageCode(TEXT("Image Binarise Base"), EImageFormat::RGB_UByte, InNode->GetMessageContext(), Options);
         }
-        a = GenerateImageFormat( a, EImageFormat::IF_RGB_UBYTE );
+        a = GenerateImageFormat( a, EImageFormat::RGB_UByte );
         a = GenerateImageSize( a, Options.RectSize);
         op->SetChild( op->op.args.ImageBinarise.base, a );
 
@@ -1155,7 +1155,7 @@ namespace mu
         else
         {
             // This argument is required
-            base = GenerateMissingImageCode(TEXT("Image resize base"), EImageFormat::IF_RGB_UBYTE, InNode->GetMessageContext(), Options);
+            base = GenerateMissingImageCode(TEXT("Image resize base"), EImageFormat::RGB_UByte, InNode->GetMessageContext(), Options);
         }
 
         // Size
@@ -1274,7 +1274,7 @@ namespace mu
         if (node.m_options.Num() == 0)
 		{
 			// No options in the switch!
-            Ptr<ASTOp> missingOp = GenerateMissingImageCode(TEXT("Switch option"), EImageFormat::IF_RGBA_UBYTE, InNode->GetMessageContext(), Options);
+            Ptr<ASTOp> missingOp = GenerateMissingImageCode(TEXT("Switch option"), EImageFormat::RGBA_UByte, InNode->GetMessageContext(), Options);
 			Result.op = missingOp;
 			return;
 		}
@@ -1307,7 +1307,7 @@ namespace mu
             else
             {
                 // This argument is required
-                branch = GenerateMissingImageCode(TEXT("Switch option"), EImageFormat::IF_RGBA_UBYTE, InNode->GetMessageContext(), Options);
+                branch = GenerateMissingImageCode(TEXT("Switch option"), EImageFormat::RGBA_UByte, InNode->GetMessageContext(), Options);
             }
 
             op->cases.Emplace((int16_t)t,op,branch);
@@ -1319,11 +1319,11 @@ namespace mu
 		// Disabled: This is not always desirable. For example if the image is going to be used in a 
 		// projector, the size doesn't need to be constrained.
         //auto desc = switchAt->GetImageDesc( true );
-        //if ( desc.m_format == EImageFormat::IF_NONE )
+        //if ( desc.m_format == EImageFormat::None )
         //{
         //    // TODO: Look for the most generic of the options?
         //    // For now force a decently generic one
-        //    desc.m_format = EImageFormat::IF_RGBA_UBYTE;
+        //    desc.m_format = EImageFormat::RGBA_UByte;
         //}
 
         //if (desc.m_size[0]!=0 && desc.m_size[1]!=0)
@@ -1454,7 +1454,7 @@ namespace mu
 			Ptr<ASTOp> CurrentMeshToProjectOp = MeshResult.MeshOp;
 			//Ptr<ASTOp> CurrentMeshToProjectOp = MeshResult.BaseMeshOp;
 
-            if (projectorResult.type == PROJECTOR_TYPE::WRAPPING)
+            if (projectorResult.type == EProjectorType::Wrapping)
             {
                 // For wrapping projector we need the entire mesh. The actual project operation
                 // will remove the faces that are not in the layout block we are generating.
@@ -1552,7 +1552,7 @@ namespace mu
 			GenerateImage(Options, MaskResult, node.m_pMask);
 			Ptr<ASTOp> mask = MaskResult.op;
 
-			mask = GenerateImageFormat(mask, EImageFormat::IF_L_UBYTE);
+			mask = GenerateImageFormat(mask, EImageFormat::L_UByte);
 
 			ImageRasterOp->mask = GenerateImageSize(mask, RasterImageSize);
 
@@ -1591,7 +1591,7 @@ namespace mu
         else
         {
             // This argument is required
-			ImageRasterOp->image = GenerateMissingImageCode(TEXT("Projector image"), EImageFormat::IF_RGB_UBYTE, InNode->GetMessageContext(), Options);
+			ImageRasterOp->image = GenerateMissingImageCode(TEXT("Projector image"), EImageFormat::RGB_UByte, InNode->GetMessageContext(), Options);
         }
 
 		if (RasterImageSize.X == 0)
@@ -1720,7 +1720,7 @@ namespace mu
         else
         {
             // This argument is required
-            base = GenerateMissingImageCode(TEXT("Mipmap image"), EImageFormat::IF_RGB_UBYTE, InNode->GetMessageContext(), Options);
+            base = GenerateMissingImageCode(TEXT("Mipmap image"), EImageFormat::RGB_UByte, InNode->GetMessageContext(), Options);
         }
 
         op->Source = base;
@@ -1777,9 +1777,9 @@ namespace mu
 		else
 		{
 			// This argument is required
-			a = GenerateMissingImageCode(TEXT("Image Invert Color"), EImageFormat::IF_RGB_UBYTE, InNode->GetMessageContext(), Options);
+			a = GenerateMissingImageCode(TEXT("Image Invert Color"), EImageFormat::RGB_UByte, InNode->GetMessageContext(), Options);
 		}
-		a = GenerateImageFormat(a, EImageFormat::IF_RGB_UBYTE);
+		a = GenerateImageFormat(a, EImageFormat::RGB_UByte);
 		a = GenerateImageSize(a, Options.RectSize);
 		op->SetChild(op->op.args.ImageInvert.base, a);
 
@@ -1833,7 +1833,7 @@ namespace mu
             else
             {
                 // This argument is required
-                variationOp = GenerateMissingImageCode(TEXT("Variation option"), EImageFormat::IF_RGBA_UBYTE, InNode->GetMessageContext(), Options);
+                variationOp = GenerateMissingImageCode(TEXT("Variation option"), EImageFormat::RGBA_UByte, InNode->GetMessageContext(), Options);
             }
 
 
@@ -1853,11 +1853,11 @@ namespace mu
 			desc = currentOp->GetImageDesc(true);
 		}
 
-        if ( desc.m_format == EImageFormat::IF_NONE )
+        if ( desc.m_format == EImageFormat::None )
         {
             // TODO: Look for the most generic of the options?
             // For now force a decently generic one
-            desc.m_format = EImageFormat::IF_RGBA_UBYTE;
+            desc.m_format = EImageFormat::RGBA_UByte;
         }
 
         if ( desc.m_size[0] != 0 && desc.m_size[1] != 0 )
@@ -1955,7 +1955,7 @@ namespace mu
 
 			switch (Format)
 			{
-			case EImageFormat::IF_L_UBYTE:
+			case EImageFormat::L_UByte:
 			{
 				uint8* DataPtr = GeneratedImage->GetLODData(0);
 				for (int32 P = 0; P < Size[0]*Size[1]; ++P)
@@ -1973,7 +1973,7 @@ namespace mu
 				}
 				break;
 			}
-			case EImageFormat::IF_RGB_UBYTE:
+			case EImageFormat::RGB_UByte:
 			{
 				uint8* DataPtr = GeneratedImage->GetLODData(0);
 				for (int32 P = 0; P < Size[0]*Size[1]; ++P)
@@ -1995,8 +1995,8 @@ namespace mu
 				}
 				break;
 			}
-			case EImageFormat::IF_BGRA_UBYTE:
-			case EImageFormat::IF_RGBA_UBYTE:
+			case EImageFormat::BGRA_UByte:
+			case EImageFormat::RGBA_UByte:
 			{
 				uint8* DataPtr = GeneratedImage->GetLODData(0);
 				for (int32 P = 0; P < Size[0]*Size[1]; ++P)
@@ -2075,7 +2075,7 @@ namespace mu
     {
         Ptr<ASTOp> Result = Op;
 
-        if (InFormat!=EImageFormat::IF_NONE && Op && Op->GetImageDesc().m_format!=InFormat)
+        if (InFormat!=EImageFormat::None && Op && Op->GetImageDesc().m_format!=InFormat)
         {
             // Generate the format change code
             Ptr<ASTOpImagePixelFormat> op = new ASTOpImagePixelFormat();
diff --git a/Engine/Plugins/Mutable/Source/MutableTools/Private/MuT/CodeGenerator_Matrix.cpp b/Engine/Plugins/Mutable/Source/MutableTools/Private/MuT/CodeGenerator_Matrix.cpp
--- a/Engine/Plugins/Mutable/Source/MutableTools/Private/MuT/CodeGenerator_Matrix.cpp
+++ b/Engine/Plugins/Mutable/Source/MutableTools/Private/MuT/CodeGenerator_Matrix.cpp
@@ -58,8 +58,8 @@ namespace mu
 			param.m_name = Typed->Name;
 			const TCHAR* CStr = ToCStr(Typed->Uid);
 			param.m_uid.ImportTextItem(CStr, 0, nullptr, nullptr);
-			param.m_type = PARAMETER_TYPE::T_MATRIX;
-			param.m_defaultValue.Set<ParamMatrixType>(Typed->DefaultValue);
+			param.m_type = EParameterType::Matrix;
+			param.m_defaultValue.Set<FParamMatrixType>(Typed->DefaultValue);
 
 			op = new ASTOpParameter();
 			op->type = OP_TYPE::MA_PARAMETER;
diff --git a/Engine/Plugins/Mutable/Source/MutableTools/Private/MuT/CodeGenerator_Mesh.cpp b/Engine/Plugins/Mutable/Source/MutableTools/Private/MuT/CodeGenerator_Mesh.cpp
--- a/Engine/Plugins/Mutable/Source/MutableTools/Private/MuT/CodeGenerator_Mesh.cpp
+++ b/Engine/Plugins/Mutable/Source/MutableTools/Private/MuT/CodeGenerator_Mesh.cpp
@@ -201,12 +201,12 @@ namespace mu
 		// Get the information about the texture coordinates channel
 		int32 TexCoordsBufferIndex = -1;
 		int32 TexCoordsChannelIndex = -1;
-		Mesh->GetVertexBuffers().FindChannel(MBS_TEXCOORDS, LayoutChannel, &TexCoordsBufferIndex, &TexCoordsChannelIndex);
+		Mesh->GetVertexBuffers().FindChannel(EMeshBufferSemantic::TexCoords, LayoutChannel, &TexCoordsBufferIndex, &TexCoordsChannelIndex);
 		check(TexCoordsBufferIndex >= 0);
 		check(TexCoordsChannelIndex >= 0);
 
 		const FMeshBufferChannel& TexCoordsChannel = Mesh->VertexBuffers.Buffers[TexCoordsBufferIndex].Channels[TexCoordsChannelIndex];
-		check(TexCoordsChannel.Semantic == MBS_TEXCOORDS);
+		check(TexCoordsChannel.Semantic == EMeshBufferSemantic::TexCoords);
 
 		uint8* TexCoordData = Mesh->GetVertexBuffers().GetBufferData(TexCoordsBufferIndex);
 		int32 elemSize = Mesh->GetVertexBuffers().GetElementSize(TexCoordsBufferIndex);
@@ -225,11 +225,11 @@ namespace mu
 			for (int32 VertexIndex = 0; VertexIndex < NumVertices; ++VertexIndex)
 			{
 				FVector2f& UV = TexCoords[VertexIndex];
-				if (TexCoordsChannel.Format == MBF_FLOAT32)
+				if (TexCoordsChannel.Format == EMeshBufferFormat::Float32)
 				{
 					UV = *((FVector2f*)pVertices);
 				}
-				else if (TexCoordsChannel.Format == MBF_FLOAT16)
+				else if (TexCoordsChannel.Format == EMeshBufferFormat::Float16)
 				{
 					const FFloat16* pUV = reinterpret_cast<const FFloat16*>(pVertices);
 					UV = FVector2f(float(pUV[0]), float(pUV[1]));
@@ -284,7 +284,7 @@ namespace mu
 		constexpr uint16 NullBlockId = MAX_uint16 - 1;
 		LayoutData.Init(NullBlockId, NumVertices);
 
-		UntypedMeshBufferIteratorConst ItIndices(Mesh->GetIndexBuffers(), MBS_VERTEXINDEX);
+		UntypedMeshBufferIteratorConst ItIndices(Mesh->GetIndexBuffers(), EMeshBufferSemantic::VertexIndex);
 		for (int32 TriangleIndex = 0;  TriangleIndex < NumTriangles; ++TriangleIndex)
 		{
 			uint32 Index0 = ItIndices.GetAsUINT32();
@@ -499,9 +499,9 @@ namespace mu
 
 			// TODO
 			check(Layout->GetBlockCount() < MAX_uint16);
-			const EMeshBufferSemantic LayoutSemantic = MBS_LAYOUTBLOCK;
+			const EMeshBufferSemantic LayoutSemantic = EMeshBufferSemantic::LayoutBlock;
 			const int32 LayoutSemanticIndex = int32(LayoutChannel);
-			const EMeshBufferFormat LayoutFormat = bUseAbsoluteBlockIds ? MBF_UINT64 : MBF_UINT16;
+			const EMeshBufferFormat LayoutFormat = bUseAbsoluteBlockIds ? EMeshBufferFormat::UInt64 : EMeshBufferFormat::UInt16;
 			const int32 LayoutComponents = 1;
 			const int32 LayoutOffset = 0;
 			int32 ElementSize = bUseAbsoluteBlockIds ? sizeof(uint64) : sizeof(uint16);
@@ -561,12 +561,12 @@ namespace mu
 				}
 
 				// Copy UVs
-				if (TexCoordsChannel.Format == MBF_FLOAT32)
+				if (TexCoordsChannel.Format == EMeshBufferFormat::Float32)
 				{
 					FVector2f* pUV = reinterpret_cast<FVector2f*>(pVertices);
 					*pUV = UV;
 				}
-				else if (TexCoordsChannel.Format == MBF_FLOAT16)
+				else if (TexCoordsChannel.Format == EMeshBufferFormat::Float16)
 				{
 					FFloat16* pUV = reinterpret_cast<FFloat16*>(pVertices);
 					pUV[0] = FFloat16(UV[0]);
@@ -741,7 +741,7 @@ namespace mu
 		// UE only has position and normal morph data, optimize for this case if indicated. 
 		if (node.bOnlyPositionAndNormal)
 		{
-			op->Channels = { {static_cast<uint8>(MBS_POSITION), 0}, {static_cast<uint8>(MBS_NORMAL), 0} };
+			op->Channels = { {static_cast<uint8>(EMeshBufferSemantic::Position), 0}, {static_cast<uint8>(EMeshBufferSemantic::Normal), 0} };
 		}
 
         // Base
@@ -882,7 +882,7 @@ namespace mu
 
                     for ( size_t c=0; c<node.m_channels.Num(); ++c)
                     {
-                        check( node.m_channels[c].semantic < 256 );
+                        check( int32(node.m_channels[c].semantic) < 256 );
 						check(node.m_channels[c].semanticIndex < 256);
 						
 						ASTOpMeshDifference::FChannel Channel;
@@ -901,7 +901,7 @@ namespace mu
         if (!count)
         {
             // TODO
-            //op.args.MeshInterpolate.target[0] = GenerateMissingImageCode( "First mesh", IF_RGB_UBYTE );
+            //op.args.MeshInterpolate.target[0] = GenerateMissingImageCode( "First mesh", RGB_UByte );
             ErrorLog->GetPrivate()->Add
                 ( "Mesh interpolation: at least the first mesh is required.",
                   ELMT_ERROR, InterpolateNode->GetMessageContext());
diff --git a/Engine/Plugins/Mutable/Source/MutableTools/Private/MuT/CodeGenerator_Projector.cpp b/Engine/Plugins/Mutable/Source/MutableTools/Private/MuT/CodeGenerator_Projector.cpp
--- a/Engine/Plugins/Mutable/Source/MutableTools/Private/MuT/CodeGenerator_Projector.cpp
+++ b/Engine/Plugins/Mutable/Source/MutableTools/Private/MuT/CodeGenerator_Projector.cpp
@@ -116,10 +116,10 @@ namespace mu
 			param.m_name = node.m_name;
 			const TCHAR* CStr = ToCStr(node.m_uid);
 			param.m_uid.ImportTextItem(CStr, 0, nullptr, nullptr);
-			param.m_type = PARAMETER_TYPE::T_PROJECTOR;
+			param.m_type = EParameterType::Projector;
 
 			FProjector p = ProjectorFromNode(node);
-			param.m_defaultValue.Set<ParamProjectorType>(p);
+			param.m_defaultValue.Set<FParamProjectorType>(p);
 
 			op = new ASTOpParameter();
 			op->type = OP_TYPE::PR_PARAMETER;
@@ -140,7 +140,7 @@ namespace mu
 			op = *it;
 		}
 
-		const FProjector& Projector = op->parameter.m_defaultValue.Get<ParamProjectorType>();
+		const FProjector& Projector = op->parameter.m_defaultValue.Get<FParamProjectorType>();
 		result.type = Projector.type;
 		result.op = op;
 	}
@@ -155,7 +155,7 @@ namespace mu
 			ELMT_ERROR, errorContext);
 
 		FProjector p;
-		p.type = PROJECTOR_TYPE::PLANAR;
+		p.type = EProjectorType::Planar;
 
 		p.direction[0] = 1;
 		p.direction[1] = 0;
diff --git a/Engine/Plugins/Mutable/Source/MutableTools/Private/MuT/CodeGenerator_Scalar.cpp b/Engine/Plugins/Mutable/Source/MutableTools/Private/MuT/CodeGenerator_Scalar.cpp
--- a/Engine/Plugins/Mutable/Source/MutableTools/Private/MuT/CodeGenerator_Scalar.cpp
+++ b/Engine/Plugins/Mutable/Source/MutableTools/Private/MuT/CodeGenerator_Scalar.cpp
@@ -144,8 +144,8 @@ namespace mu
 			param.m_name = node.m_name;
 			const TCHAR* CStr = ToCStr(node.m_uid);
 			param.m_uid.ImportTextItem(CStr, 0, nullptr, nullptr);
-			param.m_type = PARAMETER_TYPE::T_FLOAT;
-			param.m_defaultValue.Set<ParamFloatType>(node.m_defaultValue);
+			param.m_type = EParameterType::Float;
+			param.m_defaultValue.Set<FParamFloatType>(node.m_defaultValue);
 
 			op = new ASTOpParameter();
 			op->type = OP_TYPE::SC_PARAMETER;
@@ -184,8 +184,8 @@ namespace mu
 			param.m_name = node.m_name;
 			const TCHAR* CStr = ToCStr(node.m_uid);
 			param.m_uid.ImportTextItem(CStr, 0, nullptr, nullptr);
-			param.m_type = PARAMETER_TYPE::T_INT;
-			param.m_defaultValue.Set<ParamIntType>(node.m_defaultValue);
+			param.m_type = EParameterType::Int;
+			param.m_defaultValue.Set<FParamIntType>(node.m_defaultValue);
 
 			param.m_possibleValues.SetNum(node.m_options.Num());
 			for (int32 i = 0; i < node.m_options.Num(); ++i)
diff --git a/Engine/Plugins/Mutable/Source/MutableTools/Private/MuT/CodeGenerator_String.cpp b/Engine/Plugins/Mutable/Source/MutableTools/Private/MuT/CodeGenerator_String.cpp
--- a/Engine/Plugins/Mutable/Source/MutableTools/Private/MuT/CodeGenerator_String.cpp
+++ b/Engine/Plugins/Mutable/Source/MutableTools/Private/MuT/CodeGenerator_String.cpp
@@ -89,8 +89,8 @@ namespace mu
 			param.m_name = node.m_name;
 			const TCHAR* CStr = ToCStr(node.m_uid);
 			param.m_uid.ImportTextItem(CStr, 0, nullptr, nullptr);
-			param.m_type = PARAMETER_TYPE::T_FLOAT;
-			param.m_defaultValue.Set<ParamStringType>(node.m_defaultValue);
+			param.m_type = EParameterType::Float;
+			param.m_defaultValue.Set<FParamStringType>(node.m_defaultValue);
 
 			op = new ASTOpParameter();
 			op->type = OP_TYPE::ST_PARAMETER;
diff --git a/Engine/Plugins/Mutable/Source/MutableTools/Private/MuT/CodeOptimiser_Parameters.cpp b/Engine/Plugins/Mutable/Source/MutableTools/Private/MuT/CodeOptimiser_Parameters.cpp
--- a/Engine/Plugins/Mutable/Source/MutableTools/Private/MuT/CodeOptimiser_Parameters.cpp
+++ b/Engine/Plugins/Mutable/Source/MutableTools/Private/MuT/CodeOptimiser_Parameters.cpp
@@ -248,7 +248,7 @@ namespace mu
             Ptr<ASTOpFixed> wplainOp = new ASTOpFixed;
             wplainOp->op.type = OP_TYPE::IM_PLAINCOLOUR;
             wplainOp->SetChild( wplainOp->op.args.ImagePlainColour.colour, whiteOp );
-            wplainOp->op.args.ImagePlainColour.format = EImageFormat::IF_L_UBYTE;
+            wplainOp->op.args.ImagePlainColour.format = EImageFormat::L_UByte;
             wplainOp->op.args.ImagePlainColour.size[0] = 4;
 			wplainOp->op.args.ImagePlainColour.size[1] = 4;
 			wplainOp->op.args.ImagePlainColour.LODs = 1;
@@ -472,7 +472,7 @@ namespace mu
                             Ptr<ASTOpFixed> plainOp = new ASTOpFixed;
                             plainOp->op.type = OP_TYPE::IM_PLAINCOLOUR;
                             plainOp->SetChild( plainOp->op.args.ImagePlainColour.colour, blackOp );
-                            plainOp->op.args.ImagePlainColour.format = EImageFormat::IF_L_UBYTE;
+                            plainOp->op.args.ImagePlainColour.format = EImageFormat::L_UByte;
                             plainOp->op.args.ImagePlainColour.size[0] = 4;
 							plainOp->op.args.ImagePlainColour.size[1] = 4;
 							plainOp->op.args.ImagePlainColour.LODs = 1;
@@ -518,7 +518,7 @@ namespace mu
 //							OP plainOp;
 //							plainOp.type = OP_TYPE::IM_PLAINCOLOUR;
 //							plainOp.args.ImagePlainColour.colour = program.AddOp( blackOp );
-//							plainOp.args.ImagePlainColour.format = IF_L_UBYTE;
+//							plainOp.args.ImagePlainColour.format = L_UByte;
 //							plainOp.args.ImagePlainColour.size = ;
 //							plainOp.args.ImagePlainColour.LODs = ;
 
@@ -579,7 +579,7 @@ namespace mu
                             Ptr<ASTOpFixed> plainOp = new ASTOpFixed;
                             plainOp->op.type = OP_TYPE::IM_PLAINCOLOUR;
                             plainOp->SetChild( plainOp->op.args.ImagePlainColour.colour, blackOp );
-                            plainOp->op.args.ImagePlainColour.format = EImageFormat::IF_L_UBYTE;
+                            plainOp->op.args.ImagePlainColour.format = EImageFormat::L_UByte;
                             plainOp->op.args.ImagePlainColour.size[0] = 4;
 							plainOp->op.args.ImagePlainColour.size[1] = 4;
 							plainOp->op.args.ImagePlainColour.LODs = 1;
@@ -1047,7 +1047,7 @@ namespace mu
                         Ptr<ASTOpFixed> plainOp = new ASTOpFixed;
                         plainOp->op.type = OP_TYPE::IM_PLAINCOLOUR;
                         plainOp->SetChild( plainOp->op.args.ImagePlainColour.colour, blackOp );
-                        plainOp->op.args.ImagePlainColour.format = EImageFormat::IF_L_UBYTE;
+                        plainOp->op.args.ImagePlainColour.format = EImageFormat::L_UByte;
                         plainOp->op.args.ImagePlainColour.size[0] = 4;
 						plainOp->op.args.ImagePlainColour.size[1] = 4;
 						plainOp->op.args.ImagePlainColour.LODs = 1;
@@ -1104,7 +1104,7 @@ namespace mu
                         Ptr<ASTOpFixed> plainOp = new ASTOpFixed;
                         plainOp->op.type = OP_TYPE::IM_PLAINCOLOUR;
                         plainOp->SetChild( plainOp->op.args.ImagePlainColour.colour, blackOp );
-                        plainOp->op.args.ImagePlainColour.format = EImageFormat::IF_L_UBYTE;
+                        plainOp->op.args.ImagePlainColour.format = EImageFormat::L_UByte;
                         plainOp->op.args.ImagePlainColour.size[0] = 4;
 						plainOp->op.args.ImagePlainColour.size[1] = 4;
 						plainOp->op.args.ImagePlainColour.LODs = 1;
@@ -1217,7 +1217,7 @@ namespace mu
                         Ptr<ASTOpFixed> plainOp = new ASTOpFixed;
                         plainOp->op.type = OP_TYPE::IM_PLAINCOLOUR;
                         plainOp->SetChild( plainOp->op.args.ImagePlainColour.colour, blackOp );
-                        plainOp->op.args.ImagePlainColour.format = EImageFormat::IF_L_UBYTE; //TODO: FORMAT_LIKE
+                        plainOp->op.args.ImagePlainColour.format = EImageFormat::L_UByte; //TODO: FORMAT_LIKE
                         plainOp->op.args.ImagePlainColour.size[0] = 4;
 						plainOp->op.args.ImagePlainColour.size[1] = 4;
 						plainOp->op.args.ImagePlainColour.LODs = 1;
@@ -1264,7 +1264,7 @@ namespace mu
                         Ptr<ASTOpFixed> plainOp = new ASTOpFixed;
                         plainOp->op.type = OP_TYPE::IM_PLAINCOLOUR;
                         plainOp->SetChild( plainOp->op.args.ImagePlainColour.colour, blackOp );
-                        plainOp->op.args.ImagePlainColour.format = EImageFormat::IF_L_UBYTE; //TODO: FORMAT_LIKE
+                        plainOp->op.args.ImagePlainColour.format = EImageFormat::L_UByte; //TODO: FORMAT_LIKE
                         plainOp->op.args.ImagePlainColour.size[0] = 4;
 						plainOp->op.args.ImagePlainColour.size[1] = 4;
 						plainOp->op.args.ImagePlainColour.LODs = 1;
@@ -1477,7 +1477,7 @@ namespace mu
 
 
     class AccumulateAllImageFormatsOpAST
-            : public Visitor_TopDown_Unique_Const< std::array<uint8_t, size_t(EImageFormat::IF_COUNT)> >
+            : public Visitor_TopDown_Unique_Const< std::array<uint8_t, size_t(EImageFormat::Count)> >
     {
     public:
 
@@ -1498,12 +1498,12 @@ namespace mu
         {
             bool recurse = false;
 
-            const std::array<uint8_t, (size_t)EImageFormat::IF_COUNT>& currentFormats = GetCurrentState();
+            const std::array<uint8_t, (size_t)EImageFormat::Count>& currentFormats = GetCurrentState();
 
             // Remove unsupported formats
             if (GetOpDataType( at->GetOpType() )==DT_IMAGE)
             {
-				std::array<uint8_t, (size_t)EImageFormat::IF_COUNT>* it = m_supportedFormats.Find(at);
+				std::array<uint8_t, (size_t)EImageFormat::Count>* it = m_supportedFormats.Find(at);
                 if (!it)
                 {
                     // Default to all supported
@@ -1511,7 +1511,7 @@ namespace mu
                     it = m_supportedFormats.Find(at);
                 }
 
-                for ( unsigned f=0; f< (unsigned)EImageFormat::IF_COUNT; ++f )
+                for ( unsigned f=0; f< (unsigned)EImageFormat::Count; ++f )
                 {
                     if ( !currentFormats[f] )
                     {
@@ -1532,10 +1532,10 @@ namespace mu
 
                 if ( typedAt->mask )
                 {
-                    std::array<uint8_t, (size_t)EImageFormat::IF_COUNT> newState;
+                    std::array<uint8_t, (size_t)EImageFormat::Count> newState;
                     newState.fill(0);
-                    newState[(size_t)EImageFormat::IF_L_UBYTE ] = 1;
-                    newState[(size_t)EImageFormat::IF_L_UBYTE_RLE ] = 1;
+                    newState[(size_t)EImageFormat::L_UByte ] = 1;
+                    newState[(size_t)EImageFormat::L_UByteRLE ] = 1;
 
                     RecurseWithState( typedAt->mask.child(), newState );
                 }
@@ -1549,11 +1549,11 @@ namespace mu
                 RecurseWithCurrentState( typedAt->base.child() );
                 RecurseWithCurrentState( typedAt->blend.child() );
 
-                std::array<uint8_t, (size_t)EImageFormat::IF_COUNT> newState;
+                std::array<uint8_t, (size_t)EImageFormat::Count> newState;
                 newState.fill(0);
                 // TODO
-                //newState[ IF_L_UBYTE ] = 1;
-                //newState[ IF_L_UBYTE_RLE ] = 1;
+                //newState[ L_UByte ] = 1;
+                //newState[ L_UByteRLE ] = 1;
 
                 if ( typedAt->mask )
                 {
@@ -1568,10 +1568,10 @@ namespace mu
 
                 RecurseWithCurrentState( typedAt->children[typedAt->op.args.ImageDisplace.source].child() );
 
-                std::array<uint8_t, (size_t)EImageFormat::IF_COUNT> newState;
+                std::array<uint8_t, (size_t)EImageFormat::Count> newState;
                 newState.fill(0);
-                newState[(size_t)EImageFormat::IF_L_UBYTE ] = 1;
-                newState[(size_t)EImageFormat::IF_L_UBYTE_RLE ] = 1;
+                newState[(size_t)EImageFormat::L_UByte ] = 1;
+                newState[(size_t)EImageFormat::L_UByteRLE ] = 1;
                 RecurseWithState( typedAt->children[typedAt->op.args.ImageDisplace.displacementMap].child(), newState );
                 break;
             }
@@ -1588,7 +1588,7 @@ namespace mu
 
         bool IsSupportedFormat( Ptr<ASTOp> at, EImageFormat format ) const
         {
-            const std::array<uint8_t, (size_t)EImageFormat::IF_COUNT>* it = m_supportedFormats.Find(at);
+            const std::array<uint8_t, (size_t)EImageFormat::Count>* it = m_supportedFormats.Find(at);
             if (!it)
             {
                 return false;
@@ -1601,14 +1601,14 @@ namespace mu
     private:
 
         //! Formats known to be supported for every instruction.
-        //! IF_COUNT*code.size() entries
-        TMap< Ptr<ASTOp>, std::array<uint8_t, (size_t)EImageFormat::IF_COUNT> > m_supportedFormats;
+        //! Count*code.size() entries
+        TMap< Ptr<ASTOp>, std::array<uint8_t, (size_t)EImageFormat::Count> > m_supportedFormats;
 
         //! Constant convenience initial value
-        std::array<uint8_t, (size_t)EImageFormat::IF_COUNT> m_initialState;
+        std::array<uint8_t, (size_t)EImageFormat::Count> m_initialState;
 
         //! Constant convenience initial value
-        std::array<uint8_t, (size_t)EImageFormat::IF_COUNT> m_allSupported;
+        std::array<uint8_t, (size_t)EImageFormat::Count> m_allSupported;
 
     };
 
@@ -1874,12 +1874,12 @@ namespace mu
                 {
                     FImageDesc desc = at->GetImageDesc();
 
-                    if ( desc.m_format!=EImageFormat::IF_L_UBYTE_RLE
+                    if ( desc.m_format!=EImageFormat::L_UByteRLE
                          &&
-                         m_opFormats.IsSupportedFormat(at, EImageFormat::IF_L_UBYTE_RLE) )
+                         m_opFormats.IsSupportedFormat(at, EImageFormat::L_UByteRLE) )
                     {
                         Ptr<ASTOpImagePixelFormat> op = new ASTOpImagePixelFormat;
-                        op->Format = EImageFormat::IF_L_UBYTE_RLE;
+                        op->Format = EImageFormat::L_UByteRLE;
                         // Note: we have to clone here, to avoid a loop with the visitor system
                         // that updates visited children before processing a node.
 						ASTOp::MapChildFunc Identity = [](const Ptr<ASTOp>& o) {return o;};
diff --git a/Engine/Plugins/Mutable/Source/MutableTools/Private/MuT/CodeOptimiser_Semantic.cpp b/Engine/Plugins/Mutable/Source/MutableTools/Private/MuT/CodeOptimiser_Semantic.cpp
--- a/Engine/Plugins/Mutable/Source/MutableTools/Private/MuT/CodeOptimiser_Semantic.cpp
+++ b/Engine/Plugins/Mutable/Source/MutableTools/Private/MuT/CodeOptimiser_Semantic.cpp
@@ -1304,7 +1304,7 @@ namespace mu
 				// Only sink the resize if we know that the pixelformat source image is uncompressed.
 				Ptr<ASTOpImagePixelFormat> SourceTyped = static_cast<ASTOpImagePixelFormat*>(sourceAt.get());
 				FImageDesc PixelFormatSourceDesc = SourceTyped ->Source->GetImageDesc();
-				if (PixelFormatSourceDesc.m_format!=EImageFormat::IF_NONE
+				if (PixelFormatSourceDesc.m_format!=EImageFormat::None
 					&&
 					!mu::IsCompressedFormat(PixelFormatSourceDesc.m_format) )
 				{
@@ -1513,7 +1513,7 @@ namespace mu
 //                        blackLayoutOp.args.ImageBlankLayout.layout = sizeSourceSize->layout;
 //                        blackLayoutOp.args.ImageBlankLayout.blockSize[0] = sizeSourceSize->factor[0];
 //                        blackLayoutOp.args.ImageBlankLayout.blockSize[1] = sizeSourceSize->factor[1];
-//                        blackLayoutOp.args.ImageBlankLayout.format = IF_L_UBYTE;
+//                        blackLayoutOp.args.ImageBlankLayout.format = L_UByte;
 
 //                        OP opResize = program.m_code[at];
 //                        opResize.args.ImageResizeLike.sizeSource = program.AddOp( blackLayoutOp );
diff --git a/Engine/Plugins/Mutable/Source/MutableTools/Private/MuT/DataPacker.cpp b/Engine/Plugins/Mutable/Source/MutableTools/Private/MuT/DataPacker.cpp
--- a/Engine/Plugins/Mutable/Source/MutableTools/Private/MuT/DataPacker.cpp
+++ b/Engine/Plugins/Mutable/Source/MutableTools/Private/MuT/DataPacker.cpp
@@ -299,11 +299,11 @@ namespace mu
             MUTABLE_CPUPROFILER_SCOPE(AccumulateImageFormatsAST);
 
             // TODO: Use statically sized vectors (std::array)
-            //vector<bool> initial( IF_COUNT, true );
+            //vector<bool> initial( Count, true );
             //std::fill( m_supportedFormats.begin(), m_supportedFormats.end(), initial );
 
 			TArray<bool> defaultState;
-			defaultState.SetNumZeroed(int32(EImageFormat::IF_COUNT));
+			defaultState.SetNumZeroed(int32(EImageFormat::Count));
 
             Traverse( roots, defaultState );
         }
@@ -316,7 +316,7 @@ namespace mu
             const TArray<bool>& currentFormats = GetCurrentState();
 
 			TArray<bool> defaultState;
-			defaultState.SetNumZeroed(int32(EImageFormat::IF_COUNT));
+			defaultState.SetNumZeroed(int32(EImageFormat::Count));
 			bool allFalse = currentFormats == defaultState;
 
             // Can we use the cache?
@@ -339,11 +339,11 @@ namespace mu
                 if (!m_supportedFormats.count(op))
                 {
 					TArray<bool> initial;
-					initial.Init( true, int32(EImageFormat::IF_COUNT) );
+					initial.Init( true, int32(EImageFormat::Count) );
                     m_supportedFormats.insert( std::make_pair(op, std::move(initial)) );
                 }
 
-                for ( unsigned f=0; f< unsigned(EImageFormat::IF_COUNT); ++f )
+                for ( unsigned f=0; f< unsigned(EImageFormat::Count); ++f )
                 {
                     if ( !currentFormats[f] )
                     {
@@ -365,14 +365,14 @@ namespace mu
 				const ASTOpImageCompose* op = static_cast<const ASTOpImageCompose*>(node.get());
 
 				TArray<bool> newState;
-				newState.Init(false, int32(EImageFormat::IF_COUNT));
+				newState.Init(false, int32(EImageFormat::Count));
                 RecurseWithState( op->Layout.child(), newState );
                 RecurseWithState( op->Base.child(), newState );
                 RecurseWithState( op->BlockImage.child(), newState );
 
                 if ( op->Mask )
                 {
-                    newState[(int32)EImageFormat::IF_L_UBIT_RLE] = true;
+                    newState[(int32)EImageFormat::L_UBitRLE] = true;
                     RecurseWithState( op->Mask.child(), newState );
                 }
                 break;
@@ -383,14 +383,14 @@ namespace mu
                 recurse = false;
 				const ASTOpImageLayerColor* op = static_cast<const ASTOpImageLayerColor*>(node.get());
 				TArray<bool> newState;
-				newState.Init(false, int32(EImageFormat::IF_COUNT));
+				newState.Init(false, int32(EImageFormat::Count));
 				RecurseWithState( op->base.child(), newState );
                 RecurseWithState( op->color.child(), newState );
 
                 if ( op->mask )
                 {
-                    newState[(int32)EImageFormat::IF_L_UBYTE] = true;
-                    newState[(int32)EImageFormat::IF_L_UBYTE_RLE] = true;
+                    newState[(int32)EImageFormat::L_UByte] = true;
+                    newState[(int32)EImageFormat::L_UByteRLE] = true;
 
                     RecurseWithState( op->mask.child(), newState );
                 }
@@ -402,14 +402,14 @@ namespace mu
                 recurse = false;
 				const ASTOpImageLayer* op = static_cast<const ASTOpImageLayer*>(node.get());
 				TArray<bool> newState;
-				newState.Init(false, int32(EImageFormat::IF_COUNT));
+				newState.Init(false, int32(EImageFormat::Count));
 				RecurseWithState( op->base.child(), newState );
                 RecurseWithState( op->blend.child(), newState );
 
                 if (op->mask)
                 {
-                    newState[(int32)EImageFormat::IF_L_UBYTE] = true;
-                    newState[(int32)EImageFormat::IF_L_UBYTE_RLE] = true;
+                    newState[(int32)EImageFormat::L_UByte] = true;
+                    newState[(int32)EImageFormat::L_UByteRLE] = true;
 
                     RecurseWithState( op->mask.child(), newState );
                 }
@@ -421,14 +421,14 @@ namespace mu
                 recurse = false;
 				const ASTOpImageMultiLayer* op = static_cast<const ASTOpImageMultiLayer*>(node.get());
 				TArray<bool> newState;
-				newState.Init(false, int32(EImageFormat::IF_COUNT));
+				newState.Init(false, int32(EImageFormat::Count));
 				RecurseWithState( op->base.child(), newState );
                 RecurseWithState( op->blend.child(), newState );
 
                 if (op->mask)
                 {
-                    newState[(size_t)EImageFormat::IF_L_UBYTE] = true;
-                    newState[(size_t)EImageFormat::IF_L_UBYTE_RLE] = true;
+                    newState[(size_t)EImageFormat::L_UByte] = true;
+                    newState[(size_t)EImageFormat::L_UByteRLE] = true;
 
                     RecurseWithState( op->mask.child(), newState );
                 }
@@ -440,11 +440,11 @@ namespace mu
                 recurse = false;
 				const ASTOpFixed* op = static_cast<const ASTOpFixed*>(node.get());
 				TArray<bool> newState;
-				newState.Init(false, int32(EImageFormat::IF_COUNT));
+				newState.Init(false, int32(EImageFormat::Count));
 				RecurseWithState( op->children[op->op.args.ImageDisplace.source].child(), newState );
 
-                newState[(int32)EImageFormat::IF_L_UBYTE ] = true;
-                newState[(int32)EImageFormat::IF_L_UBYTE_RLE ] = true;
+                newState[(int32)EImageFormat::L_UByte ] = true;
+                newState[(int32)EImageFormat::L_UByteRLE ] = true;
 
                 RecurseWithState( op->children[op->op.args.ImageDisplace.displacementMap].child(), newState );
                 break;
@@ -452,12 +452,12 @@ namespace mu
 
             default:
             {
-                //m_currentFormats.Add(vector<bool>(IF_COUNT, false));
+                //m_currentFormats.Add(vector<bool>(Count, false));
                 //Recurse(at, program);
                 //m_currentFormats.pop_back();
 
 				TArray<bool> newState;
-				newState.Init(false, int32(EImageFormat::IF_COUNT));
+				newState.Init(false, int32(EImageFormat::Count));
 				if (currentFormats != newState)
                 {
                     RecurseWithState(node, newState);
@@ -503,12 +503,12 @@ namespace mu
             MUTABLE_CPUPROFILER_SCOPE(AccumulateMeshChannelUsageAST);
 
             // Sanity check in case we add more semantics
-            static_assert(MBS_COUNT<sizeof(uint64)*8, "Too many mesh buffer semantics." );
+            static_assert(uint32(EMeshBufferSemantic::Count)<sizeof(uint64)*8, "Too many mesh buffer semantics." );
 
             // Default state: we need everything except internal semantics
             uint64 defaultState = 0xffffffffffffffff;
-            defaultState ^= (UINT64_C(1)<<MBS_LAYOUTBLOCK);
-            defaultState ^= (UINT64_C(1)<<MBS_VERTEXINDEX);
+            defaultState ^= (UINT64_C(1)<<uint32(EMeshBufferSemantic::LayoutBlock));
+            defaultState ^= (UINT64_C(1)<<uint32(EMeshBufferSemantic::VertexIndex));
 
             Traverse(roots,defaultState);
         }
@@ -553,7 +553,7 @@ namespace mu
 				const ASTOpMeshDifference* op = static_cast<const ASTOpMeshDifference*>(node.get());
 
                 uint64 newState = currentSemantics;
-                newState |= (UINT64_C(1)<<MBS_VERTEXINDEX);
+                newState |= (UINT64_C(1)<<uint32(EMeshBufferSemantic::VertexIndex));
                 RecurseWithState( op->Base.child(), newState );
 
                 RecurseWithState( op->Target.child(), currentSemantics );
@@ -567,7 +567,7 @@ namespace mu
 				const ASTOpMeshRemoveMask* op = static_cast<const ASTOpMeshRemoveMask*>(node.get());
 
                 uint64 newState = currentSemantics;
-                newState |= (UINT64_C(1)<<MBS_VERTEXINDEX);
+                newState |= (UINT64_C(1)<< uint32(EMeshBufferSemantic::VertexIndex));
                 RecurseWithState( op->source.child(), newState );
                 for( const TPair<ASTChild, ASTChild>& r: op->removes )
                 {
@@ -583,7 +583,7 @@ namespace mu
 				const ASTOpMeshMorph* op = static_cast<const ASTOpMeshMorph*>(node.get());
 
                 uint64 newState = currentSemantics;
-                newState |= (UINT64_C(1)<<MBS_VERTEXINDEX);
+                newState |= (UINT64_C(1)<< uint32(EMeshBufferSemantic::VertexIndex));
                 RecurseWithState( op->Base.child(), newState );
                 RecurseWithState( op->Target.child(), newState );
                 break;
@@ -596,7 +596,7 @@ namespace mu
 				const ASTOpFixed* op = static_cast<const ASTOpFixed*>(node.get());
 
                 uint64 newState = currentSemantics;
-                newState |= (UINT64_C(1)<<MBS_LAYOUTBLOCK);
+                newState |= (UINT64_C(1)<< uint32(EMeshBufferSemantic::LayoutBlock));
                 RecurseWithState( op->children[op->op.args.MeshApplyLayout.mesh].child(), newState );
 
                 RecurseWithState( op->children[op->op.args.MeshApplyLayout.layout].child(), currentSemantics );
@@ -610,7 +610,7 @@ namespace mu
 				const ASTOpFixed* op = static_cast<const ASTOpFixed*>(node.get());
 
 				uint64 newState = currentSemantics;
-				newState |= (UINT64_C(1) << MBS_LAYOUTBLOCK);
+				newState |= (UINT64_C(1) << uint32(EMeshBufferSemantic::LayoutBlock));
 				RecurseWithState(op->children[op->op.args.MeshProject.mesh].child(), newState);
 
 				RecurseWithState(op->children[op->op.args.MeshProject.projector].child(), currentSemantics);
@@ -624,7 +624,7 @@ namespace mu
 				const ASTOpImageRasterMesh* op = static_cast<const ASTOpImageRasterMesh*>(node.get());
 
 				uint64 newState = currentSemantics;
-				newState |= (UINT64_C(1) << MBS_LAYOUTBLOCK);
+				newState |= (UINT64_C(1) << uint32(EMeshBufferSemantic::LayoutBlock));
 				RecurseWithState(op->mesh.child(), newState);
 
 				RecurseWithState(op->image.child(), currentSemantics);
@@ -642,8 +642,8 @@ namespace mu
 
                 // todo: check if we really need all of them
                 uint64 newState = currentSemantics;
-                newState |= (UINT64_C(1)<<MBS_LAYOUTBLOCK);
-                newState |= (UINT64_C(1)<<MBS_VERTEXINDEX);
+                newState |= (UINT64_C(1)<< uint32(EMeshBufferSemantic::LayoutBlock));
+                newState |= (UINT64_C(1)<< uint32(EMeshBufferSemantic::VertexIndex));
 
                 RecurseWithState( op->Source.child(), newState );
                 break;
@@ -656,7 +656,7 @@ namespace mu
 				const ASTOpLayoutFromMesh* op = static_cast<const ASTOpLayoutFromMesh*>(node.get());
 
 				uint64 newState = currentSemantics;
-				newState |= (UINT64_C(1) << MBS_LAYOUTBLOCK);
+				newState |= (UINT64_C(1) << uint32(EMeshBufferSemantic::LayoutBlock));
 
 				RecurseWithState(op->Mesh.child(), newState);
 				break;
@@ -711,7 +711,7 @@ namespace mu
                 EMeshBufferSemantic semantic;
                 pMesh->GetVertexBuffers().GetChannel(v,c,&semantic,
                                                       nullptr, nullptr, nullptr, nullptr);
-                used = (( (UINT64_C(1)<<semantic) ) & usedSemantics) != 0;
+                used = (( (UINT64_C(1)<< uint32(semantic)) ) & usedSemantics) != 0;
             }
 
             if (!used)
@@ -728,7 +728,7 @@ namespace mu
         // TODO: hack, if we don't need layouts, remove them.
         {
             uint64 layoutSemantics = 0;
-            layoutSemantics |= (UINT64_C(1)<<MBS_LAYOUTBLOCK);
+            layoutSemantics |= (UINT64_C(1)<< uint32(EMeshBufferSemantic::LayoutBlock));
 
             if ( (usedSemantics&layoutSemantics) == 0)
             {
@@ -783,9 +783,9 @@ namespace mu
 					ASTOp::Replace(n, NewPlain);
 
 				}
-				else if ( accFormat.m_supportedFormats[typed][(int32)EImageFormat::IF_L_UBIT_RLE] )
+				else if ( accFormat.m_supportedFormats[typed][(int32)EImageFormat::L_UBitRLE] )
                 {
-                    ImagePtr pNew = ImOp.ImagePixelFormat( ImageCompressionQuality, pOld.get(), EImageFormat::IF_L_UBIT_RLE );
+                    ImagePtr pNew = ImOp.ImagePixelFormat( ImageCompressionQuality, pOld.get(), EImageFormat::L_UBitRLE );
 
                     // Only replace if the compression was worth!
                     int32 oldSize = pOld->GetDataSize();
@@ -795,9 +795,9 @@ namespace mu
                         typed->SetValue(pNew, OptimizeOptions.DiskCacheContext);
                     }
                 }
-                else if ( accFormat.m_supportedFormats[typed][(int32)EImageFormat::IF_L_UBYTE_RLE] )
+                else if ( accFormat.m_supportedFormats[typed][(int32)EImageFormat::L_UByteRLE] )
                 {
-                    ImagePtr pNew = ImOp.ImagePixelFormat( ImageCompressionQuality, pOld.get(), EImageFormat::IF_L_UBYTE_RLE );
+                    ImagePtr pNew = ImOp.ImagePixelFormat( ImageCompressionQuality, pOld.get(), EImageFormat::L_UByteRLE );
 
                     // Only replace if the compression was worth!
 					int32 oldSize = pOld->GetDataSize();
diff --git a/Engine/Plugins/Mutable/Source/MutableTools/Private/MuT/NodeLayout.cpp b/Engine/Plugins/Mutable/Source/MutableTools/Private/MuT/NodeLayout.cpp
--- a/Engine/Plugins/Mutable/Source/MutableTools/Private/MuT/NodeLayout.cpp
+++ b/Engine/Plugins/Mutable/Source/MutableTools/Private/MuT/NodeLayout.cpp
@@ -33,8 +33,8 @@ namespace mu
 		{
 			UVs.SetNumUninitialized(IndexCount * 2);
 
-			UntypedMeshBufferIteratorConst IndexIt(Mesh->GetIndexBuffers(), MBS_VERTEXINDEX, 0);
-			UntypedMeshBufferIteratorConst TexIt(Mesh->GetVertexBuffers(), MBS_TEXCOORDS, LayoutIndex);
+			UntypedMeshBufferIteratorConst IndexIt(Mesh->GetIndexBuffers(), EMeshBufferSemantic::VertexIndex, 0);
+			UntypedMeshBufferIteratorConst TexIt(Mesh->GetVertexBuffers(), EMeshBufferSemantic::TexCoords, LayoutIndex);
 
 			int32 TriangleCount = IndexCount / 3;
 			for (int32 TriangleIndex = 0; TriangleIndex < TriangleCount; ++TriangleIndex)
@@ -249,7 +249,7 @@ namespace mu
 		// Find Unique Vertices
 		MeshCreateCollapsedVertexMap(Mesh.get(), CollapsedVertices);
 
-		UntypedMeshBufferIteratorConst ItIndices(Mesh->GetIndexBuffers(), MBS_VERTEXINDEX);
+		UntypedMeshBufferIteratorConst ItIndices(Mesh->GetIndexBuffers(), EMeshBufferSemantic::VertexIndex);
 		for (int32 TriangleIndex = 0; TriangleIndex < NumTriangles; ++TriangleIndex)
 		{
 			uint32 Index0 = ItIndices.GetAsUINT32();
@@ -283,12 +283,12 @@ namespace mu
 		{
 			int32 TexCoordsBufferIndex = -1;
 			int32 TexCoordsChannelIndex = -1;
-			Mesh->GetVertexBuffers().FindChannel(MBS_TEXCOORDS, LayoutIndex, &TexCoordsBufferIndex, &TexCoordsChannelIndex);
+			Mesh->GetVertexBuffers().FindChannel(EMeshBufferSemantic::TexCoords, LayoutIndex, &TexCoordsBufferIndex, &TexCoordsChannelIndex);
 			check(TexCoordsBufferIndex >= 0);
 			check(TexCoordsChannelIndex >= 0);
 
 			const FMeshBufferChannel& TexCoordsChannel = Mesh->VertexBuffers.Buffers[TexCoordsBufferIndex].Channels[TexCoordsChannelIndex];
-			check(TexCoordsChannel.Semantic == MBS_TEXCOORDS);
+			check(TexCoordsChannel.Semantic == EMeshBufferSemantic::TexCoords);
 
 			uint8* TexCoordData = Mesh->GetVertexBuffers().GetBufferData(TexCoordsBufferIndex);
 			int32 ElemSize = Mesh->GetVertexBuffers().GetElementSize(TexCoordsBufferIndex);
@@ -299,11 +299,11 @@ namespace mu
 			for (int32 VertexIndex = 0; VertexIndex < NumVertices; ++VertexIndex)
 			{
 				FVector2f& UV = TexCoords[VertexIndex];
-				if (TexCoordsChannel.Format == MBF_FLOAT32)
+				if (TexCoordsChannel.Format == EMeshBufferFormat::Float32)
 				{
 					UV = *((FVector2f*)pVertices);
 				}
-				else if (TexCoordsChannel.Format == MBF_FLOAT16)
+				else if (TexCoordsChannel.Format == EMeshBufferFormat::Float16)
 				{
 					const FFloat16* pUV = reinterpret_cast<const FFloat16*>(pVertices);
 					UV = FVector2f(float(pUV[0]), float(pUV[1]));
@@ -487,7 +487,7 @@ namespace mu
 			// TODO: Size?
 			int32 SizeX = 1024;
 			int32 SizeY = 1024;
-			Ptr<Image> Mask = new Image(SizeX, SizeY, 1, EImageFormat::IF_L_UBYTE, EInitializationType::Black);
+			Ptr<Image> Mask = new Image(SizeX, SizeY, 1, EImageFormat::L_UByte, EInitializationType::Black);
 			const TArrayView<uint8> ImageData = Mask->DataStorage.GetLOD(0);
 
 			for (int32 TriangleIndex = 0; TriangleIndex < NumTriangles; ++TriangleIndex)
diff --git a/Engine/Plugins/Mutable/Source/MutableTools/Private/MuT/NodeMesh.cpp b/Engine/Plugins/Mutable/Source/MutableTools/Private/MuT/NodeMesh.cpp
--- a/Engine/Plugins/Mutable/Source/MutableTools/Private/MuT/NodeMesh.cpp
+++ b/Engine/Plugins/Mutable/Source/MutableTools/Private/MuT/NodeMesh.cpp
@@ -39,7 +39,7 @@ namespace mu
 		TArray<FVector3f> Vertices;
 		Vertices.SetNumUninitialized(NumVertices);
 
-		mu::UntypedMeshBufferIteratorConst ItPosition = mu::UntypedMeshBufferIteratorConst(Mesh->GetVertexBuffers(), mu::MBS_POSITION);
+		mu::UntypedMeshBufferIteratorConst ItPosition = mu::UntypedMeshBufferIteratorConst(Mesh->GetVertexBuffers(), mu::EMeshBufferSemantic::Position);
 
 		FVector3f* VertexData = Vertices.GetData();
 
diff --git a/Engine/Plugins/Mutable/Source/MutableTools/Private/MuT/NodeMeshInterpolatePrivate.h b/Engine/Plugins/Mutable/Source/MutableTools/Private/MuT/NodeMeshInterpolatePrivate.h
--- a/Engine/Plugins/Mutable/Source/MutableTools/Private/MuT/NodeMeshInterpolatePrivate.h
+++ b/Engine/Plugins/Mutable/Source/MutableTools/Private/MuT/NodeMeshInterpolatePrivate.h
@@ -25,7 +25,7 @@ namespace mu
 		{
 			CHANNEL()
 			{
-				semantic = MBS_NONE;
+				semantic = EMeshBufferSemantic::None;
 				semanticIndex = 0;
 			}
 
diff --git a/Engine/Plugins/Mutable/Source/MutableTools/Private/MuT/NodeProjectorConstant.cpp b/Engine/Plugins/Mutable/Source/MutableTools/Private/MuT/NodeProjectorConstant.cpp
--- a/Engine/Plugins/Mutable/Source/MutableTools/Private/MuT/NodeProjectorConstant.cpp
+++ b/Engine/Plugins/Mutable/Source/MutableTools/Private/MuT/NodeProjectorConstant.cpp
@@ -29,7 +29,7 @@ namespace mu
 	//---------------------------------------------------------------------------------------------
 	// Own Interface
 	//---------------------------------------------------------------------------------------------
-    void NodeProjectorConstant::GetValue( PROJECTOR_TYPE* OutType,
+    void NodeProjectorConstant::GetValue( EProjectorType* OutType,
 		FVector3f* OutPos,
 		FVector3f* OutDir,
 		FVector3f* OutUp,
@@ -48,7 +48,7 @@ namespace mu
 
 
 	//---------------------------------------------------------------------------------------------
-    void NodeProjectorConstant::SetValue( PROJECTOR_TYPE type, FVector3f pos, FVector3f dir, FVector3f up, FVector3f scale, float projectionAngle )
+    void NodeProjectorConstant::SetValue( EProjectorType type, FVector3f pos, FVector3f dir, FVector3f up, FVector3f scale, float projectionAngle )
 	{
         m_pD->m_type = type;
 		m_pD->m_position = pos;
diff --git a/Engine/Plugins/Mutable/Source/MutableTools/Private/MuT/NodeProjectorParameter.cpp b/Engine/Plugins/Mutable/Source/MutableTools/Private/MuT/NodeProjectorParameter.cpp
--- a/Engine/Plugins/Mutable/Source/MutableTools/Private/MuT/NodeProjectorParameter.cpp
+++ b/Engine/Plugins/Mutable/Source/MutableTools/Private/MuT/NodeProjectorParameter.cpp
@@ -51,7 +51,7 @@ namespace mu
 
 
 	//---------------------------------------------------------------------------------------------
-    void NodeProjectorParameter::SetDefaultValue( PROJECTOR_TYPE type,
+    void NodeProjectorParameter::SetDefaultValue( EProjectorType type,
 		FVector3f pos,
 		FVector3f dir,
 		FVector3f up,
diff --git a/Engine/Plugins/Mutable/Source/MutableTools/Private/MuT/TypeInfo.cpp b/Engine/Plugins/Mutable/Source/MutableTools/Private/MuT/TypeInfo.cpp
--- a/Engine/Plugins/Mutable/Source/MutableTools/Private/MuT/TypeInfo.cpp
+++ b/Engine/Plugins/Mutable/Source/MutableTools/Private/MuT/TypeInfo.cpp
@@ -8,7 +8,7 @@
 namespace mu
 {
 
-	const char* TypeInfo::s_imageFormatName[size_t(EImageFormat::IF_COUNT)] =
+	const char* TypeInfo::s_imageFormatName[size_t(EImageFormat::Count)] =
 	{
 		"None",
 		"RGB U888",
@@ -55,7 +55,7 @@ namespace mu
 
 	};
 
-	static_assert(sizeof(TypeInfo::s_imageFormatName) / sizeof(void*) == int32(EImageFormat::IF_COUNT));
+	static_assert(sizeof(TypeInfo::s_imageFormatName) / sizeof(void*) == int32(EImageFormat::Count));
 
 
 	const char* TypeInfo::s_meshBufferSemanticName[] =
@@ -88,50 +88,50 @@ namespace mu
 		"AltSkinWeight"
 	};
 
-	static_assert(sizeof(TypeInfo::s_meshBufferSemanticName) / sizeof(void*) == int32(MBS_COUNT));
+	static_assert(sizeof(TypeInfo::s_meshBufferSemanticName) / sizeof(void*) == int32(EMeshBufferSemantic::Count));
 
-	const char* TypeInfo::s_meshBufferFormatName[MBF_COUNT] =
+	const char* TypeInfo::s_meshBufferFormatName[int32(EMeshBufferFormat::Count)] =
 	{
 		"None",
-		"Float16",
-		"Float32",
-
-		"UInt8",
-		"UInt16",
-		"UInt32",
-		"Int8",
-		"Int16",
-		"Int32",
-
-		"NUInt8",
-		"NUInt16",
-		"NUInt32",
-		"NInt8",
-		"NInt16",
-		"NInt32",
-
-		"PackedDir8",
+		"EMeshBufferFormat::Float16",
+		"EMeshBufferFormat::Float32",
+
+		"EMeshBufferFormat::UInt8",
+		"EMeshBufferFormat::UInt16",
+		"EMeshBufferFormat::UInt32",
+		"EMeshBufferFormat::Int8",
+		"EMeshBufferFormat::Int16",
+		"EMeshBufferFormat::Int32",
+
+		"EMeshBufferFormat::NUInt8",
+		"EMeshBufferFormat::NUInt16",
+		"EMeshBufferFormat::NUInt32",
+		"EMeshBufferFormat::NInt8",
+		"EMeshBufferFormat::NInt16",
+		"EMeshBufferFormat::NInt32",
+
+		"EMeshBufferFormat::PackedDir8",
 		"PackedDir8_WTangentSign",
-		"PackedDirS8",
+		"EMeshBufferFormat::PackedDirS8",
 		"PackedDirS8_WTangentSign",
 
-		"Float64",
+		"EMeshBufferFormat::Float64",
 
-		"UInt64",
-		"Int64",
+		"EMeshBufferFormat::UInt64",
+		"EMeshBufferFormat::Int64",
 		"NUint64",
-		"NInt64",
+		"EMeshBufferFormat::NInt64",
 	};
 
-	static_assert(sizeof(TypeInfo::s_meshBufferFormatName) / sizeof(void*) == int32(MBF_COUNT));
+	static_assert(sizeof(TypeInfo::s_meshBufferFormatName) / sizeof(void*) == int32(EMeshBufferFormat::Count));
 
-	const char* TypeInfo::s_projectorTypeName[static_cast<uint32>(PROJECTOR_TYPE::COUNT)] =
+	const char* TypeInfo::s_projectorTypeName[static_cast<uint32>(EProjectorType::Count)] =
 	{
 		"Planar",
 		"Cylindrical",
 		"Wrapping",
 	};
-	static_assert(sizeof(TypeInfo::s_projectorTypeName) / sizeof(void*) == int32(PROJECTOR_TYPE::COUNT));
+	static_assert(sizeof(TypeInfo::s_projectorTypeName) / sizeof(void*) == int32(EProjectorType::Count));
 
 	const char* TypeInfo::s_blendModeName[uint32(EBlendType::_BT_COUNT)] =
 	{
