[mutable] Cleanup MutableRuntime API moving to use more unreal types and naming standards.

diff --git a/Engine/Plugins/Mutable/Source/MutableTools/Private/MuT/UnrealPixelFormatOverride.cpp b/Engine/Plugins/Mutable/Source/MutableTools/Private/MuT/UnrealPixelFormatOverride.cpp
--- a/Engine/Plugins/Mutable/Source/MutableTools/Private/MuT/UnrealPixelFormatOverride.cpp
+++ b/Engine/Plugins/Mutable/Source/MutableTools/Private/MuT/UnrealPixelFormatOverride.cpp
@@ -33,101 +33,101 @@ void FillBuildSettingsFromMutableFormat(FTextureBuildSettings& Settings, bool& b
 
 	switch (Format)
 	{
-	case mu::EImageFormat::IF_ASTC_4x4_RGBA_LDR:
+	case mu::EImageFormat::ASTC_4x4_RGBA_LDR:
 		Settings.TextureFormatName = TEXT("ASTC_RGBA_HQ");
 		Settings.CompressionQuality = 4; // See GetQualityFormat in TextureFormatASTC.cpp
 		bOutHasAlpha = true;
 		break;
-	case mu::EImageFormat::IF_ASTC_6x6_RGBA_LDR:
+	case mu::EImageFormat::ASTC_6x6_RGBA_LDR:
 		Settings.TextureFormatName = TEXT("ASTC_RGBA");
 		Settings.CompressionQuality = 3; // See GetQualityFormat in TextureFormatASTC.cpp
 		bOutHasAlpha = true;
 		break;
-	case mu::EImageFormat::IF_ASTC_8x8_RGBA_LDR:
+	case mu::EImageFormat::ASTC_8x8_RGBA_LDR:
 		Settings.TextureFormatName = TEXT("ASTC_RGBA");
 		Settings.CompressionQuality = 2; // See GetQualityFormat in TextureFormatASTC.cpp
 		bOutHasAlpha = true;
 		break;
-	case mu::EImageFormat::IF_ASTC_10x10_RGBA_LDR:
+	case mu::EImageFormat::ASTC_10x10_RGBA_LDR:
 		Settings.TextureFormatName = TEXT("ASTC_RGBA");
 		Settings.CompressionQuality = 1; // See GetQualityFormat in TextureFormatASTC.cpp
 		bOutHasAlpha = true;
 		break;
-	case mu::EImageFormat::IF_ASTC_12x12_RGBA_LDR:
+	case mu::EImageFormat::ASTC_12x12_RGBA_LDR:
 		Settings.TextureFormatName = TEXT("ASTC_RGBA");
 		Settings.CompressionQuality = 0; // See GetQualityFormat in TextureFormatASTC.cpp
 		bOutHasAlpha = true;
 		break;
 
-	case mu::EImageFormat::IF_ASTC_4x4_RGB_LDR:
+	case mu::EImageFormat::ASTC_4x4_RGB_LDR:
 		Settings.TextureFormatName = TEXT("ASTC_RGBA_HQ");
 		Settings.CompressionQuality = 4; // See GetQualityFormat in TextureFormatASTC.cpp
 		bOutHasAlpha = false;
 		break;
-	case mu::EImageFormat::IF_ASTC_6x6_RGB_LDR:
+	case mu::EImageFormat::ASTC_6x6_RGB_LDR:
 		Settings.TextureFormatName = TEXT("ASTC_RGBA");
 		Settings.CompressionQuality = 3; // See GetQualityFormat in TextureFormatASTC.cpp
 		bOutHasAlpha = false;
 		break;
-	case mu::EImageFormat::IF_ASTC_8x8_RGB_LDR:
+	case mu::EImageFormat::ASTC_8x8_RGB_LDR:
 		Settings.TextureFormatName = TEXT("ASTC_RGBA");
 		Settings.CompressionQuality = 2; // See GetQualityFormat in TextureFormatASTC.cpp
 		bOutHasAlpha = false;
 		break;
-	case mu::EImageFormat::IF_ASTC_10x10_RGB_LDR:
+	case mu::EImageFormat::ASTC_10x10_RGB_LDR:
 		Settings.TextureFormatName = TEXT("ASTC_RGBA");
 		Settings.CompressionQuality = 1; // See GetQualityFormat in TextureFormatASTC.cpp
 		bOutHasAlpha = false;
 		break;
-	case mu::EImageFormat::IF_ASTC_12x12_RGB_LDR:
+	case mu::EImageFormat::ASTC_12x12_RGB_LDR:
 		Settings.TextureFormatName = TEXT("ASTC_RGBA");
 		Settings.CompressionQuality = 0; // See GetQualityFormat in TextureFormatASTC.cpp
 		bOutHasAlpha = false;
 		break;
 
-	case mu::EImageFormat::IF_ASTC_4x4_RG_LDR:
+	case mu::EImageFormat::ASTC_4x4_RG_LDR:
 		Settings.TextureFormatName = TEXT("ASTC_RGB"); // There is no way to get a 4x4 RG in the ASTC compressor from a TextureFormatName
 		Settings.CompressionQuality = 4; // See GetQualityFormat in TextureFormatASTC.cpp
 		bOutHasAlpha = false;
 		break;
-	case mu::EImageFormat::IF_ASTC_6x6_RG_LDR:
+	case mu::EImageFormat::ASTC_6x6_RG_LDR:
 		Settings.TextureFormatName = TEXT("ASTC_NormalRG");
 		Settings.CompressionQuality = 3; // See GetQualityFormat in TextureFormatASTC.cpp
 		bOutHasAlpha = false;
 		break;
-	case mu::EImageFormat::IF_ASTC_8x8_RG_LDR:
+	case mu::EImageFormat::ASTC_8x8_RG_LDR:
 		Settings.TextureFormatName = TEXT("ASTC_NormalLA");
 		Settings.CompressionQuality = 2; // See GetQualityFormat in TextureFormatASTC.cpp
 		bOutHasAlpha = false;
 		break;
-	case mu::EImageFormat::IF_ASTC_10x10_RG_LDR:
+	case mu::EImageFormat::ASTC_10x10_RG_LDR:
 		Settings.TextureFormatName = TEXT("ASTC_NormalRG");
 		Settings.CompressionQuality = 1; // See GetQualityFormat in TextureFormatASTC.cpp
 		bOutHasAlpha = false;
 		break;
-	case mu::EImageFormat::IF_ASTC_12x12_RG_LDR:
+	case mu::EImageFormat::ASTC_12x12_RG_LDR:
 		Settings.TextureFormatName = TEXT("ASTC_NormalRG");
 		Settings.CompressionQuality = 0; // See GetQualityFormat in TextureFormatASTC.cpp
 		bOutHasAlpha = false;
 		break;
 
-	case mu::EImageFormat::IF_BC1:
+	case mu::EImageFormat::BC1:
 		Settings.TextureFormatName = TEXT("DXT1");
 		bOutHasAlpha = false;
 		break;
-	case mu::EImageFormat::IF_BC2:
+	case mu::EImageFormat::BC2:
 		Settings.TextureFormatName = TEXT("DXT3");
 		bOutHasAlpha = true;
 		break;
-	case mu::EImageFormat::IF_BC3:
+	case mu::EImageFormat::BC3:
 		Settings.TextureFormatName = TEXT("DXT5");
 		bOutHasAlpha = true;
 		break;
-	case mu::EImageFormat::IF_BC4:
+	case mu::EImageFormat::BC4:
 		Settings.TextureFormatName = TEXT("BC4");
 		bOutHasAlpha = false;
 		break;
-	case mu::EImageFormat::IF_BC5:
+	case mu::EImageFormat::BC5:
 		Settings.TextureFormatName = TEXT("BC5");
 		bOutHasAlpha = true;
 		break;
@@ -149,14 +149,14 @@ void MutableToImageCore(const mu::Image* InMutable, FImage& CoreImage, int32 LOD
 	ERawImageFormat::Type CoreImageFormat;
 	switch (Mutable->GetFormat())
 	{
-	case mu::EImageFormat::IF_BGRA_UBYTE: CoreImageFormat = ERawImageFormat::BGRA8; break;
-		//case mu::EImageFormat::IF_L_UBYTE: CoreImageFormat = ERawImageFormat::G8; break;
+	case mu::EImageFormat::BGRA_UByte: CoreImageFormat = ERawImageFormat::BGRA8; break;
+		//case mu::EImageFormat::L_UByte: CoreImageFormat = ERawImageFormat::G8; break;
 
 	default:
 	{
 		// Unsupported format: force conversion
 		mu::FImageOperator ImOp = mu::FImageOperator::GetDefault(mu::FImageOperator::FImagePixelFormatFunc());
-		Mutable = ImOp.ImagePixelFormat(4, InMutable, mu::EImageFormat::IF_BGRA_UBYTE,LOD);
+		Mutable = ImOp.ImagePixelFormat(4, InMutable, mu::EImageFormat::BGRA_UByte,LOD);
 		// We are extracting one LOD, so always access LOD 0 of the resulting mutable image 
 		LOD = 0;
 		CoreImageFormat = ERawImageFormat::BGRA8;
@@ -196,29 +196,29 @@ mu::EImageFormat UnrealToMutablePixelFormat(EPixelFormat PlatformFormat, bool bH
 {
 	switch (PlatformFormat)
 	{
-	case PF_ASTC_4x4: return bHasAlpha ? mu::EImageFormat::IF_ASTC_4x4_RGBA_LDR : mu::EImageFormat::IF_ASTC_4x4_RGB_LDR;
-	case PF_ASTC_6x6: return bHasAlpha ? mu::EImageFormat::IF_ASTC_6x6_RGBA_LDR : mu::EImageFormat::IF_ASTC_6x6_RGB_LDR;
-	case PF_ASTC_8x8: return bHasAlpha ? mu::EImageFormat::IF_ASTC_8x8_RGBA_LDR : mu::EImageFormat::IF_ASTC_8x8_RGB_LDR;
-	case PF_ASTC_10x10: return bHasAlpha ? mu::EImageFormat::IF_ASTC_10x10_RGBA_LDR : mu::EImageFormat::IF_ASTC_10x10_RGB_LDR;
-	case PF_ASTC_12x12: return bHasAlpha ? mu::EImageFormat::IF_ASTC_12x12_RGBA_LDR : mu::EImageFormat::IF_ASTC_12x12_RGB_LDR;
-	case PF_ASTC_4x4_NORM_RG: return mu::EImageFormat::IF_ASTC_4x4_RG_LDR;
-	case PF_ASTC_6x6_NORM_RG: return mu::EImageFormat::IF_ASTC_6x6_RG_LDR;
-	case PF_ASTC_8x8_NORM_RG: return mu::EImageFormat::IF_ASTC_8x8_RG_LDR;
-	case PF_ASTC_10x10_NORM_RG: return mu::EImageFormat::IF_ASTC_10x10_RG_LDR;
-	case PF_ASTC_12x12_NORM_RG: return mu::EImageFormat::IF_ASTC_12x12_RG_LDR;
-	case PF_DXT1: return mu::EImageFormat::IF_BC1;
-	case PF_DXT3: return mu::EImageFormat::IF_BC2;
-	case PF_DXT5: return mu::EImageFormat::IF_BC3;
-	case PF_BC4: return mu::EImageFormat::IF_BC4;
-	case PF_BC5: return mu::EImageFormat::IF_BC5;
-	case PF_G8: return mu::EImageFormat::IF_L_UBYTE;
-	case PF_L8: return mu::EImageFormat::IF_L_UBYTE;
-	case PF_A8: return mu::EImageFormat::IF_L_UBYTE;
-	case PF_R8G8B8A8: return mu::EImageFormat::IF_RGBA_UBYTE;
-	case PF_A8R8G8B8: return mu::EImageFormat::IF_RGBA_UBYTE;
-	case PF_B8G8R8A8: return mu::EImageFormat::IF_BGRA_UBYTE;
+	case PF_ASTC_4x4: return bHasAlpha ? mu::EImageFormat::ASTC_4x4_RGBA_LDR : mu::EImageFormat::ASTC_4x4_RGB_LDR;
+	case PF_ASTC_6x6: return bHasAlpha ? mu::EImageFormat::ASTC_6x6_RGBA_LDR : mu::EImageFormat::ASTC_6x6_RGB_LDR;
+	case PF_ASTC_8x8: return bHasAlpha ? mu::EImageFormat::ASTC_8x8_RGBA_LDR : mu::EImageFormat::ASTC_8x8_RGB_LDR;
+	case PF_ASTC_10x10: return bHasAlpha ? mu::EImageFormat::ASTC_10x10_RGBA_LDR : mu::EImageFormat::ASTC_10x10_RGB_LDR;
+	case PF_ASTC_12x12: return bHasAlpha ? mu::EImageFormat::ASTC_12x12_RGBA_LDR : mu::EImageFormat::ASTC_12x12_RGB_LDR;
+	case PF_ASTC_4x4_NORM_RG: return mu::EImageFormat::ASTC_4x4_RG_LDR;
+	case PF_ASTC_6x6_NORM_RG: return mu::EImageFormat::ASTC_6x6_RG_LDR;
+	case PF_ASTC_8x8_NORM_RG: return mu::EImageFormat::ASTC_8x8_RG_LDR;
+	case PF_ASTC_10x10_NORM_RG: return mu::EImageFormat::ASTC_10x10_RG_LDR;
+	case PF_ASTC_12x12_NORM_RG: return mu::EImageFormat::ASTC_12x12_RG_LDR;
+	case PF_DXT1: return mu::EImageFormat::BC1;
+	case PF_DXT3: return mu::EImageFormat::BC2;
+	case PF_DXT5: return mu::EImageFormat::BC3;
+	case PF_BC4: return mu::EImageFormat::BC4;
+	case PF_BC5: return mu::EImageFormat::BC5;
+	case PF_G8: return mu::EImageFormat::L_UByte;
+	case PF_L8: return mu::EImageFormat::L_UByte;
+	case PF_A8: return mu::EImageFormat::L_UByte;
+	case PF_R8G8B8A8: return mu::EImageFormat::RGBA_UByte;
+	case PF_A8R8G8B8: return mu::EImageFormat::RGBA_UByte;
+	case PF_B8G8R8A8: return mu::EImageFormat::BGRA_UByte;
 	default:
-		return mu::EImageFormat::IF_NONE;
+		return mu::EImageFormat::None;
 	}
 }
 
@@ -227,21 +227,21 @@ mu::EImageFormat QualityAndPerformanceFix(mu::EImageFormat Format)
 {
 	switch (Format)
 	{
-	case mu::EImageFormat::IF_ASTC_8x8_RGB_LDR:		return mu::EImageFormat::IF_ASTC_4x4_RGB_LDR; break;
-	case mu::EImageFormat::IF_ASTC_8x8_RGBA_LDR:	return mu::EImageFormat::IF_ASTC_4x4_RGBA_LDR; break;
-	case mu::EImageFormat::IF_ASTC_8x8_RG_LDR:		return mu::EImageFormat::IF_ASTC_4x4_RG_LDR; break;
-	case mu::EImageFormat::IF_ASTC_12x12_RGB_LDR:	return mu::EImageFormat::IF_ASTC_4x4_RGB_LDR; break;
-	case mu::EImageFormat::IF_ASTC_12x12_RGBA_LDR:	return mu::EImageFormat::IF_ASTC_4x4_RGBA_LDR; break;
-	case mu::EImageFormat::IF_ASTC_12x12_RG_LDR:	return mu::EImageFormat::IF_ASTC_4x4_RG_LDR; break;
-	case mu::EImageFormat::IF_ASTC_6x6_RGB_LDR:		return mu::EImageFormat::IF_ASTC_4x4_RGB_LDR; break;
-	case mu::EImageFormat::IF_ASTC_6x6_RGBA_LDR:	return mu::EImageFormat::IF_ASTC_4x4_RGBA_LDR; break;
-	case mu::EImageFormat::IF_ASTC_6x6_RG_LDR:		return mu::EImageFormat::IF_ASTC_4x4_RG_LDR; break;
-	case mu::EImageFormat::IF_ASTC_10x10_RGB_LDR:	return mu::EImageFormat::IF_ASTC_4x4_RGB_LDR; break;
-	case mu::EImageFormat::IF_ASTC_10x10_RGBA_LDR:	return mu::EImageFormat::IF_ASTC_4x4_RGBA_LDR; break;
-	case mu::EImageFormat::IF_ASTC_10x10_RG_LDR:	return mu::EImageFormat::IF_ASTC_4x4_RG_LDR; break;
+	case mu::EImageFormat::ASTC_8x8_RGB_LDR:		return mu::EImageFormat::ASTC_4x4_RGB_LDR; break;
+	case mu::EImageFormat::ASTC_8x8_RGBA_LDR:	return mu::EImageFormat::ASTC_4x4_RGBA_LDR; break;
+	case mu::EImageFormat::ASTC_8x8_RG_LDR:		return mu::EImageFormat::ASTC_4x4_RG_LDR; break;
+	case mu::EImageFormat::ASTC_12x12_RGB_LDR:	return mu::EImageFormat::ASTC_4x4_RGB_LDR; break;
+	case mu::EImageFormat::ASTC_12x12_RGBA_LDR:	return mu::EImageFormat::ASTC_4x4_RGBA_LDR; break;
+	case mu::EImageFormat::ASTC_12x12_RG_LDR:	return mu::EImageFormat::ASTC_4x4_RG_LDR; break;
+	case mu::EImageFormat::ASTC_6x6_RGB_LDR:		return mu::EImageFormat::ASTC_4x4_RGB_LDR; break;
+	case mu::EImageFormat::ASTC_6x6_RGBA_LDR:	return mu::EImageFormat::ASTC_4x4_RGBA_LDR; break;
+	case mu::EImageFormat::ASTC_6x6_RG_LDR:		return mu::EImageFormat::ASTC_4x4_RG_LDR; break;
+	case mu::EImageFormat::ASTC_10x10_RGB_LDR:	return mu::EImageFormat::ASTC_4x4_RGB_LDR; break;
+	case mu::EImageFormat::ASTC_10x10_RGBA_LDR:	return mu::EImageFormat::ASTC_4x4_RGBA_LDR; break;
+	case mu::EImageFormat::ASTC_10x10_RG_LDR:	return mu::EImageFormat::ASTC_4x4_RG_LDR; break;
 
 	// This is more of a performance fix.
-	case mu::EImageFormat::IF_BGRA_UBYTE:			return mu::EImageFormat::IF_RGBA_UBYTE;
+	case mu::EImageFormat::BGRA_UByte:			return mu::EImageFormat::RGBA_UByte;
 
 	default:
 		break;
diff --git a/Engine/Plugins/Mutable/Source/MutableTools/Public/MuT/NodeImageFormat.h b/Engine/Plugins/Mutable/Source/MutableTools/Public/MuT/NodeImageFormat.h
--- a/Engine/Plugins/Mutable/Source/MutableTools/Public/MuT/NodeImageFormat.h
+++ b/Engine/Plugins/Mutable/Source/MutableTools/Public/MuT/NodeImageFormat.h
@@ -40,7 +40,7 @@ namespace mu
 
 		//!
 		EImageFormat GetFormat() const;
-        void SetFormat(EImageFormat format, EImageFormat formatIfAlpha = EImageFormat::IF_NONE );
+        void SetFormat(EImageFormat format, EImageFormat formatIfAlpha = EImageFormat::None );
 
 		//!
 		NodeImagePtr GetSource() const;
diff --git a/Engine/Plugins/Mutable/Source/MutableTools/Public/MuT/NodeProjector.h b/Engine/Plugins/Mutable/Source/MutableTools/Public/MuT/NodeProjector.h
--- a/Engine/Plugins/Mutable/Source/MutableTools/Public/MuT/NodeProjector.h
+++ b/Engine/Plugins/Mutable/Source/MutableTools/Public/MuT/NodeProjector.h
@@ -74,7 +74,7 @@ namespace mu
 		//-----------------------------------------------------------------------------------------
 
 		//! Get the value that this node returns
-        void GetValue( PROJECTOR_TYPE* OutType,
+        void GetValue( EProjectorType* OutType,
 			FVector3f* OutPos,
 			FVector3f* OutDir,
 			FVector3f* OutUp,
@@ -82,7 +82,7 @@ namespace mu
 			float* OutProjectionAngle ) const;
 
 		//! Set the value to be returned by this node
-        void SetValue( PROJECTOR_TYPE type, FVector3f pos, FVector3f dir, FVector3f up, FVector3f scale, float projectionAngle);
+        void SetValue( EProjectorType type, FVector3f pos, FVector3f dir, FVector3f up, FVector3f scale, float projectionAngle);
 
 		//-----------------------------------------------------------------------------------------
 		// Interface pattern
@@ -129,7 +129,7 @@ namespace mu
 		void SetUid( const FString& );
 
 		//! Set the default value of the parameter.
-        void SetDefaultValue( PROJECTOR_TYPE type,
+        void SetDefaultValue( EProjectorType type,
 			FVector3f pos,
 			FVector3f dir,
 			FVector3f up,
diff --git a/Engine/Plugins/Mutable/Source/MutableTools/Public/MuT/TypeInfo.h b/Engine/Plugins/Mutable/Source/MutableTools/Public/MuT/TypeInfo.h
--- a/Engine/Plugins/Mutable/Source/MutableTools/Public/MuT/TypeInfo.h
+++ b/Engine/Plugins/Mutable/Source/MutableTools/Public/MuT/TypeInfo.h
@@ -18,15 +18,15 @@ namespace mu
 	{
 	public:
 
-		static const char* s_imageFormatName[size_t(EImageFormat::IF_COUNT)];
+		static const char* s_imageFormatName[size_t(EImageFormat::Count)];
 
-		static const char* s_meshBufferSemanticName[ MBS_COUNT ];
+		static const char* s_meshBufferSemanticName[uint32(EMeshBufferSemantic::Count)];
 
-		static const char* s_meshBufferFormatName[ MBF_COUNT ];
+		static const char* s_meshBufferFormatName[uint32(EMeshBufferFormat::Count) ];
 
 		static const char* s_blendModeName[uint32(EBlendType::_BT_COUNT) ];
 
-		static const char* s_projectorTypeName [ static_cast<uint32>(mu::PROJECTOR_TYPE::COUNT) ];
+		static const char* s_projectorTypeName [ static_cast<uint32>(mu::EProjectorType::Count) ];
 
 	};
 
