[Backout] - CL37218173
Original CL Desc
-----------------------------------------------------------------
[Mutable] Enable async CO compilation at cook time.

#rnx
#tests Standard Build Frontend

diff --git a/Engine/Plugins/Mutable/Source/CustomizableObject/Private/MuCO/CustomizableObject.cpp b/Engine/Plugins/Mutable/Source/CustomizableObject/Private/MuCO/CustomizableObject.cpp
--- a/Engine/Plugins/Mutable/Source/CustomizableObject/Private/MuCO/CustomizableObject.cpp
+++ b/Engine/Plugins/Mutable/Source/CustomizableObject/Private/MuCO/CustomizableObject.cpp
@@ -89,7 +89,7 @@ TAutoConsoleVariable<int32> CVarMutableDerivedDataCacheUsage(
 
 TAutoConsoleVariable<bool> CVarMutableAsyncCook(
 	TEXT("Mutable.CookAsync"),
-	true,
+	false,
 	TEXT("True - Customizable Objects will be compiled asynchronously during cook.\n")
 	TEXT("False - Sync compilation.\n"));
 
