[Mutable] Disable access to the DDC when cooking CO assets by default.

- mutable.DerivedDataCacheUsage set to 0 by default.
- Issues related to how data is stored to the ddc could be related to clothing mesh to mesh attachment data not being correct.

#rnx

diff --git a/Engine/Plugins/Mutable/Source/CustomizableObject/Private/MuCO/CustomizableObject.cpp b/Engine/Plugins/Mutable/Source/CustomizableObject/Private/MuCO/CustomizableObject.cpp
--- a/Engine/Plugins/Mutable/Source/CustomizableObject/Private/MuCO/CustomizableObject.cpp
+++ b/Engine/Plugins/Mutable/Source/CustomizableObject/Private/MuCO/CustomizableObject.cpp
@@ -79,7 +79,7 @@ TAutoConsoleVariable<bool> CVarMutableUseBulkData(
 
 TAutoConsoleVariable<int32> CVarMutableDerivedDataCacheUsage(
 	TEXT("mutable.DerivedDataCacheUsage"),
-	2,
+	0,
 	TEXT("Derived data cache access for cooked data.")
 	TEXT("0 - None. Disables access to the cache.")
 	TEXT("1 - Local. Allow cache requests to query and store records and values in local caches.")
