[Mutable] Fix group projector poses for optimized COs.

- Mesh projection did not have skinning buffers when optimized. Now a generic skinning format is forced when required.

#rnx

diff --git a/Engine/Plugins/Mutable/Source/MutableTools/Private/MuT/ASTOpMeshFormat.cpp b/Engine/Plugins/Mutable/Source/MutableTools/Private/MuT/ASTOpMeshFormat.cpp
--- a/Engine/Plugins/Mutable/Source/MutableTools/Private/MuT/ASTOpMeshFormat.cpp
+++ b/Engine/Plugins/Mutable/Source/MutableTools/Private/MuT/ASTOpMeshFormat.cpp
@@ -21,6 +21,7 @@
 #include "MuT/ASTOpMeshApplyPose.h"
 #include "MuT/ASTOpSwitch.h"
 
+#include "GPUSkinPublicDefs.h"
 
 namespace mu
 {
@@ -224,6 +225,50 @@ namespace
 		return pTargetMorphFormat;
 	}
 
+    Ptr<const Mesh> EnsureFormatHasSkinningBuffers(Ptr<const Mesh>& FormatMesh)
+    {
+        const FMeshBufferSet& FormatMeshVertexBuffers = FormatMesh->GetVertexBuffers();
+    
+        int32 SourceSkinningBufferIndex = -1;         
+        int32 SourceSkinningChannelIndex = -1;
+
+        // Assume bone indices implies it also has weights.
+        FormatMeshVertexBuffers.FindChannel(MBS_BONEINDICES, 0, &SourceSkinningBufferIndex, &SourceSkinningChannelIndex);		
+
+        bool bSourceHasSkinningData = SourceSkinningBufferIndex != -1;
+			
+        if (bSourceHasSkinningData)
+        {
+            return FormatMesh;
+        }
+
+        Ptr<Mesh> NewMesh = FormatMesh->Clone();
+        FMeshBufferSet& MeshBuffers = NewMesh->GetVertexBuffers();
+        
+        FMeshBuffer& Buffer = MeshBuffers.Buffers.AddDefaulted_GetRef();
+
+        FMeshBufferChannel BoneIndices;
+        BoneIndices.Semantic = MBS_BONEINDICES;
+        BoneIndices.Format = MBF_UINT16;
+        BoneIndices.SemanticIndex = 0;
+        BoneIndices.Offset = 0;
+        BoneIndices.ComponentCount = MAX_TOTAL_INFLUENCES;
+
+        FMeshBufferChannel BoneWeights;
+        BoneWeights.Semantic = MBS_BONEWEIGHTS;
+        BoneWeights.Format = MBF_NUINT16;
+        BoneWeights.SemanticIndex = 0;
+        BoneWeights.Offset = MAX_TOTAL_INFLUENCES*2;
+        BoneWeights.ComponentCount = MAX_TOTAL_INFLUENCES;
+
+
+        Buffer.ElementSize = MAX_TOTAL_INFLUENCES*4;
+        Buffer.Channels.Add(BoneIndices);
+        Buffer.Channels.Add(BoneWeights);
+
+        return NewMesh;
+    }
+
 }
 
 
@@ -325,8 +370,26 @@ mu::Ptr<ASTOp> Sink_MeshFormatAST::Visit(const mu::Ptr<ASTOp>& at, const ASTOpMe
 
 	case OP_TYPE::ME_APPLYPOSE:
 	{
+		Ptr<const Mesh> TargetFormat = FindBaseMeshConstant(currentFormatOp->Format.child());
+        TargetFormat = EnsureFormatHasSkinningBuffers(TargetFormat);
+
 		Ptr<ASTOpMeshApplyPose> NewOp = mu::Clone<ASTOpMeshApplyPose>(at);
-		NewOp->base = Visit(NewOp->base.child(), currentFormatOp);
+        mu::Ptr<ASTOpMeshFormat> NewFormat = mu::Clone<ASTOpMeshFormat>(currentFormatOp);
+        
+		mu::Ptr<ASTOpConstantResource> NewFormatConstant = new ASTOpConstantResource();
+		NewFormatConstant->Type = OP_TYPE::ME_CONSTANT; 
+		NewFormatConstant->SetValue(TargetFormat, nullptr);
+		NewFormatConstant->SourceDataDescriptor = at->GetSourceDataDescriptor();
+	
+        NewFormat->Flags = NewFormat->Flags | OP::MeshFormatArgs::OptimizeBuffers;
+
+        //TODO: Optimize, in case no skinning data is found in the format mesh a generic buffer that can represent
+        // all possible skinning formats is added. This is not optimal, we may want to add a flag to the format op
+        // to indicate it should copy the skinning from the base mesh.
+        NewFormat->Format = NewFormatConstant;
+
+        NewOp->base = Visit(NewOp->base.child(), NewFormat.get());
+		
 		newAt = NewOp;
 		break;
 	}
