CI fix to take into account changes made in 37257640

diff --git a/Engine/Plugins/Mutable/Source/MutableTools/Private/MuT/ASTOpMeshFormat.cpp b/Engine/Plugins/Mutable/Source/MutableTools/Private/MuT/ASTOpMeshFormat.cpp
--- a/Engine/Plugins/Mutable/Source/MutableTools/Private/MuT/ASTOpMeshFormat.cpp
+++ b/Engine/Plugins/Mutable/Source/MutableTools/Private/MuT/ASTOpMeshFormat.cpp
@@ -233,7 +233,7 @@ namespace
         int32 SourceSkinningChannelIndex = -1;
 
         // Assume bone indices implies it also has weights.
-        FormatMeshVertexBuffers.FindChannel(MBS_BONEINDICES, 0, &SourceSkinningBufferIndex, &SourceSkinningChannelIndex);		
+        FormatMeshVertexBuffers.FindChannel(EMeshBufferSemantic::BoneIndices, 0, &SourceSkinningBufferIndex, &SourceSkinningChannelIndex);
 
         bool bSourceHasSkinningData = SourceSkinningBufferIndex != -1;
 			
@@ -248,15 +248,15 @@ namespace
         FMeshBuffer& Buffer = MeshBuffers.Buffers.AddDefaulted_GetRef();
 
         FMeshBufferChannel BoneIndices;
-        BoneIndices.Semantic = MBS_BONEINDICES;
-        BoneIndices.Format = MBF_UINT16;
+        BoneIndices.Semantic = EMeshBufferSemantic::BoneIndices;
+        BoneIndices.Format = EMeshBufferFormat::UInt16;
         BoneIndices.SemanticIndex = 0;
         BoneIndices.Offset = 0;
         BoneIndices.ComponentCount = MAX_TOTAL_INFLUENCES;
 
         FMeshBufferChannel BoneWeights;
-        BoneWeights.Semantic = MBS_BONEWEIGHTS;
-        BoneWeights.Format = MBF_NUINT16;
+        BoneWeights.Semantic = EMeshBufferSemantic::BoneWeights;
+        BoneWeights.Format = EMeshBufferFormat::NUInt16;
         BoneWeights.SemanticIndex = 0;
         BoneWeights.Offset = MAX_TOTAL_INFLUENCES*2;
         BoneWeights.ComponentCount = MAX_TOTAL_INFLUENCES;
