[Mutable] Add compile option to disable image and mesh generation. Disable image generation for server builds.

- Now the code generation happens even if the resource generation is disabled. That way parameters will be consisten between server (with images disabled) and client builds.
- Now an AddInstanceOp value is optional. In general, all cases were already considering it optional.

#rnx

diff --git a/Engine/Plugins/Mutable/Source/CustomizableObject/Internal/MuCO/CustomizableObjectPrivate.h b/Engine/Plugins/Mutable/Source/CustomizableObject/Internal/MuCO/CustomizableObjectPrivate.h
--- a/Engine/Plugins/Mutable/Source/CustomizableObject/Internal/MuCO/CustomizableObjectPrivate.h
+++ b/Engine/Plugins/Mutable/Source/CustomizableObject/Internal/MuCO/CustomizableObjectPrivate.h
@@ -1571,6 +1571,8 @@ public:
 
 		FixesForMeshSectionMultipleOutputs,
 
+		ImageParametersInServerBuilds,
+
 		// -----<new versions can be added above this line>--------
 		LastCustomizableObjectVersion
 	};
diff --git a/Engine/Plugins/Mutable/Source/CustomizableObjectEditor/Private/MuCOE/CustomizableObjectCompileRunnable.cpp b/Engine/Plugins/Mutable/Source/CustomizableObjectEditor/Private/MuCOE/CustomizableObjectCompileRunnable.cpp
--- a/Engine/Plugins/Mutable/Source/CustomizableObjectEditor/Private/MuCOE/CustomizableObjectCompileRunnable.cpp
+++ b/Engine/Plugins/Mutable/Source/CustomizableObjectEditor/Private/MuCOE/CustomizableObjectCompileRunnable.cpp
@@ -22,7 +22,7 @@
 #include "DerivedDataCache.h"
 #include "DerivedDataRequestOwner.h"
 
-class ITargetPlatform;
+#include "Interfaces/ITargetPlatform.h"
 
 #define LOCTEXT_NAMESPACE "CustomizableObjectEditor"
 
@@ -183,6 +183,12 @@ uint32 FCustomizableObjectCompileRunnable::Run()
 	
 	CompilerOptions->SetImageTiling(Options.ImageTiling);
 
+	// On server builds we don't want the images to be generted.
+	if (Options.TargetPlatform && Options.TargetPlatform->IsServerOnly())
+	{
+		CompilerOptions->SetDisableImageGeneration(true);
+	}
+
 	mu::Ptr<mu::Compiler> Compiler = new mu::Compiler(CompilerOptions);
 
 	UE_LOG(LogMutable, Verbose, TEXT("PROFILE: [ %16.8f ] FCustomizableObjectCompileRunnable Compile start."), FPlatformTime::Seconds());
diff --git a/Engine/Plugins/Mutable/Source/CustomizableObjectEditor/Private/MuCOE/GenerateMutableSource/GenerateMutableSourceSurface.cpp b/Engine/Plugins/Mutable/Source/CustomizableObjectEditor/Private/MuCOE/GenerateMutableSource/GenerateMutableSourceSurface.cpp
--- a/Engine/Plugins/Mutable/Source/CustomizableObjectEditor/Private/MuCOE/GenerateMutableSource/GenerateMutableSourceSurface.cpp
+++ b/Engine/Plugins/Mutable/Source/CustomizableObjectEditor/Private/MuCOE/GenerateMutableSource/GenerateMutableSourceSurface.cpp
@@ -405,12 +405,6 @@ mu::Ptr<mu::NodeSurface> GenerateMutableSourceSurface(const UEdGraphPin * Pin, F
 		int32 NumImages = TypedNodeMat->GetNumParameters(EMaterialParameterType::Texture);
 		SurfNode->Images.SetNum(NumImages);
 
-		if (!GenerationContext.Options.TargetPlatform || GenerationContext.Options.TargetPlatform->IsServerOnly())
-		{
-			// Don't generate the images in the server
-			NumImages = 0;
-		}
-
 		for (int32 ImageIndex = 0; ImageIndex < NumImages; ++ImageIndex)
 		{
 			const UEdGraphPin* ImagePin = TypedNodeMat->GetParameterPin(EMaterialParameterType::Texture, ImageIndex);
@@ -717,7 +711,10 @@ mu::Ptr<mu::NodeSurface> GenerateMutableSourceSurface(const UEdGraphPin * Pin, F
 							if (GenerationContext.Options.TargetPlatform)
 							{
 								ReferenceTexture->GetTargetPlatformBuildSettings(GenerationContext.Options.TargetPlatform, BuildSettingsPerFormatPerLayer);
-								if (BuildSettingsPerFormatPerLayer.IsEmpty())
+
+								bool bIsServerOnly = GenerationContext.Options.TargetPlatform->IsServerOnly();
+								// Suppress the message for server only platforms. Images are discarded is ServerOnly. 
+								if (!bIsServerOnly && BuildSettingsPerFormatPerLayer.IsEmpty())
 								{
 									const FString ReplacedImageFormatMsg = FString::Printf(TEXT("In object [%s] for platform [%s] the unsupported image format of texture [%s] is used, RGBA_UByte will be used instead."),
 										*GenerationContext.Object->GetName(),
diff --git a/Engine/Plugins/Mutable/Source/MutableTools/Internal/MuT/AST.h b/Engine/Plugins/Mutable/Source/MutableTools/Internal/MuT/AST.h
--- a/Engine/Plugins/Mutable/Source/MutableTools/Internal/MuT/AST.h
+++ b/Engine/Plugins/Mutable/Source/MutableTools/Internal/MuT/AST.h
@@ -231,6 +231,9 @@ namespace mu
 
 		/** Function used to request an engine resource (image) from the compiler. */
 		FReferencedResourceFunc ReferencedResourceProvider;
+
+        bool bDisableImageGeneration = false;
+        bool bDisableMeshGeneration = false;
 	};
 
 
diff --git a/Engine/Plugins/Mutable/Source/MutableTools/Private/MuT/CodeOptimiser_Parameters.cpp b/Engine/Plugins/Mutable/Source/MutableTools/Private/MuT/CodeOptimiser_Parameters.cpp
--- a/Engine/Plugins/Mutable/Source/MutableTools/Private/MuT/CodeOptimiser_Parameters.cpp
+++ b/Engine/Plugins/Mutable/Source/MutableTools/Private/MuT/CodeOptimiser_Parameters.cpp
@@ -1926,34 +1926,38 @@ namespace mu
         if( type==OP_TYPE::IN_ADDIMAGE )
         {
 			ASTOpInstanceAdd* typedAt = static_cast<ASTOpInstanceAdd*>(at.get());
-			Ptr<ASTOp> imageAt = typedAt->value.child();
+			
+			if (typedAt && typedAt->value)
+			{			
+				Ptr<ASTOp> imageAt = typedAt->value.child();
 
-            // Does it have a runtime parameter in its subtree?
-            bool hasRuntimeParameter = m_hasRuntimeParamVisitor.HasAny( imageAt );
+				// Does it have a runtime parameter in its subtree?
+				bool hasRuntimeParameter = m_hasRuntimeParamVisitor.HasAny( imageAt );
 
-            if (bAlwaysUncompress || hasRuntimeParameter)
-            {
-                FImageDesc imageDesc = imageAt->GetImageDesc( true );
+				if (bAlwaysUncompress || hasRuntimeParameter)
+				{
+					FImageDesc imageDesc = imageAt->GetImageDesc( true );
 
-                // Is it a compressed format?
-				EImageFormat format = imageDesc.m_format;
-				EImageFormat uncompressedFormat = GetUncompressedFormat( format );
-                bool isCompressedFormat = (uncompressedFormat != format);
+					// Is it a compressed format?
+					EImageFormat format = imageDesc.m_format;
+					EImageFormat uncompressedFormat = GetUncompressedFormat( format );
+					bool isCompressedFormat = (uncompressedFormat != format);
 
-                if (isCompressedFormat)
-                {
-                    Ptr<ASTOpInstanceAdd> newAt = mu::Clone<ASTOpInstanceAdd>(at);
+					if (isCompressedFormat)
+					{
+						Ptr<ASTOpInstanceAdd> newAt = mu::Clone<ASTOpInstanceAdd>(at);
 
-                    // Add a new format operation to uncompress the image
-                    Ptr<ASTOpImagePixelFormat> fop = new ASTOpImagePixelFormat;
-                    fop->Format = uncompressedFormat;
-                    fop->FormatIfAlpha = uncompressedFormat;
-                    fop->Source = imageAt;
+						// Add a new format operation to uncompress the image
+						Ptr<ASTOpImagePixelFormat> fop = new ASTOpImagePixelFormat;
+						fop->Format = uncompressedFormat;
+						fop->FormatIfAlpha = uncompressedFormat;
+						fop->Source = imageAt;
 
-                    newAt->value = fop;
-                    at = newAt;
-                }
-            }
+						newAt->value = fop;
+						at = newAt;
+					}
+				}
+			}
         }
 
         return at;
diff --git a/Engine/Plugins/Mutable/Source/MutableTools/Private/MuT/Compiler.cpp b/Engine/Plugins/Mutable/Source/MutableTools/Private/MuT/Compiler.cpp
--- a/Engine/Plugins/Mutable/Source/MutableTools/Private/MuT/Compiler.cpp
+++ b/Engine/Plugins/Mutable/Source/MutableTools/Private/MuT/Compiler.cpp
@@ -5,6 +5,7 @@
 
 #include "MuT/AST.h"
 #include "MuT/ASTOpParameter.h"
+#include "MuT/ASTOpInstanceAdd.h"
 #include "MuT/CodeGenerator.h"
 #include "MuT/CodeOptimiser.h"
 #include "MuT/CompilerPrivate.h"
@@ -185,6 +186,15 @@ namespace mu
 		m_pD->OptimisationOptions.ReferencedResourceProvider = Provider;
 	}
 
+    void CompilerOptions::SetDisableImageGeneration(bool bDisabled)
+    {
+		m_pD->OptimisationOptions.bDisableImageGeneration = bDisabled;
+    }
+
+    void CompilerOptions::SetDisableMeshGeneration(bool bDisabled)
+    {
+		m_pD->OptimisationOptions.bDisableMeshGeneration = bDisabled;
+    }
 
 	void CompilerOptions::LogStats() const
 	{
@@ -217,6 +227,74 @@ namespace mu
         m_pD = 0;
     }
 
+    namespace Private
+    {
+        void RemoveDisabledAddInstanceOpsValues(const Ptr<CompilerOptions>& Options, const TArray<FStateCompilationData>& States)
+        {
+            MUTABLE_CPUPROFILER_SCOPE(RemoveDisabledAddInstanceOpsValues);
+            
+            TArray<Ptr<ASTOp>> Roots;
+            for(const FStateCompilationData& State : States)
+            {
+                Roots.Add(State.root);
+            }
+
+            TArray<Ptr<ASTOpInstanceAdd>> InstanceAddOpsToRemove;
+
+            const bool bRemoveMeshes = Options->GetPrivate()->OptimisationOptions.bDisableMeshGeneration;
+            const bool bRemoveImages = Options->GetPrivate()->OptimisationOptions.bDisableImageGeneration;
+
+            // Gather all InstanceAddOps its value needs to be removed.
+            if (bRemoveMeshes || bRemoveImages)
+            {
+                ASTOp::Traverse_TopDown_Unique_Imprecise(Roots, [&](Ptr<ASTOp> Node)
+                //ASTOp::Traverse_TopRandom_Unique_NonReentrant(Roots, [&](Ptr<ASTOp> Node)
+                {
+			        
+                    bool bProcessChildren = GetOpDataType(Node->GetOpType()) == DT_INSTANCE;
+                    if (!bProcessChildren)
+                    {
+                        return false;
+                    }
+
+                    switch (Node->GetOpType())
+                    {
+
+                    case OP_TYPE::IN_ADDMESH:
+                    {
+                        Ptr<ASTOpInstanceAdd> TypedNode = static_cast<ASTOpInstanceAdd*>(Node.get());
+                        if (bRemoveMeshes && TypedNode)
+                        {
+                            InstanceAddOpsToRemove.Add(TypedNode);
+                        }
+                        break;
+                    }
+                    case OP_TYPE::IN_ADDIMAGE:
+                    {
+                        Ptr<ASTOpInstanceAdd> TypedNode = static_cast<ASTOpInstanceAdd*>(Node.get());
+                        if (bRemoveImages && TypedNode)
+                        {
+                            InstanceAddOpsToRemove.Add(TypedNode);
+                        }
+                        break;
+                    }
+
+                    default:
+                        break;
+                    }
+
+                    return true;
+                });
+            }
+
+            // Remove the ops dangling from the InstanceAddOp but keep the InstanceAddOp for simplicity and parity with a full AST.
+            // The case of missing value is contemplated in all usages.
+            for (Ptr<ASTOpInstanceAdd>& InstanceAddOp : InstanceAddOpsToRemove)
+            {
+                InstanceAddOp->value = nullptr;
+            }
+        }
+    }
 
     //---------------------------------------------------------------------------------------------
 	TSharedPtr<Model> Compiler::Compile( const Ptr<Node>& pNode )
@@ -269,6 +347,7 @@ namespace mu
 			Parameters.Sort(SortPredicate);
 		}
 
+        mu::Private::RemoveDisabledAddInstanceOpsValues(m_pD->m_options, states);
 
         // Slow AST code verification for debugging.
         //TArray<Ptr<ASTOp>> roots;
diff --git a/Engine/Plugins/Mutable/Source/MutableTools/Public/MuT/Compiler.h b/Engine/Plugins/Mutable/Source/MutableTools/Public/MuT/Compiler.h
--- a/Engine/Plugins/Mutable/Source/MutableTools/Public/MuT/Compiler.h
+++ b/Engine/Plugins/Mutable/Source/MutableTools/Public/MuT/Compiler.h
@@ -85,6 +85,9 @@ namespace mu
 		/** */
 		void SetReferencedResourceCallback(const FReferencedResourceFunc&);
 
+        void SetDisableImageGeneration(bool bDisabled);
+        void SetDisableMeshGeneration(bool bDisabled);
+
         //! Different data packing strategies
         enum class TextureLayoutStrategy : uint8
         {
