[mutable] Remove spamming log.
#rnx

diff --git a/Engine/Plugins/Mutable/Source/MutableRuntime/Private/MuR/System.cpp b/Engine/Plugins/Mutable/Source/MutableRuntime/Private/MuR/System.cpp
--- a/Engine/Plugins/Mutable/Source/MutableRuntime/Private/MuR/System.cpp
+++ b/Engine/Plugins/Mutable/Source/MutableRuntime/Private/MuR/System.cpp
@@ -185,8 +185,6 @@ namespace mu
     {
 		LLM_SCOPE_BYNAME(TEXT("MutableRuntime"));
 
-		UE_LOG(LogMutableCore, Log, TEXT("Forcing a working memory clear."));
-
 		// rom caches
 		for (FWorkingMemoryManager::FModelCacheEntry& ModelCache : m_pD->WorkingMemoryManager.CachePerModel)
 		{
