[Mutable] Fix failed mesh generation when one of the components of a CO is empty.

#rnx

diff --git a/Engine/Plugins/Mutable/Source/CustomizableObject/Private/MuCO/CustomizableObjectInstance.cpp b/Engine/Plugins/Mutable/Source/CustomizableObject/Private/MuCO/CustomizableObjectInstance.cpp
--- a/Engine/Plugins/Mutable/Source/CustomizableObject/Private/MuCO/CustomizableObjectInstance.cpp
+++ b/Engine/Plugins/Mutable/Source/CustomizableObject/Private/MuCO/CustomizableObjectInstance.cpp
@@ -2305,7 +2305,12 @@ bool UCustomizableInstancePrivate::UpdateSkeletalMesh_PostBeginUpdate0(UCustomiz
 		{
 			continue;
 		}
-		
+
+		if (ComponentMesh->GetSurfaceCount() == 0 && !ComponentMesh->IsReference())
+		{
+			continue;
+		}
+
 		// If it is a referenced resource, only the first LOD is relevant.
 		if (ComponentMesh->IsReference())
         {
