[Mutable] Fix nullptr access when streaming Mesh LODs due to uninitialized SkinWeightProfilesData.

#rnx

diff --git a/Engine/Plugins/Mutable/Source/CustomizableObject/Internal/MuCO/UnrealConversionUtils.h b/Engine/Plugins/Mutable/Source/CustomizableObject/Internal/MuCO/UnrealConversionUtils.h
--- a/Engine/Plugins/Mutable/Source/CustomizableObject/Internal/MuCO/UnrealConversionUtils.h
+++ b/Engine/Plugins/Mutable/Source/CustomizableObject/Internal/MuCO/UnrealConversionUtils.h
@@ -90,6 +90,17 @@ namespace UnrealConversionUtils
 		const mu::Ptr<const mu::Mesh> InMutableMesh);
 	
 
+	/**
+	 *Initializes the LODResource's SkinWeightProfilesData buffer.
+	 * @param LODResource - LODRenderData to be updated.
+	 * @param bIsStreamableMesh - If StreamMeshLODs is enabled for the outer SkeletalMesh.
+	 */
+	void InitSkinWeightProfilesData(
+		FSkeletalMeshLODRenderData& LODResource,
+		bool bIsStreamableMesh
+	);
+
+
 	/**
 	 *Performs a copy of the data found on the index buffers on the mutable mesh to the buffers of the skeletal mesh
 	 * @param LODResource - LODRenderData to be updated.
diff --git a/Engine/Plugins/Mutable/Source/CustomizableObject/Private/MuCO/CustomizableObjectInstance.cpp b/Engine/Plugins/Mutable/Source/CustomizableObject/Private/MuCO/CustomizableObjectInstance.cpp
--- a/Engine/Plugins/Mutable/Source/CustomizableObject/Private/MuCO/CustomizableObjectInstance.cpp
+++ b/Engine/Plugins/Mutable/Source/CustomizableObject/Private/MuCO/CustomizableObjectInstance.cpp
@@ -5129,6 +5129,8 @@ bool UCustomizableInstancePrivate::BuildOrCopyRenderData(const TSharedRef<FUpdat
 			LOD.FirstBoneMap,
 			MeshSurfacesMetadata);
 
+		UnrealConversionUtils::InitSkinWeightProfilesData(LODResource, OperationData->bStreamMeshLODs);
+
 		if (LODResource.bStreamedDataInlined) // Non-streamable LOD
 		{
 			// Copy Vertices
diff --git a/Engine/Plugins/Mutable/Source/CustomizableObject/Private/MuCO/CustomizableObjectMeshUpdate.cpp b/Engine/Plugins/Mutable/Source/CustomizableObject/Private/MuCO/CustomizableObjectMeshUpdate.cpp
--- a/Engine/Plugins/Mutable/Source/CustomizableObject/Private/MuCO/CustomizableObjectMeshUpdate.cpp
+++ b/Engine/Plugins/Mutable/Source/CustomizableObject/Private/MuCO/CustomizableObjectMeshUpdate.cpp
@@ -273,6 +273,13 @@ void FCustomizableObjectMeshStreamIn::ConvertMesh(const FContext& Context)
 			return;
 		}
 
+		if (MutableMesh->GetVertexCount() == 0 || MutableMesh->GetSurfaceCount() == 0)
+		{
+			check(false);
+			Abort();
+			return;
+		}
+
 		FSkeletalMeshLODRenderData& LODResource = *Context.LODResourcesView[LODIndex];
 		UnrealConversionUtils::CopyMutableVertexBuffers(LODResource, MutableMesh, true);
 		UnrealConversionUtils::CopyMutableIndexBuffers(LODResource, MutableMesh);
@@ -281,6 +288,7 @@ void FCustomizableObjectMeshStreamIn::ConvertMesh(const FContext& Context)
 		// SkinWeightProfilesData
 		// MorphTargets (MorphTargetVertexInfoBuffers)
 		// Clothing (ClothVertexBuffer)
+		UnrealConversionUtils::InitSkinWeightProfilesData(LODResource, true /* bIsStreamableMesh */);
 
 		UnrealConversionUtils::UpdateSkeletalMeshLODRenderDataBuffersSize(LODResource);
 	}
diff --git a/Engine/Plugins/Mutable/Source/CustomizableObject/Private/MuCO/UnrealConversionUtils.cpp b/Engine/Plugins/Mutable/Source/CustomizableObject/Private/MuCO/UnrealConversionUtils.cpp
--- a/Engine/Plugins/Mutable/Source/CustomizableObject/Private/MuCO/UnrealConversionUtils.cpp
+++ b/Engine/Plugins/Mutable/Source/CustomizableObject/Private/MuCO/UnrealConversionUtils.cpp
@@ -15,6 +15,13 @@
 
 class USkeleton;
 
+TAutoConsoleVariable<int32> CVarMutableInitSkinWightProfileData(
+	TEXT("mutable.InitSkinWeightProfileData"),
+	1,
+	TEXT("0 - Do not initialize SkinWeightProfilesData buffer by default")
+	TEXT("1 - Initialize on streamable SkeletalMeshes")
+	TEXT("2 - Initialize on all generated SkeletalMeshes"));
+
 namespace UnrealConversionUtils
 {
 	// Hidden functions only used internally to aid other functions
@@ -462,6 +469,18 @@ namespace UnrealConversionUtils
 	}
 
 
+	void InitSkinWeightProfilesData(FSkeletalMeshLODRenderData& LODResource, bool bIsStreamableMesh)
+	{
+		const int32 InitProfileDataMode = CVarMutableInitSkinWightProfileData.GetValueOnAnyThread();
+		const bool bInitSkinWeightProfilesData = (InitProfileDataMode == 1 && bIsStreamableMesh) || InitProfileDataMode == 2;
+
+		if (bInitSkinWeightProfilesData)
+		{
+			LODResource.SkinWeightProfilesData.Init(&LODResource.SkinWeightVertexBuffer);
+		}
+	}
+
+
 	struct FAuxMutableSkinWeightVertexKey
 	{
 		FAuxMutableSkinWeightVertexKey(const uint8* InKey, uint8 InKeySize, uint32 InHash)
