Update FMutableGraphGenerationContext::LoadObject to use Cook load scope of type ECookLoadType::EditorOnly

diff --git a/Engine/Plugins/Mutable/Source/CustomizableObjectEditor/Private/MuCOE/GenerateMutableSource/GenerateMutableSource.cpp b/Engine/Plugins/Mutable/Source/CustomizableObjectEditor/Private/MuCOE/GenerateMutableSource/GenerateMutableSource.cpp
--- a/Engine/Plugins/Mutable/Source/CustomizableObjectEditor/Private/MuCOE/GenerateMutableSource/GenerateMutableSource.cpp
+++ b/Engine/Plugins/Mutable/Source/CustomizableObjectEditor/Private/MuCOE/GenerateMutableSource/GenerateMutableSource.cpp
@@ -57,9 +57,15 @@
 #include "Hash/CityHash.h"
 #include "MuCOE/Nodes/CustomizableObjectNodeCopyMaterial.h"
 #include "MuCOE/Nodes/CustomizableObjectNodeGroupProjectorParameter.h"
+#include "UObject/ICookInfo.h"
 
 #define LOCTEXT_NAMESPACE "CustomizableObjectEditor"
 
+TAutoConsoleVariable<bool> CVarLoadObjectCookLoadTypeEditorOnly(
+	TEXT("Mutable.LoadObjectCookLoadTypeEditorOnly"),
+	true,
+	TEXT("Controls whether calls to LoadObject use a CookLoadScope of type EditorOnly")
+);
 
 uint32 GetTypeHash(const FGraphCycleKey& Key)
 {
@@ -441,9 +447,19 @@ FMutableComponentInfo* FMutableGraphGenerationContext::GetCurrentComponentInfo()
 
 UObject* FMutableGraphGenerationContext::LoadObject(const FSoftObjectPtr& SoftObject, bool bParticipatingObjectsPassLoad)
 {
-	return bLoadObjects && (!bParticipatingObjectsPass || bParticipatingObjectsPassLoad) ?
-		SoftObject.LoadSynchronous() :
-		nullptr;
+	if (CVarLoadObjectCookLoadTypeEditorOnly.GetValueOnAnyThread())
+	{ 
+		FCookLoadScope CookLoadScope(ECookLoadType::EditorOnly);
+		return bLoadObjects && (!bParticipatingObjectsPass || bParticipatingObjectsPassLoad) ?
+			SoftObject.LoadSynchronous() :
+			nullptr;
+	}
+	else
+	{
+		return bLoadObjects && (!bParticipatingObjectsPass || bParticipatingObjectsPassLoad) ?
+			SoftObject.LoadSynchronous() :
+			nullptr;
+	}
 }
 
 
