[Mutable] Fix crash with UCustomizableSkeletalComponent nullptr private.

Some assets got serialized with the private part to nullptr. Recreate it if necessary once deserialized.

#rnx

diff --git a/Engine/Plugins/Mutable/Source/CustomizableObject/Private/MuCO/CustomizableSkeletalComponent.cpp b/Engine/Plugins/Mutable/Source/CustomizableObject/Private/MuCO/CustomizableSkeletalComponent.cpp
--- a/Engine/Plugins/Mutable/Source/CustomizableObject/Private/MuCO/CustomizableSkeletalComponent.cpp
+++ b/Engine/Plugins/Mutable/Source/CustomizableObject/Private/MuCO/CustomizableSkeletalComponent.cpp
@@ -11,6 +11,13 @@
 #include "MuCO/CustomizableObject.h"
 
 
+static TAutoConsoleVariable<bool> CVarCustomizableComponentPrivateNullptr(
+	TEXT("mutable.CustomizableComponentPrivateNullptr"),
+	true,
+	TEXT("Rollback. If true, force create the private part after the serialization."),
+	ECVF_Scalability);
+
+
 UCustomizableSkeletalComponentPrivate::UCustomizableSkeletalComponentPrivate()
 {
 	// This object may get instantiated into a level, as a part of UCustomizableSkeletalComponent, so needs to be public to ensure
@@ -49,6 +56,17 @@ UCustomizableSkeletalComponent::UCustomizableSkeletalComponent()
 {
 }
 
+void UCustomizableSkeletalComponent::Serialize(FArchive& Ar)
+{
+	Super::Serialize(Ar);
+	
+	// Somehow we got some assets serialized with Private = nullptr. Fix them.
+	if (Ar.IsLoading() && !Private && CVarCustomizableComponentPrivateNullptr.GetValueOnAnyThread())
+	{
+		Private = NewObject<UCustomizableSkeletalComponentPrivate>(this, FName("Private"), RF_Public);
+	}
+}
+
 
 #if WITH_EDITOR
 void UCustomizableSkeletalComponent::PostEditChangeProperty(FPropertyChangedEvent& PropertyChangedEvent)
diff --git a/Engine/Plugins/Mutable/Source/CustomizableObject/Public/MuCO/CustomizableSkeletalComponent.h b/Engine/Plugins/Mutable/Source/CustomizableObject/Public/MuCO/CustomizableSkeletalComponent.h
--- a/Engine/Plugins/Mutable/Source/CustomizableObject/Public/MuCO/CustomizableSkeletalComponent.h
+++ b/Engine/Plugins/Mutable/Source/CustomizableObject/Public/MuCO/CustomizableSkeletalComponent.h
@@ -64,6 +64,7 @@ protected:
 #if WITH_EDITOR
 	virtual void PostEditChangeProperty(FPropertyChangedEvent& PropertyChangedEvent) override;
 #endif
+	virtual void Serialize(FArchive& Ar) override;
 
 	// UObject interface
 	virtual void PostInitProperties() override;
