[mutable] More enforcement of coding standards.
- Removed unused code.
- Removed unnecessary versioning from serialization.
- Remove most of use of the custom intrusive pointer from the MutableRuntime module.
#rnx

diff --git a/Engine/Plugins/Mutable/Source/CustomizableObject/Internal/MuCO/CustomizableObjectInstanceDescriptor.h b/Engine/Plugins/Mutable/Source/CustomizableObject/Internal/MuCO/CustomizableObjectInstanceDescriptor.h
--- a/Engine/Plugins/Mutable/Source/CustomizableObject/Internal/MuCO/CustomizableObjectInstanceDescriptor.h
+++ b/Engine/Plugins/Mutable/Source/CustomizableObject/Internal/MuCO/CustomizableObjectInstanceDescriptor.h
@@ -23,7 +23,7 @@ typedef TMap<const UCustomizableObjectInstance*, FMutableUpdateCandidate> FMutab
 
 namespace mu
 {
-	class Parameters;
+	class FParameters;
 	
 	template<typename Type>
 	class Ptr;
@@ -85,7 +85,7 @@ struct CUSTOMIZABLEOBJECT_API FCustomizableObjectInstanceDescriptor
 	uint16 GetRequestedLODs(FCustomizableObjectComponentIndex) const;
 
 	// ------------------------------------------------------------
-	// Parameters
+	// FParameters
 	// ------------------------------------------------------------
 
 	const TArray<FCustomizableObjectBoolParameterValue>& GetBoolParameters() const;
@@ -234,21 +234,21 @@ struct CUSTOMIZABLEOBJECT_API FCustomizableObjectInstanceDescriptor
 	 * The added value is initialized with the default projector as set up in the editor and is the last one of the range. */
 	int32 AddValueToProjectorRange(const FString& ParamName);
 
-	/** Remove the RangeIndex element of the integer range of values from the parameter ParamName. If Range index is -1, removes the last element.
+	/** Remove the FRangeIndex element of the integer range of values from the parameter ParamName. If Range index is -1, removes the last element.
 		Returns the index of the last valid integer, -1 if no values left. */
 	int32 RemoveValueFromIntRange(const FString& ParamName, int32 RangeIndex = -1);
 
-	/** Remove the RangeIndex element of the float range of values from the parameter ParamName. If Range index is -1, removes the last element.
+	/** Remove the FRangeIndex element of the float range of values from the parameter ParamName. If Range index is -1, removes the last element.
 		Returns the index of the last valid float, -1 if no values left. */
 	int32 RemoveValueFromFloatRange(const FString& ParamName, int32 RangeIndex = -1);
 
 	/** Remove the last of the texture range of values from the parameter ParamName, returns the index of the last valid float, -1 if no values left. */
 	int32 RemoveValueFromTextureRange(const FString& ParamName);
 
-	/** Remove the RangeIndex element of the texture range of values from the parameter ParamName, returns the index of the last valid float, -1 if no values left. */
+	/** Remove the FRangeIndex element of the texture range of values from the parameter ParamName, returns the index of the last valid float, -1 if no values left. */
 	int32 RemoveValueFromTextureRange(const FString& ParamName, int32 RangeIndex);
 
-	/** Remove the RangeIndex element of the projector range of values from the parameter ParamName. If Range index is -1, removes the last element.
+	/** Remove the FRangeIndex element of the projector range of values from the parameter ParamName. If Range index is -1, removes the last element.
 		Returns the index of the last valid projector, -1 if no values left.
 	*/
 	int32 RemoveValueFromProjectorRange(const FString& ParamName, int32 RangeIndex = -1);
@@ -302,7 +302,7 @@ struct CUSTOMIZABLEOBJECT_API FCustomizableObjectInstanceDescriptor
 	void UpdateLayer(const FName& ParamName, int32 Index, const FMultilayerProjectorLayer& Layer);
 
 	/** Return a Mutable Core object containing all parameters. */
-	mu::Ptr<mu::Parameters> GetParameters() const;
+	TSharedPtr<mu::FParameters> GetParameters() const;
 
 	FString ToString() const;
 	
diff --git a/Engine/Plugins/Mutable/Source/CustomizableObject/Internal/MuCO/CustomizableObjectInstancePrivate.h b/Engine/Plugins/Mutable/Source/CustomizableObject/Internal/MuCO/CustomizableObjectInstancePrivate.h
--- a/Engine/Plugins/Mutable/Source/CustomizableObject/Internal/MuCO/CustomizableObjectInstancePrivate.h
+++ b/Engine/Plugins/Mutable/Source/CustomizableObject/Internal/MuCO/CustomizableObjectInstancePrivate.h
@@ -19,8 +19,8 @@
 
 namespace mu 
 {
-	class PhysicsBody;
-	class Mesh;
+	class FPhysicsBody;
+	class FMesh;
 	typedef uint64 FResourceID;
 }
 
@@ -41,7 +41,7 @@ class UCustomizableObjectExtension;
 extern const FString MULTILAYER_PROJECTOR_PARAMETERS_INVALID;
 
 	
-// Parameters encoding
+// FParameters encoding
 CUSTOMIZABLEOBJECT_API extern const FString NUM_LAYERS_PARAMETER_POSTFIX;
 CUSTOMIZABLEOBJECT_API extern const FString OPACITY_PARAMETER_POSTFIX;
 CUSTOMIZABLEOBJECT_API extern const FString IMAGE_PARAMETER_POSTFIX;
@@ -50,7 +50,7 @@ CUSTOMIZABLEOBJECT_API extern const FString POSE_PARAMETER_POSTFIX;
 
 /** \param OnlyLOD: If not 0, extract and convert only one single LOD from the source image.
   * \param ExtractChannel: If different than -1, extract a single-channel image with the specified source channel data. */
-CUSTOMIZABLEOBJECT_API void ConvertImage(UTexture2D* Texture, mu::Ptr<const mu::Image> MutableImage, const FMutableModelImageProperties& Props, int32 OnlyLOD = -1, int32 ExtractChannel = -1);
+CUSTOMIZABLEOBJECT_API void ConvertImage(UTexture2D* Texture, TSharedPtr<const mu::FImage> MutableImage, const FMutableModelImageProperties& Props, int32 OnlyLOD = -1, int32 ExtractChannel = -1);
 
 
 /** CustomizableObject Instance flags for internal use  */
@@ -225,7 +225,7 @@ public:
 	TArray<FSoftObjectPath> AssetsToStream;
 
 	// Only used in LiveUpdateMode to reuse core instances between updates and their temp data to speed up updates, but spend way more memory
-	mu::Instance::ID LiveUpdateModeInstanceID = 0;
+	mu::FInstance::FID LiveUpdateModeInstanceID = 0;
 
 	UCustomizableInstancePrivate();
 
@@ -295,7 +295,7 @@ private:
 	void InitSkeletalMeshData(const TSharedRef<FUpdateContextPrivate>& OperationData, USkeletalMesh* SkeletalMesh, const FMutableRefSkeletalMeshData& RefSkeletalMeshData, const UCustomizableObject& CustomizableObject, FCustomizableObjectComponentIndex ObjectComponentIndex);
 
 	bool BuildSkeletonData(const TSharedRef<FUpdateContextPrivate>& OperationData, USkeletalMesh& SkeletalMesh, const FMutableRefSkeletalMeshData& RefSkeletalMeshData, UCustomizableObject& CustomizableObject, FCustomizableObjectInstanceComponentIndex InstanceComponentIndex);
-	void BuildMeshSockets(const TSharedRef<FUpdateContextPrivate>& OperationData, USkeletalMesh* SkeletalMesh, const FModelResources& ModelResources, const FMutableRefSkeletalMeshData& RefSkeletalMeshData, mu::Ptr<const mu::Mesh> MutableMesh);
+	void BuildMeshSockets(const TSharedRef<FUpdateContextPrivate>& OperationData, USkeletalMesh* SkeletalMesh, const FModelResources& ModelResources, const FMutableRefSkeletalMeshData& RefSkeletalMeshData, TSharedPtr<const mu::FMesh> MutableMesh);
 	void BuildOrCopyElementData(const TSharedRef<FUpdateContextPrivate>& OperationData, USkeletalMesh* SkeletalMesh, UCustomizableObjectInstance* CustomizableObjectInstance, FCustomizableObjectInstanceComponentIndex InstanceComponentIndex);
 	void BuildOrCopyMorphTargetsData(const TSharedRef<FUpdateContextPrivate>& OperationData, USkeletalMesh* SkeletalMesh, const USkeletalMesh* SrcSkeletalMesh, UCustomizableObjectInstance* CustomizableObjectInstance, FCustomizableObjectInstanceComponentIndex InstanceComponentIndex);
 	bool BuildOrCopyRenderData(const TSharedRef<FUpdateContextPrivate>& OperationData, USkeletalMesh* SkeletalMesh, const USkeletalMesh* SrcSkeletalMesh, UCustomizableObjectInstance* CustomizableObjectInstance, FCustomizableObjectInstanceComponentIndex InstanceComponentIndex);
@@ -309,7 +309,7 @@ private:
 	USkeleton* MergeSkeletons(UCustomizableObject& CustomizableObject, const FMutableRefSkeletalMeshData& RefSkeletalMeshData, FCustomizableObjectComponentIndex ObjectComponentIndex, bool& bOutCreatedNewSkeleton);
 
 	//
-	UPhysicsAsset* GetOrBuildMainPhysicsAsset(const TSharedRef<FUpdateContextPrivate>& OperationData, TObjectPtr<class UPhysicsAsset> TamplateAsset, const mu::PhysicsBody* PhysicsBody, bool bDisableCollisionBetweenAssets, FCustomizableObjectInstanceComponentIndex InstanceComponentIndex);
+	UPhysicsAsset* GetOrBuildMainPhysicsAsset(const TSharedRef<FUpdateContextPrivate>& OperationData, TObjectPtr<class UPhysicsAsset> TamplateAsset, const mu::FPhysicsBody* PhysicsBody, bool bDisableCollisionBetweenAssets, FCustomizableObjectInstanceComponentIndex InstanceComponentIndex);
 	
 	// Create a transient texture and add it to the TextureTrackerArray
 	UTexture2D* CreateTexture();
diff --git a/Engine/Plugins/Mutable/Source/CustomizableObject/Internal/MuCO/CustomizableObjectMipDataProvider.h b/Engine/Plugins/Mutable/Source/CustomizableObject/Internal/MuCO/CustomizableObjectMipDataProvider.h
--- a/Engine/Plugins/Mutable/Source/CustomizableObject/Internal/MuCO/CustomizableObjectMipDataProvider.h
+++ b/Engine/Plugins/Mutable/Source/CustomizableObject/Internal/MuCO/CustomizableObjectMipDataProvider.h
@@ -15,7 +15,7 @@
 #include "CustomizableObjectMipDataProvider.generated.h"
 
 enum EPixelFormat : uint8;
-namespace mu { class Parameters; }
+namespace mu { class FParameters; }
 namespace mu { class System; }
 
 class FThreadSafeCounter;
@@ -31,8 +31,8 @@ class FMutableUpdateContext
 public:
 	FMutableUpdateContext() = default;
 	
-	FMutableUpdateContext(const FString& InCustomizableObjectPathName, const FString& InInstancePathName, mu::Ptr<mu::System> InSystem,
-		TSharedPtr<mu::Model, ESPMode::ThreadSafe> InModel, mu::Ptr<const mu::Parameters> InParameters, int32 InState);
+	FMutableUpdateContext(const FString& InCustomizableObjectPathName, const FString& InInstancePathName, TSharedPtr<mu::FSystem> InSystem,
+		TSharedPtr<mu::FModel, ESPMode::ThreadSafe> InModel, TSharedPtr<const mu::FParameters> InParameters, int32 InState);
 
 	~FMutableUpdateContext();
 
@@ -40,15 +40,15 @@ public:
 
 	const FString& GetInstancePathName() const;
 	
-	mu::Ptr<mu::System> GetSystem() const;
+	TSharedPtr<mu::FSystem> GetSystem() const;
 	
-	TSharedPtr<mu::Model, ESPMode::ThreadSafe> GetModel() const;
+	TSharedPtr<mu::FModel, ESPMode::ThreadSafe> GetModel() const;
 
-	mu::Ptr<const mu::Parameters> GetParameters() const;
+	TSharedPtr<const mu::FParameters> GetParameters() const;
 	
 	int32 GetState() const;
 
-	const TArray<mu::Ptr<const mu::Image>>& GetImageParameterValues() const;
+	const TArray<TSharedPtr<const mu::FImage>>& GetImageParameterValues() const;
 
 	// Benchmarking Utility data (it may not always be present)
 	FString CapturedDescriptor;
@@ -58,12 +58,12 @@ private:
 	FString CustomizableObjectPathName;
 	FString InstancePathName;
 
-	mu::Ptr<mu::System> System;
-	TSharedPtr<mu::Model, ESPMode::ThreadSafe> Model;
-	mu::Ptr<const mu::Parameters> Parameters;
+	TSharedPtr<mu::FSystem> System;
+	TSharedPtr<mu::FModel, ESPMode::ThreadSafe> Model;
+	TSharedPtr<const mu::FParameters> Parameters;
 	int32 State = -1;
 
-	TArray<mu::Ptr<const mu::Image>> ImageParameterValues;
+	TArray<TSharedPtr<const mu::FImage>> ImageParameterValues;
 };
 
 
@@ -90,7 +90,7 @@ struct FMutableImageOperationData
 
 	TSharedPtr<FMutableUpdateContext> UpdateContext;
 
-	mu::ImagePtrConst Result;
+	TSharedPtr<const mu::FImage> Result;
 
 	TArray<FMutableMipUpdateLevel> Levels;
 
diff --git a/Engine/Plugins/Mutable/Source/CustomizableObject/Internal/MuCO/CustomizableObjectPrivate.h b/Engine/Plugins/Mutable/Source/CustomizableObject/Internal/MuCO/CustomizableObjectPrivate.h
--- a/Engine/Plugins/Mutable/Source/CustomizableObject/Internal/MuCO/CustomizableObjectPrivate.h
+++ b/Engine/Plugins/Mutable/Source/CustomizableObject/Internal/MuCO/CustomizableObjectPrivate.h
@@ -23,7 +23,7 @@
 
 namespace mu 
 { 
-	class Model;
+	class FModel;
 	struct FBoneName;
 }
 
@@ -901,7 +901,7 @@ namespace MutablePrivate
 
 	struct CUSTOMIZABLEOBJECT_API FMutableCachedPlatformData
 	{
-		/** Serialized mu::Model */
+		/** Serialized mu::FModel */
 		TArray64<uint8> ModelData;
 
 		/** Serialized FModelResources */
@@ -928,7 +928,7 @@ namespace MutablePrivate
 	 *  Resources will be split into two buckets for non-optional and optional BulkData.
 	 */
 	void CUSTOMIZABLEOBJECT_API GenerateBulkDataFilesListWithFileLimit(
-		TSharedPtr<const mu::Model, ESPMode::ThreadSafe> Model,
+		TSharedPtr<const mu::FModel, ESPMode::ThreadSafe> Model,
 		FModelStreamableBulkData& ModelStreamableBulkData,
 		uint32 NumFilesPerBucket,
 		TArray<FFile>& OutBulkDataFiles);
@@ -936,7 +936,7 @@ namespace MutablePrivate
 	/** Generate the list of BulkData files with a soft restriction to the size of the files.
 	 */
 	void CUSTOMIZABLEOBJECT_API GenerateBulkDataFilesListWithSizeLimit(
-		TSharedPtr<const mu::Model, ESPMode::ThreadSafe> Model,
+		TSharedPtr<const mu::FModel, ESPMode::ThreadSafe> Model,
 		FModelStreamableBulkData& ModelStreamableBulkData,
 		const ITargetPlatform* TargetPlatform,
 		uint64 TargetBulkDataFileBytes,
@@ -946,7 +946,7 @@ namespace MutablePrivate
 	 * the streamable's FileIds and Offsets.
 	 */
 	void GenerateBulkDataFilesList(
-		TSharedPtr<const mu::Model, ESPMode::ThreadSafe> Model,
+		TSharedPtr<const mu::FModel, ESPMode::ThreadSafe> Model,
 		FModelStreamableBulkData& StreamableBulkData,
 		bool bUseRomTypeAndFlagsToFilter,
 		TFunctionRef<void(const FFileCategoryID&, const FClassifyNode&, TArray<FFile>&)> CreateFileList,
@@ -1037,24 +1037,24 @@ struct FModelResources
 	UPROPERTY()
 	TArray<FMutableRefSkeletalMeshData> ReferenceSkeletalMeshesData;
 
-	/** Skeletons used by the compiled mu::Model. */
+	/** Skeletons used by the compiled mu::FModel. */
 	UPROPERTY()
 	TArray<TSoftObjectPtr<USkeleton>> Skeletons;
 
-	/** Materials used by the compiled mu::Model */
+	/** Materials used by the compiled mu::FModel */
 	UPROPERTY()
 	TArray<TSoftObjectPtr<UMaterialInterface>> Materials;
 
-	/** PassThrough textures used by the mu::Model. */
+	/** PassThrough textures used by the mu::FModel. */
 	UPROPERTY()
 	TArray<TSoftObjectPtr<UTexture>> PassThroughTextures;
 
-	/** PassThrough meshes used by the mu::Model. */
+	/** PassThrough meshes used by the mu::FModel. */
 	UPROPERTY()
 	TArray<TSoftObjectPtr<USkeletalMesh>> PassThroughMeshes;
 
 #if WITH_EDITORONLY_DATA
-	/** Runtime referenced textures used by the mu::Model. */
+	/** Runtime referenced textures used by the mu::FModel. */
 	UPROPERTY()
 	TArray<TSoftObjectPtr<const UTexture>> RuntimeReferencedTextures;
 #endif
@@ -1241,7 +1241,7 @@ class CUSTOMIZABLEOBJECT_API UCustomizableObjectPrivate : public UObject
 {
 	GENERATED_BODY()
 
-	TSharedPtr<mu::Model, ESPMode::ThreadSafe> MutableModel;
+	TSharedPtr<mu::FModel, ESPMode::ThreadSafe> MutableModel;
 
 	/** Stores streamable data info to be used by MutableModel In-Game. Cooked resources. */
 	UPROPERTY()
@@ -1267,9 +1267,9 @@ class CUSTOMIZABLEOBJECT_API UCustomizableObjectPrivate : public UObject
 
 public:
 	/** Must be called after unlocking the CustomizableObject. */
-	void SetModel(const TSharedPtr<mu::Model, ESPMode::ThreadSafe>& Model, const FGuid Identifier);
-	const TSharedPtr<mu::Model, ESPMode::ThreadSafe>& GetModel();
-	const TSharedPtr<const mu::Model, ESPMode::ThreadSafe> GetModel() const;
+	void SetModel(const TSharedPtr<mu::FModel, ESPMode::ThreadSafe>& Model, const FGuid Identifier);
+	const TSharedPtr<mu::FModel, ESPMode::ThreadSafe>& GetModel();
+	const TSharedPtr<const mu::FModel, ESPMode::ThreadSafe> GetModel() const;
 
 	const FModelResources& GetModelResources() const;
 
@@ -1288,7 +1288,7 @@ public:
 	bool IsLocked() const;
 
 	/** Modify the provided mutable parameters so that the forced values for the given customizable object state are applied. */
-	void ApplyStateForcedValuesToParameters(int32 State, mu::Parameters* Parameters);
+	void ApplyStateForcedValuesToParameters(int32 State, mu::FParameters* Parameters);
 
 	int32 FindParameter(const FString& Name) const;
 	int32 FindParameterTyped(const FString& Name, EMutableParameterType Type) const;
@@ -1322,7 +1322,7 @@ public:
 #endif
 	
 	/** Rebuild ParameterProperties from the current compiled model. */
-	void UpdateParameterPropertiesFromModel(const TSharedPtr<mu::Model>& Model);
+	void UpdateParameterPropertiesFromModel(const TSharedPtr<mu::FModel>& Model);
 
 	void AddUncompiledCOWarning(const FString& AdditionalLoggingInfo);
 
@@ -1365,7 +1365,7 @@ public:
 
 	void SetIsChildObject(bool bIsChildObject);
 
-	/** Return the names used by mutable to identify which mu::Image should be considered of LowPriority. */
+	/** Return the names used by mutable to identify which mu::FImage should be considered of LowPriority. */
 	void GetLowPriorityTextureNames(TArray<FString>& OutTextureNames);
 
 	/** Return the MinLOD index to generate based on the active LODSettings (PerPlatformMinLOD or PerQualityLevelMinLOD) */
@@ -1573,6 +1573,8 @@ public:
 
 		ImageParametersInServerBuilds,
 
+		RemovedUnnecessarySerializationVersioning,
+
 		// -----<new versions can be added above this line>--------
 		LastCustomizableObjectVersion
 	};
diff --git a/Engine/Plugins/Mutable/Source/CustomizableObject/Internal/MuCO/CustomizableObjectSystemPrivate.h b/Engine/Plugins/Mutable/Source/CustomizableObject/Internal/MuCO/CustomizableObjectSystemPrivate.h
--- a/Engine/Plugins/Mutable/Source/CustomizableObject/Internal/MuCO/CustomizableObjectSystemPrivate.h
+++ b/Engine/Plugins/Mutable/Source/CustomizableObject/Internal/MuCO/CustomizableObjectSystemPrivate.h
@@ -188,7 +188,7 @@ class FMutablePendingInstanceWork
 
 	TSet<FMutablePendingInstanceDiscard, FPendingInstanceDiscardKeyFuncs> PendingInstanceDiscards;
 
-	TSet<mu::Instance::ID> PendingIDsToRelease;
+	TSet<mu::FInstance::FID> PendingIDsToRelease;
 
 public:
 	// Returns the number of pending instance updates, LOD Updates and discards.
@@ -216,13 +216,13 @@ public:
 		return PendingInstanceDiscards.CreateIterator();
 	}
 
-	TSet<mu::Instance::ID>::TIterator GetIDsToReleaseIterator()
+	TSet<mu::FInstance::FID>::TIterator GetIDsToReleaseIterator()
 	{
 		return PendingIDsToRelease.CreateIterator();
 	}
 
 	void AddDiscard(const FMutablePendingInstanceDiscard& TaskToEnqueue);
-	void AddIDRelease(mu::Instance::ID IDToRelease);
+	void AddIDRelease(mu::FInstance::FID IDToRelease);
 
 	void RemoveAllUpdatesAndDiscardsAndReleases();
 };
@@ -348,7 +348,7 @@ struct FInstanceUpdateData
 		int32 BaseMip;
 		
 		uint16 FullImageSizeX, FullImageSizeY;
-		mu::Ptr<const mu::Image> Image;
+		TSharedPtr<const mu::FImage> Image;
 		TWeakObjectPtr<UTexture2D> Cached;
 
 		bool bIsPassThrough = false;
@@ -397,7 +397,7 @@ struct FInstanceUpdateData
 	struct FLOD
 	{		
 		mu::FResourceID MeshID;
-		mu::Ptr<const mu::Mesh> Mesh;
+		TSharedPtr<const mu::FMesh> Mesh;
 
 		/** Range in the Surfaces array */
 		uint16 FirstSurface = 0;
@@ -488,7 +488,7 @@ struct FInstanceUpdateData
 
 	struct FNamedExtensionData
 	{
-		mu::ExtensionDataPtrConst Data;
+		TSharedPtr<const mu::FExtensionData> Data;
 		FName Name;
 	};
 	TArray<FNamedExtensionData> ExtendedInputPins;
@@ -559,8 +559,8 @@ private:
 
 public:
 	/** Instance parameters at the time of the operation request. */
-	mu::Ptr<mu::Parameters> Parameters; 
-	mu::Ptr<mu::System> MutableSystem;
+	TSharedPtr<mu::FParameters> Parameters; 
+	TSharedPtr<mu::FSystem> MutableSystem;
 
 	bool bOnlyUpdateIfNotGenerated = false;
 	bool bIgnoreCloseDist = false;
@@ -587,10 +587,10 @@ public:
 	/** This option comes from the operation request. It is used to reduce the number of mipmaps that mutable must generate for images.  */
 	int32 MipsToSkip = 0;
 
-	mu::Instance::ID InstanceID = 0; // Redundant
-	mu::Ptr<const mu::Instance> MutableInstance;
+	mu::FInstance::FID InstanceID = 0; // Redundant
+	TSharedPtr<const mu::FInstance> MutableInstance;
 
-	TSharedPtr<mu::Model> Model;
+	TSharedPtr<mu::FModel> Model;
 
 	// Number of possible components in the entire CO
 	// TODO: Redundant because it is in the CO
@@ -718,7 +718,7 @@ public:
 	static UCustomizableObjectSystem* SSystem;
 
 	// Pointer to the lower level mutable system that actually does the work.
-	mu::Ptr<mu::System> MutableSystem;
+	TSharedPtr<mu::FSystem> MutableSystem;
 
 	/** Store the last streaming memory size in bytes, to change it when it is safe. */
 	uint64 LastWorkingMemoryBytes = 0;
@@ -800,7 +800,7 @@ public:
 	void InitDiscardResourcesSkeletalMesh(UCustomizableObjectInstance* InCustomizableObjectInstance);
 
 	// Init the async release of a Mutable Core Instance ID and all the temp resources associated with it
-	void InitInstanceIDRelease(mu::Instance::ID IDToRelease);
+	void InitInstanceIDRelease(mu::FInstance::FID IDToRelease);
 
 	void GetMipStreamingConfig(const UCustomizableObjectInstance& Instance, bool& bOutNeverStream, int32& OutMipsToSkip) const;
 	
@@ -932,5 +932,5 @@ namespace impl
 
 
 /** Set OnlyLOD to -1 to generate all mips */
-CUSTOMIZABLEOBJECT_API FTexturePlatformData* MutableCreateImagePlatformData(mu::Ptr<const mu::Image> MutableImage, int32 OnlyLOD, uint16 FullSizeX, uint16 FullSizeY);
+CUSTOMIZABLEOBJECT_API FTexturePlatformData* MutableCreateImagePlatformData(TSharedPtr<const mu::FImage> MutableImage, int32 OnlyLOD, uint16 FullSizeX, uint16 FullSizeY);
 
diff --git a/Engine/Plugins/Mutable/Source/CustomizableObject/Internal/MuCO/UnrealConversionUtils.h b/Engine/Plugins/Mutable/Source/CustomizableObject/Internal/MuCO/UnrealConversionUtils.h
--- a/Engine/Plugins/Mutable/Source/CustomizableObject/Internal/MuCO/UnrealConversionUtils.h
+++ b/Engine/Plugins/Mutable/Source/CustomizableObject/Internal/MuCO/UnrealConversionUtils.h
@@ -13,7 +13,7 @@
 
 namespace mu
 {
-	class Mesh;
+	class FMesh;
 	class FMeshBufferSet;
 	struct FBoneName;
 }
@@ -42,7 +42,7 @@ namespace UnrealConversionUtils
 	 */
 	CUSTOMIZABLEOBJECT_API void SetupRenderSections(
 		FSkeletalMeshLODRenderData& LODResource,
-		const mu::Ptr<const mu::Mesh> InMutableMesh,
+		const mu::FMesh* InMutableMesh,
 		const TArray<mu::FBoneName>& InBoneMap,
 		const TMap<mu::FBoneName, TPair<FName, uint16>>& BoneInfoMap,
 		const int32 InFirstBoneMapIndex,
@@ -56,7 +56,7 @@ namespace UnrealConversionUtils
 	 */
 	CUSTOMIZABLEOBJECT_API void InitVertexBuffersWithDummyData(
 		FSkeletalMeshLODRenderData& LODResource,
-		const mu::Ptr<const mu::Mesh> InMutableMesh,
+		const mu::FMesh* InMutableMesh,
 		const bool bAllowCPUAccess);
 
 	/** Performs a copy of the data found on the vertex buffers on the mutable mesh to the buffers of the skeletal mesh
@@ -66,7 +66,7 @@ namespace UnrealConversionUtils
 	 */
 	CUSTOMIZABLEOBJECT_API void CopyMutableVertexBuffers(
 		FSkeletalMeshLODRenderData& LODResource,
-		const mu::Ptr<const mu::Mesh> InMutableMesh,
+		const mu::FMesh* InMutableMesh,
 		const bool bAllowCPUAccess);
 
 	
@@ -77,7 +77,7 @@ namespace UnrealConversionUtils
 	 */
 	CUSTOMIZABLEOBJECT_API void InitIndexBuffersWithDummyData(
 		FSkeletalMeshLODRenderData& LODResource,
-		const mu::Ptr<const mu::Mesh> InMutableMesh);
+		const mu::FMesh* InMutableMesh);
 
 	/**
 	 *Performs a copy of the data found on the index buffers on the mutable mesh to the buffers of the skeletal mesh
@@ -87,7 +87,7 @@ namespace UnrealConversionUtils
 	 */
 	CUSTOMIZABLEOBJECT_API bool CopyMutableIndexBuffers(
 		FSkeletalMeshLODRenderData& LODResource,
-		const mu::Ptr<const mu::Mesh> InMutableMesh);
+		const mu::FMesh* InMutableMesh);
 	
 
 	/**
diff --git a/Engine/Plugins/Mutable/Source/CustomizableObject/Internal/MuCO/UnrealMutableModelDiskStreamer.h b/Engine/Plugins/Mutable/Source/CustomizableObject/Internal/MuCO/UnrealMutableModelDiskStreamer.h
--- a/Engine/Plugins/Mutable/Source/CustomizableObject/Internal/MuCO/UnrealMutableModelDiskStreamer.h
+++ b/Engine/Plugins/Mutable/Source/CustomizableObject/Internal/MuCO/UnrealMutableModelDiskStreamer.h
@@ -14,16 +14,16 @@ class IAsyncReadFileHandle;
 class IAsyncReadRequest;
 class IBulkDataIORequest;
 class UCustomizableObject;
-namespace mu { class Model; }
+namespace mu { class FModel; }
 struct FMutableStreamableBlock;
 
-class UnrealMutableInputStream : public mu::InputStream
+class FUnrealMutableInputStream : public mu::FInputStream
 {
 public:
 
-	UnrealMutableInputStream(FArchive& ar);
+	FUnrealMutableInputStream(FArchive& ar);
 
-	// mu::InputStream interface
+	// mu::FInputStream interface
 	void Read(void* pData, uint64 size) override;
 
 private:
@@ -32,7 +32,7 @@ private:
 
 
 // Implementation of a mutable streamer using bulk storage.
-class CUSTOMIZABLEOBJECT_API FUnrealMutableModelBulkReader : public mu::ModelReader
+class CUSTOMIZABLEOBJECT_API FUnrealMutableModelBulkReader : public mu::FModelReader
 {
 public:
 	// 
@@ -60,7 +60,7 @@ public:
 	void EndStreaming();
 
 	// mu::ModelReader interface
-	FOperationID BeginReadBlock(const mu::Model*, uint32 key0, void* pBuffer, uint64 size, TFunction<void(bool bSuccess)>* CompletionCallback) override;
+	FOperationID BeginReadBlock(const mu::FModel*, uint32 key0, void* pBuffer, uint64 size, TFunction<void(bool bSuccess)>* CompletionCallback) override;
 	bool IsReadCompleted(FOperationID) override;
 	bool EndRead(FOperationID) override;
 
@@ -80,7 +80,7 @@ protected:
 	/** Streaming data for one object. */
 	struct FObjectData
 	{
-		TWeakPtr<const mu::Model> Model;
+		TWeakPtr<const mu::FModel> Model;
 
 		FString BulkFilePrefix;
 		TMap<FOperationID, FReadRequest> CurrentReadRequests;
@@ -107,13 +107,13 @@ protected:
 
 #if WITH_EDITOR
 
-class UnrealMutableOutputStream : public mu::OutputStream
+class UnrealMutableOutputStream : public mu::FOutputStream
 {
 public:
 
 	UnrealMutableOutputStream(FArchive& ar);
 
-	// mu::OutputStream interface
+	// mu::FOutputStream interface
 	void Write(const void* pData, uint64 size) override;
 
 private:
@@ -122,13 +122,13 @@ private:
 
 
 // Implementation of a mutable streamer using bulk storage.
-class CUSTOMIZABLEOBJECT_API FUnrealMutableModelBulkWriterEditor : public mu::ModelWriter
+class CUSTOMIZABLEOBJECT_API FUnrealMutableModelBulkWriterEditor : public mu::FModelWriter
 {
 public:
 	// 
 	FUnrealMutableModelBulkWriterEditor(FArchive* InMainDataArchive = nullptr, FArchive* InStreamedDataArchive = nullptr);
 
-	// mu::ModelWriter interface
+	// mu::FModelWriter interface
 	void OpenWriteFile(uint32 BlockKey) override;
 	void Write(const void* pBuffer, uint64 size) override;
 	void CloseWriteFile() override;
@@ -147,13 +147,13 @@ protected:
 
 
 // Implementation of a mutable streamer using bulk storage.
-class CUSTOMIZABLEOBJECT_API FUnrealMutableModelBulkWriterCook : public mu::ModelWriter
+class CUSTOMIZABLEOBJECT_API FUnrealMutableModelBulkWriterCook : public mu::FModelWriter
 {
 public:
 	// 
 	FUnrealMutableModelBulkWriterCook(FArchive* InMainDataArchive = nullptr, MutablePrivate::FModelStreamableData* InStreamedData = nullptr);
 
-	// mu::ModelWriter interface
+	// mu::FModelWriter interface
 	void OpenWriteFile(uint32 BlockKey) override;
 	void Write(const void* pBuffer, uint64 size) override;
 	void CloseWriteFile() override;
diff --git a/Engine/Plugins/Mutable/Source/CustomizableObject/Internal/MuCO/UnrealToMutableTextureConversionUtils.h b/Engine/Plugins/Mutable/Source/CustomizableObject/Internal/MuCO/UnrealToMutableTextureConversionUtils.h
--- a/Engine/Plugins/Mutable/Source/CustomizableObject/Internal/MuCO/UnrealToMutableTextureConversionUtils.h
+++ b/Engine/Plugins/Mutable/Source/CustomizableObject/Internal/MuCO/UnrealToMutableTextureConversionUtils.h
@@ -33,7 +33,7 @@ private:
 
 namespace mu
 {
-	class Image;
+	class FImage;
 }
 
 // Forward declarations
@@ -48,6 +48,6 @@ enum class EUnrealToMutableConversionError
     Unknown
 };
 
-CUSTOMIZABLEOBJECT_API EUnrealToMutableConversionError ConvertTextureUnrealSourceToMutable(mu::Image* OutResult, FMutableSourceTextureData&, uint8 MipmapsToSkip);
+CUSTOMIZABLEOBJECT_API EUnrealToMutableConversionError ConvertTextureUnrealSourceToMutable(mu::FImage* OutResult, FMutableSourceTextureData&, uint8 MipmapsToSkip);
 
 #endif
diff --git a/Engine/Plugins/Mutable/Source/CustomizableObject/Private/MuCO/CustomizableObject.cpp b/Engine/Plugins/Mutable/Source/CustomizableObject/Private/MuCO/CustomizableObject.cpp
--- a/Engine/Plugins/Mutable/Source/CustomizableObject/Private/MuCO/CustomizableObject.cpp
+++ b/Engine/Plugins/Mutable/Source/CustomizableObject/Private/MuCO/CustomizableObject.cpp
@@ -551,11 +551,11 @@ void UCustomizableObjectPrivate::LoadModelStreamableBulk(FArchive& MemoryReader,
 
 void UCustomizableObjectPrivate::LoadModel(FArchive& MemoryReader)
 {
-	TSharedPtr<mu::Model, ESPMode::ThreadSafe> LoadedModel;
+	TSharedPtr<mu::FModel, ESPMode::ThreadSafe> LoadedModel;
 
-	UnrealMutableInputStream Stream(MemoryReader);
-	mu::InputArchive Arch(&Stream);
-	LoadedModel = mu::Model::StaticUnserialise(Arch);
+	FUnrealMutableInputStream Stream(MemoryReader);
+	mu::FInputArchive Arch(&Stream);
+	LoadedModel = mu::FModel::StaticUnserialise(Arch);
 
 	SetModel(LoadedModel, GenerateIdentifier(*GetPublic()));
 }
@@ -1197,7 +1197,7 @@ FString UCustomizableObject::GetDesc()
 {
 	int32 States = GetStateCount();
 	int32 Params = GetParameterCount();
-	return FString::Printf(TEXT("%d States, %d Parameters"), States, Params);
+	return FString::Printf(TEXT("%d States, %d FParameters"), States, Params);
 }
 
 
@@ -1205,7 +1205,7 @@ void UCustomizableObjectPrivate::SaveEmbeddedData(FArchive& Ar)
 {
 	UE_LOG(LogMutable, Verbose, TEXT("Saving embedded data for Customizable Object [%s] now at position %d."), *GetName(), int(Ar.Tell()));
 
-	TSharedPtr<mu::Model> Model = GetModel();
+	TSharedPtr<mu::FModel> Model = GetModel();
 
 	int32 InternalVersion = Model ? CurrentSupportedVersion : -1;
 	Ar << InternalVersion;
@@ -1215,8 +1215,8 @@ void UCustomizableObjectPrivate::SaveEmbeddedData(FArchive& Ar)
 		// Serialise the entire model, but unload the streamable data first.
 		{
 			UnrealMutableOutputStream Stream(Ar);
-			mu::OutputArchive Arch(&Stream);
-			mu::Model::Serialise(Model.Get(), Arch);
+			mu::FOutputArchive Arch(&Stream);
+			mu::FModel::Serialise(Model.Get(), Arch);
 		}
 
 		UE_LOG(LogMutable, Verbose, TEXT("Saved embedded data for Customizable Object [%s] now at position %d."), *GetName(), int(Ar.Tell()));
@@ -1237,9 +1237,9 @@ void UCustomizableObjectPrivate::LoadEmbeddedData(FArchive& Ar)
 	if (ensure(CurrentSupportedVersion == InternalVersion))
 	{		
 		// Load model
-		UnrealMutableInputStream Stream(Ar);
-		mu::InputArchive Arch(&Stream);
-		TSharedPtr<mu::Model, ESPMode::ThreadSafe> Model = mu::Model::StaticUnserialise(Arch);
+		FUnrealMutableInputStream Stream(Ar);
+		mu::FInputArchive Arch(&Stream);
+		TSharedPtr<mu::FModel, ESPMode::ThreadSafe> Model = mu::FModel::StaticUnserialise(Arch);
 
 		SetModel(Model, FGuid());
 	}
@@ -1553,11 +1553,11 @@ const FString & UCustomizableObject::GetParameterName(int32 ParamIndex) const
 }
 
 
-void UCustomizableObjectPrivate::UpdateParameterPropertiesFromModel(const TSharedPtr<mu::Model>& Model)
+void UCustomizableObjectPrivate::UpdateParameterPropertiesFromModel(const TSharedPtr<mu::FModel>& Model)
 {
 	if (Model)
 	{
-		mu::ParametersPtr MutableParameters = mu::Model::NewParameters(Model);
+		TSharedPtr<mu::FParameters> MutableParameters = mu::FModel::NewParameters(Model);
 		const int32 NumParameters = MutableParameters->GetCount();
 
 		TArray<int32> TypedParametersCount;
@@ -1854,7 +1854,7 @@ float UCustomizableObject::GetFloatParameterDefaultValue(const FString& InParame
 		return FCustomizableObjectFloatParameterValue::DEFAULT_PARAMETER_VALUE;
 	}
 
-	const TSharedPtr<const mu::Model>& Model = GetPrivate()->GetModel();
+	const TSharedPtr<const mu::FModel>& Model = GetPrivate()->GetModel();
 	if (!Model)
 	{
 		checkNoEntry();
@@ -1874,7 +1874,7 @@ int32 UCustomizableObject::GetIntParameterDefaultValue(const FString& InParamete
 		return FCustomizableObjectIntParameterValue::DEFAULT_PARAMETER_VALUE;
 	}
 
-	const TSharedPtr<const mu::Model>& Model = GetPrivate()->GetModel();
+	const TSharedPtr<const mu::FModel>& Model = GetPrivate()->GetModel();
 	if (!Model)
 	{
 		checkNoEntry();
@@ -1894,7 +1894,7 @@ bool UCustomizableObject::GetBoolParameterDefaultValue(const FString& InParamete
 		return FCustomizableObjectBoolParameterValue::DEFAULT_PARAMETER_VALUE;
 	}
 
-	const TSharedPtr<const mu::Model>& Model = GetPrivate()->GetModel();
+	const TSharedPtr<const mu::FModel>& Model = GetPrivate()->GetModel();
 	if (!Model)
 	{
 		checkNoEntry();
@@ -1914,7 +1914,7 @@ FLinearColor UCustomizableObject::GetColorParameterDefaultValue(const FString& I
 		return FCustomizableObjectVectorParameterValue::DEFAULT_PARAMETER_VALUE;
 	}
 
-	const TSharedPtr<const mu::Model>& Model = GetPrivate()->GetModel();
+	const TSharedPtr<const mu::FModel>& Model = GetPrivate()->GetModel();
 	if (!Model)
 	{
 		checkNoEntry();
@@ -1937,7 +1937,7 @@ FTransform UCustomizableObject::GetTransformParameterDefaultValue(const FString&
 		return FCustomizableObjectTransformParameterValue::DEFAULT_PARAMETER_VALUE;
 	}
 
-	const TSharedPtr<const mu::Model>& Model = GetPrivate()->GetModel();
+	const TSharedPtr<const mu::FModel>& Model = GetPrivate()->GetModel();
 	if (!Model)
 	{
 		checkNoEntry();
@@ -1974,7 +1974,7 @@ FCustomizableObjectProjector UCustomizableObject::GetProjectorParameterDefaultVa
 		return FCustomizableObjectProjectorParameterValue::DEFAULT_PARAMETER_VALUE;
 	}
 
-	const TSharedPtr<const mu::Model>& Model = GetPrivate()->GetModel();
+	const TSharedPtr<const mu::FModel>& Model = GetPrivate()->GetModel();
 	if (!Model)
 	{
 		checkNoEntry();
@@ -1999,7 +1999,7 @@ FName UCustomizableObject::GetTextureParameterDefaultValue(const FString& InPara
 		return FCustomizableObjectTextureParameterValue::DEFAULT_PARAMETER_VALUE;
 	}
 
-	const TSharedPtr<const mu::Model>& Model = GetPrivate()->GetModel();
+	const TSharedPtr<const mu::FModel>& Model = GetPrivate()->GetModel();
 	if (!Model)
 	{
 		checkNoEntry();
@@ -2035,7 +2035,7 @@ bool UCustomizableObject::IsParameterMultidimensional(const int32& InParamIndex)
 }
 
 
-void UCustomizableObjectPrivate::ApplyStateForcedValuesToParameters(int32 State, mu::Parameters* Parameters)
+void UCustomizableObjectPrivate::ApplyStateForcedValuesToParameters(int32 State, mu::FParameters* Parameters)
 {
 	const FString& StateName = GetPublic()->GetStateName(State);
 	const FMutableStateData* StateData = GetModelResources().StateUIDataMap.Find(StateName);
@@ -2052,7 +2052,7 @@ void UCustomizableObjectPrivate::ApplyStateForcedValuesToParameters(int32 State,
 			continue;
 		}
 
-		bool bIsMultidimensional = Parameters->NewRangeIndex(ForcedParameterIndex).get() != nullptr;
+		bool bIsMultidimensional = Parameters->NewRangeIndex(ForcedParameterIndex) != nullptr;
 		if (!bIsMultidimensional)
 		{
 			switch (GetParameterType(ForcedParameterIndex))
@@ -2259,7 +2259,7 @@ void UModelStreamableData::Serialize(FArchive& Ar)
 }
 
 
-void UCustomizableObjectPrivate::SetModel(const TSharedPtr<mu::Model, ESPMode::ThreadSafe>& Model, const FGuid Id)
+void UCustomizableObjectPrivate::SetModel(const TSharedPtr<mu::FModel, ESPMode::ThreadSafe>& Model, const FGuid Id)
 {
 	if (MutableModel == Model
 #if WITH_EDITOR
@@ -2289,13 +2289,13 @@ void UCustomizableObjectPrivate::SetModel(const TSharedPtr<mu::Model, ESPMode::T
 }
 
 
-const TSharedPtr<mu::Model, ESPMode::ThreadSafe>& UCustomizableObjectPrivate::GetModel()
+const TSharedPtr<mu::FModel, ESPMode::ThreadSafe>& UCustomizableObjectPrivate::GetModel()
 {
 	return MutableModel;
 }
 
 
-const TSharedPtr<const mu::Model, ESPMode::ThreadSafe> UCustomizableObjectPrivate::GetModel() const
+const TSharedPtr<const mu::FModel, ESPMode::ThreadSafe> UCustomizableObjectPrivate::GetModel() const
 {
 	return MutableModel;
 }
@@ -2613,7 +2613,7 @@ namespace MutablePrivate
 	}
 	
 	void GenerateBulkDataFilesListWithFileLimit(
-		TSharedPtr<const mu::Model, ESPMode::ThreadSafe> Model,
+		TSharedPtr<const mu::FModel, ESPMode::ThreadSafe> Model,
 		FModelStreamableBulkData& ModelStreamableBulkData,
 		uint32 NumFilesPerBucket,
 		TArray<FFile>& OutBulkDataFiles)
@@ -2760,7 +2760,7 @@ namespace MutablePrivate
 	
 
 	void GenerateBulkDataFilesListWithSizeLimit(
-		TSharedPtr<const mu::Model, ESPMode::ThreadSafe> Model,
+		TSharedPtr<const mu::FModel, ESPMode::ThreadSafe> Model,
 		FModelStreamableBulkData& ModelStreamableBulkData,
 		const ITargetPlatform* TargetPlatform,
 		uint64 TargetBulkDataFileBytes,
@@ -2847,7 +2847,7 @@ namespace MutablePrivate
 
 
 	void GenerateBulkDataFilesList(
-		TSharedPtr<const mu::Model, ESPMode::ThreadSafe> Model,
+		TSharedPtr<const mu::FModel, ESPMode::ThreadSafe> Model,
 		FModelStreamableBulkData& ModelStreamableBulkData,
 		bool bUseRomTypeAndFlagsToFilter,
 		TFunctionRef<void(const FFileCategoryID&, const FClassifyNode&, TArray<FFile>&)> CreateFileList,
@@ -3230,7 +3230,7 @@ void UCustomizableObjectPrivate::LoadCompiledDataFromDDC(FCompilationOptions Opt
 			LoadModel(ModelReader);
 		}
 
-		TSharedPtr<mu::Model> Model = GetModel();
+		TSharedPtr<mu::FModel> Model = GetModel();
 		TSharedPtr<FModelStreamableBulkData> ModelStreamables = GetModelStreamableBulkData(Options.bIsCooking);
 
 		// Cache CookedPlatfomData. 
diff --git a/Engine/Plugins/Mutable/Source/CustomizableObject/Private/MuCO/CustomizableObjectInstance.cpp b/Engine/Plugins/Mutable/Source/CustomizableObject/Private/MuCO/CustomizableObjectInstance.cpp
--- a/Engine/Plugins/Mutable/Source/CustomizableObject/Private/MuCO/CustomizableObjectInstance.cpp
+++ b/Engine/Plugins/Mutable/Source/CustomizableObject/Private/MuCO/CustomizableObjectInstance.cpp
@@ -1330,7 +1330,7 @@ USkeleton* UCustomizableInstancePrivate::MergeSkeletons(UCustomizableObject& Cus
 
 namespace
 {
-	FKAggregateGeom MakeAggGeomFromMutablePhysics(int32 BodyIndex, const mu::PhysicsBody* MutablePhysicsBody)
+	FKAggregateGeom MakeAggGeomFromMutablePhysics(int32 BodyIndex, const mu::FPhysicsBody* MutablePhysicsBody)
 	{
 		FKAggregateGeom BodyAggGeom;
 
@@ -1475,7 +1475,7 @@ namespace
 	TObjectPtr<UPhysicsAsset> MakePhysicsAssetFromTemplateAndMutableBody(
 		const TSharedRef<FUpdateContextPrivate>& OperationData,
 		TObjectPtr<UPhysicsAsset> TemplateAsset,
-		const mu::PhysicsBody* MutablePhysics,
+		const mu::FPhysicsBody* MutablePhysics,
 		FCustomizableObjectInstanceComponentIndex InstanceComponentIndex)
 	{
 		check(TemplateAsset);
@@ -1602,7 +1602,7 @@ namespace
 UPhysicsAsset* UCustomizableInstancePrivate::GetOrBuildMainPhysicsAsset(
 	const TSharedRef<FUpdateContextPrivate>& OperationData,
 	TObjectPtr<UPhysicsAsset> TemplateAsset,
-	const mu::PhysicsBody* MutablePhysics,
+	const mu::FPhysicsBody* MutablePhysics,
 	bool bDisableCollisionsBetweenDifferentAssets,
 	FCustomizableObjectInstanceComponentIndex InstanceComponentIndex)
 {
@@ -2157,11 +2157,11 @@ bool UCustomizableInstancePrivate::UpdateSkeletalMesh_PostBeginUpdate0(UCustomiz
 				continue;
 			}
 
-			if (ExtensionOutput.Data->Origin == mu::ExtensionData::EOrigin::Invalid)
+			if (ExtensionOutput.Data->Origin == mu::FExtensionData::EOrigin::Invalid)
 			{
 				// Null data was produced
 				//
-				// This can happen if a node produces an ExtensionData but doesn't initialize it
+				// This can happen if a node produces an FExtensionData but doesn't initialize it
 				UE_LOG(LogMutable, Error, TEXT("Invalid data sent to Object node input pin %s"), *ExtensionOutput.Name.ToString());
 				continue;
 			}
@@ -2179,14 +2179,14 @@ bool UCustomizableInstancePrivate::UpdateSkeletalMesh_PostBeginUpdate0(UCustomiz
 			const FCustomizableObjectResourceData* ReferencedExtensionData = nullptr;
 			switch (ExtensionOutput.Data->Origin)
 			{
-				case mu::ExtensionData::EOrigin::ConstantAlwaysLoaded:
+				case mu::FExtensionData::EOrigin::ConstantAlwaysLoaded:
 				{
 					check(CustomizableObject->GetPrivate()->GetAlwaysLoadedExtensionData().IsValidIndex(ExtensionOutput.Data->Index));
 					ReferencedExtensionData = &CustomizableObject->GetPrivate()->GetAlwaysLoadedExtensionData()[ExtensionOutput.Data->Index];
 				}
 				break;
 
-				case mu::ExtensionData::EOrigin::ConstantStreamed:
+				case mu::FExtensionData::EOrigin::ConstantStreamed:
 				{
 					check(CustomizableObject->GetPrivate()->GetStreamedExtensionData().IsValidIndex(ExtensionOutput.Data->Index));
 
@@ -2295,7 +2295,7 @@ bool UCustomizableInstancePrivate::UpdateSkeletalMesh_PostBeginUpdate0(UCustomiz
 		ComponentData->LastMeshIdPerLOD.Init(MAX_uint64, MAX_MESH_LOD_COUNT);
 
 		// We need the first valid mesh. get it from the component, considering that some LOSs may have been skipped.
-		mu::Ptr<const mu::Mesh> ComponentMesh;
+		TSharedPtr<const mu::FMesh> ComponentMesh;
 		for (int32 FirstValidLODIndex = Component->FirstLOD; FirstValidLODIndex < OperationData->InstanceUpdateData.LODs.Num() && !ComponentMesh; ++FirstValidLODIndex)
 		{
 			ComponentMesh = OperationData->InstanceUpdateData.LODs[FirstValidLODIndex].Mesh;
@@ -2365,11 +2365,11 @@ bool UCustomizableInstancePrivate::UpdateSkeletalMesh_PostBeginUpdate0(UCustomiz
 		}
 
 		// Build PhysicsAsset merging physics assets coming from SubMeshes of the newly generated Mesh
-		if (mu::Ptr<const mu::PhysicsBody> MutablePhysics = ComponentMesh->GetPhysicsBody())
+		if (TSharedPtr<const mu::FPhysicsBody> MutablePhysics = ComponentMesh->GetPhysicsBody())
 		{
 			constexpr bool bDisallowCollisionBetweenAssets = true;
 			UPhysicsAsset* PhysicsAssetResult = GetOrBuildMainPhysicsAsset(
-				OperationData, RefSkeletalMeshData.PhysicsAsset, MutablePhysics.get(), bDisallowCollisionBetweenAssets, FCustomizableObjectInstanceComponentIndex(InstanceComponentIndex) );
+				OperationData, RefSkeletalMeshData.PhysicsAsset, MutablePhysics.Get(), bDisallowCollisionBetweenAssets, FCustomizableObjectInstanceComponentIndex(InstanceComponentIndex) );
 
 			SkeletalMesh->SetPhysicsAsset(PhysicsAssetResult);
 
@@ -2385,7 +2385,7 @@ bool UCustomizableInstancePrivate::UpdateSkeletalMesh_PostBeginUpdate0(UCustomiz
 		const int32 NumAdditionalPhysicsNum = ComponentMesh->AdditionalPhysicsBodies.Num();
 		for (int32 I = 0; I < NumAdditionalPhysicsNum; ++I)
 		{
-			mu::Ptr<const mu::PhysicsBody> AdditionalPhysiscBody = ComponentMesh->AdditionalPhysicsBodies[I];
+			TSharedPtr<const mu::FPhysicsBody> AdditionalPhysiscBody = ComponentMesh->AdditionalPhysicsBodies[I];
 			
 			check(AdditionalPhysiscBody);
 			if (!AdditionalPhysiscBody->bBodiesModified)
@@ -2416,7 +2416,7 @@ bool UCustomizableInstancePrivate::UpdateSkeletalMesh_PostBeginUpdate0(UCustomiz
 
 			Entry.PropertyIndex = Info.PropertyIndex;
 			Entry.PhysicsAsset = MakePhysicsAssetFromTemplateAndMutableBody(
-				OperationData,PhysicsAssetTemplate, AdditionalPhysiscBody.get(), InstanceComponentIndex);
+				OperationData,PhysicsAssetTemplate, AdditionalPhysiscBody.Get(), InstanceComponentIndex);
 		}
 
 		// Add sockets from the SkeletalMesh of reference and from the MutableMesh
@@ -3196,7 +3196,7 @@ void FMutableUpdateCandidate::ApplyLODUpdateParamsToInstance(FUpdateContextPriva
 
 // The memory allocated in the function and pointed by the returned pointer is owned by the caller and must be freed. 
 // If assigned to a UTexture2D, it will be freed by that UTexture2D
-FTexturePlatformData* MutableCreateImagePlatformData(mu::Ptr<const mu::Image> MutableImage, int32 OnlyLOD, uint16 FullSizeX, uint16 FullSizeY)
+FTexturePlatformData* MutableCreateImagePlatformData(TSharedPtr<const mu::FImage> MutableImage, int32 OnlyLOD, uint16 FullSizeX, uint16 FullSizeY)
 {
 	int32 SizeX = FMath::Max(MutableImage->GetSize()[0], FullSizeX);
 	int32 SizeY = FMath::Max(MutableImage->GetSize()[1], FullSizeY);
@@ -3323,13 +3323,13 @@ FTexturePlatformData* MutableCreateImagePlatformData(mu::Ptr<const mu::Image> Mu
 		{
 		case 1:
 			PlatformFormat = PF_R8;
-			MutableImage = ImOp.ImagePixelFormat(0, MutableImage.get(), mu::EImageFormat::L_UByte);
+			MutableImage = ImOp.ImagePixelFormat(0, MutableImage.Get(), mu::EImageFormat::L_UByte);
 			break;
 		case 2:
 		case 3:
 		case 4:
 			PlatformFormat = PF_R8G8B8A8;
-			MutableImage = ImOp.ImagePixelFormat(0, MutableImage.get(), mu::EImageFormat::RGBA_UByte);
+			MutableImage = ImOp.ImagePixelFormat(0, MutableImage.Get(), mu::EImageFormat::RGBA_UByte);
 			break;
 		default: 
 			// Absolutely worst case
@@ -3469,7 +3469,7 @@ FTexturePlatformData* MutableCreateImagePlatformData(mu::Ptr<const mu::Image> Mu
 }
 
 
-void ConvertImage(UTexture2D* Texture, mu::Ptr<const mu::Image> MutableImage, const FMutableModelImageProperties& Props, int OnlyLOD, int32 ExtractChannel)
+void ConvertImage(UTexture2D* Texture, TSharedPtr<const mu::FImage> MutableImage, const FMutableModelImageProperties& Props, int OnlyLOD, int32 ExtractChannel)
 {
 	MUTABLE_CPUPROFILER_SCOPE(UCustomizableInstancePrivate::ConvertImage);
 
@@ -3482,7 +3482,7 @@ void ConvertImage(UTexture2D* Texture, mu::Ptr<const mu::Image> MutableImage, co
 	{
 		mu::FImageOperator ImOp = mu::FImageOperator::GetDefault(mu::FImageOperator::FImagePixelFormatFunc());
 
-		MutableImage = ImOp.ImagePixelFormat( 4, MutableImage.get(), mu::EImageFormat::RGBA_UByte );
+		MutableImage = ImOp.ImagePixelFormat( 4, MutableImage.Get(), mu::EImageFormat::RGBA_UByte );
 
 		uint8_t Channel = uint8_t( FMath::Clamp(ExtractChannel,0,3) );
 		MutableImage = ImOp.ImageSwizzle( mu::EImageFormat::L_UByte, &MutableImage, &Channel );
@@ -3495,7 +3495,7 @@ void ConvertImage(UTexture2D* Texture, mu::Ptr<const mu::Image> MutableImage, co
 		UE_LOG(LogMutable, Warning, TEXT("Building instance: a texture was generated in RGB format, which is slow to convert to Unreal."));
 
 		// Expand the image.
-		mu::ImagePtr Converted = new mu::Image(MutableImage->GetSizeX(), MutableImage->GetSizeY(), MutableImage->GetLODCount(), mu::EImageFormat::RGBA_UByte, mu::EInitializationType::NotInitialized);
+		TSharedPtr<mu::FImage> Converted = MakeShared<mu::FImage>(MutableImage->GetSizeX(), MutableImage->GetSizeY(), MutableImage->GetLODCount(), mu::EImageFormat::RGBA_UByte, mu::EInitializationType::NotInitialized);
 
 		for (int32 LODIndex = 0; LODIndex < Converted->GetLODCount(); ++LODIndex)
 		{
@@ -3521,7 +3521,7 @@ void ConvertImage(UTexture2D* Texture, mu::Ptr<const mu::Image> MutableImage, co
 		MUTABLE_CPUPROFILER_SCOPE(Swizzle);
 		// Swizzle the image.
 		// \TODO: Raise a warning?
-		mu::ImagePtr Converted = new mu::Image(MutableImage->GetSizeX(), MutableImage->GetSizeY(), 1, mu::EImageFormat::RGBA_UByte, mu::EInitializationType::NotInitialized);
+		TSharedPtr<mu::FImage> Converted = MakeShared<mu::FImage>(MutableImage->GetSizeX(), MutableImage->GetSizeY(), 1, mu::EImageFormat::RGBA_UByte, mu::EInitializationType::NotInitialized);
 		int32 PixelCount = MutableImage->GetSizeX() * MutableImage->GetSizeY();
 
 		const uint8* pSource = MutableImage->GetLODData(0);
@@ -3846,7 +3846,7 @@ bool UCustomizableInstancePrivate::BuildSkeletonData(const TSharedRef<FUpdateCon
 }
 
 
-void UCustomizableInstancePrivate::BuildMeshSockets(const TSharedRef<FUpdateContextPrivate>& OperationData, USkeletalMesh* SkeletalMesh, const FModelResources& ModelResources, const FMutableRefSkeletalMeshData& RefSkeletalMeshData, mu::Ptr<const mu::Mesh> MutableMesh)
+void UCustomizableInstancePrivate::BuildMeshSockets(const TSharedRef<FUpdateContextPrivate>& OperationData, USkeletalMesh* SkeletalMesh, const FModelResources& ModelResources, const FMutableRefSkeletalMeshData& RefSkeletalMeshData, TSharedPtr<const mu::FMesh> MutableMesh)
 {
 	// Build mesh sockets.
 	MUTABLE_CPUPROFILER_SCOPE(UCustomizableInstancePrivate::BuildMeshSockets);
@@ -4107,7 +4107,7 @@ void UCustomizableInstancePrivate::BuildOrCopyClothingData(const TSharedRef<FUpd
 				continue;
 			}
 
-			if (mu::MeshPtrConst MutableMesh = LOD.Mesh)
+			if (TSharedPtr<const mu::FMesh> MutableMesh = LOD.Mesh)
 			{
 				const mu::FMeshBufferSet& MeshSet = MutableMesh->GetVertexBuffers();
 				const mu::FMeshBufferSet& IndicesSet = MutableMesh->GetIndexBuffers();
@@ -4827,7 +4827,7 @@ void UCustomizableInstancePrivate::BuildOrCopyClothingData(const TSharedRef<FUpd
 		{
 			const FInstanceUpdateData::FLOD& LOD = OperationData->InstanceUpdateData.LODs[Component->FirstLOD+LODIndex];
 
-			if (mu::MeshPtrConst MutableMesh = LOD.Mesh)
+			if (TSharedPtr<const mu::FMesh> MutableMesh = LOD.Mesh)
 			{
 				TArray<FMeshToMeshVertData>& MappingData = LodMappingData[LODIndex];
 				TArray<FClothBufferIndexMapping>& ClothingIndexMapping = LodClothingIndexMapping[LODIndex];
@@ -5123,7 +5123,7 @@ bool UCustomizableInstancePrivate::BuildOrCopyRenderData(const TSharedRef<FUpdat
 		// Set RenderSections
 		UnrealConversionUtils::SetupRenderSections(
 			LODResource,
-			LOD.Mesh,
+			LOD.Mesh.Get(),
 			OperationData->InstanceUpdateData.BoneMaps,
 			BoneInfoMap,
 			LOD.FirstBoneMap,
@@ -5136,11 +5136,11 @@ bool UCustomizableInstancePrivate::BuildOrCopyRenderData(const TSharedRef<FUpdat
 			// Copy Vertices
 			UnrealConversionUtils::CopyMutableVertexBuffers(
 				LODResource,
-				LOD.Mesh,
+				LOD.Mesh.Get(),
 				SkeletalMesh->GetLODInfo(LODIndex)->bAllowCPUAccess);
 
 			// Copy indices.
-			if (!UnrealConversionUtils::CopyMutableIndexBuffers(LODResource, LOD.Mesh))
+			if (!UnrealConversionUtils::CopyMutableIndexBuffers(LODResource, LOD.Mesh.Get()))
 			{
 				// End with failure
 				return false;
@@ -5204,11 +5204,11 @@ bool UCustomizableInstancePrivate::BuildOrCopyRenderData(const TSharedRef<FUpdat
 			// Init VertexBuffers for streaming
 			UnrealConversionUtils::InitVertexBuffersWithDummyData(
 				LODResource,
-				LOD.Mesh,
+				LOD.Mesh.Get(),
 				SkeletalMesh->GetLODInfo(LODIndex)->bAllowCPUAccess);
 
 			// Init IndexBuffers for streaming
-			UnrealConversionUtils::InitIndexBuffersWithDummyData(LODResource, LOD.Mesh);
+			UnrealConversionUtils::InitIndexBuffersWithDummyData(LODResource, LOD.Mesh.Get());
 
 			// SkinWeightProfilesInfo Not supported yet
 		}
@@ -5374,7 +5374,7 @@ UE::Tasks::FTask UCustomizableInstancePrivate::LoadAdditionalAssetsAndData(
 			{
 				const FInstanceUpdateData::FLOD& LOD = OperationData->InstanceUpdateData.LODs[Component.FirstLOD + LODIndex];
 
-				mu::Ptr<const mu::Mesh> MutableMesh = LOD.Mesh;
+				TSharedPtr<const mu::FMesh> MutableMesh = LOD.Mesh;
 				if (!MutableMesh)
 				{
 					continue;
@@ -6507,7 +6507,7 @@ void UCustomizableInstancePrivate::BuildMaterials(const TSharedRef<FUpdateContex
 					{
 						const FInstanceUpdateData::FImage& Image = OperationData->InstanceUpdateData.Images[Surface.FirstImage + ImageIndex];
 						FString KeyName = Image.Name.ToString();
-						mu::ImagePtrConst MutableImage = Image.Image;
+						TSharedPtr<const mu::FImage> MutableImage = Image.Image;
 
 						UTexture2D* MutableTexture = nullptr; // Texture generated by mutable
 						UTexture* PassThroughTexture = nullptr; // Texture not generated by mutable
diff --git a/Engine/Plugins/Mutable/Source/CustomizableObject/Private/MuCO/CustomizableObjectInstanceDescriptor.cpp b/Engine/Plugins/Mutable/Source/CustomizableObject/Private/MuCO/CustomizableObjectInstanceDescriptor.cpp
--- a/Engine/Plugins/Mutable/Source/CustomizableObject/Private/MuCO/CustomizableObjectInstanceDescriptor.cpp
+++ b/Engine/Plugins/Mutable/Source/CustomizableObject/Private/MuCO/CustomizableObjectInstanceDescriptor.cpp
@@ -10,7 +10,6 @@
 #include "MuCO/DefaultImageProvider.h"
 #include "MuCO/MutableProjectorTypeUtils.h"
 #include "MuR/Model.h"
-#include "MuR/MutableMemory.h"
 #include "MuR/Parameters.h"
 #include "MuR/Ptr.h"
 
@@ -507,7 +506,7 @@ void FCustomizableObjectInstanceDescriptor::SetBuildParameterRelevancy(const boo
 }
 
 
-mu::ParametersPtr FCustomizableObjectInstanceDescriptor::GetParameters() const
+TSharedPtr<mu::FParameters> FCustomizableObjectInstanceDescriptor::GetParameters() const
 {
 	if (!CustomizableObject)
 	{
@@ -525,7 +524,7 @@ mu::ParametersPtr FCustomizableObjectInstanceDescriptor::GetParameters() const
 		return nullptr;
 	}
 
-	mu::ParametersPtr MutableParameters = mu::Model::NewParameters( CustomizableObject->GetPrivate()->GetModel() );
+	TSharedPtr<mu::FParameters> MutableParameters = mu::FModel::NewParameters( CustomizableObject->GetPrivate()->GetModel() );
 
 	const int32 ParamCount = MutableParameters->GetCount();
 	for (int32 ParamIndex = 0; ParamIndex < ParamCount; ++ParamIndex)
@@ -556,7 +555,7 @@ mu::ParametersPtr FCustomizableObjectInstanceDescriptor::GetParameters() const
 			{
 				if (IntParameter.ParameterName.Equals(Name, ESearchCase::CaseSensitive) || (Uid.IsValid() && IntParameter.Id == Uid))
 				{
-					if (mu::RangeIndexPtr RangeIdxPtr = MutableParameters->NewRangeIndex(ParamIndex))
+					if (TSharedPtr<mu::FRangeIndex> RangeIdxPtr = MutableParameters->NewRangeIndex(ParamIndex))
 					{
 						for (int32 RangeIndex = 0; RangeIndex < IntParameter.ParameterRangeValueNames.Num(); ++RangeIndex)
 						{
@@ -564,7 +563,7 @@ mu::ParametersPtr FCustomizableObjectInstanceDescriptor::GetParameters() const
 
 							const FString& RangeValueName = IntParameter.ParameterRangeValueNames[RangeIndex];
 							const int32 Value = CustomizableObject->FindIntParameterValue(ParamIndex, RangeValueName);
-							MutableParameters->SetIntValue(ParamIndex, Value, RangeIdxPtr);
+							MutableParameters->SetIntValue(ParamIndex, Value, RangeIdxPtr.Get());
 						}
 					}
 					else
@@ -586,12 +585,12 @@ mu::ParametersPtr FCustomizableObjectInstanceDescriptor::GetParameters() const
 			{
 				if (FloatParameter.ParameterName == Name || (Uid.IsValid() && FloatParameter.Id == Uid))
 				{
-					if (mu::RangeIndexPtr RangeIdxPtr = MutableParameters->NewRangeIndex(ParamIndex))
+					if (TSharedPtr<mu::FRangeIndex> RangeIdxPtr = MutableParameters->NewRangeIndex(ParamIndex))
 					{
 						for (int32 RangeIndex = 0; RangeIndex < FloatParameter.ParameterRangeValues.Num(); ++RangeIndex)
 						{
 							RangeIdxPtr->SetPosition(0, RangeIndex);
-							MutableParameters->SetFloatValue(ParamIndex, FloatParameter.ParameterRangeValues[RangeIndex], RangeIdxPtr);
+							MutableParameters->SetFloatValue(ParamIndex, FloatParameter.ParameterRangeValues[RangeIndex], RangeIdxPtr.Get());
 						}
 					}
 					else
@@ -642,7 +641,7 @@ mu::ParametersPtr FCustomizableObjectInstanceDescriptor::GetParameters() const
 			{
 				if (ProjectorParameter.ParameterName == Name || (Uid.IsValid() && ProjectorParameter.Id == Uid))
 				{
-					auto CopyProjector = [&MutableParameters, ParamIndex](const FCustomizableObjectProjector& Value, const mu::RangeIndexPtr RangeIdxPtr = nullptr)
+					auto CopyProjector = [&MutableParameters, ParamIndex](const FCustomizableObjectProjector& Value, const mu::FRangeIndex* RangeIdxPtr = nullptr)
 					{
 						switch (Value.ProjectionType)
 						{
@@ -681,12 +680,12 @@ mu::ParametersPtr FCustomizableObjectInstanceDescriptor::GetParameters() const
 
 					CopyProjector(ProjectorParameter.Value);
 
-					if (const mu::RangeIndexPtr RangeIdxPtr = MutableParameters->NewRangeIndex(ParamIndex))
+					if (const TSharedPtr<mu::FRangeIndex> RangeIdxPtr = MutableParameters->NewRangeIndex(ParamIndex))
 					{
 						for (int32 RangeIndex = 0; RangeIndex < ProjectorParameter.RangeValues.Num(); ++RangeIndex)
 						{
 							RangeIdxPtr->SetPosition(0, RangeIndex);
-							CopyProjector(ProjectorParameter.RangeValues[RangeIndex], RangeIdxPtr);
+							CopyProjector(ProjectorParameter.RangeValues[RangeIndex], RangeIdxPtr.Get());
 						}
 					}
 				}
@@ -701,12 +700,12 @@ mu::ParametersPtr FCustomizableObjectInstanceDescriptor::GetParameters() const
 			{
 				if (TextureParameter.ParameterName == Name || (Uid.IsValid() && TextureParameter.Id == Uid))
 				{
-					if (mu::RangeIndexPtr RangeIdxPtr = MutableParameters->NewRangeIndex(ParamIndex))
+					if (TSharedPtr<mu::FRangeIndex> RangeIdxPtr = MutableParameters->NewRangeIndex(ParamIndex))
 					{
 						for (int32 RangeIndex = 0; RangeIndex < TextureParameter.ParameterRangeValues.Num(); ++RangeIndex)
 						{
 							RangeIdxPtr->SetPosition(0, RangeIndex);
-							MutableParameters->SetImageValue(ParamIndex, TextureParameter.ParameterRangeValues[RangeIndex], RangeIdxPtr);
+							MutableParameters->SetImageValue(ParamIndex, TextureParameter.ParameterRangeValues[RangeIndex], RangeIdxPtr.Get());
 						}
 					}
 					else
@@ -786,7 +785,7 @@ void FCustomizableObjectInstanceDescriptor::ReloadParameters()
 		return;
 	}
 
-	mu::ParametersPtr MutableParameters = mu::Model::NewParameters(CustomizableObject->GetPrivate()->GetModel());
+	TSharedPtr<mu::FParameters> MutableParameters = mu::FModel::NewParameters(CustomizableObject->GetPrivate()->GetModel());
 
 	int32 ParamCount = MutableParameters->GetCount();
 	for (int32 ParamIndex = 0; ParamIndex < ParamCount; ++ParamIndex)
@@ -849,7 +848,7 @@ void FCustomizableObjectInstanceDescriptor::ReloadParameters()
 					return false;
 				};
 				
-				if (mu::RangeIndexPtr RangeIdxPtr = MutableParameters->NewRangeIndex(ParamIndex)) // Is multidimensional
+				if (TSharedPtr<mu::FRangeIndex> RangeIdxPtr = MutableParameters->NewRangeIndex(ParamIndex)) // Is multidimensional
 				{
 					// Get num of ranges (layers) from the instance
 					int32 ValueCount = Result->ParameterRangeValueNames.Num();
@@ -864,7 +863,7 @@ void FCustomizableObjectInstanceDescriptor::ReloadParameters()
 						}
 						else
 						{
-							const int32 Value = MutableParameters->GetIntValue(ParamIndex, RangeIdxPtr);
+							const int32 Value = MutableParameters->GetIntValue(ParamIndex, RangeIdxPtr.Get());
 							const FString AuxParameterValueName = CustomizableObject->FindIntParameterValueName(ParamIndex, Value);
 							Param.ParameterRangeValueNames.Add(AuxParameterValueName);
 						}
@@ -894,13 +893,13 @@ void FCustomizableObjectInstanceDescriptor::ReloadParameters()
 			} 
 			else // Not found in Instance Parameters. Use Mutable Parameters.
 			{
-				if (mu::RangeIndexPtr RangeIdxPtr = MutableParameters->NewRangeIndex(ParamIndex))
+				if (TSharedPtr<mu::FRangeIndex> RangeIdxPtr = MutableParameters->NewRangeIndex(ParamIndex))
 				{
 					const int32 ValueCount = MutableParameters->GetValueCount(ParamIndex);
 
 					for (int32 ValueIndex = 0; ValueIndex < ValueCount; ++ValueIndex)
 					{	
-						const mu::RangeIndexPtr RangeValueIdxPtr = MutableParameters->GetValueIndex(ParamIndex, ValueIndex);
+						const TSharedPtr<mu::FRangeIndex> RangeValueIdxPtr = MutableParameters->GetValueIndex(ParamIndex, ValueIndex);
 						const int32 RangeIndex = RangeValueIdxPtr->GetPosition(0);
 
 						if (!Param.ParameterRangeValueNames.IsValidIndex(RangeIndex))
@@ -908,7 +907,7 @@ void FCustomizableObjectInstanceDescriptor::ReloadParameters()
 							Param.ParameterRangeValueNames.AddDefaulted(RangeIndex + 1 - Param.ParameterRangeValueNames.Num());
 						}
 
-						const int32 Value = MutableParameters->GetIntValue(ParamIndex, RangeValueIdxPtr);
+						const int32 Value = MutableParameters->GetIntValue(ParamIndex, RangeValueIdxPtr.Get());
 						const FString AuxParameterValueName = CustomizableObject->FindIntParameterValueName(ParamIndex, Value);
 						Param.ParameterRangeValueNames[RangeIndex] = AuxParameterValueName;
 					}
@@ -938,7 +937,7 @@ void FCustomizableObjectInstanceDescriptor::ReloadParameters()
 			
 			if (FCustomizableObjectFloatParameterValue* Result = OldFloatParameters.FindByPredicate(FindByNameAndUid))
 			{	
-				if (mu::RangeIndexPtr RangeIdxPtr = MutableParameters->NewRangeIndex(ParamIndex))
+				if (TSharedPtr<mu::FRangeIndex> RangeIdxPtr = MutableParameters->NewRangeIndex(ParamIndex))
 				{
 					Param.ParameterRangeValues = Result->ParameterRangeValues;
 				}
@@ -949,13 +948,13 @@ void FCustomizableObjectInstanceDescriptor::ReloadParameters()
 			} 
 			else // Not found in Instance Parameters. Use Mutable Parameters.
 			{
-				if (mu::RangeIndexPtr RangeIdxPtr = MutableParameters->NewRangeIndex(ParamIndex))
+				if (TSharedPtr<mu::FRangeIndex> RangeIdxPtr = MutableParameters->NewRangeIndex(ParamIndex))
 				{
 					const int32 ValueCount = MutableParameters->GetValueCount(ParamIndex);
 
 					for (int32 ValueIndex = 0; ValueIndex < ValueCount; ++ValueIndex)
 					{	
-						mu::RangeIndexPtr RangeValueIdxPtr = MutableParameters->GetValueIndex(ParamIndex, ValueIndex);
+						TSharedPtr<mu::FRangeIndex> RangeValueIdxPtr = MutableParameters->GetValueIndex(ParamIndex, ValueIndex);
 						int32 RangeIndex = RangeValueIdxPtr->GetPosition(0);
 
 						if (!Param.ParameterRangeValues.IsValidIndex(RangeIndex))
@@ -963,7 +962,7 @@ void FCustomizableObjectInstanceDescriptor::ReloadParameters()
 							Param.ParameterRangeValues.AddDefaulted(RangeIndex + 1 - Param.ParameterRangeValues.Num());
 						}
 
-						Param.ParameterRangeValues[RangeIndex] = MutableParameters->GetFloatValue(ParamIndex, RangeValueIdxPtr);
+						Param.ParameterRangeValues[RangeIndex] = MutableParameters->GetFloatValue(ParamIndex, RangeValueIdxPtr.Get());
 					}
 				}
 				else
@@ -1044,7 +1043,7 @@ void FCustomizableObjectInstanceDescriptor::ReloadParameters()
 			
 			if (FCustomizableObjectProjectorParameterValue* Result = OldProjectorParameters.FindByPredicate(FindByNameAndUid))
 			{	
-				if (mu::RangeIndexPtr RangeIdxPtr = MutableParameters->NewRangeIndex(ParamIndex))
+				if (TSharedPtr<mu::FRangeIndex> RangeIdxPtr = MutableParameters->NewRangeIndex(ParamIndex))
 				{
 					Param.RangeValues = Result->RangeValues;
 					Param.Value.ProjectionType = DefaultProjectorValue.ProjectionType;
@@ -1062,7 +1061,7 @@ void FCustomizableObjectInstanceDescriptor::ReloadParameters()
 			} 
 			else // Not found in Instance Parameters. Use Mutable Parameters.
 			{
-				auto GetProjector = [&MutableParameters, ParamIndex](FCustomizableObjectProjector &Value, const mu::RangeIndexPtr RangeIndex = nullptr)
+				auto GetProjector = [&MutableParameters, ParamIndex](FCustomizableObjectProjector &Value, const TSharedPtr<mu::FRangeIndex> RangeIndex = nullptr)
 				{
 					mu::EProjectorType Type;
 					MutableParameters->GetProjectorValue(ParamIndex,
@@ -1072,7 +1071,7 @@ void FCustomizableObjectInstanceDescriptor::ReloadParameters()
 						&Value.Up,
 						&Value.Scale,
 						&Value.Angle,
-						RangeIndex);
+						RangeIndex.Get());
 					
 					Value.ProjectionType = ProjectorUtils::GetEquivalentProjectorType(Type);
 					if (Value.ProjectionType == ECustomizableObjectProjectorType::Cylindrical)
@@ -1088,13 +1087,13 @@ void FCustomizableObjectInstanceDescriptor::ReloadParameters()
 
 				GetProjector(Param.Value);
 
-				if (mu::RangeIndexPtr RangeIdxPtr = MutableParameters->NewRangeIndex(ParamIndex))
+				if (TSharedPtr<mu::FRangeIndex> RangeIdxPtr = MutableParameters->NewRangeIndex(ParamIndex))
 				{
 					const int32 ValueCount = MutableParameters->GetValueCount(ParamIndex);
 
 					for (int32 ValueIndex = 0; ValueIndex < ValueCount; ++ValueIndex)
 					{	
-						mu::RangeIndexPtr RangeValueIdxPtr = MutableParameters->GetValueIndex(ParamIndex, ValueIndex);
+						TSharedPtr<mu::FRangeIndex> RangeValueIdxPtr = MutableParameters->GetValueIndex(ParamIndex, ValueIndex);
 						int32 RangeIndex = RangeValueIdxPtr->GetPosition(0);
 
 						if (!Param.RangeValues.IsValidIndex(RangeIndex))
@@ -1124,7 +1123,7 @@ void FCustomizableObjectInstanceDescriptor::ReloadParameters()
 			
 			if (FCustomizableObjectTextureParameterValue* Result = OldTextureParameters.FindByPredicate(FindByNameAndUid))
 			{	
-				if (mu::RangeIndexPtr RangeIdxPtr = MutableParameters->NewRangeIndex(ParamIndex))
+				if (TSharedPtr<mu::FRangeIndex> RangeIdxPtr = MutableParameters->NewRangeIndex(ParamIndex))
 				{
 					Param.ParameterRangeValues = Result->ParameterRangeValues;
 				}
@@ -1135,13 +1134,13 @@ void FCustomizableObjectInstanceDescriptor::ReloadParameters()
 			} 
 			else // Not found in Instance Parameters. Use Mutable Parameters.
 			{
-				if (mu::RangeIndexPtr RangeIdxPtr = MutableParameters->NewRangeIndex(ParamIndex))
+				if (TSharedPtr<mu::FRangeIndex> RangeIdxPtr = MutableParameters->NewRangeIndex(ParamIndex))
 				{
 					const int32 ValueCount = MutableParameters->GetValueCount(ParamIndex);
 
 					for (int32 ValueIndex = 0; ValueIndex < ValueCount; ++ValueIndex)
 					{	
-						mu::RangeIndexPtr RangeValueIdxPtr = MutableParameters->GetValueIndex(ParamIndex, ValueIndex);
+						TSharedPtr<mu::FRangeIndex> RangeValueIdxPtr = MutableParameters->GetValueIndex(ParamIndex, ValueIndex);
 						int32 RangeIndex = RangeValueIdxPtr->GetPosition(0);
 
 						if (!Param.ParameterRangeValues.IsValidIndex(RangeIndex))
@@ -1155,7 +1154,7 @@ void FCustomizableObjectInstanceDescriptor::ReloadParameters()
 							}
 						}
 
-						Param.ParameterRangeValues[RangeIndex] = MutableParameters->GetImageValue(ParamIndex, RangeValueIdxPtr);
+						Param.ParameterRangeValues[RangeIndex] = MutableParameters->GetImageValue(ParamIndex, RangeValueIdxPtr.Get());
 					}
 				}
 				else
diff --git a/Engine/Plugins/Mutable/Source/CustomizableObject/Private/MuCO/CustomizableObjectMeshUpdate.cpp b/Engine/Plugins/Mutable/Source/CustomizableObject/Private/MuCO/CustomizableObjectMeshUpdate.cpp
--- a/Engine/Plugins/Mutable/Source/CustomizableObject/Private/MuCO/CustomizableObjectMeshUpdate.cpp
+++ b/Engine/Plugins/Mutable/Source/CustomizableObject/Private/MuCO/CustomizableObjectMeshUpdate.cpp
@@ -106,7 +106,7 @@ void FCustomizableObjectMeshStreamIn::DoCancelMeshUpdate(const FContext& Context
 
 namespace impl
 {
-	void Task_Mutable_UpdateMesh_End(const TSharedPtr<FMutableMeshOperationData> OperationData, TRefCountPtr<FCustomizableObjectMeshStreamIn>& Task, mu::Instance::ID InstanceID)
+	void Task_Mutable_UpdateMesh_End(const TSharedPtr<FMutableMeshOperationData> OperationData, TRefCountPtr<FCustomizableObjectMeshStreamIn>& Task, mu::FInstance::FID InstanceID)
 	{
 		MUTABLE_CPUPROFILER_SCOPE(Task_Mutable_UpdateMesh_End);
 
@@ -127,7 +127,7 @@ namespace impl
 	void Task_Mutable_UpdateMesh_Loop(
 		const TSharedPtr<FMutableMeshOperationData> OperationData,
 		TRefCountPtr<FCustomizableObjectMeshStreamIn>& Task,
-		mu::Instance::ID InstanceID,
+		mu::FInstance::FID InstanceID,
 		int32 LODIndex)
 	{
 		MUTABLE_CPUPROFILER_SCOPE(Task_Mutable_UpdateMesh_Loop);
@@ -139,7 +139,7 @@ namespace impl
 		}
 
 		const mu::FResourceID MeshID = OperationData->MeshIDs[LODIndex];
-		UE::Tasks::TTask<mu::Ptr<const mu::Mesh>> GetMeshTask = OperationData->System->GetMesh(InstanceID, MeshID);
+		UE::Tasks::TTask<TSharedPtr<const mu::FMesh>> GetMeshTask = OperationData->System->GetMesh(InstanceID, MeshID);
 
 		UE::Tasks::AddNested(UE::Tasks::Launch(TEXT("Task_MutableGetMeshes_GetMesh_Post"), [=]() mutable
 			{
@@ -156,8 +156,8 @@ namespace impl
 		MUTABLE_CPUPROFILER_SCOPE(Task_Mutable_UpdateMesh);
 		TRACE_BEGIN_REGION(UE_MUTABLE_UPDATE_MESH_REGION);
 
-		mu::SystemPtr System = OperationData->System;
-		const TSharedPtr<mu::Model, ESPMode::ThreadSafe> Model = OperationData->Model;
+		TSharedPtr<mu::FSystem> System = OperationData->System;
+		const TSharedPtr<mu::FModel, ESPMode::ThreadSafe> Model = OperationData->Model;
 
 #if WITH_EDITOR
 		// Recompiling a CO in the editor will invalidate the previously generated Model. Check that it is valid before accessing the streamed data.
@@ -169,14 +169,14 @@ namespace impl
 #endif
 
 		// For now, we are forcing the recreation of mutable-side instances with every update.
-		mu::Instance::ID InstanceID = System->NewInstance(Model);
+		mu::FInstance::FID InstanceID = System->NewInstance(Model);
 		UE_LOG(LogMutable, Verbose, TEXT("Creating Mutable instance with id [%d] for a mesh update"), InstanceID);
 
 		// LOD mask, set to all ones to build all LODs
 		const uint32 LODMask = 0xFFFFFFFF;
 
 		// Main instance generation step
-		mu::Ptr<const mu::Instance> Instance = System->BeginUpdate(InstanceID, OperationData->Parameters, OperationData->State, LODMask);
+		TSharedPtr<const mu::FInstance> Instance = System->BeginUpdate(InstanceID, OperationData->Parameters, OperationData->State, LODMask);
 		check(Instance);
 
 		Task_Mutable_UpdateMesh_Loop(OperationData, Task, InstanceID, OperationData->PendingFirstLODIdx);
@@ -264,7 +264,7 @@ void FCustomizableObjectMeshStreamIn::ConvertMesh(const FContext& Context)
 
 	for (int32 LODIndex = PendingFirstLODIdx; LODIndex < CurrentFirstLODIdx; ++LODIndex)
 	{
-		mu::MeshPtrConst MutableMesh = OperationData->Meshes[LODIndex];
+		TSharedPtr<const mu::FMesh> MutableMesh = OperationData->Meshes[LODIndex];
 
 		if (!MutableMesh)
 		{
@@ -281,8 +281,8 @@ void FCustomizableObjectMeshStreamIn::ConvertMesh(const FContext& Context)
 		}
 
 		FSkeletalMeshLODRenderData& LODResource = *Context.LODResourcesView[LODIndex];
-		UnrealConversionUtils::CopyMutableVertexBuffers(LODResource, MutableMesh, true);
-		UnrealConversionUtils::CopyMutableIndexBuffers(LODResource, MutableMesh);
+		UnrealConversionUtils::CopyMutableVertexBuffers(LODResource, MutableMesh.Get(), true);
+		UnrealConversionUtils::CopyMutableIndexBuffers(LODResource, MutableMesh.Get());
 
 		// TODO PERE: Implement missing buffers
 		// SkinWeightProfilesData
diff --git a/Engine/Plugins/Mutable/Source/CustomizableObject/Private/MuCO/CustomizableObjectMeshUpdate.h b/Engine/Plugins/Mutable/Source/CustomizableObject/Private/MuCO/CustomizableObjectMeshUpdate.h
--- a/Engine/Plugins/Mutable/Source/CustomizableObject/Private/MuCO/CustomizableObjectMeshUpdate.h
+++ b/Engine/Plugins/Mutable/Source/CustomizableObject/Private/MuCO/CustomizableObjectMeshUpdate.h
@@ -16,22 +16,22 @@ class UCustomizableObjectSkeletalMesh;
 
 namespace mu 
 { 
-	class Parameters;
+	class FParameters;
 }
 
 struct FMutableMeshOperationData
 {
 	/** Mutable */
-	mu::Ptr<mu::System> System;
-	TSharedPtr<mu::Model, ESPMode::ThreadSafe> Model;
-	mu::Ptr<const mu::Parameters> Parameters;
+	TSharedPtr<mu::FSystem> System;
+	TSharedPtr<mu::FModel, ESPMode::ThreadSafe> Model;
+	TSharedPtr<const mu::FParameters> Parameters;
 	int32 State = -1;
 
 	// IDs of the meshes to generate per LOD.
 	TArray<mu::FResourceID> MeshIDs;
 
 	// Meshes generated per LOD
-	TArray<mu::MeshPtrConst> Meshes;
+	TArray<TSharedPtr<const mu::FMesh>> Meshes;
 
 	/** SkeletalMesh */
 	/** The resident first LOD resource index.With domain = [0, ResourceState.NumLODs[.NOT THE ASSET LOD INDEX! */ 
@@ -71,7 +71,7 @@ private:
 	/** Cancels the MeshUpdate task. */
 	void CancelMeshUpdate(const FContext& Context);
 
-	/** Converts from mu::Mesh to FSkeletalMeshLODRenderData. */
+	/** Converts from mu::FMesh to FSkeletalMeshLODRenderData. */
 	void ConvertMesh(const FContext& Context);
 
 	/** TODO: Stream priority. */
diff --git a/Engine/Plugins/Mutable/Source/CustomizableObject/Private/MuCO/CustomizableObjectMipDataProvider.cpp b/Engine/Plugins/Mutable/Source/CustomizableObject/Private/MuCO/CustomizableObjectMipDataProvider.cpp
--- a/Engine/Plugins/Mutable/Source/CustomizableObject/Private/MuCO/CustomizableObjectMipDataProvider.cpp
+++ b/Engine/Plugins/Mutable/Source/CustomizableObject/Private/MuCO/CustomizableObjectMipDataProvider.cpp
@@ -25,9 +25,9 @@ UMutableTextureMipDataProviderFactory::UMutableTextureMipDataProviderFactory(con
 FMutableUpdateContext::FMutableUpdateContext(
 	const FString& InCustomizableObjectPathName,
 	const FString& InInstancePathName,
-	mu::Ptr<mu::System> InSystem,
-	TSharedPtr<mu::Model, ESPMode::ThreadSafe> InModel,
-	mu::Ptr<const mu::Parameters> InParameters,
+	TSharedPtr<mu::FSystem> InSystem,
+	TSharedPtr<mu::FModel, ESPMode::ThreadSafe> InModel,
+	TSharedPtr<const mu::FParameters> InParameters,
 	int32 InState) :
 	CustomizableObjectPathName(InCustomizableObjectPathName),
 	InstancePathName(InInstancePathName),
@@ -67,19 +67,19 @@ const FString& FMutableUpdateContext::GetInstancePathName() const
 }
 
 
-mu::Ptr<mu::System> FMutableUpdateContext::GetSystem() const
+TSharedPtr<mu::FSystem> FMutableUpdateContext::GetSystem() const
 {
 	return System;
 }
 
 
-TSharedPtr<mu::Model, ESPMode::ThreadSafe> FMutableUpdateContext::GetModel() const
+TSharedPtr<mu::FModel, ESPMode::ThreadSafe> FMutableUpdateContext::GetModel() const
 {
 	return Model;
 }
 
 
-mu::Ptr<const mu::Parameters> FMutableUpdateContext::GetParameters() const
+TSharedPtr<const mu::FParameters> FMutableUpdateContext::GetParameters() const
 {
 	return Parameters;
 }
@@ -91,7 +91,7 @@ int32 FMutableUpdateContext::GetState() const
 }
 
 
-const TArray<mu::Ptr<const mu::Image>>& FMutableUpdateContext::GetImageParameterValues() const
+const TArray<TSharedPtr<const mu::FImage>>& FMutableUpdateContext::GetImageParameterValues() const
 {
 	return ImageParameterValues;
 }
@@ -149,9 +149,9 @@ namespace Impl
 
 		// This runs in a worker thread.
 		check(OperationData.IsValid());
-		check(OperationData->UpdateContext->GetSystem().get());
+		check(OperationData->UpdateContext->GetSystem());
 		check(OperationData->UpdateContext->GetModel());
-		check(OperationData->UpdateContext->GetParameters().get());
+		check(OperationData->UpdateContext->GetParameters());
 
 		if (!OperationData.IsValid())
 		{
@@ -160,8 +160,8 @@ namespace Impl
 
 		TRACE_BEGIN_REGION(UE_MUTABLE_MIPDATA_PROVIDER_UPDATE_IMAGE_REGION);
 		
-		mu::Ptr<mu::System> System = OperationData->UpdateContext->GetSystem();
-		const TSharedPtr<mu::Model, ESPMode::ThreadSafe> Model = OperationData->UpdateContext->GetModel();
+		TSharedPtr<mu::FSystem> System = OperationData->UpdateContext->GetSystem();
+		const TSharedPtr<mu::FModel, ESPMode::ThreadSafe> Model = OperationData->UpdateContext->GetModel();
 
 		auto EndUpdateImage = [](TSharedPtr<FMutableImageOperationData>& OperationData)
 		{
@@ -193,10 +193,10 @@ namespace Impl
 #endif
 
 		// For now, we are forcing the recreation of mutable-side instances with every update.
-		mu::Instance::ID InstanceID = System->NewInstance(Model);
+		mu::FInstance::FID InstanceID = System->NewInstance(Model);
 		UE_LOG(LogMutable, Verbose, TEXT("Creating Mutable instance with id [%d] for a single UpdateImage"), InstanceID)
 
-		mu::Ptr<const mu::Instance> Instance;
+		TSharedPtr<const mu::FInstance> Instance;
 
 		// Main instance generation step
 		{
@@ -216,7 +216,7 @@ namespace Impl
 		// runtime cannot remember all the resources it has built, and only remembers a fixed amount.
 		mu::FResourceID MipImageID = Instance->GetImageId(ImageRef.Component, ImageRef.LOD, SurfaceIndex, ImageRef.Image);
 
-		UE::Tasks::TTask<mu::Ptr<const mu::Image>> GetImageTask = 
+		UE::Tasks::TTask<TSharedPtr<const mu::FImage>> GetImageTask = 
 				System->GetImage(InstanceID, MipImageID, ImageRef.BaseMip + OperationData->MipsToSkip, ImageRef.LOD);
 
 		UE::Tasks::AddNested(UE::Tasks::Launch(TEXT("MipDataProvider_EndUpdateImagesTask"),
@@ -224,7 +224,7 @@ namespace Impl
 				{
 					check(GetImageTask.IsCompleted());
 
-					mu::Ptr<const mu::Image> ResultImage = GetImageTask.GetResult();
+					TSharedPtr<const mu::FImage> ResultImage = GetImageTask.GetResult();
 
 					check(ResultImage);
 
@@ -241,7 +241,7 @@ namespace Impl
 						UE_LOG(LogMutable, Warning, TEXT("Mutable generated an incomplete mip chain for image."));
 
 						// Force the right number of mips. The missing data will be black.
-						mu::Ptr<mu::Image> NewImage = new mu::Image(ResultImage->GetSizeX(), ResultImage->GetSizeY(), FullMipCount, ResultImage->GetFormat(), mu::EInitializationType::Black);
+						TSharedPtr<mu::FImage> NewImage = MakeShared<mu::FImage>(ResultImage->GetSizeX(), ResultImage->GetSizeY(), FullMipCount, ResultImage->GetFormat(), mu::EInitializationType::Black);
 
 						// Formats with BytesPerBlock == 0 will not allocate memory. This type of images are not expected here.
 						check(!NewImage->DataStorage.IsEmpty());
@@ -444,7 +444,7 @@ bool FMutableTextureMipDataProvider::PollMips(const FTextureUpdateSyncOptions& S
 		// The counter must be zero meaning the Mutable image operation has finished
 		check(SyncOptions.Counter->GetValue() == 0);
 
-		mu::Ptr<const mu::Image> Image = OperationData->Result;
+		TSharedPtr<const mu::FImage> Image = OperationData->Result;
 		
 		int32 ImageLODCount = 0;
 		if (Image)
diff --git a/Engine/Plugins/Mutable/Source/CustomizableObject/Private/MuCO/CustomizableObjectSystem.cpp b/Engine/Plugins/Mutable/Source/CustomizableObject/Private/MuCO/CustomizableObjectSystem.cpp
--- a/Engine/Plugins/Mutable/Source/CustomizableObject/Private/MuCO/CustomizableObjectSystem.cpp
+++ b/Engine/Plugins/Mutable/Source/CustomizableObject/Private/MuCO/CustomizableObjectSystem.cpp
@@ -169,7 +169,7 @@ FUpdateContextPrivate::FUpdateContextPrivate(UCustomizableObjectInstance& InInst
 	MutableSystem = UCustomizableObjectSystem::GetInstance()->GetPrivate()->MutableSystem;	
 	check(MutableSystem);
 	
-	InInstance.GetCustomizableObject()->GetPrivate()->ApplyStateForcedValuesToParameters(CapturedDescriptor.GetState(), Parameters.get());
+	InInstance.GetCustomizableObject()->GetPrivate()->ApplyStateForcedValuesToParameters(CapturedDescriptor.GetState(), Parameters.Get());
 
 	UCustomizableObjectSystem* System = UCustomizableObjectSystem::GetInstance();
 	System->GetPrivate()->CacheTextureParameters(CapturedDescriptor.GetTextureParameters());
@@ -433,7 +433,7 @@ void FMutablePendingInstanceWork::AddDiscard(const FMutablePendingInstanceDiscar
 }
 
 
-void FMutablePendingInstanceWork::AddIDRelease(mu::Instance::ID IDToRelease)
+void FMutablePendingInstanceWork::AddIDRelease(mu::FInstance::FID IDToRelease)
 {
 	PendingIDsToRelease.Add(IDToRelease);
 }
@@ -534,12 +534,11 @@ void UCustomizableObjectSystem::InitSystem()
 	Private->LastWorkingMemoryBytes = CVarWorkingMemoryKB->GetInt() * 1024;
 	Private->LastGeneratedResourceCacheSize = CVarGeneratedResourcesCacheSize.GetValueOnGameThread();
 
-	const mu::Ptr<mu::Settings> pSettings = new mu::Settings;
-	check(pSettings);
-	pSettings->SetProfile(false);
-	pSettings->SetWorkingMemoryBytes(Private->LastWorkingMemoryBytes);
+	mu::FSettings Settings;
+	Settings.SetProfile(false);
+	Settings.SetWorkingMemoryBytes(Private->LastWorkingMemoryBytes);
 	Private->ExtensionDataStreamer = MakeShared<FUnrealExtensionDataStreamer>(Private);
-	Private->MutableSystem = new mu::System(pSettings, Private->ExtensionDataStreamer);
+	Private->MutableSystem = MakeShared<mu::FSystem>(Settings, Private->ExtensionDataStreamer);
 	check(Private->MutableSystem);
 
 	Private->Streamer = MakeShared<FUnrealMutableModelBulkReader>();
@@ -821,7 +820,7 @@ void FinishUpdateGlobal(const TSharedRef<FUpdateContextPrivate>& Context)
 
 			if (SystemPrivate)
 			{
-				// Cache new Texture Parameters
+				// Cache new Texture FParameters
 				SystemPrivate->CacheTextureParameters(PrivateInstance->CommittedDescriptor.GetTextureParameters());
 			}
 			
@@ -1323,7 +1322,7 @@ void UCustomizableObjectSystemPrivate::InitDiscardResourcesSkeletalMesh(UCustomi
 }
 
 
-void UCustomizableObjectSystemPrivate::InitInstanceIDRelease(mu::Instance::ID IDToRelease)
+void UCustomizableObjectSystemPrivate::InitInstanceIDRelease(mu::FInstance::FID IDToRelease)
 {
 	check(IsInGameThread());
 
@@ -1660,7 +1659,7 @@ namespace impl
 		check(SystemPrivate->Streamer != nullptr);
 		SystemPrivate->Streamer->PrepareStreamingForObject(Operation->Instance->GetCustomizableObject());			
 
-		const mu::Ptr<mu::System> MutableSystem = SystemPrivate->MutableSystem;
+		const TSharedPtr<mu::FSystem> MutableSystem = SystemPrivate->MutableSystem;
 		
 		if (Operation->bLiveUpdateMode)
 		{
@@ -1685,7 +1684,7 @@ namespace impl
 			UE_LOG(LogMutable, Verbose, TEXT("Creating Mutable instance with id [%d] "), Operation->InstanceID);
 		}
 
-		Operation->MutableInstance = MutableSystem->BeginUpdate(Operation->InstanceID, Operation->Parameters, Operation->GetCapturedDescriptor().GetState(), mu::System::AllLODs);
+		Operation->MutableInstance = MutableSystem->BeginUpdate(Operation->InstanceID, Operation->Parameters, Operation->GetCapturedDescriptor().GetState(), mu::FSystem::AllLODs);
 		Operation->NumInstanceComponents = Operation->MutableInstance->GetComponentCount();
 	}
 
@@ -1757,7 +1756,7 @@ namespace impl
 				const FInstanceUpdateData::FImage& Image = OperationData->InstanceUpdateData.Images[Surface.FirstImage+ImageIndex];
 
 				const FName KeyName = Image.Name;
-				mu::ImagePtrConst MutableImage = Image.Image;
+				TSharedPtr<const mu::FImage> MutableImage = Image.Image;
 
 				// If the image is null, it must be in the cache (or repeated in this instance), and we don't need to do anything here.
 				if (MutableImage)
@@ -1798,7 +1797,7 @@ namespace impl
 			{
 				FInstanceUpdateData::FLOD& LOD = OperationData->InstanceUpdateData.LODs[Component.FirstLOD+LODIndex];
 
-				mu::Ptr<const mu::Mesh> Mesh = LOD.Mesh;
+				TSharedPtr<const mu::FMesh> Mesh = LOD.Mesh;
 				if (!Mesh || Mesh->IsReference())
 				{
 					continue;
@@ -2102,10 +2101,10 @@ namespace impl
 			OperationData->RelevantParametersInProgress.Reset();
 		}
 
-		// Copy ExtensionData Object node input from the Instance to the InstanceUpdateData
+		// Copy FExtensionData Object node input from the Instance to the InstanceUpdateData
 		for (int32 ExtensionDataIndex = 0; ExtensionDataIndex < OperationData->MutableInstance->GetExtensionDataCount(); ExtensionDataIndex++)
 		{
-			mu::Ptr<const mu::ExtensionData> ExtensionData;
+			TSharedPtr<const mu::FExtensionData> ExtensionData;
 			FName Name;
 			OperationData->MutableInstance->GetExtensionData(ExtensionDataIndex, ExtensionData, Name);
 
@@ -2129,7 +2128,7 @@ namespace impl
 		double StartTime,
 		const TSharedRef<TArray<FGetImageData>>& GetImagesData,
 		int32 GetImageIndex,
-		UE::Tasks::TTask<mu::Ptr<const mu::Image>> GetImageTask)
+		UE::Tasks::TTask<TSharedPtr<const mu::FImage>> GetImageTask)
 	{
 		MUTABLE_CPUPROFILER_SCOPE(Task_Mutable_GetMeshes_GetImage_Post)
 
@@ -2178,7 +2177,7 @@ namespace impl
 
 		const FGetImageData& ImageData = (*GetImagesData)[GetImageIndex];
 		
-		UE::Tasks::TTask<mu::Ptr<const mu::Image>> GetImageTask = OperationData->MutableSystem->GetImage(OperationData->InstanceID, ImageData.ImageID, 0, 0);
+		UE::Tasks::TTask<TSharedPtr<const mu::FImage>> GetImageTask = OperationData->MutableSystem->GetImage(OperationData->InstanceID, ImageData.ImageID, 0, 0);
 
 		UE::Tasks::AddNested(UE::Tasks::Launch(TEXT("Task_Mutable_GetMeshes_GetImage_Post"), [=]()
 		{
@@ -2200,7 +2199,7 @@ namespace impl
 		const UCustomizableObject* CustomizableObject = Instance->GetCustomizableObject();
 		const FModelResources& ModelResources = CustomizableObject->GetPrivate()->GetModelResources();
 
-		const mu::Instance* MutableInstance = OperationData->MutableInstance.get();
+		const mu::FInstance* MutableInstance = OperationData->MutableInstance.Get();
 
 		TArray<int32> SurfacesSharedId;
 
@@ -2390,7 +2389,7 @@ namespace impl
 		double StartTime,
 		const TSharedRef<TArray<FGetMeshData>>& GetMeshesData,
 		int32 GetMeshIndex,
-		UE::Tasks::TTask<mu::Ptr<const mu::Mesh>> GetMeshTask)
+		UE::Tasks::TTask<TSharedPtr<const mu::FMesh>> GetMeshTask)
 	{
 		MUTABLE_CPUPROFILER_SCOPE(Task_MutableGetMeshes_GetMesh_Post)
 
@@ -2440,7 +2439,7 @@ namespace impl
 		}
 		
 		const FGetMeshData& MeshData = (*GetMeshesData)[GetMeshIndex];
-		UE::Tasks::TTask<mu::Ptr<const mu::Mesh>> GetMeshTask = OperationData->MutableSystem->GetMesh(OperationData->InstanceID, MeshData.MeshID);
+		UE::Tasks::TTask<TSharedPtr<const mu::FMesh>> GetMeshTask = OperationData->MutableSystem->GetMesh(OperationData->InstanceID, MeshData.MeshID);
 		
 		UE::Tasks::AddNested(UE::Tasks::Launch(TEXT("Task_MutableGetMeshes_GetMesh_Post"), [=]()
 		{
@@ -2485,7 +2484,7 @@ namespace impl
 			}
 
 			// Main instance generation step
-			const mu::Instance* Instance = OperationData->MutableInstance.get(); // TODO GMTFuture remove
+			const mu::FInstance* Instance = OperationData->MutableInstance.Get(); // TODO GMTFuture remove
 			if (!Instance)
 			{
 				UE_LOG(LogMutable, Warning, TEXT("An Instace update has failed."));
@@ -2605,7 +2604,7 @@ namespace impl
     	double StartTime,
     	const TSharedPtr<TArray<mu::FResourceID>>& ImagesInThisInstance,
     	int32 ImageIndex,
-    	UE::Tasks::TTask<mu::Ptr<const mu::Image>> GetImageTask,
+    	UE::Tasks::TTask<TSharedPtr<const mu::FImage>> GetImageTask,
     	int32 MipSizeX,
 		int32 MipSizeY,
 		int32 FullLODContent,
@@ -2625,7 +2624,7 @@ namespace impl
 		{
 			// Generate a correctly-sized but empty image instead, to avoid crashes.
 			UE_LOG(LogMutable, Warning, TEXT("Mutable generated a wrongly-sized image %llu."), Image.ImageID);
-			Image.Image = new mu::Image(MipSizeX, MipSizeY, FullLODContent - MipsToSkip, Image.Image->GetFormat(), mu::EInitializationType::Black);
+			Image.Image = MakeShared<mu::FImage>(MipSizeX, MipSizeY, FullLODContent - MipsToSkip, Image.Image->GetFormat(), mu::EInitializationType::Black);
 		}
 
 		// We need one mip or the complete chain. Otherwise there was a bug.
@@ -2643,7 +2642,7 @@ namespace impl
 			UE_LOG(LogMutable, Warning, TEXT("Mutable generated an incomplete mip chain for image %llu."), Image.ImageID);
 
 			// Force the right number of mips. The missing data will be black.
-			const mu::Ptr<mu::Image> NewImage = new mu::Image(Image.Image->GetSizeX(), Image.Image->GetSizeY(), FullMipCount, Image.Image->GetFormat(), mu::EInitializationType::Black);
+			const TSharedPtr<mu::FImage> NewImage = MakeShared<mu::FImage>(Image.Image->GetSizeX(), Image.Image->GetSizeY(), FullMipCount, Image.Image->GetFormat(), mu::EInitializationType::Black);
 			check(NewImage);	
 			// Formats with BytesPerBlock == 0 will not allocate memory. This type of images are not expected here.
 			check(!NewImage->DataStorage.IsEmpty());
@@ -2725,14 +2724,14 @@ namespace impl
 		const int32 MipSizeY = FMath::Max(Image.FullImageSizeY >> MipsToSkip, 1);
 		if (MipsToSkip > 0 && CustomizableObjectSystemPrivateData->EnableSkipGenerateResidentMips != 0 && OperationData->LowPriorityTextures.Find(Image.Name.ToString()) != INDEX_NONE)
 		{
-			mu::Ptr<const mu::Image> NewImage = new mu::Image(MipSizeX, MipSizeY, FullLODCount - MipsToSkip, ImageDesc.m_format, mu::EInitializationType::Black);
+			TSharedPtr<const mu::FImage> NewImage = MakeShared<mu::FImage>(MipSizeX, MipSizeY, FullLODCount - MipsToSkip, ImageDesc.m_format, mu::EInitializationType::Black);
 
-			UE::Tasks::TTask<mu::Ptr<const mu::Image>> DummyTask = UE::Tasks::MakeCompletedTask<mu::Ptr<const mu::Image>>(NewImage);
+			UE::Tasks::TTask<TSharedPtr<const mu::FImage>> DummyTask = UE::Tasks::MakeCompletedTask<TSharedPtr<const mu::FImage>>(NewImage);
 			Task_Mutable_GetImages_GetImage_Post(OperationData, StartTime, ImagesInThisInstance, ImageIndex, DummyTask, MipSizeX, MipSizeY, FullLODCount, MipsToSkip);
 		}
 		else
 		{
-			const UE::Tasks::TTask<mu::Ptr<const mu::Image>> GetImageTask = OperationData->MutableSystem->GetImage(OperationData->InstanceID, Image.ImageID, Image.BaseMip + MipsToSkip, Image.BaseLOD);
+			const UE::Tasks::TTask<TSharedPtr<const mu::FImage>> GetImageTask = OperationData->MutableSystem->GetImage(OperationData->InstanceID, Image.ImageID, Image.BaseMip + MipsToSkip, Image.BaseLOD);
 			
 			UE::Tasks::AddNested(UE::Tasks::Launch(TEXT("Task_Mutable_GetImages_GetImage_Post"), [=]()
 			{
@@ -2794,7 +2793,7 @@ namespace impl
 
 
 	// This runs in a worker thread.
-	void Task_Mutable_ReleaseInstance(mu::Instance::ID InstanceID, mu::Ptr<mu::System> MutableSystem, bool bLiveUpdateMode)
+	void Task_Mutable_ReleaseInstance(mu::FInstance::FID InstanceID, TSharedPtr<mu::FSystem> MutableSystem, bool bLiveUpdateMode)
 	{
 		MUTABLE_CPUPROFILER_SCOPE(Task_Mutable_ReleaseInstance)
 
@@ -2822,7 +2821,7 @@ namespace impl
 
 
 	// This runs in a worker thread.
-	void Task_Mutable_ReleaseInstanceID(const mu::Instance::ID InstanceID, const mu::Ptr<mu::System>& MutableSystem)
+	void Task_Mutable_ReleaseInstanceID(const mu::FInstance::FID InstanceID, const TSharedPtr<mu::FSystem>& MutableSystem)
 	{
 		MUTABLE_CPUPROFILER_SCOPE(Task_Mutable_ReleaseInstanceID)
 
@@ -2977,7 +2976,7 @@ namespace impl
 		// Next Task: Release Mutable. We need this regardless if we cancel or not
 		//-------------------------------------------------------------		
 		{
-			const mu::Ptr<mu::System> MutableSystem = CustomizableObjectSystemPrivateData->MutableSystem;
+			const TSharedPtr<mu::FSystem> MutableSystem = CustomizableObjectSystemPrivateData->MutableSystem;
 			CustomizableObjectSystemPrivateData->MutableTaskGraph.AddMutableThreadTask(
 				TEXT("Task_Mutable_ReleaseInstance"),
 				[InstanceID = OperationData->InstanceID, MutableSystem, bLiveUpdateMode = OperationData->bLiveUpdateMode]()
@@ -3105,7 +3104,7 @@ namespace impl
 		// guaranteed that the next CO program hasn't started yet, we need to call ClearActiveObject
 		// and CancelPendingLoads on SystemPrivateData->ExtensionDataStreamer.
 		//
-		// ExtensionDataStreamer->AreAnyLoadsPending should return false if the program succeeded.
+		// FExtensionDataStreamer->AreAnyLoadsPending should return false if the program succeeded.
 		//
 		// If the program aborted, AreAnyLoadsPending may return true, as the program doesn't cancel
 		// its own loads on exit (maybe it should?)
@@ -3167,14 +3166,14 @@ namespace impl
 
 
 	/** Enqueue the release ID operation in the Mutable queue */
-	void Task_Game_ReleaseInstanceID(const mu::Instance::ID IDToRelease)
+	void Task_Game_ReleaseInstanceID(const mu::FInstance::FID IDToRelease)
 	{
 		MUTABLE_CPUPROFILER_SCOPE(Task_Game_ReleaseInstanceID)
 
 		UCustomizableObjectSystem* System = UCustomizableObjectSystem::GetInstanceChecked();
 		UCustomizableObjectSystemPrivate* SystemPrivateData = System->GetPrivate();
 
-		const mu::Ptr<mu::System> MutableSystem = SystemPrivateData->MutableSystem;
+		const TSharedPtr<mu::FSystem> MutableSystem = SystemPrivateData->MutableSystem;
 
 		// Task: Release Instance ID
 		//-------------------------------------------------------------
@@ -3341,7 +3340,7 @@ namespace impl
 			bCancel = true;
 		}
 
-		mu::Ptr<const mu::Parameters> Parameters = Operation->Parameters;
+		TSharedPtr<const mu::FParameters> Parameters = Operation->Parameters;
 		if (!Parameters)
 		{
 			bCancel = true;
diff --git a/Engine/Plugins/Mutable/Source/CustomizableObject/Private/MuCO/UnrealConversionUtils.cpp b/Engine/Plugins/Mutable/Source/CustomizableObject/Private/MuCO/UnrealConversionUtils.cpp
--- a/Engine/Plugins/Mutable/Source/CustomizableObject/Private/MuCO/UnrealConversionUtils.cpp
+++ b/Engine/Plugins/Mutable/Source/CustomizableObject/Private/MuCO/UnrealConversionUtils.cpp
@@ -170,7 +170,7 @@ namespace UnrealConversionUtils
 
 	void SetupRenderSections(
 		FSkeletalMeshLODRenderData& LODResource,
-		const mu::MeshPtrConst InMutableMesh,
+		const mu::FMesh* InMutableMesh,
 		const TArray<mu::FBoneName>& InBoneMap,
 		const TMap<mu::FBoneName, TPair<FName, uint16>>& BoneInfoMap,
 		const int32 InFirstBoneMapIndex,
@@ -239,7 +239,7 @@ namespace UnrealConversionUtils
 
 	void InitVertexBuffersWithDummyData(
 		FSkeletalMeshLODRenderData& LODResource,
-		const mu::Ptr<const mu::Mesh> InMutableMesh,
+		const mu::FMesh* InMutableMesh,
 		const bool bAllowCPUAccess)
 	{
 		MUTABLE_CPUPROFILER_SCOPE(InitVertexBuffersWithDummyData);
@@ -331,7 +331,7 @@ namespace UnrealConversionUtils
 
 	void CopyMutableVertexBuffers(
 		FSkeletalMeshLODRenderData& LODResource,
-		const mu::MeshPtrConst MutableMesh,
+		const mu::FMesh* MutableMesh,
 		const bool bAllowCPUAccess)
 
 	{
@@ -423,7 +423,7 @@ namespace UnrealConversionUtils
 	}
 
 		
-	void InitIndexBuffersWithDummyData(FSkeletalMeshLODRenderData& LODResource, const mu::Ptr<const mu::Mesh> InMutableMesh)
+	void InitIndexBuffersWithDummyData(FSkeletalMeshLODRenderData& LODResource, const mu::FMesh* InMutableMesh)
 	{
 		MUTABLE_CPUPROFILER_SCOPE(InitIndexBuffersWithDummyData);
 
@@ -439,7 +439,7 @@ namespace UnrealConversionUtils
 
 	bool CopyMutableIndexBuffers(
 		FSkeletalMeshLODRenderData& LODResource,
-		const mu::Ptr<const mu::Mesh> InMutableMesh)
+		const mu::FMesh* InMutableMesh)
 	{
 		MUTABLE_CPUPROFILER_SCOPE(CopyMutableIndexBuffers);
 
diff --git a/Engine/Plugins/Mutable/Source/CustomizableObject/Private/MuCO/UnrealExtensionDataStreamer.cpp b/Engine/Plugins/Mutable/Source/CustomizableObject/Private/MuCO/UnrealExtensionDataStreamer.cpp
--- a/Engine/Plugins/Mutable/Source/CustomizableObject/Private/MuCO/UnrealExtensionDataStreamer.cpp
+++ b/Engine/Plugins/Mutable/Source/CustomizableObject/Private/MuCO/UnrealExtensionDataStreamer.cpp
@@ -95,15 +95,15 @@ void FUnrealExtensionDataStreamer::CancelPendingLoads()
 	PendingLoads.Empty();
 }
 
-mu::ExtensionDataPtr FUnrealExtensionDataStreamer::CloneExtensionData(const mu::ExtensionDataPtrConst& Source)
+TSharedPtr<mu::FExtensionData> FUnrealExtensionDataStreamer::CloneExtensionData(const TSharedPtr<const mu::FExtensionData>& Source)
 {
 	unimplemented();
 	return nullptr;
 }
 
 TSharedRef<const mu::FExtensionDataLoadHandle> FUnrealExtensionDataStreamer::StartLoad(
-	const mu::ExtensionDataPtrConst& Data,
-	TArray<mu::ExtensionDataPtrConst>& OutUnloadedConstants)
+	const TSharedPtr<const mu::FExtensionData>& Data,
+	TArray<TSharedPtr<const mu::FExtensionData>>& OutUnloadedConstants)
 {
 	// Note that this can be called from any thread
 
@@ -157,7 +157,7 @@ TSharedPtr<FStreamableHandle> FUnrealExtensionDataStreamer::StartLoadOnGameThrea
 	const TSharedRef<mu::FExtensionDataLoadHandle>& LoadHandle)
 {
 	check(IsInGameThread());
-	check(LoadHandle->Data->Origin == mu::ExtensionData::EOrigin::ConstantStreamed);
+	check(LoadHandle->Data->Origin == mu::FExtensionData::EOrigin::ConstantStreamed);
 
 	UCustomizableObjectSystemPrivate* SystemPrivate = SystemPrivateWeak.Get();
 	if (!SystemPrivate)
@@ -204,7 +204,7 @@ TSharedPtr<FStreamableHandle> FUnrealExtensionDataStreamer::StartLoadOnGameThrea
 	TargetsToStream.Add(StreamedData.GetPath().ToSoftObjectPath());
 
 	check(SystemPrivate->ExtensionDataStreamer);
-	TFunction<void()> OnLoadComplete = [ExtensionDataStreamer = SystemPrivate->ExtensionDataStreamer, ObjectToLoadFor, LoadHandle]()
+	TFunction<void()> OnLoadComplete = [FExtensionDataStreamer = SystemPrivate->ExtensionDataStreamer, ObjectToLoadFor, LoadHandle]()
 	{
 		UCustomizableObject* Object = ObjectToLoadFor.Get();
 		if (!Object)
@@ -213,7 +213,7 @@ TSharedPtr<FStreamableHandle> FUnrealExtensionDataStreamer::StartLoadOnGameThrea
 			return;
 		}
 
-		ExtensionDataStreamer->NotifyLoadCompleted(Object, LoadHandle);
+		FExtensionDataStreamer->NotifyLoadCompleted(Object, LoadHandle);
 	};
 
 	const bool bManageActiveHandle = false;
