[mutable] More enforcement of coding standards.

diff --git a/Engine/Plugins/Mutable/Source/CustomizableObject/Private/MuCO/UnrealMutableImageProvider.cpp b/Engine/Plugins/Mutable/Source/CustomizableObject/Private/MuCO/UnrealMutableImageProvider.cpp
--- a/Engine/Plugins/Mutable/Source/CustomizableObject/Private/MuCO/UnrealMutableImageProvider.cpp
+++ b/Engine/Plugins/Mutable/Source/CustomizableObject/Private/MuCO/UnrealMutableImageProvider.cpp
@@ -18,7 +18,7 @@
 //-------------------------------------------------------------------------------------------------
 namespace
 {
-	void ConvertTextureUnrealPlatformToMutable(mu::Image* OutResult, UTexture2D* Texture, uint8 MipmapsToSkip)
+	void ConvertTextureUnrealPlatformToMutable(mu::FImage* OutResult, UTexture2D* Texture, uint8 MipmapsToSkip)
 	{		
 		check(Texture->GetPlatformData()->Mips[MipmapsToSkip].BulkData.IsBulkDataLoaded());
 
@@ -82,7 +82,7 @@ mu::EImageFormat GetMutablePixelFormat(EPixelFormat InTextureFormat)
 
 
 //-------------------------------------------------------------------------------------------------
-TTuple<UE::Tasks::FTask, TFunction<void()>> FUnrealMutableResourceProvider::GetImageAsync(FName Id, uint8 MipmapsToSkip, TFunction<void(mu::Ptr<mu::Image>)>& ResultCallback)
+TTuple<UE::Tasks::FTask, TFunction<void()>> FUnrealMutableResourceProvider::GetImageAsync(FName Id, uint8 MipmapsToSkip, TFunction<void(TSharedPtr<mu::FImage>)>& ResultCallback)
 {
 	// Thread: worker
 	MUTABLE_CPUPROFILER_SCOPE(FUnrealMutableImageProvider::GetImageAsync);
@@ -132,9 +132,9 @@ TTuple<UE::Tasks::FTask, TFunction<void()>> FUnrealMutableResourceProvider::GetI
 			check(MipIndex >= 0);
 			
 			// In the editor the src data can be directly accessed
-			mu::Ptr<mu::Image> Image = new mu::Image();
+			TSharedPtr<mu::FImage> Image = MakeShared<mu::FImage>();
 
-			EUnrealToMutableConversionError Error = ConvertTextureUnrealSourceToMutable(Image.get(), *SourceTextureData, MipIndex);
+			EUnrealToMutableConversionError Error = ConvertTextureUnrealSourceToMutable(Image.Get(), *SourceTextureData, MipIndex);
 			if (Error != EUnrealToMutableConversionError::Success)
 			{
 				// This could happen in the editor, because some source textures may have changed while there was a background compilation.
@@ -180,7 +180,7 @@ TTuple<UE::Tasks::FTask, TFunction<void()>> FUnrealMutableResourceProvider::GetI
 			int32 SizeY = TextureToLoad->GetSizeY() >> MipIndex;
 
 			check(LODs == 1);
-			mu::Ptr<mu::Image> Image = new mu::Image(SizeX, SizeY, LODs, MutImageFormat, mu::EInitializationType::NotInitialized);
+			TSharedPtr<mu::FImage> Image = MakeShared<mu::FImage>(SizeX, SizeY, LODs, MutImageFormat, mu::EInitializationType::NotInitialized);
 			TArrayView<uint8> MutImageDataView = Image->DataStorage.GetLOD(0);
 
 			// In a packaged game the bulk data has to be loaded
@@ -333,7 +333,7 @@ TTuple<UE::Tasks::FTask, TFunction<void()>> FUnrealMutableResourceProvider::GetI
 
 
 //-------------------------------------------------------------------------------------------------
-TTuple<UE::Tasks::FTask, TFunction<void()>> FUnrealMutableResourceProvider::GetReferencedImageAsync(const void* ModelPtr, int32 Id, uint8 MipmapsToSkip, TFunction<void(mu::Ptr<mu::Image>)>& ResultCallback)
+TTuple<UE::Tasks::FTask, TFunction<void()>> FUnrealMutableResourceProvider::GetReferencedImageAsync(const void* ModelPtr, int32 Id, uint8 MipmapsToSkip, TFunction<void(TSharedPtr<mu::FImage>)>& ResultCallback)
 {
 	// Thread: worker
 	MUTABLE_CPUPROFILER_SCOPE(FUnrealMutableImageProvider::GetReferencedImageAsync);
@@ -345,7 +345,7 @@ TTuple<UE::Tasks::FTask, TFunction<void()>> FUnrealMutableResourceProvider::GetR
 
 #if WITH_EDITOR
 
-	mu::Ptr<mu::Image> Image = new mu::Image();
+	TSharedPtr<mu::FImage> Image = MakeShared<mu::FImage>();
 
 	FScopeLock Lock(&RuntimeReferencedLock);
 	
@@ -365,7 +365,7 @@ TTuple<UE::Tasks::FTask, TFunction<void()>> FUnrealMutableResourceProvider::GetR
 	const int32 MipIndex = FMath::Min(static_cast<int32>(MipmapsToSkip), SourceTextureData.GetSource().GetNumMips() - 1);
 	check(MipIndex >= 0);
 	
-	EUnrealToMutableConversionError Error = ConvertTextureUnrealSourceToMutable(Image.get(), SourceTextureData, MipIndex);
+	EUnrealToMutableConversionError Error = ConvertTextureUnrealSourceToMutable(Image.Get(), SourceTextureData, MipIndex);
 	if (Error != EUnrealToMutableConversionError::Success)
 	{
 		// This could happen in the editor, because some source textures may have changed while updating.
@@ -474,9 +474,9 @@ mu::FImageDesc FUnrealMutableResourceProvider::GetImageDesc(FName Id, uint8 Mipm
 }
 
 
-TTuple<UE::Tasks::FTask, TFunction<void()>> FUnrealMutableResourceProvider::GetMeshAsync(FName Id, TFunction<void(mu::Ptr<mu::Mesh>)>& ResultCallback)
+TTuple<UE::Tasks::FTask, TFunction<void()>> FUnrealMutableResourceProvider::GetMeshAsync(FName Id, TFunction<void(TSharedPtr<mu::FMesh>)>& ResultCallback)
 {
-	// TODO: This is a stub function where Mesh Parameters would be supported.
+	// TODO: This is a stub function where Mesh FParameters would be supported.
 
 	// Thread: worker
 	MUTABLE_CPUPROFILER_SCOPE(FUnrealMutableImageProvider::GetMeshAsync);
@@ -489,13 +489,13 @@ TTuple<UE::Tasks::FTask, TFunction<void()>> FUnrealMutableResourceProvider::GetM
 	// No USkeletalMesh was provided, cannot do anything, just provide a dummy mesh
 	UE_LOG(LogMutable, Warning, TEXT("No USkeletalMesh was provided for an application-specific mesh parameter."));
 
-	mu::Ptr<mu::Mesh> Dummy = new mu::Mesh;
+	TSharedPtr<mu::FMesh> Dummy = MakeShared<mu::FMesh>();
 	ResultCallback( Dummy );
 	return Invoke(TrivialReturn);
 }
 
 
-TTuple<UE::Tasks::FTask, TFunction<void()>> FUnrealMutableResourceProvider::GetReferencedMeshAsync(const void* ModelPtr, int32 Id, TFunction<void(mu::Ptr<mu::Mesh>)>& ResultCallback)
+TTuple<UE::Tasks::FTask, TFunction<void()>> FUnrealMutableResourceProvider::GetReferencedMeshAsync(const void* ModelPtr, int32 Id, TFunction<void(TSharedPtr<mu::FMesh>)>& ResultCallback)
 {
 	// TODO: This is a stub function where referenced Meshes would be supported.
 
@@ -507,7 +507,7 @@ TTuple<UE::Tasks::FTask, TFunction<void()>> FUnrealMutableResourceProvider::GetR
 			return MakeTuple(UE::Tasks::MakeCompletedTask<void>(), []() -> void {});
 		};
 
-	mu::Ptr<mu::Mesh> Mesh = new mu::Mesh();
+	TSharedPtr<mu::FMesh> Mesh = MakeShared<mu::FMesh>();
 
 	// Not supported outside editor yet.
 	UE_LOG(LogMutable, Warning, TEXT("Failed to get referenced mesh. Not implemented yet."));
@@ -565,7 +565,7 @@ void FUnrealMutableResourceProvider::CacheImage(FName Id, bool bUser)
 				case UCustomizableSystemImageProvider::ValueType::Raw:
 					{
 						FIntVector desc = Provider->GetTextureParameterValueSize(Id);
-						mu::Ptr<mu::Image> pResult = new mu::Image(desc[0], desc[1], 1, mu::EImageFormat::RGBA_UByte, mu::EInitializationType::Black);
+						TSharedPtr<mu::FImage> pResult = MakeShared<mu::FImage>(desc[0], desc[1], 1, mu::EImageFormat::RGBA_UByte, mu::EInitializationType::Black);
 						Provider->GetTextureParameterValueData(Id, pResult->GetLODData(0));
 
 						ImageInfo = FUnrealMutableImageInfo(pResult);
@@ -577,11 +577,11 @@ void FUnrealMutableResourceProvider::CacheImage(FName Id, bool bUser)
 				case UCustomizableSystemImageProvider::ValueType::Unreal:
 					{
 						UTexture2D* UnrealTexture = Provider->GetTextureParameterValue(Id);
-						mu::Ptr<mu::Image> pResult = new mu::Image();
+						TSharedPtr<mu::FImage> pResult = MakeShared<mu::FImage>();
 
 #if WITH_EDITOR
 						FMutableSourceTextureData Tex(*UnrealTexture);
-						EUnrealToMutableConversionError Error = ConvertTextureUnrealSourceToMutable(pResult.get(), Tex, 0);
+						EUnrealToMutableConversionError Error = ConvertTextureUnrealSourceToMutable(pResult.Get(), Tex, 0);
 						if (Error != EUnrealToMutableConversionError::Success)
 						{
 							// This could happen in the editor, because some source textures may have changed while there was a background compilation.
@@ -589,7 +589,7 @@ void FUnrealMutableResourceProvider::CacheImage(FName Id, bool bUser)
 							UE_LOG(LogMutable, Warning, TEXT("Failed to load some source texture data for [%s]. Some textures may be corrupted."), *UnrealTexture->GetName());
 						}
 #else
-						ConvertTextureUnrealPlatformToMutable(pResult.get(), UnrealTexture, 0);
+						ConvertTextureUnrealPlatformToMutable(pResult.Get(), UnrealTexture, 0);
 #endif
 
 						ImageInfo = FUnrealMutableImageInfo(pResult);
@@ -678,7 +678,7 @@ void FUnrealMutableResourceProvider::ClearCache(bool bUser)
 }
 
 
-void FUnrealMutableResourceProvider::CacheImages(const mu::Parameters& Parameters)
+void FUnrealMutableResourceProvider::CacheImages(const mu::FParameters& Parameters)
 {
 	const int32 NumParams = Parameters.GetCount();
 	for (int32 ParamIndex = 0; ParamIndex < NumParams; ++ParamIndex)
@@ -696,8 +696,8 @@ void FUnrealMutableResourceProvider::CacheImages(const mu::Parameters& Parameter
 		const int32 NumValues = Parameters.GetValueCount(ParamIndex);
 		for (int32 ValueIndex = 0; ValueIndex < NumValues; ++ValueIndex)
 		{
-			mu::Ptr<const mu::RangeIndex> Range = Parameters.GetValueIndex(ParamIndex, ValueIndex);
-			const FName TextureId = Parameters.GetImageValue(ParamIndex, Range);
+			TSharedPtr<const mu::FRangeIndex> Range = Parameters.GetValueIndex(ParamIndex, ValueIndex);
+			const FName TextureId = Parameters.GetImageValue(ParamIndex, Range.Get());
 
 			CacheImage(TextureId, false);
 		}
@@ -705,7 +705,7 @@ void FUnrealMutableResourceProvider::CacheImages(const mu::Parameters& Parameter
 }
 
 
-void FUnrealMutableResourceProvider::UnCacheImages(const mu::Parameters& Parameters)
+void FUnrealMutableResourceProvider::UnCacheImages(const mu::FParameters& Parameters)
 {
 	const int32 NumParams = Parameters.GetCount();
 	for (int32 ParamIndex = 0; ParamIndex < NumParams; ++ParamIndex)
@@ -723,8 +723,8 @@ void FUnrealMutableResourceProvider::UnCacheImages(const mu::Parameters& Paramet
 		const int32 NumValues = Parameters.GetValueCount(ParamIndex);
 		for (int32 ValueIndex = 0; ValueIndex < NumValues; ++ValueIndex)
 		{
-			mu::Ptr<const mu::RangeIndex> Range = Parameters.GetValueIndex(ParamIndex, ValueIndex);
-			const FName TextureId = Parameters.GetImageValue(ParamIndex, Range);
+			TSharedPtr<const mu::FRangeIndex> Range = Parameters.GetValueIndex(ParamIndex, ValueIndex);
+			const FName TextureId = Parameters.GetImageValue(ParamIndex, Range.Get());
 
 			UnCacheImage(TextureId, false);
 		}
@@ -733,7 +733,7 @@ void FUnrealMutableResourceProvider::UnCacheImages(const mu::Parameters& Paramet
 
 
 #if WITH_EDITOR
-void FUnrealMutableResourceProvider::CacheRuntimeReferencedImages(const TSharedRef<const mu::Model>& Model, const TArray<TSoftObjectPtr<const UTexture>>& RuntimeReferencedTextures)
+void FUnrealMutableResourceProvider::CacheRuntimeReferencedImages(const TSharedRef<const mu::FModel>& Model, const TArray<TSoftObjectPtr<const UTexture>>& RuntimeReferencedTextures)
 {
 	check(IsInGameThread());
 	
@@ -760,7 +760,7 @@ void FUnrealMutableResourceProvider::CacheRuntimeReferencedImages(const TSharedR
 #endif
 
 
-mu::ImagePtr FUnrealMutableResourceProvider::CreateDummy()
+TSharedPtr<mu::FImage> FUnrealMutableResourceProvider::CreateDummy()
 {
 	// Create a dummy image
 	const int32 Size = DUMMY_IMAGE_DESC.m_size[0];
@@ -773,7 +773,7 @@ mu::ImagePtr FUnrealMutableResourceProvider::CreateDummy()
 	uint8 Colors[CheckerTileCount][4] = {{255, 255, 0, 0}, {0, 0, 255, 0}};
 #endif
 
-	mu::ImagePtr pResult = new mu::Image(Size, Size, DUMMY_IMAGE_DESC.m_lods, DUMMY_IMAGE_DESC.m_format, mu::EInitializationType::NotInitialized);
+	TSharedPtr<mu::FImage> pResult = MakeShared<mu::FImage>(Size, Size, DUMMY_IMAGE_DESC.m_lods, DUMMY_IMAGE_DESC.m_format, mu::EInitializationType::NotInitialized);
 
 	check(pResult->GetLODCount() == 1);
 	check(pResult->GetFormat() == mu::EImageFormat::RGBA_UByte || pResult->GetFormat() == mu::EImageFormat::BGRA_UByte);
@@ -846,7 +846,7 @@ void FUnrealMutableResourceProvider::AddReferencedObjects(FReferenceCollector& C
 }
 
 
-FUnrealMutableResourceProvider::FUnrealMutableImageInfo::FUnrealMutableImageInfo(const mu::ImagePtr& InImage)
+FUnrealMutableResourceProvider::FUnrealMutableImageInfo::FUnrealMutableImageInfo(const TSharedPtr<mu::FImage>& InImage)
 {
 	check(IsInGameThread())
 	
diff --git a/Engine/Plugins/Mutable/Source/CustomizableObject/Private/MuCO/UnrealMutableImageProvider.h b/Engine/Plugins/Mutable/Source/CustomizableObject/Private/MuCO/UnrealMutableImageProvider.h
--- a/Engine/Plugins/Mutable/Source/CustomizableObject/Private/MuCO/UnrealMutableImageProvider.h
+++ b/Engine/Plugins/Mutable/Source/CustomizableObject/Private/MuCO/UnrealMutableImageProvider.h
@@ -18,17 +18,17 @@ class UTexture2D;
 
 
 /** Implementation of a mutable core provider for image parameters that are application-specific. */
-class FUnrealMutableResourceProvider : public mu::ExternalResourceProvider, public FGCObject
+class FUnrealMutableResourceProvider : public mu::FExternalResourceProvider, public FGCObject
 {
 
 public:
 	// mu::ExternalResourceProvider interface
 	// Thread: worker
-	virtual TTuple<UE::Tasks::FTask, TFunction<void()>> GetImageAsync(FName Id, uint8 MipmapsToSkip, TFunction<void(mu::Ptr<mu::Image>)>& ResultCallback) override;
-	virtual TTuple<UE::Tasks::FTask, TFunction<void()>> GetReferencedImageAsync(const void* ModelPtr, int32 Id, uint8 MipmapsToSkip, TFunction<void(mu::Ptr<mu::Image>)>& ResultCallback) override;
+	virtual TTuple<UE::Tasks::FTask, TFunction<void()>> GetImageAsync(FName Id, uint8 MipmapsToSkip, TFunction<void(TSharedPtr<mu::FImage>)>& ResultCallback) override;
+	virtual TTuple<UE::Tasks::FTask, TFunction<void()>> GetReferencedImageAsync(const void* ModelPtr, int32 Id, uint8 MipmapsToSkip, TFunction<void(TSharedPtr<mu::FImage>)>& ResultCallback) override;
 	virtual mu::FImageDesc GetImageDesc(FName Id, uint8 MipmapsToSkip) override;
-	virtual TTuple<UE::Tasks::FTask, TFunction<void()>> GetMeshAsync(FName Id, TFunction<void(mu::Ptr<mu::Mesh>)>& ResultCallback) override;
-	virtual TTuple<UE::Tasks::FTask, TFunction<void()>> GetReferencedMeshAsync(const void* ModelPtr, int32 Id, TFunction<void(mu::Ptr<mu::Mesh>)>& ResultCallback) override;
+	virtual TTuple<UE::Tasks::FTask, TFunction<void()>> GetMeshAsync(FName Id, TFunction<void(TSharedPtr<mu::FMesh>)>& ResultCallback) override;
+	virtual TTuple<UE::Tasks::FTask, TFunction<void()>> GetReferencedMeshAsync(const void* ModelPtr, int32 Id, TFunction<void(TSharedPtr<mu::FMesh>)>& ResultCallback) override;
 
 	
 	// Own interface	
@@ -45,11 +45,11 @@ public:
 	 * @param bUser if true, removes a references from the user reference counter. If false, removes a reference from the system reference counter. */
 	void ClearCache(bool bUser);
 
-	void CacheImages(const mu::Parameters& Parameters);
-	void UnCacheImages(const mu::Parameters& Parameters);
+	void CacheImages(const mu::FParameters& Parameters);
+	void UnCacheImages(const mu::FParameters& Parameters);
 
 #if WITH_EDITOR
-	void CacheRuntimeReferencedImages(const TSharedRef<const mu::Model>& Model, const TArray<TSoftObjectPtr<const UTexture>>& RuntimeReferencedTextures);
+	void CacheRuntimeReferencedImages(const TSharedRef<const mu::FModel>& Model, const TArray<TSoftObjectPtr<const UTexture>>& RuntimeReferencedTextures);
 #endif
 	
 	/** List of actual image providers that have been registered to the CustomizableObjectSystem. */
@@ -66,11 +66,11 @@ private:
 	{
 		FUnrealMutableImageInfo() = default;
 
-		FUnrealMutableImageInfo(const mu::ImagePtr& InImage);
+		FUnrealMutableImageInfo(const TSharedPtr<mu::FImage>& InImage);
 
 		FUnrealMutableImageInfo(UTexture2D& Texture);
 
-		mu::ImagePtr Image;
+		TSharedPtr<mu::FImage> Image;
 
 #if WITH_EDITOR
 		TSharedPtr<FMutableSourceTextureData> SourceTextureData;
@@ -91,7 +91,7 @@ private:
 			mu::FImageDesc(mu::FImageSize(32, 32), mu::EImageFormat::RGBA_UByte, 1);
 
 	/** This will be called if an image Id has been requested by Mutable core but it has not been provided by any provider. */
-	static mu::ImagePtr CreateDummy();
+	static TSharedPtr<mu::FImage> CreateDummy();
 	static mu::FImageDesc CreateDummyDesc();
 
 	/** Map of Ids to external textures that may be required for any instance or Mutable texture mip under construction.
@@ -103,7 +103,7 @@ private:
 	struct FRuntimeReferencedImages
 	{
 		TArray<FMutableSourceTextureData> SourceTextures;
-		TWeakPtr<const mu::Model> Model;
+		TWeakPtr<const mu::FModel> Model;
 	};
 	
 	TMap<const void*, FRuntimeReferencedImages> RuntimeReferencedImages;
diff --git a/Engine/Plugins/Mutable/Source/CustomizableObject/Private/MuCO/UnrealMutableModelDiskStreamer.cpp b/Engine/Plugins/Mutable/Source/CustomizableObject/Private/MuCO/UnrealMutableModelDiskStreamer.cpp
--- a/Engine/Plugins/Mutable/Source/CustomizableObject/Private/MuCO/UnrealMutableModelDiskStreamer.cpp
+++ b/Engine/Plugins/Mutable/Source/CustomizableObject/Private/MuCO/UnrealMutableModelDiskStreamer.cpp
@@ -40,13 +40,13 @@ void UnrealMutableOutputStream::Write(const void* pData, uint64 size)
 #endif
 
 //-------------------------------------------------------------------------------------------------
-UnrealMutableInputStream::UnrealMutableInputStream(FArchive& ar)
+FUnrealMutableInputStream::FUnrealMutableInputStream(FArchive& ar)
 	: m_ar(ar)
 {
 }
 
 
-void UnrealMutableInputStream::Read(void* pData, uint64 size)
+void FUnrealMutableInputStream::Read(void* pData, uint64 size)
 {
 	m_ar.Serialize(pData, size);
 }
@@ -94,7 +94,7 @@ bool FUnrealMutableModelBulkReader::PrepareStreamingForObject(UCustomizableObjec
 	if (!bAlreadyStreaming)
 	{
 		FObjectData NewData;
-		NewData.Model = TWeakPtr<const mu::Model>(CustomizableObject->GetPrivate()->GetModel());
+		NewData.Model = TWeakPtr<const mu::FModel>(CustomizableObject->GetPrivate()->GetModel());
 		NewData.ModelStreamableBulkData = CustomizableObject->GetPrivate()->GetModelStreamableBulkData();
 
 		if (!NewData.ModelStreamableBulkData || NewData.ModelStreamableBulkData->ModelStreamables.IsEmpty())
@@ -213,7 +213,7 @@ FAutoConsoleVariableRef CVarStreamPriority(
 	TEXT(""));
 
 
-mu::ModelReader::FOperationID FUnrealMutableModelBulkReader::BeginReadBlock(const mu::Model* Model, uint32 Key, void* pBuffer, uint64 size, TFunction<void(bool bSuccess)>* CompletionCallback)
+mu::FModelReader::FOperationID FUnrealMutableModelBulkReader::BeginReadBlock(const mu::FModel* Model, uint32 Key, void* Buffer, uint64 Size, TFunction<void(bool bSuccess)>* CompletionCallback)
 {
 	MUTABLE_CPUPROFILER_SCOPE(FUnrealMutableModelBulkStreamer::OpenReadFile);
 
@@ -269,10 +269,10 @@ mu::ModelReader::FOperationID FUnrealMutableModelBulkReader::BeginReadBlock(cons
 
 		Request.BulkReadRequest = MakeShareable(BulkData.CreateStreamingRequest(
 			Block->Offset,
-			size,
+			Size,
 			(EAsyncIOPriorityAndFlags)StreamPriority,
 			&IOCallback,
-			reinterpret_cast<uint8*>(pBuffer)));
+			reinterpret_cast<uint8*>(Buffer)));
 	}
 	else
 	{
@@ -305,21 +305,21 @@ mu::ModelReader::FOperationID FUnrealMutableModelBulkReader::BeginReadBlock(cons
 			ReadRequest.DDCReadRequest = MakeShared<FRequestOwner>(EPriority::High);
 
 			GetCache().Get(MakeArrayView(&Request, 1), *ReadRequest.DDCReadRequest.Get(),
-				[CompletionCallbackCapture = *CompletionCallback, Block, ResourceId, pBuffer, size](FCacheGetResponse&& Response)
+				[CompletionCallbackCapture = *CompletionCallback, Block, ResourceId, Buffer, Size](FCacheGetResponse&& Response)
 				{
 					bool bSuccess = Response.Status == EStatus::Ok;
 					if (bSuccess)
 					{
 						const FCompressedBuffer& CompressedBuffer = Response.Record.GetValue(ResourceId).GetData();
-						if (size < CompressedBuffer.GetRawSize())
+						if (Size < CompressedBuffer.GetRawSize())
 						{
-							check(CompressedBuffer.GetRawSize() >= Block->Offset + size);
+							check(CompressedBuffer.GetRawSize() >= Block->Offset + Size);
 							FSharedBuffer DecompressedBuffer = CompressedBuffer.Decompress();
-							FMemory::Memcpy(pBuffer, reinterpret_cast<const uint8*>(DecompressedBuffer.GetData()) + Block->Offset, size);
+							FMemory::Memcpy(Buffer, reinterpret_cast<const uint8*>(DecompressedBuffer.GetData()) + Block->Offset, Size);
 						}
 						else
 						{
-							check(CompressedBuffer.TryDecompressTo(MakeMemoryView(pBuffer, size)));
+							check(CompressedBuffer.TryDecompressTo(MakeMemoryView(Buffer, Size)));
 						}
 					}
 					else
@@ -384,10 +384,10 @@ mu::ModelReader::FOperationID FUnrealMutableModelBulkReader::BeginReadBlock(cons
 
 			ReadRequest.FileReadRequest = MakeShareable(FileHandle->ReadRequest(
 				BulkDataOffsetInFile + Block->Offset,
-				size,
+				Size,
 				(EAsyncIOPriorityAndFlags)StreamPriority,
 				ReadRequest.FileCallback.Get(),
-				reinterpret_cast<uint8*>(pBuffer)));
+				reinterpret_cast<uint8*>(Buffer)));
 
 			ObjectData->CurrentReadRequests.Add(Result, ReadRequest);
 		}
@@ -399,7 +399,7 @@ mu::ModelReader::FOperationID FUnrealMutableModelBulkReader::BeginReadBlock(cons
 }
 
 
-bool FUnrealMutableModelBulkReader::IsReadCompleted(mu::ModelReader::FOperationID OperationId)
+bool FUnrealMutableModelBulkReader::IsReadCompleted(mu::FModelReader::FOperationID OperationId)
 {
 	MUTABLE_CPUPROFILER_SCOPE(FUnrealMutableModelBulkStreamer::IsReadCompleted);
 
@@ -430,7 +430,7 @@ bool FUnrealMutableModelBulkReader::IsReadCompleted(mu::ModelReader::FOperationI
 }
 
 
-bool FUnrealMutableModelBulkReader::EndRead(mu::ModelReader::FOperationID OperationId)
+bool FUnrealMutableModelBulkReader::EndRead(mu::FModelReader::FOperationID OperationId)
 {
 	MUTABLE_CPUPROFILER_SCOPE(FUnrealMutableModelBulkStreamer::EndRead);
 
diff --git a/Engine/Plugins/Mutable/Source/CustomizableObject/Private/MuCO/UnrealToMutableTextureConversionUtils.cpp b/Engine/Plugins/Mutable/Source/CustomizableObject/Private/MuCO/UnrealToMutableTextureConversionUtils.cpp
--- a/Engine/Plugins/Mutable/Source/CustomizableObject/Private/MuCO/UnrealToMutableTextureConversionUtils.cpp
+++ b/Engine/Plugins/Mutable/Source/CustomizableObject/Private/MuCO/UnrealToMutableTextureConversionUtils.cpp
@@ -217,7 +217,7 @@ bool FMutableSourceTextureData::IsNormalComposite() const
 }
 
 
-EUnrealToMutableConversionError ConvertTextureUnrealSourceToMutable(mu::Image* OutResult, FMutableSourceTextureData& Tex, uint8 MipmapsToSkip)
+EUnrealToMutableConversionError ConvertTextureUnrealSourceToMutable(mu::FImage* OutResult, FMutableSourceTextureData& Tex, uint8 MipmapsToSkip)
 {
 	MUTABLE_CPUPROFILER_SCOPE(ConvertTextureUnrealSourceToMutable);
 
diff --git a/Engine/Plugins/Mutable/Source/CustomizableObject/Public/MuCO/CustomizableInstanceLODManagement.h b/Engine/Plugins/Mutable/Source/CustomizableObject/Public/MuCO/CustomizableInstanceLODManagement.h
--- a/Engine/Plugins/Mutable/Source/CustomizableObject/Public/MuCO/CustomizableInstanceLODManagement.h
+++ b/Engine/Plugins/Mutable/Source/CustomizableObject/Public/MuCO/CustomizableInstanceLODManagement.h
@@ -8,7 +8,7 @@
 class UCustomizableObjectInstanceUsage;
 typedef TMap<const class UCustomizableObjectInstance*, class FMutableUpdateCandidate> FMutableInstanceUpdateMap;
 
-// This is an abstract base class, override it to create a new Instance LOD management system and register with UCustomizableObjectSystem::SetInstanceLODManagement
+// This is an abstract base class, override it to create a MakeShared<FInstance> LOD management system and register with UCustomizableObjectSystem::SetInstanceLODManagement
 UCLASS(Abstract, Blueprintable, BlueprintType)
 class CUSTOMIZABLEOBJECT_API UCustomizableInstanceLODManagementBase : public UObject
 {
diff --git a/Engine/Plugins/Mutable/Source/CustomizableObject/Public/MuCO/CustomizableObject.h b/Engine/Plugins/Mutable/Source/CustomizableObject/Public/MuCO/CustomizableObject.h
--- a/Engine/Plugins/Mutable/Source/CustomizableObject/Public/MuCO/CustomizableObject.h
+++ b/Engine/Plugins/Mutable/Source/CustomizableObject/Public/MuCO/CustomizableObject.h
@@ -32,8 +32,8 @@ template <typename FuncType> class TFunctionRef;
 // Forward declaration of Mutable classes necessary for the interface
 namespace mu
 {
-	class Model;
-	class Parameters;
+	class FModel;
+	class FParameters;
 }
 
 
@@ -316,11 +316,11 @@ public:
 	FMutableLODSettings LODSettings;
 	
 private:
-	// mu::ExtensionData::Index is an index into this array when mu::ExtensionData::Origin is ConstantAlwaysLoaded
+	// mu::FExtensionData::Index is an index into this array when mu::ExtensionData::Origin is ConstantAlwaysLoaded
 	UPROPERTY()
 	TArray<FCustomizableObjectResourceData> AlwaysLoadedExtensionData;
 
-	// mu::ExtensionData::Index is an index into this array when mu::ExtensionData::Origin is ConstantStreamed
+	// mu::FExtensionData::Index is an index into this array when mu::ExtensionData::Origin is ConstantStreamed
 	UPROPERTY()
 	TArray<FCustomizableObjectStreamedResourceData> StreamedExtensionData;
 
diff --git a/Engine/Plugins/Mutable/Source/CustomizableObject/Public/MuCO/CustomizableObjectInstance.h b/Engine/Plugins/Mutable/Source/CustomizableObject/Public/MuCO/CustomizableObjectInstance.h
--- a/Engine/Plugins/Mutable/Source/CustomizableObject/Public/MuCO/CustomizableObjectInstance.h
+++ b/Engine/Plugins/Mutable/Source/CustomizableObject/Public/MuCO/CustomizableObjectInstance.h
@@ -11,7 +11,7 @@
 
 namespace mu
 {
-	class Image;
+	class FImage;
 }
 
 class UCustomizableObjectSystemPrivate;
@@ -543,17 +543,17 @@ public:
 	UFUNCTION(BlueprintCallable, Category = CustomizableObjectInstance)
 	int32 AddValueToProjectorRange(const FString& ParamName);
 
-	/** Remove the RangeIndex element of the integer range of values from the parameter ParamName. If RangeValue is -1 removes the last of the integer range of values.
+	/** Remove the FRangeIndex element of the integer range of values from the parameter ParamName. If RangeValue is -1 removes the last of the integer range of values.
 		Returns the index of the last valid integer, -1 if no values left. */
 	UFUNCTION(BlueprintCallable, Category = CustomizableObjectInstance)
 	int32 RemoveValueFromIntRange(const FString& ParamName, int32 RangeIndex = -1);
 
-	/** Remove the RangeIndex element of the float range of values from the parameter ParamName. If RangeValue is -1 removes the last of the float range of values.
+	/** Remove the FRangeIndex element of the float range of values from the parameter ParamName. If RangeValue is -1 removes the last of the float range of values.
 		Returns the index of the last valid float, -1 if no values left. */
 	UFUNCTION(BlueprintCallable, Category = CustomizableObjectInstance)
 	int32 RemoveValueFromFloatRange(const FString& ParamName, int32 RangeIndex = -1);
 
-	/** Remove the RangeIndex element of the projector range of values from the parameter ParamName. If RangeValue is -1 removes the last of the projector range of values.
+	/** Remove the FRangeIndex element of the projector range of values from the parameter ParamName. If RangeValue is -1 removes the last of the projector range of values.
 		Returns the index of the last valid projector, -1 if no values left. */
 	UFUNCTION(BlueprintCallable, Category = CustomizableObjectInstance)
 	int32 RemoveValueFromProjectorRange(const FString& ParamName, int32 RangeIndex = -1);
diff --git a/Engine/Plugins/Mutable/Source/CustomizableObject/Public/MuCO/CustomizableObjectSkeletalMesh.h b/Engine/Plugins/Mutable/Source/CustomizableObject/Public/MuCO/CustomizableObjectSkeletalMesh.h
--- a/Engine/Plugins/Mutable/Source/CustomizableObject/Public/MuCO/CustomizableObjectSkeletalMesh.h
+++ b/Engine/Plugins/Mutable/Source/CustomizableObject/Public/MuCO/CustomizableObjectSkeletalMesh.h
@@ -53,10 +53,10 @@ private: // Private Variables
 	FString CustomizableObjectPathName;
 	FString InstancePathName;
 
-	TSharedPtr<mu::Model, ESPMode::ThreadSafe> Model;
+	TSharedPtr<mu::FModel, ESPMode::ThreadSafe> Model;
 	
 	/** Instance parameters at the time of creation. */
-	mu::Ptr<mu::Parameters> Parameters;
+	TSharedPtr<mu::FParameters> Parameters;
 	int32 State = -1;
 
 	/** Mesh IDs per LOD */
diff --git a/Engine/Plugins/Mutable/Source/CustomizableObject/Public/MuCO/CustomizableObjectSystem.h b/Engine/Plugins/Mutable/Source/CustomizableObject/Public/MuCO/CustomizableObjectSystem.h
--- a/Engine/Plugins/Mutable/Source/CustomizableObject/Public/MuCO/CustomizableObjectSystem.h
+++ b/Engine/Plugins/Mutable/Source/CustomizableObject/Public/MuCO/CustomizableObjectSystem.h
@@ -251,26 +251,26 @@ public:
 	void SetImagePixelFormatOverride(const mu::FImageOperator::FImagePixelFormatFunc&);
 #endif
 
-	/** [Texture Parameters] Get a list of all the possible values for external texture parameters according to the various providers registered with RegisterImageProvider. */
+	/** [Texture FParameters] Get a list of all the possible values for external texture parameters according to the various providers registered with RegisterImageProvider. */
 	TArray<FCustomizableObjectExternalTexture> GetTextureParameterValues();
 
-	/** [Texture Parameters] Add a new image provider to the CustomizableObject System. This will be queried for when an external image ID is provided to mutable in an Texture Parameter node. */
+	/** [Texture FParameters] Add a new image provider to the CustomizableObject System. This will be queried for when an external image ID is provided to mutable in an Texture Parameter node. */
 	void RegisterImageProvider(UCustomizableSystemImageProvider* Provider);
 
-	/** [Texture Parameters] Remove a previously registered provider. */
+	/** [Texture FParameters] Remove a previously registered provider. */
 	void UnregisterImageProvider(UCustomizableSystemImageProvider* Provider);
 
-	/** [Texture Parameters] Interface to actually cache Images in the Mutable system and make them available at run-time.
+	/** [Texture FParameters] Interface to actually cache Images in the Mutable system and make them available at run-time.
 		Any cached image has to be registered by an Image provider before caching it.
 		Have in mind that once an image has been cached, it will spend memory according to its size, except in the case
 		of images of type UCustomizableSystemImageProvider::ValueType::Unreal_Deferred, where only a very small amount of
 		memory is used and the real texel data is loaded when needed during an update and then immediately discarded */
 
-	/** [Texture Parameters] Cache an image which has to have been previously registered by an Image provider with the parameter id. */
+	/** [Texture FParameters] Cache an image which has to have been previously registered by an Image provider with the parameter id. */
 	void CacheImage(FName ImageId);
-	/** [Texture Parameters] Remove an image from the cache. */
+	/** [Texture FParameters] Remove an image from the cache. */
 	void UnCacheImage(FName ImageId);
-	/** [Texture Parameters] Remove all images from the cache. */
+	/** [Texture FParameters] Remove all images from the cache. */
 	void ClearImageCache();
 
 	// Show a warning on-screen and via a notification (if in Editor) and log an error when a CustomizableObject is
diff --git a/Engine/Plugins/Mutable/Source/CustomizableObject/Public/MuCO/DefaultImageProvider.h b/Engine/Plugins/Mutable/Source/CustomizableObject/Public/MuCO/DefaultImageProvider.h
--- a/Engine/Plugins/Mutable/Source/CustomizableObject/Public/MuCO/DefaultImageProvider.h
+++ b/Engine/Plugins/Mutable/Source/CustomizableObject/Public/MuCO/DefaultImageProvider.h
@@ -11,7 +11,7 @@ class UCustomizableObjectInstance;
 
 namespace mu
 {
-	class Parameters;
+	class FParameters;
 }
 
 /** Simple image provider that translates UTexture2D to Mutable IDs.
diff --git a/Engine/Plugins/Mutable/Source/CustomizableObject/Public/MuCO/UnrealExtensionDataStreamer.h b/Engine/Plugins/Mutable/Source/CustomizableObject/Public/MuCO/UnrealExtensionDataStreamer.h
--- a/Engine/Plugins/Mutable/Source/CustomizableObject/Public/MuCO/UnrealExtensionDataStreamer.h
+++ b/Engine/Plugins/Mutable/Source/CustomizableObject/Public/MuCO/UnrealExtensionDataStreamer.h
@@ -11,8 +11,8 @@
 class UCustomizableObjectSystemPrivate;
 class UCustomizableObject;
 
-/** An implementation of mu::ExtensionDataStreamer designed to work with the Customizable Object integration */
-class CUSTOMIZABLEOBJECT_API FUnrealExtensionDataStreamer : public mu::ExtensionDataStreamer
+/** An implementation of mu::FExtensionDataStreamer designed to work with the Customizable Object integration */
+class CUSTOMIZABLEOBJECT_API FUnrealExtensionDataStreamer : public mu::FExtensionDataStreamer
 {
 public:
 	FUnrealExtensionDataStreamer(const TWeakObjectPtr<UCustomizableObjectSystemPrivate>& InSystemPrivateWeak);
@@ -23,7 +23,7 @@ public:
 	FUnrealExtensionDataStreamer& operator=(const FUnrealExtensionDataStreamer& Other) = delete;
 
 	/**
-	 * Note that *only* the functions from the mu::ExtensionDataStreamer interface may be called
+	 * Note that *only* the functions from the mu::FExtensionDataStreamer interface may be called
 	 * from other threads. All other functions may only be called from the Game thread.
 	 */
 
@@ -48,12 +48,12 @@ public:
 	/** Cancels all loads requested by StartLoad that haven't yet completed. */
 	void CancelPendingLoads();
 
-	/** mu::ExtensionDataStreamer interface */
-	virtual mu::ExtensionDataPtr CloneExtensionData(const mu::ExtensionDataPtrConst& Source) override;
+	/** mu::FExtensionDataStreamer interface */
+	virtual TSharedPtr<mu::FExtensionData> CloneExtensionData(const TSharedPtr<const mu::FExtensionData>& Source) override;
 	
 	virtual TSharedRef<const mu::FExtensionDataLoadHandle> StartLoad(
-		const mu::ExtensionDataPtrConst& Data,
-		TArray<mu::ExtensionDataPtrConst>& OutUnloadedConstants) override;
+		const TSharedPtr<const mu::FExtensionData>& Data,
+		TArray<TSharedPtr<const mu::FExtensionData>>& OutUnloadedConstants) override;
 
 private:
 	/** Queued up on the Game thread to start the async load */
diff --git a/Engine/Plugins/Mutable/Source/CustomizableObjectEditor/Private/MuCOE/CustomizableInstanceDetails.cpp b/Engine/Plugins/Mutable/Source/CustomizableObjectEditor/Private/MuCOE/CustomizableInstanceDetails.cpp
--- a/Engine/Plugins/Mutable/Source/CustomizableObjectEditor/Private/MuCOE/CustomizableInstanceDetails.cpp
+++ b/Engine/Plugins/Mutable/Source/CustomizableObjectEditor/Private/MuCOE/CustomizableInstanceDetails.cpp
@@ -90,8 +90,8 @@ void FCustomizableInstanceDetails::CustomizeDetails(const TSharedPtr<IDetailLayo
 	// New Category that will store all properties widgets
 	IDetailCategoryBuilder& ResourcesCategory = DetailBuilder->EditCategory("Generated Resources");
 	IDetailCategoryBuilder& VisibilitySettingsCategory = DetailBuilder->EditCategory("ParametersVisibility");
-	IDetailCategoryBuilder& ParametersCategory = DetailBuilder->EditCategory("Instance Parameters");
-	IDetailCategoryBuilder& OldParametersCategory = DetailBuilder->EditCategory("Old Instance Parameters");
+	IDetailCategoryBuilder& ParametersCategory = DetailBuilder->EditCategory("Instance FParameters");
+	IDetailCategoryBuilder& OldParametersCategory = DetailBuilder->EditCategory("Old Instance FParameters");
 	IDetailCategoryBuilder& SkeletalMeshCategory = DetailBuilder->EditCategory("CustomizableSkeletalMesh");
 	IDetailCategoryBuilder& TextureParametersCategory = DetailBuilder->EditCategory("TextureParameter");
 
@@ -99,7 +99,7 @@ void FCustomizableInstanceDetails::CustomizeDetails(const TSharedPtr<IDetailLayo
 	UCustomizableObject* CustomizableObject = CustomInstance->GetCustomizableObject();
 	if (!CustomizableObject)
 	{
-		VisibilitySettingsCategory.AddCustomRow(LOCTEXT("CustomizableInstanceDetails_NoCOMessage", "Instance Parameters"))
+		VisibilitySettingsCategory.AddCustomRow(LOCTEXT("CustomizableInstanceDetails_NoCOMessage", "Instance FParameters"))
 		[
 			SNew(STextBlock).Text(LOCTEXT("Model not compiled", "Model not compiled"))
 		];
@@ -149,7 +149,7 @@ void FCustomizableInstanceDetails::CustomizeDetails(const TSharedPtr<IDetailLayo
 		GenerateInstanceProfileSelector()
 	];
 
-	// Only Runtime Parameters Option
+	// Only Runtime FParameters Option
 	VisibilitySettingsCategory.AddCustomRow( LOCTEXT("CustomizableInstanceDetails_RuntimeParm", "Only Runtime") )
 	.NameContent()
 	[
@@ -164,7 +164,7 @@ void FCustomizableInstanceDetails::CustomizeDetails(const TSharedPtr<IDetailLayo
 		.OnCheckStateChanged(this, &FCustomizableInstanceDetails::OnShowOnlyRuntimeSelectionChanged)
 	];
 	
-	// Only Relevant Parameters Option
+	// Only Relevant FParameters Option
 	VisibilitySettingsCategory.AddCustomRow( LOCTEXT("CustomizableInstanceDetails_RelevantParam", "Only Relevant") )
 	.NameContent()
 	[
@@ -221,7 +221,7 @@ void FCustomizableInstanceDetails::CustomizeDetails(const TSharedPtr<IDetailLayo
 		FilterType = EditorProperties->Filter.GameplayTagsFilterType;
 	}
 
-	// Copy, Paste and Reset Parameters
+	// Copy, Paste and Reset FParameters
 	ParametersCategory.AddCustomRow(LOCTEXT("CustomizableInstanceDetails_CopyPasteResetButtons", "Copy Paste Reset"))
 	[
 		SNew(SHorizontalBox)
@@ -231,7 +231,7 @@ void FCustomizableInstanceDetails::CustomizeDetails(const TSharedPtr<IDetailLayo
 		[
 			SNew(SButton)
 			.HAlign(HAlign_Center)
-			.Text(LOCTEXT("Copy_Parameters", "Copy Parameters"))
+			.Text(LOCTEXT("Copy_Parameters", "Copy FParameters"))
 			.OnClicked(this, &FCustomizableInstanceDetails::OnCopyAllParameters)
 			.IsEnabled(CustomInstance->HasAnyParameters())
 			.ToolTipText(FText::FromString(FString("Copy the preview Instance parameters")))
@@ -242,7 +242,7 @@ void FCustomizableInstanceDetails::CustomizeDetails(const TSharedPtr<IDetailLayo
 		[
 			SNew(SButton)
 			.HAlign(HAlign_Center)
-			.Text(LOCTEXT("Paste_Parameters", "Paste Parameters"))
+			.Text(LOCTEXT("Paste_Parameters", "Paste FParameters"))
 			.OnClicked(this, &FCustomizableInstanceDetails::OnPasteAllParameters)
 			.IsEnabled(CustomInstance->HasAnyParameters())
 			.ToolTipText(FText::FromString(FString("Paste the preview Instance parameters")))
@@ -260,14 +260,14 @@ void FCustomizableInstanceDetails::CustomizeDetails(const TSharedPtr<IDetailLayo
 		]
 	];
 
-	// Parameters Widgets
+	// FParameters Widgets
 	bool bHiddenParamsRuntime = GenerateParametersView(ParametersCategory);
 
 	if (bHiddenParamsRuntime)
 	{
-		FText HiddenParamsRuntimeMessage = LOCTEXT("CustomizableInstanceDetails_HiddemParamsRuntime", "Parameters are hidden due to their Runtime type. \nUncheck the Only Runtime checkbox to see them.");
+		FText HiddenParamsRuntimeMessage = LOCTEXT("CustomizableInstanceDetails_HiddemParamsRuntime", "FParameters are hidden due to their Runtime type. \nUncheck the Only Runtime checkbox to see them.");
 
-		ParametersCategory.AddCustomRow(LOCTEXT("CustomizableInstanceDetails_HiddemParamsRuntimeRow", "Parameters are hidden"))
+		ParametersCategory.AddCustomRow(LOCTEXT("CustomizableInstanceDetails_HiddemParamsRuntimeRow", "FParameters are hidden"))
 		[
 			SNew(SHorizontalBox)
 				+ SHorizontalBox::Slot()
@@ -2258,7 +2258,7 @@ FReply FCustomizableInstanceDetails::OnPasteAllParameters()
 	const int32 MinLOD = Descriptor.GetMinLod();
 	const TArray<uint16> RequestedLODLevels = Descriptor.GetRequestedLODLevels();
 	
-	FScopedTransaction LocalTransaction(LOCTEXT("OnPasteAllParameters", "Paste All Parameters"));
+	FScopedTransaction LocalTransaction(LOCTEXT("OnPasteAllParameters", "Paste All FParameters"));
 	CustomInstance->Modify();
 
 	if (Struct->ImportText(*ClipText, &Descriptor, nullptr, 0, GLog, GetPathNameSafe(Struct)))
@@ -2278,7 +2278,7 @@ FReply FCustomizableInstanceDetails::OnPasteAllParameters()
 
 FReply FCustomizableInstanceDetails::OnResetAllParameters()
 {
-	FScopedTransaction LocalTransaction(LOCTEXT("OnResetAllParameters", "Reset All Parameters"));
+	FScopedTransaction LocalTransaction(LOCTEXT("OnResetAllParameters", "Reset All FParameters"));
 	CustomInstance->Modify();
 
 	CustomInstance->SetDefaultValues();
diff --git a/Engine/Plugins/Mutable/Source/CustomizableObjectEditor/Private/MuCOE/CustomizableInstanceDetails.h b/Engine/Plugins/Mutable/Source/CustomizableObjectEditor/Private/MuCOE/CustomizableInstanceDetails.h
--- a/Engine/Plugins/Mutable/Source/CustomizableObjectEditor/Private/MuCOE/CustomizableInstanceDetails.h
+++ b/Engine/Plugins/Mutable/Source/CustomizableObjectEditor/Private/MuCOE/CustomizableInstanceDetails.h
@@ -90,12 +90,12 @@ private:
 	void GenerateWidgetRow(FDetailWidgetRow& WidgetRow, const UCustomizableObject& CustomizableObject, const FString& ParamName, const int32 ParamIndexInObject);
 	TSharedRef<SWidget> GenerateParameterWidget(const UCustomizableObject& CustomizableObject, const FString& ParamName, const int32 ParamIndexInObject);
 
-	// Int Parameters Functions
+	// Int FParameters Functions
 	TSharedRef<SWidget> GenerateIntWidget(const UCustomizableObject& CustomizableObject, const FString& ParamName, const int32 ParamIndexInObject);
 	void OnIntParameterComboBoxChanged(TSharedPtr<FString> Selection, ESelectInfo::Type SelectInfo, const FString ParamName);
 	TSharedRef<SWidget> OnGenerateWidgetIntParameter(TSharedPtr<FString> OptionName, const FString ParameterName);
 	
-	// Float Parameters Functions
+	// Float FParameters Functions
 	TSharedRef<SWidget> GenerateFloatWidget(const UCustomizableObject& CustomizableObject, const FString& ParamName);
 	float GetFloatParameterValue(const FString ParamName, int32 RangeIndex) const;
 	void OnFloatParameterChanged(float Value, const FString ParamName, int32 RangeIndex);
@@ -104,30 +104,30 @@ private:
 	void OnFloatParameterSliderEnd();
 	void OnFloatParameterCommited(float Value, ETextCommit::Type Type, const FString ParamName, int32 RangeIndex);	// Needed to have undo/redo
  
-	// Texture Parameters Functions
+	// Texture FParameters Functions
 	TSharedRef<SWidget> GenerateTextureWidget(const UCustomizableObject& CustomizableObject, const FString& ParamName);
 	void GenerateTextureParameterOptions();
 	void OnTextureParameterComboBoxSelectionChanged(TSharedPtr<FString> Selection, ESelectInfo::Type SelectInfo, const FString ParamName);
 	
-	// Color Parameters Functions
+	// Color FParameters Functions
 	TSharedRef<SWidget> GenerateColorWidget(const FString& ParamName);
 	FLinearColor GetColorParameterValue(const FString ParamName) const;
 	FReply OnColorBlockMouseButtonDown(const FGeometry& MyGeometry, const FPointerEvent& MouseEvent, const FString ParamName);
 	void OnSetColorFromColorPicker(FLinearColor NewColor, const FString PickerParamName);
 
-	// Transform Parameters Functions
+	// Transform FParameters Functions
 	TSharedRef<SWidget> GenerateTransformWidget(const FString& ParamName);
 	FTransform GetTransformParameterValue(const FString ParamName) const;
 	void OnTransformParameterChanged(FTransform NewTransform, const FString ParamName);
 	void OnTransformParameterCommitted(FTransform NewTransform, ETextCommit::Type Type, const FString ParamName);
 
 
-	// Bool Parameters Functions
+	// Bool FParameters Functions
 	TSharedRef<SWidget> GenerateBoolWidget(const FString& ParamName);
 	ECheckBoxState GetBoolParameterValue(const FString ParamName) const;
 	void OnBoolParameterChanged(ECheckBoxState InCheckboxState, const FString ParamName);
 
-	// Projector Parameters Functions
+	// Projector FParameters Functions
 	TSharedRef<SWidget> GenerateSimpleProjector(const FString& ParamName);
 	TSharedRef<SWidget> GenerateMultidimensionalProjector(const UCustomizableObject& CustomizableObject, const FString& ParamName, const int32 ParamIndexInObject);
 	TSharedPtr<ICustomizableObjectInstanceEditor> GetEditorChecked() const;
diff --git a/Engine/Plugins/Mutable/Source/CustomizableObjectEditor/Private/MuCOE/CustomizableObjectCompileRunnable.cpp b/Engine/Plugins/Mutable/Source/CustomizableObjectEditor/Private/MuCOE/CustomizableObjectCompileRunnable.cpp
--- a/Engine/Plugins/Mutable/Source/CustomizableObjectEditor/Private/MuCOE/CustomizableObjectCompileRunnable.cpp
+++ b/Engine/Plugins/Mutable/Source/CustomizableObjectEditor/Private/MuCOE/CustomizableObjectCompileRunnable.cpp
@@ -56,11 +56,11 @@ FCustomizableObjectCompileRunnable::FCustomizableObjectCompileRunnable(mu::Ptr<m
 }
 
 
-mu::Ptr<mu::Image> FCustomizableObjectCompileRunnable::LoadResourceReferenced(int32 ID)
+TSharedPtr<mu::FImage> FCustomizableObjectCompileRunnable::LoadResourceReferenced(int32 ID)
 {
 	MUTABLE_CPUPROFILER_SCOPE(LoadResourceReferenced);
 
-	mu::Ptr<mu::Image> Image;
+	TSharedPtr<mu::FImage> Image;
 	if (!ReferencedTextures.IsValidIndex(ID))
 	{
 		// The id is not valid for this CO
@@ -72,9 +72,9 @@ mu::Ptr<mu::Image> FCustomizableObjectCompileRunnable::LoadResourceReferenced(in
 	FMutableSourceTextureData& TextureData = ReferencedTextures[ID];
 
 	// In the editor the src data can be directly accessed
-	Image = new mu::Image();
+	Image = MakeShared<mu::FImage>();
 	int32 MipmapsToSkip = 0;
-	EUnrealToMutableConversionError Error = ConvertTextureUnrealSourceToMutable(Image.get(), TextureData, MipmapsToSkip);
+	EUnrealToMutableConversionError Error = ConvertTextureUnrealSourceToMutable(Image.Get(), TextureData, MipmapsToSkip);
 
 	if (Error != EUnrealToMutableConversionError::Success)
 	{
@@ -160,12 +160,12 @@ uint32 FCustomizableObjectCompileRunnable::Run()
 		CompilerOptions->SetImagePixelFormatOverride( UnrealPixelFormatFunc );
 	}
 
-	CompilerOptions->SetReferencedResourceCallback([this](int32 ID, TSharedPtr<mu::Ptr<mu::Image>> ResolvedImage, bool bRunImmediatlyIfPossible)
+	CompilerOptions->SetReferencedResourceCallback([this](int32 ID, TSharedPtr<TSharedPtr<mu::FImage>> ResolvedImage, bool bRunImmediatlyIfPossible)
 		{
 			UE::Tasks::FTask LaunchTask = UE::Tasks::Launch(TEXT("ConstantGeneratorLaunchTasks"),
 				[ID,ResolvedImage,this]()
 				{
-					mu::Ptr<mu::Image> Result = LoadResourceReferenced(ID);
+					TSharedPtr<mu::FImage> Result = LoadResourceReferenced(ID);
 					*ResolvedImage = Result;
 				},
 				LowLevelTasks::ETaskPriority::BackgroundLow
@@ -286,7 +286,7 @@ void FCustomizableObjectCompileRunnable::Tick()
 
 
 FCustomizableObjectSaveDDRunnable::FCustomizableObjectSaveDDRunnable(const TSharedPtr<FCompilationRequest>& InRequest,
-	TSharedPtr<mu::Model> InModel,
+	TSharedPtr<mu::FModel> InModel,
 	FModelResources& InModelResources,
 	TSharedPtr<FModelStreamableBulkData> InModelStreamables)
 {
@@ -412,15 +412,14 @@ void FCustomizableObjectSaveDDRunnable::CachePlatfromData()
 		}
 	}
 
-
 	// Cache Model and Model Roms
 	{
 		FMemoryWriter64 ModelMemoryWriter(PlatformData.ModelData);
 		FUnrealMutableModelBulkWriterCook Streamer(&ModelMemoryWriter, &PlatformData.ModelStreamableData);
 
-		// Serialize mu::Model and streamable resources 
+		// Serialize mu::FModel and streamable resources 
 		constexpr bool bDropData = true;
-		mu::Model::Serialise(Model.Get(), Streamer, bDropData);
+		mu::FModel::Serialise(Model.Get(), Streamer, bDropData);
 	}
 }
 
diff --git a/Engine/Plugins/Mutable/Source/CustomizableObjectEditor/Private/MuCOE/CustomizableObjectCompileRunnable.h b/Engine/Plugins/Mutable/Source/CustomizableObjectEditor/Private/MuCOE/CustomizableObjectCompileRunnable.h
--- a/Engine/Plugins/Mutable/Source/CustomizableObjectEditor/Private/MuCOE/CustomizableObjectCompileRunnable.h
+++ b/Engine/Plugins/Mutable/Source/CustomizableObjectEditor/Private/MuCOE/CustomizableObjectCompileRunnable.h
@@ -50,12 +50,12 @@ private:
 	struct FReferenceResourceRequest
 	{
 		int32 ID = -1;
-		TSharedPtr<mu::Ptr<mu::Image>> ResolvedImage;
+		TSharedPtr<TSharedPtr<mu::FImage>> ResolvedImage;
 		TSharedPtr<UE::Tasks::FTaskEvent> CompletionEvent;
 	};
 	TQueue<FReferenceResourceRequest, EQueueMode::Mpsc> PendingResourceReferenceRequests;
 
-	mu::Ptr<mu::Image> LoadResourceReferenced(int32 ID);
+	TSharedPtr<mu::FImage> LoadResourceReferenced(int32 ID);
 
 
 public:
@@ -75,7 +75,7 @@ public:
 
 	void Tick();
 
-	TSharedPtr<mu::Model, ESPMode::ThreadSafe> Model;
+	TSharedPtr<mu::FModel, ESPMode::ThreadSafe> Model;
 
 	FCompilationOptions Options;
 	
@@ -92,7 +92,7 @@ class FCustomizableObjectSaveDDRunnable : public FRunnable
 {
 public:
 
-	FCustomizableObjectSaveDDRunnable(const TSharedPtr<FCompilationRequest>& InRequest, TSharedPtr<mu::Model> InModel, FModelResources& ModelResources, TSharedPtr<FModelStreamableBulkData> ModelStreamables);
+	FCustomizableObjectSaveDDRunnable(const TSharedPtr<FCompilationRequest>& InRequest, TSharedPtr<mu::FModel> InModel, FModelResources& ModelResources, TSharedPtr<FModelStreamableBulkData> ModelStreamables);
 
 	// FRunnable interface
 	uint32 Run() override;
@@ -130,7 +130,7 @@ private:
 
 public:
 
-	TSharedPtr<mu::Model, ESPMode::ThreadSafe> Model;
+	TSharedPtr<mu::FModel, ESPMode::ThreadSafe> Model;
 	TSharedPtr<FModelStreamableBulkData> ModelStreamables;
 
 	// Cached platform data
diff --git a/Engine/Plugins/Mutable/Source/CustomizableObjectEditor/Private/MuCOE/CustomizableObjectCompiler.h b/Engine/Plugins/Mutable/Source/CustomizableObjectEditor/Private/MuCOE/CustomizableObjectCompiler.h
--- a/Engine/Plugins/Mutable/Source/CustomizableObjectEditor/Private/MuCOE/CustomizableObjectCompiler.h
+++ b/Engine/Plugins/Mutable/Source/CustomizableObjectEditor/Private/MuCOE/CustomizableObjectCompiler.h
@@ -131,7 +131,7 @@ private:
 	/** Array used to protect from garbage collection those COs loaded asynchronously */
 	TArray<TObjectPtr<UObject>> ArrayGCProtect;
 
-	TSharedPtr<mu::Model, ESPMode::ThreadSafe> Model;
+	TSharedPtr<mu::FModel, ESPMode::ThreadSafe> Model;
 	TSharedPtr<FModelStreamableBulkData> ModelStreamableBulkData;
 
 	// Protected from GC with FCustomizableObjectCompiler::AddReferencedObjects
diff --git a/Engine/Plugins/Mutable/Source/CustomizableObjectEditor/Private/MuCOE/CustomizableObjectDetails.cpp b/Engine/Plugins/Mutable/Source/CustomizableObjectEditor/Private/MuCOE/CustomizableObjectDetails.cpp
--- a/Engine/Plugins/Mutable/Source/CustomizableObjectEditor/Private/MuCOE/CustomizableObjectDetails.cpp
+++ b/Engine/Plugins/Mutable/Source/CustomizableObjectEditor/Private/MuCOE/CustomizableObjectDetails.cpp
@@ -34,7 +34,7 @@ public:
 		ParameterIndex = -1;
 	}
 
-	TSharedPtr<mu::Model, ESPMode::ThreadSafe> Model;
+	TSharedPtr<mu::FModel, ESPMode::ThreadSafe> Model;
 	int StateIndex;
 	int ParameterIndex;
 
@@ -48,7 +48,7 @@ public:
 		}
 		else
 		{
-			mu::ParametersPtr TempParams = mu::Model::NewParameters(Model);
+			TSharedPtr<mu::FParameters> TempParams = mu::FModel::NewParameters(Model);
 			res = TempParams->GetName( ParameterIndex );
 		}
 
diff --git a/Engine/Plugins/Mutable/Source/CustomizableObjectEditor/Private/MuCOE/CustomizableObjectLayout.cpp b/Engine/Plugins/Mutable/Source/CustomizableObjectEditor/Private/MuCOE/CustomizableObjectLayout.cpp
--- a/Engine/Plugins/Mutable/Source/CustomizableObjectEditor/Private/MuCOE/CustomizableObjectLayout.cpp
+++ b/Engine/Plugins/Mutable/Source/CustomizableObjectEditor/Private/MuCOE/CustomizableObjectLayout.cpp
@@ -137,7 +137,7 @@ void UCustomizableObjectLayout::GenerateAutomaticBlocksFromUVs()
 			Block.bIsAutomatic = true;
 			Block.Id = FGuid::NewGuid();
 
-			mu::Ptr<mu::Image> Mask = LayoutNode->Blocks[i].Mask;
+			TSharedPtr<mu::FImage> Mask = LayoutNode->Blocks[i].Mask;
 			if (Mask)
 			{
 				UTexture2D* UnrealImage = NewObject<UTexture2D>(UTexture2D::StaticClass());
diff --git a/Engine/Plugins/Mutable/Source/CustomizableObjectEditor/Private/MuCOE/ExtensionDataCompilerInterface.cpp b/Engine/Plugins/Mutable/Source/CustomizableObjectEditor/Private/MuCOE/ExtensionDataCompilerInterface.cpp
--- a/Engine/Plugins/Mutable/Source/CustomizableObjectEditor/Private/MuCOE/ExtensionDataCompilerInterface.cpp
+++ b/Engine/Plugins/Mutable/Source/CustomizableObjectEditor/Private/MuCOE/ExtensionDataCompilerInterface.cpp
@@ -13,10 +13,10 @@ FExtensionDataCompilerInterface::FExtensionDataCompilerInterface(FMutableGraphGe
 {
 }
 
-mu::ExtensionDataPtrConst FExtensionDataCompilerInterface::MakeStreamedExtensionData(UCustomizableObjectResourceDataContainer*& OutContainer)
+TSharedPtr<const mu::FExtensionData> FExtensionDataCompilerInterface::MakeStreamedExtensionData(UCustomizableObjectResourceDataContainer*& OutContainer)
 {
-	mu::ExtensionDataPtr Result = new mu::ExtensionData;
-	Result->Origin = mu::ExtensionData::EOrigin::ConstantStreamed;
+	TSharedPtr<mu::FExtensionData> Result = MakeShared<mu::FExtensionData>();
+	Result->Origin = mu::FExtensionData::EOrigin::ConstantStreamed;
 	Result->Index = GenerationContext.StreamedExtensionData.Num();
 	
 	if (!GenerationContext.bParticipatingObjectsPass)
@@ -35,10 +35,10 @@ mu::ExtensionDataPtrConst FExtensionDataCompilerInterface::MakeStreamedExtension
 	return Result;
 }
 
-mu::ExtensionDataPtrConst FExtensionDataCompilerInterface::MakeAlwaysLoadedExtensionData(FInstancedStruct&& Data)
+TSharedPtr<const mu::FExtensionData> FExtensionDataCompilerInterface::MakeAlwaysLoadedExtensionData(FInstancedStruct&& Data)
 {
-	mu::ExtensionDataPtr Result = new mu::ExtensionData;
-	Result->Origin = mu::ExtensionData::EOrigin::ConstantAlwaysLoaded;
+	TSharedPtr<mu::FExtensionData> Result = MakeShared<mu::FExtensionData>();
+	Result->Origin = mu::FExtensionData::EOrigin::ConstantAlwaysLoaded;
 	Result->Index = GenerationContext.AlwaysLoadedExtensionData.Num();
 
 	FCustomizableObjectResourceData* CompileTimeExtensionData = &GenerationContext.AlwaysLoadedExtensionData.AddDefaulted_GetRef();
diff --git a/Engine/Plugins/Mutable/Source/CustomizableObjectEditor/Private/MuCOE/GenerateMutableSource/GenerateMutableSource.cpp b/Engine/Plugins/Mutable/Source/CustomizableObjectEditor/Private/MuCOE/GenerateMutableSource/GenerateMutableSource.cpp
--- a/Engine/Plugins/Mutable/Source/CustomizableObjectEditor/Private/MuCOE/GenerateMutableSource/GenerateMutableSource.cpp
+++ b/Engine/Plugins/Mutable/Source/CustomizableObjectEditor/Private/MuCOE/GenerateMutableSource/GenerateMutableSource.cpp
@@ -206,7 +206,7 @@ void FMutableGraphGenerationContext::Log(const FText& Message, const UObject* Co
 }
 
 
-mu::Ptr<mu::Mesh> FMutableGraphGenerationContext::FindGeneratedMesh( const FGeneratedMeshData::FKey& Key )
+TSharedPtr<mu::FMesh> FMutableGraphGenerationContext::FindGeneratedMesh( const FGeneratedMeshData::FKey& Key )
 {
 	for (const FGeneratedMeshData& d : GeneratedMeshes)
 	{
@@ -803,11 +803,11 @@ mu::NodeMeshApplyPosePtr CreateNodeMeshApplyPose(FMutableGraphGenerationContext&
 {
 	check(ArrayBoneName.Num() == ArrayTransform.Num());
 
-	mu::Ptr<mu::Mesh> MutableMesh = new mu::Mesh();
+	TSharedPtr<mu::FMesh> MutableMesh = MakeShared<mu::FMesh>();
 	mu::NodeMeshConstantPtr PoseNodeMesh = new mu::NodeMeshConstant;
 	PoseNodeMesh->SetValue(MutableMesh);
 
-	mu::Ptr<mu::Skeleton> MutableSkeleton = new mu::Skeleton;
+	TSharedPtr<mu::FSkeleton> MutableSkeleton = MakeShared<mu::FSkeleton>();
 	MutableMesh->SetSkeleton(MutableSkeleton);
 	MutableMesh->SetBonePoseCount(ArrayBoneName.Num());
 	MutableSkeleton->SetBoneCount(ArrayBoneName.Num());
@@ -1320,7 +1320,7 @@ mu::Ptr<mu::NodeObject> GenerateMutableSource(const UEdGraphPin * Pin, FMutableG
 }
 
 
-int32 AddTagToMutableMeshUnique(mu::Mesh& MutableMesh, const FString& Tag)
+int32 AddTagToMutableMeshUnique(mu::FMesh& MutableMesh, const FString& Tag)
 {
 	const int32 TagCount = MutableMesh.GetTagCount();
 
@@ -1635,7 +1635,7 @@ mu::FImageDesc GenerateImageDescriptor(UTexture* Texture)
 }
 
 
-mu::Ptr<mu::Image> GenerateImageConstant(UTexture* Texture, FMutableGraphGenerationContext& GenerationContext, bool bIsReference)
+TSharedPtr<mu::FImage> GenerateImageConstant(UTexture* Texture, FMutableGraphGenerationContext& GenerationContext, bool bIsReference)
 {
 	MUTABLE_CPUPROFILER_SCOPE(GenerateImageConstant);
 
@@ -1692,12 +1692,12 @@ mu::Ptr<mu::Image> GenerateImageConstant(UTexture* Texture, FMutableGraphGenerat
 	}
 	
 	// Compile-time references that are left should be resolved immediately (should only happen in editor).
-	mu::Ptr<mu::Image> Result = mu::Image::CreateAsReference(Entry->ID, ImageDesc, bForceLoad);
+	TSharedPtr<mu::FImage> Result = mu::FImage::CreateAsReference(Entry->ID, ImageDesc, bForceLoad);
 	return Result;
 }
 
 
-mu::Ptr<mu::Mesh> GenerateMeshConstant(USkeletalMesh* Mesh, FMutableGraphGenerationContext& GenerationContext, bool bIsReference)
+TSharedPtr<mu::FMesh> GenerateMeshConstant(USkeletalMesh* Mesh, FMutableGraphGenerationContext& GenerationContext, bool bIsReference)
 {
 	MUTABLE_CPUPROFILER_SCOPE(GenerateMeshConstant)
 
@@ -1756,12 +1756,12 @@ mu::Ptr<mu::Mesh> GenerateMeshConstant(USkeletalMesh* Mesh, FMutableGraphGenerat
 		Entry->ID = Num - 1;
 	}
 
-	mu::Ptr<mu::Mesh> Result = mu::Mesh::CreateAsReference(Entry->ID, bForceLoad);
+	TSharedPtr<mu::FMesh> Result = mu::FMesh::CreateAsReference(Entry->ID, bForceLoad);
 	return Result;
 }
 
 
-void AddSocketTagsToMesh(const USkeletalMesh* SourceMesh, mu::Ptr<mu::Mesh> MutableMesh, FMutableGraphGenerationContext& GenerationContext)
+void AddSocketTagsToMesh(const USkeletalMesh* SourceMesh, TSharedPtr<mu::FMesh> MutableMesh, FMutableGraphGenerationContext& GenerationContext)
 {
 	for (int32 SocketIndex = 0; SocketIndex < SourceMesh->NumSockets(); ++SocketIndex)
 	{
diff --git a/Engine/Plugins/Mutable/Source/CustomizableObjectEditor/Private/MuCOE/GenerateMutableSource/GenerateMutableSource.h b/Engine/Plugins/Mutable/Source/CustomizableObjectEditor/Private/MuCOE/GenerateMutableSource/GenerateMutableSource.h
--- a/Engine/Plugins/Mutable/Source/CustomizableObjectEditor/Private/MuCOE/GenerateMutableSource/GenerateMutableSource.h
+++ b/Engine/Plugins/Mutable/Source/CustomizableObjectEditor/Private/MuCOE/GenerateMutableSource/GenerateMutableSource.h
@@ -24,7 +24,6 @@
 #include "MuT/NodeScalarParameter.h"
 #include "MuT/NodeSurfaceNew.h"
 #include "MuT/Table.h"
-#include "MuR/MutableMemory.h"
 #include "MuR/Skeleton.h"
 #include "MuR/Mesh.h"
 #include "MuR/Ptr.h"
@@ -617,7 +616,7 @@ public:
 	struct FGeneratedDataTablesData
 	{
 		// Pointer to the generated mutable Table
-		mu::Ptr<mu::Table> GeneratedTable;
+		mu::Ptr<mu::FTable> GeneratedTable;
 
 		// Table Node used to fill this info
 		const UCustomizableObjectNodeTable* ReferenceNode;
@@ -736,7 +735,7 @@ public:
 		FKey Key;
 
 		/** Generated mesh. */
-		mu::Ptr<mu::Mesh> Generated;
+		TSharedPtr<mu::FMesh> Generated;
 	};
 	TArray<FGeneratedMeshData> GeneratedMeshes;
 
@@ -744,7 +743,7 @@ public:
 	{
 		FString PinName;
 		FName PinType;
-		const mu::Ptr<mu::Table> Table;
+		const mu::Ptr<mu::FTable> Table;
 		const UCustomizableObjectNodeTable* TableNode;
 
 		bool operator==(const FGeneratedTableImageData& Other) const
@@ -766,7 +765,7 @@ public:
 	TArray<TArray<UCustomizableObjectNodeGroupProjectorParameter*>> CurrentGroupProjectors;
 
 	/** Find a mesh if already generated for a given source and flags. */
-	mu::Ptr<mu::Mesh> FindGeneratedMesh(const FGeneratedMeshData::FKey& Key);
+	TSharedPtr<mu::FMesh> FindGeneratedMesh(const FGeneratedMeshData::FKey& Key);
 
 	/** Add a resource to the streamed resources array.
 	  * OutStreamedResourceContainer - Container to store the streamed resource.
@@ -1070,9 +1069,9 @@ mu::NodeMeshApplyPosePtr CreateNodeMeshApplyPose(FMutableGraphGenerationContext&
 
 
 /** Adds Tag to MutableMesh uniquely, returns the index were the tag has been inserted or the index where an intance of the tag has been found */
-int32 AddTagToMutableMeshUnique(mu::Mesh& MutableMesh, const FString& Tag);
+int32 AddTagToMutableMeshUnique(mu::FMesh& MutableMesh, const FString& Tag);
 
-void AddSocketTagsToMesh(const USkeletalMesh* SourceMesh, mu::Ptr<mu::Mesh> MutableMesh, FMutableGraphGenerationContext& GenerationContext);
+void AddSocketTagsToMesh(const USkeletalMesh* SourceMesh, TSharedPtr<mu::FMesh> MutableMesh, FMutableGraphGenerationContext& GenerationContext);
 
 // Generates the tag for an animation instance
 FString GenerateAnimationInstanceTag(const int32 AnimInstanceIndex, const FName& SlotIndex);
@@ -1092,8 +1091,8 @@ int32 GetMaxTextureSize(const UTexture2D& ReferenceTexture, const UTextureLODSet
 // Max texture size of the texture with per platform MaxTextureSize and LODBias applied.
 int32 GetTextureSizeInGame(const UTexture2D& Texture, const UTextureLODSettings& LODSettings, uint8 SurfaceLODBias = 0);
 
-mu::Ptr<mu::Image> GenerateImageConstant(UTexture*, FMutableGraphGenerationContext&, bool bIsReference);
-mu::Ptr<mu::Mesh> GenerateMeshConstant(USkeletalMesh*, FMutableGraphGenerationContext&, bool bIsReference);
+TSharedPtr<mu::FImage> GenerateImageConstant(UTexture*, FMutableGraphGenerationContext&, bool bIsReference);
+TSharedPtr<mu::FMesh> GenerateMeshConstant(USkeletalMesh*, FMutableGraphGenerationContext&, bool bIsReference);
 
 /** Generates a mutable image descriptor from an unreal engine texture */
 mu::FImageDesc GenerateImageDescriptor(UTexture* Texture);
diff --git a/Engine/Plugins/Mutable/Source/CustomizableObjectEditor/Private/MuCOE/GenerateMutableSource/GenerateMutableSourceColor.cpp b/Engine/Plugins/Mutable/Source/CustomizableObjectEditor/Private/MuCOE/GenerateMutableSource/GenerateMutableSourceColor.cpp
--- a/Engine/Plugins/Mutable/Source/CustomizableObjectEditor/Private/MuCOE/GenerateMutableSource/GenerateMutableSourceColor.cpp
+++ b/Engine/Plugins/Mutable/Source/CustomizableObjectEditor/Private/MuCOE/GenerateMutableSource/GenerateMutableSourceColor.cpp
@@ -291,7 +291,7 @@ mu::Ptr<mu::NodeColour> GenerateMutableSourceColor(const UEdGraphPin* Pin, FMuta
 			if (bSuccess)
 			{
 				// Generating a new data table if not exists
-				mu::TablePtr Table;
+				mu::Ptr<mu::FTable> Table;
 				Table = GenerateMutableSourceTable(DataTable, TypedNodeTable, GenerationContext);
 
 				if (Table)
diff --git a/Engine/Plugins/Mutable/Source/CustomizableObjectEditor/Private/MuCOE/GenerateMutableSource/GenerateMutableSourceComponent.cpp b/Engine/Plugins/Mutable/Source/CustomizableObjectEditor/Private/MuCOE/GenerateMutableSource/GenerateMutableSourceComponent.cpp
--- a/Engine/Plugins/Mutable/Source/CustomizableObjectEditor/Private/MuCOE/GenerateMutableSource/GenerateMutableSourceComponent.cpp
+++ b/Engine/Plugins/Mutable/Source/CustomizableObjectEditor/Private/MuCOE/GenerateMutableSource/GenerateMutableSourceComponent.cpp
@@ -227,7 +227,7 @@ mu::Ptr<mu::NodeComponent> GenerateMutableSourceComponent(const UEdGraphPin* Pin
 			FString MeshUniqueTags;
 			constexpr bool bIsReference = true;
 			TSoftClassPtr<UAnimInstance> AnimInstance;
-			mu::Ptr<mu::Mesh> MutableMesh = GenerateMutableMesh(SkeletalMesh, AnimInstance, 0, 0, 0, 0, MeshUniqueTags, 0, 
+			TSharedPtr<mu::FMesh> MutableMesh = GenerateMutableMesh(SkeletalMesh, AnimInstance, 0, 0, 0, 0, MeshUniqueTags, 0, 
 																GenerationContext, TypedComponentPassthroughMesh, nullptr, bIsReference);
 
 			MeshNode->SetValue(MutableMesh);
diff --git a/Engine/Plugins/Mutable/Source/CustomizableObjectEditor/Private/MuCOE/GenerateMutableSource/GenerateMutableSourceFloat.cpp b/Engine/Plugins/Mutable/Source/CustomizableObjectEditor/Private/MuCOE/GenerateMutableSource/GenerateMutableSourceFloat.cpp
--- a/Engine/Plugins/Mutable/Source/CustomizableObjectEditor/Private/MuCOE/GenerateMutableSource/GenerateMutableSourceFloat.cpp
+++ b/Engine/Plugins/Mutable/Source/CustomizableObjectEditor/Private/MuCOE/GenerateMutableSource/GenerateMutableSourceFloat.cpp
@@ -332,7 +332,7 @@ mu::NodeScalarPtr GenerateMutableSourceFloat(const UEdGraphPin* Pin, FMutableGra
 			if (bSuccess)
 			{
 				// Generating a new data table if not exists
-				mu::TablePtr Table;
+				mu::Ptr<mu::FTable> Table;
 				Table = GenerateMutableSourceTable(DataTable, TypedNodeTable, GenerationContext);
 
 				if (Table)
diff --git a/Engine/Plugins/Mutable/Source/CustomizableObjectEditor/Private/MuCOE/GenerateMutableSource/GenerateMutableSourceGroupProjector.cpp b/Engine/Plugins/Mutable/Source/CustomizableObjectEditor/Private/MuCOE/GenerateMutableSource/GenerateMutableSourceGroupProjector.cpp
--- a/Engine/Plugins/Mutable/Source/CustomizableObjectEditor/Private/MuCOE/GenerateMutableSource/GenerateMutableSourceGroupProjector.cpp
+++ b/Engine/Plugins/Mutable/Source/CustomizableObjectEditor/Private/MuCOE/GenerateMutableSource/GenerateMutableSourceGroupProjector.cpp
@@ -515,10 +515,10 @@ TOptional<FGroupProjectorTempData> GenerateMutableGroupProjector(UCustomizableOb
 		{
 			if (const TObjectPtr<UTexture2D>& Texture = ArrayOptionTexture[SelectorIndex].OptionTexture)
 			{
-				mu::Ptr<mu::Image> ImageConstant = GenerateImageConstant(ArrayOptionTexture[SelectorIndex].OptionTexture, GenerationContext, false);
+				TSharedPtr<mu::FImage> ImageConstant = GenerateImageConstant(ArrayOptionTexture[SelectorIndex].OptionTexture, GenerationContext, false);
 
 				mu::Ptr<mu::NodeImageConstant> ImageNode = new mu::NodeImageConstant();
-				ImageNode->SetValue(ImageConstant.get());
+				ImageNode->SetValue(ImageConstant);
 
 				const uint32 MipsToSkip = ComputeLODBiasForTexture(GenerationContext, *Texture, ProjParamNode->ReferenceTexture) + AdditionalLODBias;
 				SwitchNode->SetOption(SelectorIndex, ResizeTextureByNumMips(ImageNode, MipsToSkip));
@@ -526,7 +526,7 @@ TOptional<FGroupProjectorTempData> GenerateMutableGroupProjector(UCustomizableOb
 				// Calculate the number of mips to tag as high res for this image.
 				if (ProjParamNode->ReferenceTexture)
 				{
-					int32 TotalMips = mu::Image::GetMipmapCount(ProjParamNode->ReferenceTexture->GetSizeX(), ProjParamNode->ReferenceTexture->GetSizeY());
+					int32 TotalMips = mu::FImage::GetMipmapCount(ProjParamNode->ReferenceTexture->GetSizeX(), ProjParamNode->ReferenceTexture->GetSizeY());
 					int32 NumMipsBeyondMin = FMath::Max(0, TotalMips - GenerationContext.Options.MinDiskMips);
 					int32 HighResMipsForThisImage = FMath::Min(NumMipsBeyondMin, GenerationContext.Options.NumHighResImageMips);
 					ImageNode->SourceDataDescriptor.SourceHighResMips = HighResMipsForThisImage;
diff --git a/Engine/Plugins/Mutable/Source/CustomizableObjectEditor/Private/MuCOE/GenerateMutableSource/GenerateMutableSourceImage.cpp b/Engine/Plugins/Mutable/Source/CustomizableObjectEditor/Private/MuCOE/GenerateMutableSource/GenerateMutableSourceImage.cpp
--- a/Engine/Plugins/Mutable/Source/CustomizableObjectEditor/Private/MuCOE/GenerateMutableSource/GenerateMutableSourceImage.cpp
+++ b/Engine/Plugins/Mutable/Source/CustomizableObjectEditor/Private/MuCOE/GenerateMutableSource/GenerateMutableSourceImage.cpp
@@ -118,12 +118,12 @@ mu::NodeImagePtr GenerateMutableSourceImage(const UEdGraphPin* Pin, FMutableGrap
 			}
 			else
 			{
-				mu::Ptr<mu::Image> ImageConstant = GenerateImageConstant(BaseTexture, GenerationContext, false);
+				TSharedPtr<mu::FImage> ImageConstant = GenerateImageConstant(BaseTexture, GenerationContext, false);
 
 				mu::Ptr<mu::NodeImageConstant> ConstantImageNode = new mu::NodeImageConstant;
-				ConstantImageNode->SetValue( ImageConstant.get() );
+				ConstantImageNode->SetValue( ImageConstant );
 				
-				int32 TotalMips = mu::Image::GetMipmapCount(ImageConstant->GetSizeX(), ImageConstant->GetSizeY());
+				int32 TotalMips = mu::FImage::GetMipmapCount(ImageConstant->GetSizeX(), ImageConstant->GetSizeY());
 				int32 NumMipsBeyondMin = FMath::Max(0, TotalMips - int32(MipsToSkip) - GenerationContext.Options.MinDiskMips);
 				int32 HighResMipsForThisImage = FMath::Min(NumMipsBeyondMin, GenerationContext.Options.NumHighResImageMips);
 				ConstantImageNode->SourceDataDescriptor.SourceHighResMips = HighResMipsForThisImage;
@@ -219,13 +219,13 @@ mu::NodeImagePtr GenerateMutableSourceImage(const UEdGraphPin* Pin, FMutableGrap
 
 		if (Texture)
 		{
-			mu::Ptr<mu::Image> ImageConstant = GenerateImageConstant(Texture, GenerationContext, false);
-			ImageNode->SetValue(ImageConstant.get());
+			TSharedPtr<mu::FImage> ImageConstant = GenerateImageConstant(Texture, GenerationContext, false);
+			ImageNode->SetValue(ImageConstant);
 
 			const uint32 MipsToSkip = ComputeLODBiasForTexture(GenerationContext, *Texture, nullptr, ReferenceTextureSize);
 			Result = ResizeTextureByNumMips(ImageNode, MipsToSkip);
 
-			int32 TotalMips = mu::Image::GetMipmapCount(ImageConstant->GetSizeX(), ImageConstant->GetSizeY());
+			int32 TotalMips = mu::FImage::GetMipmapCount(ImageConstant->GetSizeX(), ImageConstant->GetSizeY());
 			int32 NumMipsBeyondMin = FMath::Max(0, TotalMips - int32(MipsToSkip) - GenerationContext.Options.MinDiskMips);
 			int32 HighResMipsForThisImage = FMath::Min(NumMipsBeyondMin, GenerationContext.Options.NumHighResImageMips);
 			ImageNode->SourceDataDescriptor.SourceHighResMips = HighResMipsForThisImage;
@@ -901,7 +901,7 @@ mu::NodeImagePtr GenerateMutableSourceImage(const UEdGraphPin* Pin, FMutableGrap
 			mu::Ptr<mu::NodeImageConstant> ImageNode = new mu::NodeImageConstant();
 			Result = ImageNode;
 
-			ImageNode->SetValue(GenerateImageConstant(BaseTexture, GenerationContext, true).get());
+			ImageNode->SetValue(GenerateImageConstant(BaseTexture, GenerationContext, true));
 		}
 		else
 		{
@@ -1044,7 +1044,7 @@ mu::NodeImagePtr GenerateMutableSourceImage(const UEdGraphPin* Pin, FMutableGrap
 				if (bSuccess)
 				{
 					// Generating a new data table if not exists
-					mu::TablePtr Table = nullptr;
+					mu::Ptr<mu::FTable> Table = nullptr;
 					Table = GenerateMutableSourceTable(DataTable, TypedNodeTable, GenerationContext);
 
 					if (Table)
@@ -1125,7 +1125,7 @@ mu::NodeImagePtr GenerateMutableSourceImage(const UEdGraphPin* Pin, FMutableGrap
 								ImageTableNode->ReferenceImageDesc = ImageDesc;
 
 								// Calculate the number of mips to tag as high res for this image.
-								int32 TotalMips = mu::Image::GetMipmapCount(ImageDesc.m_size[0], ImageDesc.m_size[1]);
+								int32 TotalMips = mu::FImage::GetMipmapCount(ImageDesc.m_size[0], ImageDesc.m_size[1]);
 								int32 NumMipsBeyondMin = FMath::Max(0, TotalMips - GenerationContext.Options.MinDiskMips);
 								int32 HighResMipsForThisImage = FMath::Min(NumMipsBeyondMin, GenerationContext.Options.NumHighResImageMips);
 								ImageTableNode->SourceDataDescriptor.SourceHighResMips = HighResMipsForThisImage;
diff --git a/Engine/Plugins/Mutable/Source/CustomizableObjectEditor/Private/MuCOE/GenerateMutableSource/GenerateMutableSourceImage.h b/Engine/Plugins/Mutable/Source/CustomizableObjectEditor/Private/MuCOE/GenerateMutableSource/GenerateMutableSourceImage.h
--- a/Engine/Plugins/Mutable/Source/CustomizableObjectEditor/Private/MuCOE/GenerateMutableSource/GenerateMutableSourceImage.h
+++ b/Engine/Plugins/Mutable/Source/CustomizableObjectEditor/Private/MuCOE/GenerateMutableSource/GenerateMutableSourceImage.h
@@ -14,12 +14,12 @@ struct FMutableGraphGenerationContext;
 
 namespace mu
 {
-	class Image;
+	class FImage;
 	class NodeImage;
 }
 
 
-//mu::Ptr<mu::Image> ConvertTextureUnrealToMutable(UTexture2D* Texture, const UCustomizableObjectNode* Node, FCustomizableObjectCompiler* Compiler, bool bIsNormalComposite);
+//TSharedPtr<mu::FImage> ConvertTextureUnrealToMutable(UTexture2D* Texture, const UCustomizableObjectNode* Node, FCustomizableObjectCompiler* Compiler, bool bIsNormalComposite);
 
 mu::Ptr<mu::NodeImage> ResizeTextureByNumMips(const mu::Ptr<mu::NodeImage>& Image, int32 MipsToSkip);
 
diff --git a/Engine/Plugins/Mutable/Source/CustomizableObjectEditor/Private/MuCOE/GenerateMutableSource/GenerateMutableSourceLayout.cpp b/Engine/Plugins/Mutable/Source/CustomizableObjectEditor/Private/MuCOE/GenerateMutableSource/GenerateMutableSourceLayout.cpp
--- a/Engine/Plugins/Mutable/Source/CustomizableObjectEditor/Private/MuCOE/GenerateMutableSource/GenerateMutableSourceLayout.cpp
+++ b/Engine/Plugins/Mutable/Source/CustomizableObjectEditor/Private/MuCOE/GenerateMutableSource/GenerateMutableSourceLayout.cpp
@@ -83,7 +83,7 @@ mu::Ptr<mu::NodeLayout> CreateMutableLayoutNode(FMutableGraphGenerationContext&
 	else
 	{
 		// Convert the UE mesh in the layout into a Mutable mesh.
-		mu::Ptr<mu::Mesh> MutableMesh = nullptr;
+		TSharedPtr<mu::FMesh> MutableMesh = nullptr;
 
 		if (USkeletalMesh* SkeletalMesh = Cast<USkeletalMesh>(UnrealLayout->GetMesh()))
 		{
@@ -144,10 +144,10 @@ mu::FSourceLayoutBlock ToMutable(FMutableGraphGenerationContext& GenerationConte
 		GenerationContext.AddParticipatingObject(*UnrealBlock.Mask);
 
 		// In the editor the src data can be directly accessed
-		mu::Ptr<mu::Image> MaskImage = new mu::Image();
+		TSharedPtr<mu::FImage> MaskImage = MakeShared<mu::FImage>();
 
 		FMutableSourceTextureData Tex(*UnrealBlock.Mask);
-		EUnrealToMutableConversionError Error = ConvertTextureUnrealSourceToMutable(MaskImage.get(), Tex, 0);
+		EUnrealToMutableConversionError Error = ConvertTextureUnrealSourceToMutable(MaskImage.Get(), Tex, 0);
 		if (Error != EUnrealToMutableConversionError::Success)
 		{
 			// This should never happen, so details are not necessary.
diff --git a/Engine/Plugins/Mutable/Source/CustomizableObjectEditor/Private/MuCOE/GenerateMutableSource/GenerateMutableSourceMesh.cpp b/Engine/Plugins/Mutable/Source/CustomizableObjectEditor/Private/MuCOE/GenerateMutableSource/GenerateMutableSourceMesh.cpp
--- a/Engine/Plugins/Mutable/Source/CustomizableObjectEditor/Private/MuCOE/GenerateMutableSource/GenerateMutableSourceMesh.cpp
+++ b/Engine/Plugins/Mutable/Source/CustomizableObjectEditor/Private/MuCOE/GenerateMutableSource/GenerateMutableSourceMesh.cpp
@@ -391,14 +391,14 @@ bool IsSkeletalMeshCompatibleWithRefSkeleton(FMutableComponentInfo& ComponentInf
 	return true;
 }
 
-void SetAndPropagatePoseBoneUsage(mu::Mesh& MutableMesh, int32 PoseIndex, mu::EBoneUsageFlags Usage)
+void SetAndPropagatePoseBoneUsage(mu::FMesh& MutableMesh, int32 PoseIndex, mu::EBoneUsageFlags Usage)
 {
 	if (!MutableMesh.GetSkeleton())
 	{
 		return;
 	}
 
-	const mu::Skeleton& MutableSkeleton = *MutableMesh.GetSkeleton();
+	const mu::FSkeleton& MutableSkeleton = *MutableMesh.GetSkeleton();
 
 	if (!MutableMesh.BonePoses.IsValidIndex(PoseIndex))
 	{
@@ -460,7 +460,7 @@ TArray<TTuple<UPhysicsAsset*, int32>> GetPhysicsAssetsFromAnimInstance(FMutableG
 	return Result;
 }
 
-TArray<uint8> MakePhysicsAssetBodySetupRelevancyMap(const FMutableGraphGenerationContext& GenerationContext, UPhysicsAsset* Asset, const mu::Ptr<mu::Mesh>& Mesh)
+TArray<uint8> MakePhysicsAssetBodySetupRelevancyMap(const FMutableGraphGenerationContext& GenerationContext, UPhysicsAsset* Asset, const TSharedPtr<mu::FMesh>& Mesh)
 {
 	const int32 BodySetupsNum = Asset->SkeletalBodySetups.Num();
 
@@ -481,7 +481,7 @@ TArray<uint8> MakePhysicsAssetBodySetupRelevancyMap(const FMutableGraphGeneratio
 	return RelevancyMap;
 }
 
-mu::Ptr<mu::PhysicsBody> MakePhysicsBodyFromAsset(FMutableGraphGenerationContext& GenerationContext, UPhysicsAsset* Asset, const TArray<uint8>& BodySetupRelevancyMap)
+TSharedPtr<mu::FPhysicsBody> MakePhysicsBodyFromAsset(FMutableGraphGenerationContext& GenerationContext, UPhysicsAsset* Asset, const TArray<uint8>& BodySetupRelevancyMap)
 {
 	check(Asset);
 	check(Asset->SkeletalBodySetups.Num() == BodySetupRelevancyMap.Num());
@@ -491,7 +491,7 @@ mu::Ptr<mu::PhysicsBody> MakePhysicsBodyFromAsset(FMutableGraphGenerationContext
 	
 	const int32 NumRelevantSetups = Algo::CountIf(BodySetupRelevancyMap, [](uint8 V) { return V; });
 
-	mu::Ptr<mu::PhysicsBody> PhysicsBody = new mu::PhysicsBody;
+	TSharedPtr<mu::FPhysicsBody> PhysicsBody = MakeShared<mu::FPhysicsBody>();
 	
 	PhysicsBody->SetBodyCount(NumRelevantSetups);
 
@@ -684,7 +684,7 @@ namespace MutablePrivate
 
 
 
-mu::MeshPtr ConvertSkeletalMeshToMutable(
+TSharedPtr<mu::FMesh> ConvertSkeletalMeshToMutable(
 		const USkeletalMesh* InSkeletalMesh, const TSoftClassPtr<UAnimInstance>& AnimBp, int32 LODIndexConnected, 
 		int32 SectionIndexConnected, int32 LODIndex, int32 SectionIndex, uint32 SurfaceMetadataId,
 		FMutableGraphGenerationContext& GenerationContext, const UCustomizableObjectNode* CurrentNode, USkeletalMesh* TableReferenceSkeletalMesh)
@@ -715,7 +715,7 @@ mu::MeshPtr ConvertSkeletalMeshToMutable(
 		if (GenerationContext.CurrentAutoLODStrategy == ECustomizableObjectAutomaticLODStrategy::AutomaticFromMesh && 
 			LODIndex != LODIndexConnected) // If we are using automatic LODs and not generating the base LOD (the connected one) is not an error.
 		{
-			return new mu::Mesh(); // Return empty mesh to preserve the layouts
+			return MakeShared<mu::FMesh>(); // Return empty mesh to preserve the layouts
 		}
 		else
 		{
@@ -736,7 +736,7 @@ mu::MeshPtr ConvertSkeletalMeshToMutable(
 		if (GenerationContext.CurrentAutoLODStrategy == ECustomizableObjectAutomaticLODStrategy::AutomaticFromMesh && 
 			SectionIndex != SectionIndexConnected) // If we are using automatic LODs and not generating the base LOD (the connected one) is not an error.
 		{
-			return new mu::Mesh(); // Return empty mesh to preserve the layouts
+			return MakeShared<mu::FMesh>(); // Return empty mesh to preserve the layouts
 		}
 		else
 		{
@@ -761,7 +761,7 @@ mu::MeshPtr ConvertSkeletalMeshToMutable(
 	const bool bIgnorePhysics = EnumHasAnyFlags(CurrentFlags, EMutableMeshConversionFlags::IgnorePhysics);
 	const bool bDoNotCreateMeshMetadata = EnumHasAnyFlags(CurrentFlags, EMutableMeshConversionFlags::DoNotCreateMeshMetadata);
 
-	mu::Ptr<mu::Mesh> MutableMesh = new mu::Mesh();
+	TSharedPtr<mu::FMesh> MutableMesh = MakeShared<mu::FMesh>();
 
 	bool bBoneMapModified = false;
 	TArray<FBoneIndexType> BoneMap;
@@ -878,7 +878,7 @@ mu::MeshPtr ConvertSkeletalMeshToMutable(
 		}
 
 		// Create the skeleton, poses, and BoneMap for this mesh
-		mu::SkeletonPtr MutableSkeleton = new mu::Skeleton;
+		TSharedPtr<mu::FSkeleton> MutableSkeleton = MakeShared<mu::FSkeleton>();
 		MutableMesh->SetSkeleton(MutableSkeleton);
 
 		const int32 NumRequiredBones = RequiredBones.Num();
@@ -2022,7 +2022,7 @@ mu::MeshPtr ConvertSkeletalMeshToMutable(
 			GenerationContext.Log(FText::FromString(PhysicsSetupsRemovedMsg), CurrentNode, EMessageSeverity::Warning);
 		}
 
-		mu::Ptr<mu::PhysicsBody> PhysicsBody = new mu::PhysicsBody;
+		TSharedPtr<mu::FPhysicsBody> PhysicsBody = MakeShared<mu::FPhysicsBody>();
 		
 		const int32 NumBodySetups = RelevantBodySetups.Num();	
 		PhysicsBody->SetBodyCount(NumBodySetups);
@@ -2137,7 +2137,7 @@ mu::MeshPtr ConvertSkeletalMeshToMutable(
 	// Set Bone Parenting usages. This has to be done after all primary usages are set.
 	for (int32 I = MutableMesh->GetBonePoseCount() - 1; I >= 0; --I)
 	{
-		mu::Mesh::FBonePose& BonePose = MutableMesh->BonePoses[I];
+		mu::FMesh::FBonePose& BonePose = MutableMesh->BonePoses[I];
 
 		constexpr mu::EBoneUsageFlags FlagsToPropagate =
 				mu::EBoneUsageFlags::Skinning | mu::EBoneUsageFlags::Physics | mu::EBoneUsageFlags::Deform;
@@ -2190,7 +2190,7 @@ mu::MeshPtr ConvertSkeletalMeshToMutable(
 
 			const int32 PhysicsAssetId = GenerationContext.AnimBpOverridePhysicsAssetsInfo.AddUnique(Info);
 
-			mu::Ptr<mu::PhysicsBody> MutableBody = MakePhysicsBodyFromAsset(GenerationContext,
+			TSharedPtr<mu::FPhysicsBody> MutableBody = MakePhysicsBodyFromAsset(GenerationContext,
 					PropertyAsset,
 					MakePhysicsAssetBodySetupRelevancyMap(GenerationContext, PropertyAsset, MutableMesh));
 			MutableBody->CustomId = PhysicsAssetId;
@@ -2249,7 +2249,7 @@ mu::MeshPtr ConvertSkeletalMeshToMutable(
 	return MutableMesh;
 }
 
-mu::MeshPtr ConvertStaticMeshToMutable(const UStaticMesh* StaticMesh, int32 LODIndex, int32 SectionIndex, FMutableGraphGenerationContext& GenerationContext, const UCustomizableObjectNode* CurrentNode)
+TSharedPtr<mu::FMesh> ConvertStaticMeshToMutable(const UStaticMesh* StaticMesh, int32 LODIndex, int32 SectionIndex, FMutableGraphGenerationContext& GenerationContext, const UCustomizableObjectNode* CurrentNode)
 {
 	if (!StaticMesh->GetRenderData() ||
 		!StaticMesh->GetRenderData()->LODResources.IsValidIndex(LODIndex) ||
@@ -2262,7 +2262,7 @@ mu::MeshPtr ConvertStaticMeshToMutable(const UStaticMesh* StaticMesh, int32 LODI
 
 	GenerationContext.AddParticipatingObject(*StaticMesh);
 
-	mu::MeshPtr MutableMesh = new mu::Mesh();
+	TSharedPtr<mu::FMesh> MutableMesh = MakeShared<mu::FMesh>();
 
 	// Vertices
 	int32 VertexStart = StaticMesh->GetRenderData()->LODResources[LODIndex].Sections[SectionIndex].MinVertexIndex;
@@ -2439,7 +2439,7 @@ mu::MeshPtr ConvertStaticMeshToMutable(const UStaticMesh* StaticMesh, int32 LODI
 
 // Convert a Mesh constant to a mutable format. UniqueTags are the tags that make this Mesh unique that cannot be merged in the cache 
 // with the exact same Mesh with other tags
-mu::Ptr<mu::Mesh> GenerateMutableMesh(UObject * Mesh, const TSoftClassPtr<UAnimInstance>& AnimInstance, int32 LODIndexConnected, int32 SectionIndexConnected, 
+TSharedPtr<mu::FMesh> GenerateMutableMesh(UObject * Mesh, const TSoftClassPtr<UAnimInstance>& AnimInstance, int32 LODIndexConnected, int32 SectionIndexConnected, 
 									  int32 LODIndex, int32 SectionIndex, const FString& UniqueTags, uint32 SurfaceMetadataId, FMutableGraphGenerationContext & GenerationContext, 
 									  const UCustomizableObjectNode* CurrentNode, USkeletalMesh* TableReferenceSkeletalMesh, bool bIsReference)
 {
@@ -2447,7 +2447,7 @@ mu::Ptr<mu::Mesh> GenerateMutableMesh(UObject * Mesh, const TSoftClassPtr<UAnimI
 	EMutableMeshConversionFlags CurrentFlags = GenerationContext.MeshGenerationFlags.Last();
 
 	FMutableGraphGenerationContext::FGeneratedMeshData::FKey Key = { Mesh, LODIndex, GenerationContext.CurrentLOD, SectionIndex, CurrentFlags, UniqueTags, CurrentNode };
-	mu::Ptr<mu::Mesh> MutableMesh = GenerationContext.FindGeneratedMesh(Key);
+	TSharedPtr<mu::FMesh> MutableMesh = GenerationContext.FindGeneratedMesh(Key);
 	if (MutableMesh)
 	{
 		return MutableMesh;
@@ -2511,10 +2511,10 @@ mu::Ptr<mu::Mesh> GenerateMutableMesh(UObject * Mesh, const TSoftClassPtr<UAnimI
 }
 
 
-mu::Ptr<mu::Mesh> BuildMorphedMutableMeshFromMesh(mu::Ptr<mu::Mesh> BaseSourceMesh, USkeletalMesh* SkeletalMesh, const FString& MorphTargetName, int32 LODIndex, int32 SectionIndex)
+TSharedPtr<mu::FMesh> BuildMorphedMutableMeshFromMesh(TSharedPtr<mu::FMesh> BaseSourceMesh, USkeletalMesh* SkeletalMesh, const FString& MorphTargetName, int32 LODIndex, int32 SectionIndex)
 {
 	// Clone it (it will probably be shared)
-	mu::Ptr<mu::Mesh> MorphedSourceMesh = BaseSourceMesh->Clone();
+	TSharedPtr<mu::FMesh> MorphedSourceMesh = BaseSourceMesh->Clone();
 
 	// Bake the morph in the new mutable mesh
 	UMorphTarget* MorphTarget = SkeletalMesh ? SkeletalMesh->FindMorphTarget(*MorphTargetName) : nullptr;
@@ -2573,7 +2573,7 @@ mu::Ptr<mu::Mesh> BuildMorphedMutableMeshFromMesh(mu::Ptr<mu::Mesh> BaseSourceMe
 }
 
 
-mu::Ptr<mu::Mesh> BuildMorphedMutableMesh(const UEdGraphPin* BaseSourcePin, const FString& MorphTargetName, FMutableGraphGenerationContext& GenerationContext, const bool bOnlyConnectedLOD, const FName& RowName)
+TSharedPtr<mu::FMesh> BuildMorphedMutableMesh(const UEdGraphPin* BaseSourcePin, const FString& MorphTargetName, FMutableGraphGenerationContext& GenerationContext, const bool bOnlyConnectedLOD, const FName& RowName)
 {
 	check(BaseSourcePin);
 	SCOPED_PIN_DATA(GenerationContext, BaseSourcePin)
@@ -2612,14 +2612,14 @@ mu::Ptr<mu::Mesh> BuildMorphedMutableMesh(const UEdGraphPin* BaseSourcePin, cons
 		}
 	}
 
-	mu::Ptr<mu::Mesh> MorphedSourceMesh;
+	TSharedPtr<mu::FMesh> MorphedSourceMesh;
 
 	if (SkeletalMesh)
 	{
 		GetLODAndSectionForAutomaticLODs(GenerationContext, *Node, *SkeletalMesh, LODIndexConnected, SectionIndexConnected, LODIndex, SectionIndex, bOnlyConnectedLOD);
 		// Get the base mesh
 		constexpr bool bIsReference = false;
-		mu::Ptr<mu::Mesh> BaseSourceMesh = GenerateMutableMesh(SkeletalMesh, TSoftClassPtr<UAnimInstance>(), LODIndexConnected, SectionIndexConnected,
+		TSharedPtr<mu::FMesh> BaseSourceMesh = GenerateMutableMesh(SkeletalMesh, TSoftClassPtr<UAnimInstance>(), LODIndexConnected, SectionIndexConnected,
 			LODIndex, SectionIndex, FString(), 0, GenerationContext, Node, nullptr, bIsReference);
 		if (BaseSourceMesh)
 		{
@@ -3158,7 +3158,7 @@ mu::NodeMeshPtr GenerateMorphMesh(const UEdGraphPin* Pin,
 	// Target
 	mu::NodeMeshPtr BaseSourceMesh = SourceNode;
 
-	mu::MeshPtr MorphedSourceMesh;
+	TSharedPtr<mu::FMesh> MorphedSourceMesh;
 
 	bool bSuccess = false;
 	
@@ -3176,7 +3176,7 @@ mu::NodeMeshPtr GenerateMorphMesh(const UEdGraphPin* Pin,
 		int32 NumRows = RowNames.Num();
 
 		// Should exist
-		mu::TablePtr Table = GenerationContext.GeneratedTables[DataTable->GetName()].GeneratedTable;
+		mu::Ptr<mu::FTable> Table = GenerationContext.GeneratedTables[DataTable->GetName()].GeneratedTable;
 		check(Table);
 
 		FString ColumnName = TableColumnName + TypedNodeMorphs[MorphIndex].MorphTargetName;
@@ -3194,8 +3194,8 @@ mu::NodeMeshPtr GenerateMorphMesh(const UEdGraphPin* Pin,
 				ColumnIndex = Table->AddColumn(ColumnName, mu::ETableColumnType::Mesh);
 			}
 
-			mu::MeshPtr MorphedSourceTableMesh = BuildMorphedMutableMesh(Pin, TypedNodeMorphs[MorphIndex].MorphTargetName, GenerationContext, bOnlyConnectedLOD, RowName);
-			Table->SetCell(ColumnIndex, RowId, MorphedSourceTableMesh.get());
+			TSharedPtr<mu::FMesh> MorphedSourceTableMesh = BuildMorphedMutableMesh(Pin, TypedNodeMorphs[MorphIndex].MorphTargetName, GenerationContext, bOnlyConnectedLOD, RowName);
+			Table->SetCell(ColumnIndex, RowId, MorphedSourceTableMesh);
 		}
 
 		if (ColumnIndex > INDEX_NONE)
@@ -3410,7 +3410,7 @@ mu::Ptr<mu::NodeMesh> GenerateMutableSourceMesh(const UEdGraphPin* Pin,
 			FSkeletalMeshModel* ImportedModel = TypedNodeSkel->SkeletalMesh->GetImportedModel();
 			
 			constexpr bool bIsReference = false;
-			mu::Ptr<mu::Mesh> MutableMesh = GenerateMutableMesh(TypedNodeSkel->SkeletalMesh, TypedNodeSkel->AnimInstance, LODIndexConnected, SectionIndexConnected, 
+			TSharedPtr<mu::FMesh> MutableMesh = GenerateMutableMesh(TypedNodeSkel->SkeletalMesh, TypedNodeSkel->AnimInstance, LODIndexConnected, SectionIndexConnected, 
 																LODIndex, SectionIndex, MeshUniqueTags, SurfaceMetadataId, GenerationContext, TypedNodeSkel, nullptr, bIsReference);
 			if (MutableMesh)
 			{
@@ -3426,7 +3426,7 @@ mu::Ptr<mu::NodeMesh> GenerateMutableSourceMesh(const UEdGraphPin* Pin,
 					for (TObjectPtr<UMorphTarget> Morph : Morphs)
 					{
 						FString MorphTargetName = Morph->GetName();
-						mu::Ptr<mu::Mesh> MorphedMesh = BuildMorphedMutableMeshFromMesh(MutableMesh, TypedNodeSkel->SkeletalMesh, MorphTargetName, LODIndex, SectionIndex);
+						TSharedPtr<mu::FMesh> MorphedMesh = BuildMorphedMutableMeshFromMesh(MutableMesh, TypedNodeSkel->SkeletalMesh, MorphTargetName, LODIndex, SectionIndex);
 						MeshNode->AddMorph(MorphTargetName,MorphedMesh);
 					}
 				}
@@ -3625,7 +3625,7 @@ mu::Ptr<mu::NodeMesh> GenerateMutableSourceMesh(const UEdGraphPin* Pin,
 			check(SectionIndex < TypedNodeStatic->LODs[LODIndex].Materials.Num());
 
 			constexpr bool bIsReference = false;
-			mu::MeshPtr MutableMesh = GenerateMutableMesh(TypedNodeStatic->StaticMesh, TSoftClassPtr<UAnimInstance>(), LODIndex, SectionIndex, LODIndex, SectionIndex, 
+			TSharedPtr<mu::FMesh> MutableMesh = GenerateMutableMesh(TypedNodeStatic->StaticMesh, TSoftClassPtr<UAnimInstance>(), LODIndex, SectionIndex, LODIndex, SectionIndex, 
 														  FString(), 0, GenerationContext, TypedNodeStatic, nullptr, bIsReference);
 			if (MutableMesh)
 			{
@@ -4219,7 +4219,7 @@ mu::Ptr<mu::NodeMesh> GenerateMutableSourceMesh(const UEdGraphPin* Pin,
 			if (bSuccess)
 			{
 				// Generating a new data table if not exists
-				mu::TablePtr Table = nullptr;
+				mu::Ptr<mu::FTable> Table = nullptr;
 				Table = GenerateMutableSourceTable(DataTable, TypedNodeTable, GenerationContext);
 
 				if (Table)
@@ -4250,7 +4250,7 @@ mu::Ptr<mu::NodeMesh> GenerateMutableSourceMesh(const UEdGraphPin* Pin,
 						}
 					}
 					
-					// Generating a new Mesh column if not exists
+					// Generating a new FMesh column if not exists
 					if (Table->FindColumn(MutableColumnName) == INDEX_NONE)
 					{
 						bSuccess = GenerateTableColumn(TypedNodeTable, Pin, Table, DataTableColumnName, Property, LODIndexConnected, SectionIndexConnected, LODIndex, SectionIndex, SurfaceMetadataId, bOnlyConnectedLOD, GenerationContext);
diff --git a/Engine/Plugins/Mutable/Source/CustomizableObjectEditor/Private/MuCOE/GenerateMutableSource/GenerateMutableSourceMesh.h b/Engine/Plugins/Mutable/Source/CustomizableObjectEditor/Private/MuCOE/GenerateMutableSource/GenerateMutableSourceMesh.h
--- a/Engine/Plugins/Mutable/Source/CustomizableObjectEditor/Private/MuCOE/GenerateMutableSource/GenerateMutableSourceMesh.h
+++ b/Engine/Plugins/Mutable/Source/CustomizableObjectEditor/Private/MuCOE/GenerateMutableSource/GenerateMutableSourceMesh.h
@@ -49,7 +49,7 @@ void GetLODAndSectionForAutomaticLODs(const FMutableGraphGenerationContext& Cont
  * @param CurrentNode The node being compiled, can be a skeletal mesh or a table node.
  * @param TableReferenceSkeletalMesh Only set if compiling a table node. It refers to the reference mesh of the skeletal mesh column being compiled.
  * @return Mutable Mesh. Nullptr if there has been an error. Empty mesh if the Skeletal Mesh does not contain the requested LOD + Section. */
-mu::Ptr<mu::Mesh> ConvertSkeletalMeshToMutable(const USkeletalMesh* InSkeletalMesh, const TSoftClassPtr<UAnimInstance>& AnimBp,
+TSharedPtr<mu::FMesh> ConvertSkeletalMeshToMutable(const USkeletalMesh* InSkeletalMesh, const TSoftClassPtr<UAnimInstance>& AnimBp,
                                          int32 LODIndexConnected,  int32 SectionIndexConnected,
                                          int32 LODIndex, int32 SectionIndex, uint32 SurfaceMetadataId,
                                          FMutableGraphGenerationContext& GenerationContext,
@@ -57,7 +57,7 @@ mu::Ptr<mu::Mesh> ConvertSkeletalMeshToMutable(const USkeletalMesh* InSkeletalMe
 										 USkeletalMesh* TableReferenceSkeletalMesh);
 
 
-mu::Ptr<mu::Mesh> ConvertStaticMeshToMutable(const UStaticMesh* StaticMesh, int32 LODIndex, int32 SectionIndex,
+TSharedPtr<mu::FMesh> ConvertStaticMeshToMutable(const UStaticMesh* StaticMesh, int32 LODIndex, int32 SectionIndex,
 	FMutableGraphGenerationContext& GenerationContext, const UCustomizableObjectNode* CurrentNode);
 
 
@@ -74,7 +74,7 @@ mu::Ptr<mu::Mesh> ConvertStaticMeshToMutable(const UStaticMesh* StaticMesh, int3
  * @param CurrentNode The node being compiled, can be a skeletal mesh or a table node.
  * @param TableReferenceSkeletalMesh Only set if compiling a table node. It refers to the reference mesh of the skeletal mesh column being compiled.
  * @return Mutable Mesh. Nullptr if there has been an error. Empty mesh if the Skeletal Mesh does not contain the requested LOD + Section. */
-mu::Ptr<mu::Mesh> GenerateMutableMesh(UObject* Mesh, const TSoftClassPtr<UAnimInstance>& AnimBp,
+TSharedPtr<mu::FMesh> GenerateMutableMesh(UObject* Mesh, const TSoftClassPtr<UAnimInstance>& AnimBp,
                                 int32 LODIndexConnected, int32 SectionIndexConnected,
                                 int32 LODIndex, int32 SectionIndex, const FString& MeshUniqueTags, uint32 SurfaceMetadataId, 
                                 FMutableGraphGenerationContext& GenerationContext, 
@@ -82,7 +82,7 @@ mu::Ptr<mu::Mesh> GenerateMutableMesh(UObject* Mesh, const TSoftClassPtr<UAnimIn
 								bool bIsReference);
 
 
-mu::Ptr<mu::Mesh> BuildMorphedMutableMesh(const UEdGraphPin* BaseSourcePin, const FString& MorphTargetName,
+TSharedPtr<mu::FMesh> BuildMorphedMutableMesh(const UEdGraphPin* BaseSourcePin, const FString& MorphTargetName,
 	FMutableGraphGenerationContext& GenerationContext, bool bOnlyConnectedLOD, const FName& RowName = "");
 
 
diff --git a/Engine/Plugins/Mutable/Source/CustomizableObjectEditor/Private/MuCOE/GenerateMutableSource/GenerateMutableSourceSurface.cpp b/Engine/Plugins/Mutable/Source/CustomizableObjectEditor/Private/MuCOE/GenerateMutableSource/GenerateMutableSourceSurface.cpp
--- a/Engine/Plugins/Mutable/Source/CustomizableObjectEditor/Private/MuCOE/GenerateMutableSource/GenerateMutableSourceSurface.cpp
+++ b/Engine/Plugins/Mutable/Source/CustomizableObjectEditor/Private/MuCOE/GenerateMutableSource/GenerateMutableSourceSurface.cpp
@@ -581,14 +581,14 @@ mu::Ptr<mu::NodeSurface> GenerateMutableSourceSurface(const UEdGraphPin * Pin, F
 									if (Texture2D)
 									{
 										mu::Ptr<mu::NodeImageConstant> ConstImageNode = new mu::NodeImageConstant();
-										mu::Ptr<mu::Image> ImageConstant = GenerateImageConstant(Texture2D, GenerationContext, false);
-										ConstImageNode->SetValue(ImageConstant.get());
+										TSharedPtr<mu::FImage> ImageConstant = GenerateImageConstant(Texture2D, GenerationContext, false);
+										ConstImageNode->SetValue(ImageConstant);
 
 										const uint32 MipsToSkip = ComputeLODBiasForTexture(GenerationContext, *Texture2D, nullptr, Props.TextureSize);
 										mu::Ptr<mu::NodeImage> Result =  ResizeTextureByNumMips(ConstImageNode, MipsToSkip);
 
 										// Calculate the number of mips to tag as high res for this image.
-										int32 TotalMips = mu::Image::GetMipmapCount(ImageConstant->GetSizeX(), ImageConstant->GetSizeY());
+										int32 TotalMips = mu::FImage::GetMipmapCount(ImageConstant->GetSizeX(), ImageConstant->GetSizeY());
 										int32 NumMipsBeyondMin = FMath::Max(0, TotalMips - int32(MipsToSkip) - GenerationContext.Options.MinDiskMips);
 										int32 HighResMipsForThisImage = FMath::Min(NumMipsBeyondMin, GenerationContext.Options.NumHighResImageMips);
 										ConstImageNode->SourceDataDescriptor.SourceHighResMips = HighResMipsForThisImage;
@@ -678,8 +678,8 @@ mu::Ptr<mu::NodeSurface> GenerateMutableSourceSurface(const UEdGraphPin * Pin, F
 									// GenerationContext.ArrayTextureUnrealToMutableTask.Add(FTextureUnrealToMutableTask(CompositeNormalImage, ReferenceCompositeNormalTexture, Node, true));
 									// TODO: The normal composite part is not propagated, so it will be unsupported. Create a task that performs the required transforms at mutable image level, and add the right operations here
 									// instead of propagating the flag and doing them on unreal-convert.
-									mu::Ptr<mu::Image> ImageConstant = GenerateImageConstant(ReferenceCompositeNormalTexture, GenerationContext, false);
-									CompositeNormalImage->SetValue(ImageConstant.get());
+									TSharedPtr<mu::FImage> ImageConstant = GenerateImageConstant(ReferenceCompositeNormalTexture, GenerationContext, false);
+									CompositeNormalImage->SetValue(ImageConstant);
 
 									mu::NodeImageMipmapPtr NormalCompositeMipmapImage = new mu::NodeImageMipmap();
 									const uint32 MipsToSkip = ComputeLODBiasForTexture(GenerationContext, *ReferenceCompositeNormalTexture, ReferenceTexture);
@@ -688,7 +688,7 @@ mu::Ptr<mu::NodeSurface> GenerateMutableSourceSurface(const UEdGraphPin * Pin, F
 
 									CompositedImage->SetNormal(NormalCompositeMipmapImage);
 
-									int32 TotalMips = mu::Image::GetMipmapCount(ImageConstant->GetSizeX(), ImageConstant->GetSizeY());
+									int32 TotalMips = mu::FImage::GetMipmapCount(ImageConstant->GetSizeX(), ImageConstant->GetSizeY());
 									int32 NumMipsBeyondMin = FMath::Max(0, TotalMips - int32(MipsToSkip) - GenerationContext.Options.MinDiskMips);
 									int32 HighResMipsForThisImage = FMath::Min(NumMipsBeyondMin, GenerationContext.Options.NumHighResImageMips);
 									CompositeNormalImage->SourceDataDescriptor.SourceHighResMips = HighResMipsForThisImage;
