[mutable] More enforcement of coding standards.

diff --git a/Engine/Plugins/Mutable/Source/CustomizableObjectEditor/Private/MuCOE/GenerateMutableSource/GenerateMutableSourceTable.cpp b/Engine/Plugins/Mutable/Source/CustomizableObjectEditor/Private/MuCOE/GenerateMutableSource/GenerateMutableSourceTable.cpp
--- a/Engine/Plugins/Mutable/Source/CustomizableObjectEditor/Private/MuCOE/GenerateMutableSource/GenerateMutableSourceTable.cpp
+++ b/Engine/Plugins/Mutable/Source/CustomizableObjectEditor/Private/MuCOE/GenerateMutableSource/GenerateMutableSourceTable.cpp
@@ -31,7 +31,7 @@
 #define LOCTEXT_NAMESPACE "CustomizableObjectEditor"
 
 
-bool FillTableColumn(const UCustomizableObjectNodeTable* TableNode, mu::TablePtr MutableTable, const FString& ColumnName, const FString& RowName, const uint32 RowId, uint8* CellData, const FProperty* ColumnProperty,
+bool FillTableColumn(const UCustomizableObjectNodeTable* TableNode, mu::Ptr<mu::FTable> MutableTable, const FString& ColumnName, const FString& RowName, const uint32 RowId, uint8* CellData, const FProperty* ColumnProperty,
 	const int LODIndexConnected, const int32 SectionIndexConnected, int32 LODIndex, int32 SectionIndex, uint32 SectionMetadataId, const bool bOnlyConnectedLOD, FMutableGraphGenerationContext& GenerationContext)
 {
 	int32 CurrentColumn;
@@ -61,8 +61,8 @@ bool FillTableColumn(const UCustomizableObjectNodeTable* TableNode, mu::TablePtr
 					CurrentColumn = MutableTable->AddColumn(MutableColumnName, mu::ETableColumnType::Mesh);
 				}
 
-				mu::MeshPtr EmptySkeletalMesh = nullptr;
-				MutableTable->SetCell(CurrentColumn, RowId, EmptySkeletalMesh.get());
+				TSharedPtr<mu::FMesh> EmptySkeletalMesh = nullptr;
+				MutableTable->SetCell(CurrentColumn, RowId, EmptySkeletalMesh);
 
 				return true;
 			}
@@ -279,7 +279,7 @@ bool FillTableColumn(const UCustomizableObjectNodeTable* TableNode, mu::TablePtr
 
 			//TODO: Add AnimBp physics to Tables.
 			constexpr bool bIsReference = false;
-			mu::Ptr<mu::Mesh> MutableMesh = GenerateMutableMesh(SkeletalMesh, TSoftClassPtr<UAnimInstance>(), LODIndexConnected, SectionIndexConnected, 
+			TSharedPtr<mu::FMesh> MutableMesh = GenerateMutableMesh(SkeletalMesh, TSoftClassPtr<UAnimInstance>(), LODIndexConnected, SectionIndexConnected, 
 																LODIndex, SectionIndex, MeshUniqueTags, SectionMetadataId, GenerationContext, TableNode, ReferenceSkeletalMesh, bIsReference);
 
 			if (MutableMesh)
@@ -321,7 +321,7 @@ bool FillTableColumn(const UCustomizableObjectNodeTable* TableNode, mu::TablePtr
 					AddTagToMutableMeshUnique(*MutableMesh, MeshTag);
 				}
 
-				MutableTable->SetCell(CurrentColumn, RowId, MutableMesh.get(), SkeletalMesh);
+				MutableTable->SetCell(CurrentColumn, RowId, MutableMesh, SkeletalMesh);
 			}
 			else
 			{
@@ -390,12 +390,12 @@ bool FillTableColumn(const UCustomizableObjectNodeTable* TableNode, mu::TablePtr
 			}
 
 			constexpr bool bIsReference = false;
-			mu::MeshPtr MutableMesh = GenerateMutableMesh(StaticMesh, TSoftClassPtr<UAnimInstance>(), CurrentLOD, SectionIndex, CurrentLOD, SectionIndex, 
+			TSharedPtr<mu::FMesh> MutableMesh = GenerateMutableMesh(StaticMesh, TSoftClassPtr<UAnimInstance>(), CurrentLOD, SectionIndex, CurrentLOD, SectionIndex, 
 														  FString(), 0, GenerationContext, TableNode, nullptr, bIsReference);
 
 			if (MutableMesh)
 			{
-				MutableTable->SetCell(CurrentColumn, RowId, MutableMesh.get(), StaticMesh);
+				MutableTable->SetCell(CurrentColumn, RowId, MutableMesh, StaticMesh);
 			}
 			else
 			{
@@ -436,7 +436,7 @@ bool FillTableColumn(const UCustomizableObjectNodeTable* TableNode, mu::TablePtr
 			}
 
 			const bool bIsPassthroughTexture = TableNode->GetColumnImageMode(PinName) == ETableTextureType::PASSTHROUGH_TEXTURE;
-			mu::Ptr<mu::ResourceProxyMemory<mu::Image>> Proxy = new mu::ResourceProxyMemory<mu::Image>(GenerateImageConstant(Texture, GenerationContext, bIsPassthroughTexture));
+			mu::Ptr<mu::TResourceProxyMemory<mu::FImage>> Proxy = new mu::TResourceProxyMemory<mu::FImage>(GenerateImageConstant(Texture, GenerationContext, bIsPassthroughTexture));
 			MutableTable->SetCell(CurrentColumn, RowId, Proxy.get());
 		}
 
@@ -568,7 +568,7 @@ bool FillTableColumn(const UCustomizableObjectNodeTable* TableNode, mu::TablePtr
 				}
 
 				bool bIsPassthroughTexture = false;
-				mu::Ptr<mu::ResourceProxyMemory<mu::Image>> Proxy = new mu::ResourceProxyMemory<mu::Image>(GenerateImageConstant(ParameterTexture, GenerationContext, bIsPassthroughTexture));
+				mu::Ptr<mu::TResourceProxyMemory<mu::FImage>> Proxy = new mu::TResourceProxyMemory<mu::FImage>(GenerateImageConstant(ParameterTexture, GenerationContext, bIsPassthroughTexture));
 				MutableTable->SetCell(ColumnIndex, RowId, Proxy.get());
 
 				return true;
@@ -594,8 +594,8 @@ bool FillTableColumn(const UCustomizableObjectNodeTable* TableNode, mu::TablePtr
 				TArray<FTransform> ArrayTransform;
 				UCustomizableObjectNodeAnimationPose::StaticRetrievePoseInformation(PoseAsset, GenerationContext.GetCurrentComponentInfo()->RefSkeletalMesh, ArrayBoneName, ArrayTransform);
 
-				mu::Ptr<mu::Mesh> MutableMesh = new mu::Mesh();
-				mu::Ptr<mu::Skeleton> MutableSkeleton = new mu::Skeleton;
+				TSharedPtr<mu::FMesh> MutableMesh = MakeShared<mu::FMesh>();
+				TSharedPtr<mu::FSkeleton> MutableSkeleton = MakeShared<mu::FSkeleton>();
 
 				MutableMesh->SetSkeleton(MutableSkeleton);
 				MutableMesh->SetBonePoseCount(ArrayBoneName.Num());
@@ -611,7 +611,7 @@ bool FillTableColumn(const UCustomizableObjectNodeTable* TableNode, mu::TablePtr
 					MutableMesh->SetBonePose(i, BoneId, (FTransform3f)ArrayTransform[i], mu::EBoneUsageFlags::Skinning);
 				}
 
-				MutableTable->SetCell(CurrentColumn, RowId, MutableMesh.get());
+				MutableTable->SetCell(CurrentColumn, RowId, MutableMesh);
 			}
 		}
 
@@ -913,7 +913,7 @@ TArray<FName> GetRowsToCompile(const UDataTable& DataTable, const UCustomizableO
 }
 
 
-bool GenerateTableColumn(const UCustomizableObjectNodeTable* TableNode, const UEdGraphPin* Pin, mu::TablePtr MutableTable, const FString& DataTableColumnName, const FProperty* ColumnProperty,
+bool GenerateTableColumn(const UCustomizableObjectNodeTable* TableNode, const UEdGraphPin* Pin, mu::Ptr<mu::FTable> MutableTable, const FString& DataTableColumnName, const FProperty* ColumnProperty,
 	const int32 LODIndexConnected, const int32 SectionIndexConnected, const int32 LODIndex, const int32 SectionIndex, uint32 SectionMetadataId, const bool bOnlyConnectedLOD, FMutableGraphGenerationContext& GenerationContext)
 {
 	MUTABLE_CPUPROFILER_SCOPE(GenerateTableColumn);
@@ -1094,7 +1094,7 @@ void GenerateTableParameterUIData(const UDataTable* DataTable, const UCustomizab
 }
 
 
-mu::TablePtr GenerateMutableSourceTable(const UDataTable* DataTable, const UCustomizableObjectNodeTable* TableNode, FMutableGraphGenerationContext& GenerationContext)
+mu::Ptr<mu::FTable> GenerateMutableSourceTable(const UDataTable* DataTable, const UCustomizableObjectNodeTable* TableNode, FMutableGraphGenerationContext& GenerationContext)
 {
 	check(DataTable && TableNode);
 
@@ -1131,7 +1131,7 @@ mu::TablePtr GenerateMutableSourceTable(const UDataTable* DataTable, const UCust
 		return CachedTable->GeneratedTable;
 	}
 
-	mu::TablePtr MutableTable = new mu::Table();
+	mu::Ptr<mu::FTable> MutableTable = new mu::FTable();
 
 	if (const UScriptStruct* TableStruct = DataTable->GetRowStruct())
 	{
diff --git a/Engine/Plugins/Mutable/Source/CustomizableObjectEditor/Private/MuCOE/GenerateMutableSource/GenerateMutableSourceTable.h b/Engine/Plugins/Mutable/Source/CustomizableObjectEditor/Private/MuCOE/GenerateMutableSource/GenerateMutableSourceTable.h
--- a/Engine/Plugins/Mutable/Source/CustomizableObjectEditor/Private/MuCOE/GenerateMutableSource/GenerateMutableSourceTable.h
+++ b/Engine/Plugins/Mutable/Source/CustomizableObjectEditor/Private/MuCOE/GenerateMutableSource/GenerateMutableSourceTable.h
@@ -26,7 +26,7 @@ struct FMutableGraphGenerationContext;
  * @param bOnlyConnectedLOD Corrected LOD and Section will unconditionally always be the connected ones.
  * @param GenerationContext 
  * @return  */
-bool GenerateTableColumn(const UCustomizableObjectNodeTable* TableNode, const UEdGraphPin* Pin, mu::TablePtr MutableTable, const FString& DataTableColumnName, const FProperty* ColumnProperty,
+bool GenerateTableColumn(const UCustomizableObjectNodeTable* TableNode, const UEdGraphPin* Pin, mu::Ptr<mu::FTable> MutableTable, const FString& DataTableColumnName, const FProperty* ColumnProperty,
 	int32 LODIndexConnected, int32 SectionIndexConnected,
 	int32 LODIndex, int32 SectionIndex, uint32 SurfaceMetadataId,
 	bool bOnlyConnectedLOD,
@@ -50,14 +50,14 @@ bool GenerateTableColumn(const UCustomizableObjectNodeTable* TableNode, const UE
  * @param bOnlyConnectedLOD Corrected LOD and Section will unconditionally always be the connected ones.
  * @param GenerationContext 
  * @return  */
-bool FillTableColumn(const UCustomizableObjectNodeTable* TableNode, mu::TablePtr MutableTable, const FString& ColumnName,
+bool FillTableColumn(const UCustomizableObjectNodeTable* TableNode, mu::Ptr<mu::FTable> MutableTable, const FString& ColumnName,
 	const FString& RowName, uint32 RowId, uint8* CellData, const FProperty* ColumnProperty,
 	int32 LODIndexConnected, int32 SectionIndexConnected, int32 LODIndex, int32 SectionIndex, uint32 SectionMetadataId,
 	bool bOnlyConnectedLOD,
 	FMutableGraphGenerationContext& GenerationContext);
 
 
-mu::TablePtr GenerateMutableSourceTable(const UDataTable* DataTable, const UCustomizableObjectNodeTable* TableNode, FMutableGraphGenerationContext& GenerationContext);
+mu::Ptr<mu::FTable> GenerateMutableSourceTable(const UDataTable* DataTable, const UCustomizableObjectNodeTable* TableNode, FMutableGraphGenerationContext& GenerationContext);
 
 /** Generates a new FParameterUIData for a new table parameter. Also, adds the UIData to the parameters map. */
 void GenerateTableParameterUIData(const UDataTable* DataTable, const UCustomizableObjectNodeTable* TableNode, FMutableGraphGenerationContext& GenerationContext);
diff --git a/Engine/Plugins/Mutable/Source/CustomizableObjectEditor/Private/MuCOE/MutableMeshPreviewUtils.cpp b/Engine/Plugins/Mutable/Source/CustomizableObjectEditor/Private/MuCOE/MutableMeshPreviewUtils.cpp
--- a/Engine/Plugins/Mutable/Source/CustomizableObjectEditor/Private/MuCOE/MutableMeshPreviewUtils.cpp
+++ b/Engine/Plugins/Mutable/Source/CustomizableObjectEditor/Private/MuCOE/MutableMeshPreviewUtils.cpp
@@ -10,6 +10,7 @@
 #include "MuCO/UnrealConversionUtils.h"
 #include "MuCO/UnrealPortabilityHelpers.h"
 #include "MuR/Skeleton.h"
+#include "MuR/Mesh.h"
 #include "MuR/OpMeshFormat.h"
 #include "Engine/SkeletalMesh.h"
 #include "MuCO/CustomizableObject.h"
@@ -29,10 +30,10 @@ namespace MutableMeshPreviewUtils
 		 * @return An identical mutable mesh to the provided one in terms of data but with the buffer structure that UE
 		 * requires
 		 */
-		mu::MeshPtrConst GenerateCompatibleMutableMesh(mu::MeshPtrConst InMutableMesh)
+		TSharedPtr<const mu::FMesh> GenerateCompatibleMutableMesh(const mu::FMesh* InMutableMesh)
 		{
 			// 1) Generate a mutable mesh to be the one hosting the buffer format we desire (it will be required for 
-			mu::MeshPtr FormattedMutableMesh = new::mu::Mesh();
+			TSharedPtr<mu::FMesh> FormattedMutableMesh = MakeShared<mu::FMesh>();
 			mu::FMeshBufferSet& FormattedVertexBuffers = FormattedMutableMesh->GetVertexBuffers();
 			mu::FMeshBufferSet& FormattedIndexBuffers = FormattedMutableMesh->GetIndexBuffers();
 			
@@ -173,10 +174,10 @@ namespace MutableMeshPreviewUtils
 			// Skeletal mesh is ready to be used as an argument for the process of reformatting the Mutable Mesh we want to preview
 
 			// 3) Call mu::MeshFormat to get the mesh with the wanted buffer structure to have the data of the original mesh
-			mu::Ptr<mu::Mesh> FormatResult = new mu::Mesh();
+			TSharedPtr<mu::FMesh> FormatResult = MakeShared<mu::FMesh>();
 
 			bool bOutSuccess = false;
-			mu::MeshFormat(FormatResult.get(), InMutableMesh.get(), FormattedMutableMesh.get(),
+			mu::MeshFormat(FormatResult.Get(), InMutableMesh, FormattedMutableMesh.Get(),
 				false, true, true, false,
 				bOutSuccess);
 		
@@ -190,7 +191,7 @@ namespace MutableMeshPreviewUtils
 		 * badly located inside the buffer set.
 		 * @return True if the buffers are where they should be, false if not correctly placed or missing.
 		 */
-		bool AreVertexBuffersCorrectlyLocated(const mu::MeshPtrConst InMutableMesh, bool bLogFindingsToConsole = false )
+		bool AreVertexBuffersCorrectlyLocated(const mu::FMesh* InMutableMesh, bool bLogFindingsToConsole = false )
 		{
 			bool bCanBeConverted = true;
 
@@ -281,7 +282,7 @@ namespace MutableMeshPreviewUtils
 		 * @param OutSkeletalMesh - The skeletal mesh whose skeleton we want to generate.
 		 * @return True if the operation could be performed successfully, false if not.
 		 */
-		bool BuildSkeletalMeshSkeletonData(mu::MeshPtrConst InMutableMesh, USkeletalMesh* OutSkeletalMesh)
+		bool BuildSkeletalMeshSkeletonData(const mu::FMesh* InMutableMesh, USkeletalMesh* OutSkeletalMesh)
 		{
 			if (!InMutableMesh || !OutSkeletalMesh)
 			{
@@ -339,7 +340,7 @@ namespace MutableMeshPreviewUtils
 		 * @param InMutableMesh - The reference mutable mesh to be reading data from
 		 * @param OutSkeletalMesh - The skeletal mesh to be getting it's materials set up.
 		 */
-		void BuildSkeletalMeshElementData(const mu::MeshPtrConst InMutableMesh, USkeletalMesh* OutSkeletalMesh)
+		void BuildSkeletalMeshElementData(const mu::FMesh* InMutableMesh, USkeletalMesh* OutSkeletalMesh)
 		{
 			// Set up Unreal's default material
 			UMaterialInterface* UnrealMaterial = UMaterial::GetDefaultMaterial(MD_Surface);
@@ -366,7 +367,7 @@ namespace MutableMeshPreviewUtils
 		 * @param OutSkeletalMesh - The skeletal mesh to set up.
 		 * @return True if the operation could be performed successfully, false if not.
 		 */
-		bool BuildSkeletalMeshRenderData(const mu::MeshPtrConst InMutableMesh, USkeletalMesh* OutSkeletalMesh)
+		bool BuildSkeletalMeshRenderData(const mu::FMesh* InMutableMesh, USkeletalMesh* OutSkeletalMesh)
 		{
 			OutSkeletalMesh->SetHasVertexColors(false);
 
@@ -432,13 +433,15 @@ namespace MutableMeshPreviewUtils
 	 * Publicly accessible functions of MutableMeshPreviewUtils
 	 */
 
-	USkeletalMesh* GenerateSkeletalMeshFromMutableMesh(mu::MeshPtrConst InMutableMesh)
+	USkeletalMesh* GenerateSkeletalMeshFromMutableMesh(const mu::FMesh* InMutableMesh)
 	{
 		if (!InMutableMesh)
 		{
 			return nullptr;
 		}
 
+		TSharedPtr<const mu::FMesh> MeshCopy;
+
 		// Buffer structures checks and corrective operations
 		{
 			// If the semantic channels for the required buffer data are on not the buffer they are expected to be then
@@ -451,7 +454,8 @@ namespace MutableMeshPreviewUtils
 				       ));
 				
 				// Since we know the mesh does not have the required buffer structure do format it to match our requirements
-				InMutableMesh = GenerateCompatibleMutableMesh(InMutableMesh);
+				MeshCopy = GenerateCompatibleMutableMesh(InMutableMesh);
+				InMutableMesh = MeshCopy.Get();
 
 				// If the conversion was required then make sure all went well
 				if (!AreVertexBuffersCorrectlyLocated(InMutableMesh,true))
diff --git a/Engine/Plugins/Mutable/Source/CustomizableObjectEditor/Private/MuCOE/SMutableCodeViewer.cpp b/Engine/Plugins/Mutable/Source/CustomizableObjectEditor/Private/MuCOE/SMutableCodeViewer.cpp
--- a/Engine/Plugins/Mutable/Source/CustomizableObjectEditor/Private/MuCOE/SMutableCodeViewer.cpp
+++ b/Engine/Plugins/Mutable/Source/CustomizableObjectEditor/Private/MuCOE/SMutableCodeViewer.cpp
@@ -294,11 +294,11 @@ void SMutableCodeViewer::ClearSelectedTreeRow() const
 	TreeView->ClearSelection();
 }
 
-void SMutableCodeViewer::SetCurrentModel(const TSharedPtr<mu::Model, ESPMode::ThreadSafe>& InMutableModel, const TArray<TSoftObjectPtr<const UTexture>>& InReferencedTextures)
+void SMutableCodeViewer::SetCurrentModel(const TSharedPtr<mu::FModel, ESPMode::ThreadSafe>& InMutableModel, const TArray<TSoftObjectPtr<const UTexture>>& InReferencedTextures)
 {
 	MutableModel = InMutableModel;
 	ReferencedTextures = InReferencedTextures;
-	PreviewParameters = mu::Model::NewParameters(MutableModel);
+	PreviewParameters = mu::FModel::NewParameters(MutableModel);
 
 	RootNodes.Empty();
 	RootNodeAddresses.Empty();
@@ -342,7 +342,7 @@ void SMutableCodeViewer::SetCurrentModel(const TSharedPtr<mu::Model, ESPMode::Th
 }
 
 
-void SMutableCodeViewer::Construct(const FArguments& InArgs, const TSharedPtr<mu::Model, ESPMode::ThreadSafe>& InMutableModel,
+void SMutableCodeViewer::Construct(const FArguments& InArgs, const TSharedPtr<mu::FModel, ESPMode::ThreadSafe>& InMutableModel,
 	const TArray<TSoftObjectPtr<const UTexture>>& InReferencedTextures )
 {
 	// Min width allowed for the column. Needed to avoid having issues with the constants space being to small
@@ -641,7 +641,7 @@ void SMutableCodeViewer::Construct(const FArguments& InArgs, const TSharedPtr<mu
 					[
 						// Generate a new Constants panel to show the data stored on the current mutable program
 						SAssignNew(ConstantsWidget, SMutableConstantsWidget,
-							&(MutableModel->GetPrivate()->m_program),
+							&(MutableModel->GetPrivate()->Program),
 							SharedThis(this))
 					]
 				]
@@ -826,7 +826,7 @@ void SMutableCodeViewer::CacheOperationsMatchingStringPattern()
 		// to start searching
 		SearchPayload.SetupRootBatch(RootNodeAddresses);
 
-		const mu::FProgram& Program = MutableModel->GetPrivate()->m_program;
+		const mu::FProgram& Program = MutableModel->GetPrivate()->Program;
 		GetOperationsMatchingStringPattern(SearchString,bIsSearchStringRegularExpression,SearchPayload, Program);
 	
 		// Dump the located resources array onto the navigation array
@@ -943,7 +943,7 @@ void SMutableCodeViewer::GenerateNavigationDropdownElements()
 		// Get the type as a string to be able to print it on the UI
 		const FText OperationTypeName = FText::FromString(ModelOperationTypeNames[OperationTypeIndex]);
 
-		const mu::OP_TYPE RepresentedType = ModelOperationTypes[OperationTypeIndex].Key;
+		const mu::EOpType RepresentedType = ModelOperationTypes[OperationTypeIndex].Key;
 		const uint32 OperationTypeInstancesCount = ModelOperationTypes[OperationTypeIndex].Value;
 		
 		// Get the Color to be used on the text that will represent the operation on the dropdown
@@ -958,7 +958,7 @@ void SMutableCodeViewer::GenerateNavigationDropdownElements()
 	{
 		const FText EntryName = FText::FromString("NONE");
 		const FSlateColor EntryColor = ColorPerComputationalCost[StaticCast<uint8>(EOperationComputationalCost::Standard)];
-		NoneOperationEntry = MakeShared<FMutableOperationElement>(mu::OP_TYPE::NONE,EntryName,0,EntryColor);
+		NoneOperationEntry = MakeShared<FMutableOperationElement>(mu::EOpType::NONE,EntryName,0,EntryColor);
 
 		// @warn While not visible this element must be part of the collection for the ComboBox to be able to work
 		// properly
@@ -969,7 +969,7 @@ void SMutableCodeViewer::GenerateNavigationDropdownElements()
 	{
 		const FText EntryName = FText::FromString("Selected Constant");	
 		const FSlateColor EntryColor = FSlateColor(FLinearColor(0.35f ,0.35f,1.0f,1));
-		ConstantBasedNavigationEntry = MakeShared<FMutableOperationElement>(mu::OP_TYPE::NONE,EntryName,0,EntryColor);
+		ConstantBasedNavigationEntry = MakeShared<FMutableOperationElement>(mu::EOpType::NONE,EntryName,0,EntryColor);
 		
 		// @warn While not visible this element must be part of the collection for the ComboBox to be able to work
 		// properly
@@ -996,12 +996,12 @@ void SMutableCodeViewer::OnNavigationSelectedOperationChanged(
 	check (MutableOperationElement.IsValid());
 
 	// Cache the currently selected operation set on the UI by the user
-	const mu::OP_TYPE NewOperationType = MutableOperationElement->OperationType;
+	const mu::EOpType NewOperationType = MutableOperationElement->OperationType;
 	OperationTypeToSearch = NewOperationType;
 	CurrentlySelectedOperationTypeElement = MutableOperationElement;
 	
 	// Only do the internal work if the type is one that makes sense searching
-	if (OperationTypeToSearch != mu::OP_TYPE::NONE)
+	if (OperationTypeToSearch != mu::EOpType::NONE)
 	{
 		// Locate all operations on the mutable operations tree (not the visual one) that do share the same operation type
 		// as the one selected. This will fill the array with the elements we should be looking for during the navigation operation
@@ -1021,7 +1021,7 @@ void SMutableCodeViewer::OnNavigationSelectedOperationChanged(
 void SMutableCodeViewer::GenerateNavigationOpTypeStrings()
 {
 	// Grab only the names from the operation types located during the caching of operation types of the model
-	for (const TTuple<mu::OP_TYPE, uint32>& LocatedOperationType : ModelOperationTypes)
+	for (const TTuple<mu::EOpType, uint32>& LocatedOperationType : ModelOperationTypes)
 	{
 		// Find the name of the Operation type
 		const uint16 OperationIndex = static_cast<uint16>(LocatedOperationType.Key);
@@ -1043,8 +1043,8 @@ void SMutableCodeViewer::OnSelectedOperationTypeFromTree()
 	
 	const TSharedPtr<FMutableCodeTreeElement> ReferenceOperationElement = TreeView->GetSelectedItems()[0];
 	
-	const mu::OP_TYPE OperationType =
-		MutableModel->GetPrivate()->m_program.GetOpType(ReferenceOperationElement->MutableOperation);
+	const mu::EOpType OperationType =
+		MutableModel->GetPrivate()->Program.GetOpType(ReferenceOperationElement->MutableOperation);
 
 	// Find the operation type directly in our array of operation elements (from the drop down)
 	const TSharedPtr<const FMutableOperationElement>* RepresentativeElement = FoundModelOperationTypeElements.
@@ -1088,7 +1088,7 @@ void SMutableCodeViewer::CacheAddressesOfOperationsOfType()
 	SearchPayload.SetupRootBatch(RootNodeAddresses);
 	
 	// Main update procedure run for the targeted state and the targeted parameter values
-	const mu::FProgram& Program = MutableModel->GetPrivate()->m_program;
+	const mu::FProgram& Program = MutableModel->GetPrivate()->Program;
 	GetOperationsOfType(OperationTypeToSearch,SearchPayload, Program);
 	
 	if (!SearchPayload.FoundElements.IsEmpty())
@@ -1102,7 +1102,7 @@ void SMutableCodeViewer::CacheAddressesOfOperationsOfType()
 	}
 }
 
-void SMutableCodeViewer::GetOperationsOfType(const mu::OP_TYPE& TargetOperationType,
+void SMutableCodeViewer::GetOperationsOfType(const mu::EOpType& TargetOperationType,
                                              FElementsSearchCache& InSearchPayload,
                                              const mu::FProgram& InProgram)
 {
@@ -1144,12 +1144,12 @@ void SMutableCodeViewer::CacheOperationTypesPresentOnModel()
 
 	// Initialize NodeOperationTypes with empty TPair for each possible mutable operation type
 	{
-		constexpr uint32 OperationTypesCount = static_cast<uint16>(mu::OP_TYPE::COUNT);
+		constexpr uint32 OperationTypesCount = static_cast<uint16>(mu::EOpType::COUNT);
 		ModelOperationTypes.Empty(OperationTypesCount);
 		for (uint32 Index = 0; Index < OperationTypesCount ; Index++)
 		{
-			mu::OP_TYPE TargetType = StaticCast<mu::OP_TYPE>(Index);
-			ModelOperationTypes.Add(TPair<mu::OP_TYPE,uint32>{TargetType,0});
+			mu::EOpType TargetType = StaticCast<mu::EOpType>(Index);
+			ModelOperationTypes.Add(TPair<mu::EOpType,uint32>{TargetType,0});
 		}
 	}
 	
@@ -1157,17 +1157,17 @@ void SMutableCodeViewer::CacheOperationTypesPresentOnModel()
 	// there are
 	{
 		// Get the types and the amount of instances of each unique operation on the mutable model
-		const mu::FProgram& Program = MutableModel->GetPrivate()->m_program;
-		const uint32 ProgramAddressesCount = Program.m_opAddress.Num();
+		const mu::FProgram& Program = MutableModel->GetPrivate()->Program;
+		const uint32 ProgramAddressesCount = Program.OpAddress.Num();
 		
 		// Ensure first operation type is NONE since we are skipping it due to it having to have that type
-		check (Program.GetOpType(Program.m_opAddress[0]) == mu::OP_TYPE::NONE);
+		check (Program.GetOpType(Program.OpAddress[0]) == mu::EOpType::NONE);
 
 		// Iterate over the addresses of the program and count how many instances each type has.
 		for (uint32 ProgramAddressesIndex = 1 ; ProgramAddressesIndex < ProgramAddressesCount; ProgramAddressesIndex++)
 		{
 			// Locate what is the position (index) of the operation type of the address on our collection of types found until now
-			const mu::OP_TYPE OperationType = Program.GetOpType(ProgramAddressesIndex);
+			const mu::EOpType OperationType = Program.GetOpType(ProgramAddressesIndex);
 			
 			// Increase the counter for this operation type
 			const uint16 TypeAsInteger = StaticCast<uint16>(OperationType);
@@ -1178,14 +1178,14 @@ void SMutableCodeViewer::CacheOperationTypesPresentOnModel()
 	// Remove all operation types that do have no operations present on the model
 	{
 		ModelOperationTypes.RemoveAll(
-		[](const TPair<mu::OP_TYPE, uint32>& Current)
+		[](const TPair<mu::EOpType, uint32>& Current)
 			{
 				return Current.Value == 0;
 			});
 	}
 	
 	// Sort the contents of the array of mutable operation types alphabetically
-	ModelOperationTypes.StableSort([&](const TPair<mu::OP_TYPE,uint32>& A, const TPair<mu::OP_TYPE,uint32>& B)
+	ModelOperationTypes.StableSort([&](const TPair<mu::EOpType,uint32>& A, const TPair<mu::EOpType,uint32>& B)
 	{
 		// Find the name
 		FString AString;
@@ -1332,27 +1332,27 @@ void SMutableCodeViewer::GenerateAllTreeElements()
 	// 	be ignored due to how we handle them when expanding and contracting elements (OnExpansionChanged)
 
 	// Generate all root nodes
-	const mu::Model::Private* ModelPrivate = MutableModel->GetPrivate();
-	const mu::FProgram& Program = ModelPrivate->m_program;
-	const uint32 StateCount = ModelPrivate->m_program.m_states.Num();
+	const mu::FModel::Private* ModelPrivate = MutableModel->GetPrivate();
+	const mu::FProgram& Program = ModelPrivate->Program;
+	const uint32 StateCount = ModelPrivate->Program.States.Num();
 	for ( uint32 StateIndex=0; StateIndex<StateCount; ++StateIndex )
 	{
-		const mu::FProgram::FState& State = ModelPrivate->m_program.m_states[StateIndex];
+		const mu::FProgram::FState& State = ModelPrivate->Program.States[StateIndex];
 		FString Caption = FString::Printf( TEXT("state [%s]"), *State.Name );
 
 		const FSlateColor LabelColor = ColorPerComputationalCost[StaticCast<uint8>(GetOperationTypeComputationalCost(
-			ModelPrivate->m_program.GetOpType(State.m_root)))];
+			ModelPrivate->Program.GetOpType(State.Root)))];
 
-		RootNodes.Add(MakeShareable(new FMutableCodeTreeElement(ItemCache.Num(),StateIndex, MutableModel, State.m_root, Caption,LabelColor)));
+		RootNodes.Add(MakeShareable(new FMutableCodeTreeElement(ItemCache.Num(),StateIndex, MutableModel, State.Root, Caption,LabelColor)));
 
 		// Iterate over each root node and generate all the elements in a human readable pattern (Z Pattern)
 
 		constexpr  mu::OP::ADDRESS CommonParent = 0;
-		const FItemCacheKey Key = { CommonParent, State.m_root, StateIndex };
+		const FItemCacheKey Key = { CommonParent, State.Root, StateIndex };
 
 		// Add the element to the cache so we keep the indices straight.
 		ItemCache.Add(Key, RootNodes.Last());
-		GenerateElementRecursive(StateIndex,State.m_root,Program);
+		GenerateElementRecursive(StateIndex,State.Root,Program);
 	}
 	
 }
@@ -1398,17 +1398,17 @@ void SMutableCodeViewer::GenerateElementRecursive(const int32& InStateIndex, mu:
 	
 	// For some specific parent operation types we create more detailed subtrees.
 	bool bUseGeneric = false;
-	const mu::OP_TYPE ParentOperationType = InProgram.GetOpType(InParentAddress);
+	const mu::EOpType ParentOperationType = InProgram.GetOpType(InParentAddress);
 	switch (ParentOperationType)
 	{
-	case mu::OP_TYPE::IM_CONDITIONAL:
-	case mu::OP_TYPE::LA_CONDITIONAL:
-	case mu::OP_TYPE::ME_CONDITIONAL:
-	case mu::OP_TYPE::CO_CONDITIONAL:
-	case mu::OP_TYPE::SC_CONDITIONAL:
-	case mu::OP_TYPE::NU_CONDITIONAL:
-	case mu::OP_TYPE::IN_CONDITIONAL:
-	case mu::OP_TYPE::ED_CONDITIONAL:
+	case mu::EOpType::IM_CONDITIONAL:
+	case mu::EOpType::LA_CONDITIONAL:
+	case mu::EOpType::ME_CONDITIONAL:
+	case mu::EOpType::CO_CONDITIONAL:
+	case mu::EOpType::SC_CONDITIONAL:
+	case mu::EOpType::NU_CONDITIONAL:
+	case mu::EOpType::IN_CONDITIONAL:
+	case mu::EOpType::ED_CONDITIONAL:
 	{
 		mu::OP::ConditionalArgs Args = InProgram.GetOpArgs<mu::OP::ConditionalArgs>(InParentAddress);
 		AddOpFunc(Args.condition, TEXT("cond "));
@@ -1417,14 +1417,14 @@ void SMutableCodeViewer::GenerateElementRecursive(const int32& InStateIndex, mu:
 		break;
 	}
 
-	case mu::OP_TYPE::IM_SWITCH:
-	case mu::OP_TYPE::LA_SWITCH:
-	case mu::OP_TYPE::ME_SWITCH:
-	case mu::OP_TYPE::CO_SWITCH:
-	case mu::OP_TYPE::SC_SWITCH:
-	case mu::OP_TYPE::NU_SWITCH:
-	case mu::OP_TYPE::IN_SWITCH:
-	case mu::OP_TYPE::ED_SWITCH:
+	case mu::EOpType::IM_SWITCH:
+	case mu::EOpType::LA_SWITCH:
+	case mu::EOpType::ME_SWITCH:
+	case mu::EOpType::CO_SWITCH:
+	case mu::EOpType::SC_SWITCH:
+	case mu::EOpType::NU_SWITCH:
+	case mu::EOpType::IN_SWITCH:
+	case mu::EOpType::ED_SWITCH:
 	{
 		const uint8* OpData = InProgram.GetOpArgsPointer(InParentAddress);
 
@@ -1463,7 +1463,7 @@ void SMutableCodeViewer::GenerateElementRecursive(const int32& InStateIndex, mu:
 		break;
 	}
 
-	case mu::OP_TYPE::IM_SWIZZLE:
+	case mu::EOpType::IM_SWIZZLE:
 	{
 		mu::OP::ImageSwizzleArgs Args = InProgram.GetOpArgs<mu::OP::ImageSwizzleArgs>(InParentAddress);
 		for (int32 Channel = 0; Channel < 4; ++Channel)
@@ -1474,7 +1474,7 @@ void SMutableCodeViewer::GenerateElementRecursive(const int32& InStateIndex, mu:
 		break;
 	}
 
-	case mu::OP_TYPE::CO_SWIZZLE:
+	case mu::EOpType::CO_SWIZZLE:
 	{
 		mu::OP::ColourSwizzleArgs Args = InProgram.GetOpArgs<mu::OP::ColourSwizzleArgs>(InParentAddress);
 		for (int32 Channel = 0; Channel < 4; ++Channel)
@@ -1485,7 +1485,7 @@ void SMutableCodeViewer::GenerateElementRecursive(const int32& InStateIndex, mu:
 		break;
 	}
 
-	case mu::OP_TYPE::IM_LAYER:
+	case mu::EOpType::IM_LAYER:
 	{
 		mu::OP::ImageLayerArgs Args = InProgram.GetOpArgs<mu::OP::ImageLayerArgs>(InParentAddress);
 		AddOpFunc(Args.base, TEXT("base "));
@@ -1497,7 +1497,7 @@ void SMutableCodeViewer::GenerateElementRecursive(const int32& InStateIndex, mu:
 		break;
 	}
 
-	case mu::OP_TYPE::IM_LAYERCOLOUR:
+	case mu::EOpType::IM_LAYERCOLOUR:
 	{
 		mu::OP::ImageLayerColourArgs Args = InProgram.GetOpArgs<mu::OP::ImageLayerColourArgs>(InParentAddress);
 		AddOpFunc(Args.base, TEXT("base "));
@@ -1509,7 +1509,7 @@ void SMutableCodeViewer::GenerateElementRecursive(const int32& InStateIndex, mu:
 		break;
 	}
 
-	case mu::OP_TYPE::IM_MULTILAYER:
+	case mu::EOpType::IM_MULTILAYER:
 	{
 		mu::OP::ImageMultiLayerArgs Args = InProgram.GetOpArgs<mu::OP::ImageMultiLayerArgs>(InParentAddress);
 		AddOpFunc(Args.rangeSize, TEXT("range "));
@@ -1522,7 +1522,7 @@ void SMutableCodeViewer::GenerateElementRecursive(const int32& InStateIndex, mu:
 		break;
 	}
 
-	case mu::OP_TYPE::ME_ADDTAGS:
+	case mu::EOpType::ME_ADDTAGS:
 	{
 		const uint8* OpData = InProgram.GetOpArgsPointer(InParentAddress);
 
@@ -1580,7 +1580,7 @@ void SMutableCodeViewer::GenerateElementRecursive(const int32& InStateIndex, mu:
 	}
 }
 
-SMutableCodeViewer::EOperationComputationalCost SMutableCodeViewer::GetOperationTypeComputationalCost(mu::OP_TYPE OperationType) const
+SMutableCodeViewer::EOperationComputationalCost SMutableCodeViewer::GetOperationTypeComputationalCost(mu::EOpType OperationType) const
 {
 	if (VeryExpensiveOperationTypes.Contains(OperationType))
 	{
@@ -1633,7 +1633,7 @@ void SMutableCodeViewer::GetChildrenForInfo(TSharedPtr<FMutableCodeTreeElement>
 	}
 
 	check(MutableModel);
-	const mu::FProgram& Program = MutableModel->GetPrivate()->m_program;
+	const mu::FProgram& Program = MutableModel->GetPrivate()->Program;
 
 	mu::OP::ADDRESS ParentAddress = InInfo->MutableOperation;
 
@@ -1701,7 +1701,7 @@ void SMutableCodeViewer::OnExpansionChanged(TSharedPtr<FMutableCodeTreeElement>
 void SMutableCodeViewer::GetVisibleChildren(TSharedPtr<FMutableCodeTreeElement> InInfo, TSet<TSharedPtr<FMutableCodeTreeElement>>& OutChildren)
 {
 	check(MutableModel);
-	const mu::FProgram& MutableProgram = MutableModel->GetPrivate()->m_program;
+	const mu::FProgram& MutableProgram = MutableModel->GetPrivate()->Program;
 	
 	TArray<TSharedPtr<FMutableCodeTreeElement>> ToSearchForChildren;
 	ToSearchForChildren.Add(InInfo);
@@ -1780,7 +1780,7 @@ void SMutableCodeViewer::OnSelectionChanged(TSharedPtr<FMutableCodeTreeElement>
 	bIsPreviewPendingUpdate = true;
 
 	SelectedOperationAddress = SelectedNodes[0]->MutableOperation;
-	const mu::OP_TYPE OperationType = MutableModel->GetPrivate()->m_program.GetOpType(SelectedOperationAddress);
+	const mu::EOpType OperationType = MutableModel->GetPrivate()->Program.GetOpType(SelectedOperationAddress);
 	const mu::DATATYPE OperationDataType = mu::GetOpDataType(OperationType);
 	
 	switch (OperationDataType)
@@ -2038,7 +2038,7 @@ void SMutableCodeViewer::TreeExpandElements(TArray<TSharedPtr<FMutableCodeTreeEl
 		// Filter the elements being expanded if the user has defined a desired DATATYPE
 		if (FilteringDataType != mu::DATATYPE::DT_NONE)
 		{
-			const mu::OP_TYPE OperationType = Item->MutableModel->GetPrivate()->m_program.GetOpType(OperationAddress);
+			const mu::EOpType OperationType = Item->MutableModel->GetPrivate()->Program.GetOpType(OperationAddress);
 			const mu::DATATYPE OperationDataType = mu::GetOpDataType(OperationType);
 
 			// If it is not of the desired type then ignore it and continue to the next pending element
@@ -2182,12 +2182,12 @@ void SMutableCodeViewer::CacheRootNodeAddresses()
 
 	TArray<mu::OP::ADDRESS> FoundRootNodeAddresses;
 	
-	const mu::Model::Private* ModelPrivate = MutableModel->GetPrivate();
-	const int32 StateCount = ModelPrivate->m_program.m_states.Num();
+	const mu::FModel::Private* ModelPrivate = MutableModel->GetPrivate();
+	const int32 StateCount = ModelPrivate->Program.States.Num();
 	for ( int32 StateIndex=0; StateIndex < StateCount; ++StateIndex )
 	{
-		const mu::FProgram::FState& State = ModelPrivate->m_program.m_states[StateIndex];
-		FoundRootNodeAddresses.Add(State.m_root);
+		const mu::FProgram::FState& State = ModelPrivate->Program.States[StateIndex];
+		FoundRootNodeAddresses.Add(State.Root);
 	}
 
 	RootNodeAddresses = MoveTemp(FoundRootNodeAddresses);
@@ -2213,7 +2213,7 @@ void SMutableCodeViewer::CacheRootNodeAddresses()
 	SearchPayload.SetupRootBatch(RootNodeAddresses);
 	
 	// Main update procedure run for the targeted state and the targeted parameter values
-	const mu::FProgram& Program = MutableModel->GetPrivate()->m_program; 
+	const mu::FProgram& Program = MutableModel->GetPrivate()->Program; 
 	GetOperationsReferencingConstantResource(ConstantDataType,IndexOnConstantsArray,SearchPayload, Program);
 	
 	// At this point we did get all the addresses of operations that do involve the usage of our resource
@@ -2279,7 +2279,7 @@ bool SMutableCodeViewer::IsConstantResourceUsedByOperation(const int32 IndexOnCo
 	const mu::DATATYPE ConstantDataType, const mu::OP::ADDRESS OperationAddress, const mu::FProgram& InProgram) const
 {
 	// Cache the current operation type to know where to look and what to check
-	const mu::OP_TYPE OperationType = InProgram.GetOpType(OperationAddress);
+	const mu::EOpType OperationType = InProgram.GetOpType(OperationAddress);
 
 	// Making usage of the operation data type is not valid since some operations while return one type do, in fact,
 	// contain data from other types (like the mesh constant for example that contains mesh, skeleton and physics asset)
@@ -2298,39 +2298,39 @@ bool SMutableCodeViewer::IsConstantResourceUsedByOperation(const int32 IndexOnCo
 		case mu::DATATYPE::DT_STRING:
 			{
 				// TIP: To know if they represent a constant value check the code on code runner to see if they read from the constants array
-				if (OperationType == mu::OP_TYPE::ST_CONSTANT)
+				if (OperationType == mu::EOpType::ST_CONSTANT)
 				{
 					bResourceLocated = IndexOnConstantsArray == InProgram.GetOpArgs<mu::OP::ResourceConstantArgs>(OperationAddress).value;
 				}
-				else if (OperationType == mu::OP_TYPE::IN_ADDSTRING)
+				else if (OperationType == mu::EOpType::IN_ADDSTRING)
 				{
 					bResourceLocated = IndexOnConstantsArray == InProgram.GetOpArgs<mu::OP::InstanceAddArgs>(OperationAddress).name;
 				}
-				else if (OperationType == mu::OP_TYPE::IN_ADDMESH)
+				else if (OperationType == mu::EOpType::IN_ADDMESH)
 				{
 					bResourceLocated = IndexOnConstantsArray == InProgram.GetOpArgs<mu::OP::InstanceAddArgs>(OperationAddress).name;
 				}
-				else if(OperationType == mu::OP_TYPE::IN_ADDIMAGE)
+				else if(OperationType == mu::EOpType::IN_ADDIMAGE)
 				{
 					bResourceLocated = IndexOnConstantsArray == InProgram.GetOpArgs<mu::OP::InstanceAddArgs>(OperationAddress).name;
 				}
-				else if (OperationType == mu::OP_TYPE::IN_ADDVECTOR)
+				else if (OperationType == mu::EOpType::IN_ADDVECTOR)
 				{
 					bResourceLocated = IndexOnConstantsArray == InProgram.GetOpArgs<mu::OP::InstanceAddArgs>(OperationAddress).name;
 				}
-				else if (OperationType == mu::OP_TYPE::IN_ADDSCALAR)
+				else if (OperationType == mu::EOpType::IN_ADDSCALAR)
 				{
 					bResourceLocated = IndexOnConstantsArray == InProgram.GetOpArgs<mu::OP::InstanceAddArgs>(OperationAddress).name;
 				}
-				else if (OperationType == mu::OP_TYPE::IN_ADDCOMPONENT)
+				else if (OperationType == mu::EOpType::IN_ADDCOMPONENT)
 				{
 					bResourceLocated = IndexOnConstantsArray == InProgram.GetOpArgs<mu::OP::InstanceAddArgs>(OperationAddress).name;
 				}
-				else if (OperationType == mu::OP_TYPE::IN_ADDSURFACE)
+				else if (OperationType == mu::EOpType::IN_ADDSURFACE)
 				{
 					bResourceLocated = IndexOnConstantsArray == InProgram.GetOpArgs<mu::OP::InstanceAddArgs>(OperationAddress).name;
 				}
-				else if (OperationType == mu::OP_TYPE::ME_BINDSHAPE)
+				else if (OperationType == mu::EOpType::ME_BINDSHAPE)
 				{
 					mu::OP::MeshBindShapeArgs Arguments = InProgram.GetOpArgs<mu::OP::MeshBindShapeArgs>(OperationAddress);
 					const uint8_t* Data = InProgram.GetOpArgsPointer(OperationAddress);
@@ -2361,7 +2361,7 @@ bool SMutableCodeViewer::IsConstantResourceUsedByOperation(const int32 IndexOnCo
 
 		case mu::DATATYPE::DT_IMAGE:
 			{
-				if (OperationType == mu::OP_TYPE::IM_CONSTANT)
+				if (OperationType == mu::EOpType::IM_CONSTANT)
 				{
 					bResourceLocated = IndexOnConstantsArray == InProgram.GetOpArgs<mu::OP::ResourceConstantArgs>(OperationAddress).value;
 				} 
@@ -2371,7 +2371,7 @@ bool SMutableCodeViewer::IsConstantResourceUsedByOperation(const int32 IndexOnCo
 
 		case mu::DATATYPE::DT_MESH:
 			{
-				if (OperationType == mu::OP_TYPE::ME_CONSTANT)
+				if (OperationType == mu::EOpType::ME_CONSTANT)
 				{
 					bResourceLocated = IndexOnConstantsArray == InProgram.GetOpArgs<mu::OP::MeshConstantArgs>(OperationAddress).value;
 				}
@@ -2380,7 +2380,7 @@ bool SMutableCodeViewer::IsConstantResourceUsedByOperation(const int32 IndexOnCo
 
 		case mu::DATATYPE::DT_LAYOUT:
 			{
-				if (OperationType == mu::OP_TYPE::LA_CONSTANT)
+				if (OperationType == mu::EOpType::LA_CONSTANT)
 				{
 					bResourceLocated = IndexOnConstantsArray == InProgram.GetOpArgs<mu::OP::ResourceConstantArgs>(OperationAddress).value;
 				}
@@ -2389,7 +2389,7 @@ bool SMutableCodeViewer::IsConstantResourceUsedByOperation(const int32 IndexOnCo
 
 		case mu::DATATYPE::DT_PROJECTOR:
 			{
-				if (OperationType == mu::OP_TYPE::PR_CONSTANT)
+				if (OperationType == mu::EOpType::PR_CONSTANT)
 				{
 					bResourceLocated = IndexOnConstantsArray == InProgram.GetOpArgs<mu::OP::ResourceConstantArgs>(OperationAddress).value;
 				}
@@ -2398,11 +2398,11 @@ bool SMutableCodeViewer::IsConstantResourceUsedByOperation(const int32 IndexOnCo
 
 		case mu::DATATYPE::DT_MATRIX:
 			{
-				if (OperationType == mu::OP_TYPE::ME_TRANSFORM)
+				if (OperationType == mu::EOpType::ME_TRANSFORM)
 				{
 					bResourceLocated = IndexOnConstantsArray == InProgram.GetOpArgs<mu::OP::MeshTransformArgs>(OperationAddress).matrix;
 				}
-				else if (OperationType == mu::OP_TYPE::ME_TRANSFORMWITHMESH)
+				else if (OperationType == mu::EOpType::ME_TRANSFORMWITHMESH)
 				{
 					bResourceLocated = IndexOnConstantsArray == InProgram.GetOpArgs<mu::OP::MeshTransformWithinMeshArgs>(OperationAddress).matrix;
 				}
@@ -2411,7 +2411,7 @@ bool SMutableCodeViewer::IsConstantResourceUsedByOperation(const int32 IndexOnCo
 
 		case mu::DATATYPE::DT_SHAPE:
 			{
-				if (OperationType == mu::OP_TYPE::ME_CLIPMORPHPLANE)
+				if (OperationType == mu::EOpType::ME_CLIPMORPHPLANE)
 				{
 					const mu::OP::MeshClipMorphPlaneArgs Arguments =  InProgram.GetOpArgs<mu::OP::MeshClipMorphPlaneArgs>(OperationAddress);
 					
@@ -2433,7 +2433,7 @@ bool SMutableCodeViewer::IsConstantResourceUsedByOperation(const int32 IndexOnCo
 		
 		case mu::DATATYPE::DT_CURVE:
 			{
-				if (OperationType == mu::OP_TYPE::SC_CURVE)
+				if (OperationType == mu::EOpType::SC_CURVE)
 				{
 					bResourceLocated = IndexOnConstantsArray == InProgram.GetOpArgs<mu::OP::ScalarCurveArgs>(OperationAddress).curve;
 				}
@@ -2442,7 +2442,7 @@ bool SMutableCodeViewer::IsConstantResourceUsedByOperation(const int32 IndexOnCo
 
 		case mu::DATATYPE::DT_SKELETON:
 			{
-				if (OperationType == mu::OP_TYPE::ME_CONSTANT)
+				if (OperationType == mu::EOpType::ME_CONSTANT)
 				{
 					bResourceLocated = IndexOnConstantsArray == InProgram.GetOpArgs<mu::OP::MeshConstantArgs>(OperationAddress).skeleton;
 				}
@@ -2469,7 +2469,7 @@ bool SMutableCodeViewer::IsConstantResourceUsedByOperation(const int32 IndexOnCo
 namespace
 {
 	/** Test implementation to provide image parameters. It will generate some images of a fixed size and format. */
-	class TestResourceProvider : public mu::ExternalResourceProvider
+	class TestResourceProvider : public mu::FExternalResourceProvider
 	{
 		static inline const mu::FImageDesc IMAGE_DESC = 
 			mu::FImageDesc(mu::FImageSize(1024, 1024), mu::EImageFormat::RGBA_UByte, 1);
@@ -2482,14 +2482,14 @@ namespace
 	public:
 
 
-		virtual TTuple<UE::Tasks::FTask, TFunction<void()>> GetImageAsync(FName Id, uint8 MipmapsToSkip, TFunction<void(mu::Ptr<mu::Image>)>& ResultCallback) override
+		virtual TTuple<UE::Tasks::FTask, TFunction<void()>> GetImageAsync(FName Id, uint8 MipmapsToSkip, TFunction<void(TSharedPtr<mu::FImage>)>& ResultCallback) override
 		{
 			MUTABLE_CPUPROFILER_SCOPE(TestImageProvider_GetImage);
 
 			int32 Size = IMAGE_DESC.m_size[0];
 			Size = FMath::Max(4, Size / (1 << MipmapsToSkip));
 
-			mu::Ptr<mu::Image> Image = new mu::Image(
+			TSharedPtr<mu::FImage> Image = MakeShared<mu::FImage>(
 				Size, Size, IMAGE_DESC.m_lods,
 				IMAGE_DESC.m_format,
 				mu::EInitializationType::NotInitialized);
@@ -2530,7 +2530,7 @@ namespace
 		}
 
 
-		virtual TTuple<UE::Tasks::FTask, TFunction<void()>> GetReferencedImageAsync(const void* ModelPtr, int32 Id, uint8 MipmapsToSkip, TFunction<void(mu::Ptr<mu::Image>)>& ResultCallback)
+		virtual TTuple<UE::Tasks::FTask, TFunction<void()>> GetReferencedImageAsync(const void* ModelPtr, int32 Id, uint8 MipmapsToSkip, TFunction<void(TSharedPtr<mu::FImage>)>& ResultCallback)
 		{
 			check(ReferencedTextures.IsValidIndex(Id));
 
@@ -2542,10 +2542,10 @@ namespace
 			int32 MipIndex = (MipmapsToSkip < Texture->GetPlatformData()->Mips.Num()) ? MipmapsToSkip : Texture->GetPlatformData()->Mips.Num() - 1;
 			check(MipIndex >= 0);
 
-			mu::Ptr<mu::Image> ResultImage = new mu::Image();
+			TSharedPtr<mu::FImage> ResultImage = MakeShared<mu::FImage>();
 
 			FMutableSourceTextureData Tex(*Texture);
-			EUnrealToMutableConversionError Error = ConvertTextureUnrealSourceToMutable(ResultImage.get(), Tex, MipmapsToSkip);
+			EUnrealToMutableConversionError Error = ConvertTextureUnrealSourceToMutable(ResultImage.Get(), Tex, MipmapsToSkip);
 
 			if (Error != EUnrealToMutableConversionError::Success)
 			{
@@ -2565,9 +2565,9 @@ namespace
 		}
 
 
-		virtual TTuple<UE::Tasks::FTask, TFunction<void()>> GetMeshAsync(FName Id, TFunction<void(mu::Ptr<mu::Mesh>)>& ResultCallback) override
+		virtual TTuple<UE::Tasks::FTask, TFunction<void()>> GetMeshAsync(FName Id, TFunction<void(TSharedPtr<mu::FMesh>)>& ResultCallback) override
 		{
-			// TODO: This is a stub function where Mesh Parameters would be supported.
+			// TODO: This is a stub function where Mesh FParameters would be supported.
 
 			// Thread: worker
 			MUTABLE_CPUPROFILER_SCOPE(FUnrealMutableImageProvider::GetMeshAsync);
@@ -2580,13 +2580,13 @@ namespace
 			// No USkeletalMesh was provided, cannot do anything, just provide a dummy mesh
 			UE_LOG(LogMutable, Warning, TEXT("No USkeletalMesh was provided for an application-specific mesh parameter."));
 
-			mu::Ptr<mu::Mesh> Dummy = new mu::Mesh;
+			TSharedPtr<mu::FMesh> Dummy = MakeShared<mu::FMesh>();
 			ResultCallback(Dummy);
 			return Invoke(TrivialReturn);
 		}
 
 
-		virtual TTuple<UE::Tasks::FTask, TFunction<void()>> GetReferencedMeshAsync(const void* ModelPtr, int32 Id, TFunction<void(mu::Ptr<mu::Mesh>)>& ResultCallback) override
+		virtual TTuple<UE::Tasks::FTask, TFunction<void()>> GetReferencedMeshAsync(const void* ModelPtr, int32 Id, TFunction<void(TSharedPtr<mu::FMesh>)>& ResultCallback) override
 		{
 			// TODO: This is a stub function where referenced Meshes would be supported.
 
@@ -2598,7 +2598,7 @@ namespace
 					return MakeTuple(UE::Tasks::MakeCompletedTask<void>(), []() -> void {});
 				};
 
-			mu::Ptr<mu::Mesh> Mesh = new mu::Mesh();
+			TSharedPtr<mu::FMesh> Mesh = MakeShared<mu::FMesh>();
 
 			// Not supported outside editor yet.
 			UE_LOG(LogMutable, Warning, TEXT("Failed to get referenced mesh. Not implemented yet."));
@@ -2636,11 +2636,11 @@ void SMutableCodeViewer::Tick(const FGeometry& AllottedGeometry, const double In
 
 	bIsPreviewPendingUpdate = false;
 
-	const mu::OP_TYPE OperationType = MutableModel->GetPrivate()->m_program.GetOpType(SelectedOperationAddress);
+	const mu::EOpType OperationType = MutableModel->GetPrivate()->Program.GetOpType(SelectedOperationAddress);
 	const mu::DATATYPE OperationDataType = mu::GetOpDataType(OperationType);
 
-	const mu::Ptr<mu::Settings> Settings = new mu::Settings();
-	const mu::Ptr<mu::System> System = new mu::System(Settings);
+	mu::FSettings Settings;
+	const TSharedPtr<mu::FSystem> System = MakeShared<mu::FSystem>(Settings);
 
 	TSharedPtr<TestResourceProvider> ExternalResourceProvider = MakeShared<TestResourceProvider>();
 	ExternalResourceProvider->ReferencedTextures = ReferencedTextures;
@@ -2653,7 +2653,7 @@ void SMutableCodeViewer::Tick(const FGeometry& AllottedGeometry, const double In
 	case mu::DT_LAYOUT:
 	{
 		check(PreviewLayoutViewer);
-		mu::Ptr<const mu::Layout> MutableLayout = System->GetPrivate()->BuildLayout(MutableModel, PreviewParameters.get(), SelectedOperationAddress);
+		TSharedPtr<const mu::FLayout> MutableLayout = System->GetPrivate()->BuildLayout(MutableModel, PreviewParameters.Get(), SelectedOperationAddress);
 		PreviewLayoutViewer->SetLayout(MutableLayout);
 		break;
 	}
@@ -2661,7 +2661,7 @@ void SMutableCodeViewer::Tick(const FGeometry& AllottedGeometry, const double In
 	case mu::DT_IMAGE:
 	{
 		check(PreviewImageViewer);
-		mu::Ptr<const mu::Image> MutableImage = System->GetPrivate()->BuildImage(MutableModel, PreviewParameters.get(), SelectedOperationAddress, MipsToSkip, 0);
+		TSharedPtr<const mu::FImage> MutableImage = System->GetPrivate()->BuildImage(MutableModel, PreviewParameters.Get(), SelectedOperationAddress, MipsToSkip, 0);
 		PreviewImageViewer->SetImage(MutableImage, 0);
 		break;
 	}
@@ -2669,7 +2669,7 @@ void SMutableCodeViewer::Tick(const FGeometry& AllottedGeometry, const double In
 	case mu::DT_MESH:
 	{
 		check(PreviewMeshViewer);
-		mu::Ptr<const mu::Mesh> MutableMesh = System->GetPrivate()->BuildMesh(MutableModel, PreviewParameters.get(), SelectedOperationAddress);
+		TSharedPtr<const mu::FMesh> MutableMesh = System->GetPrivate()->BuildMesh(MutableModel, PreviewParameters.Get(), SelectedOperationAddress);
 		PreviewMeshViewer->SetMesh(MutableMesh);
 		break;
 	}
@@ -2677,15 +2677,15 @@ void SMutableCodeViewer::Tick(const FGeometry& AllottedGeometry, const double In
 	case mu::DT_INSTANCE:
 	{
 		check(PreviewInstanceViewer);
-		mu::Ptr<const mu::Instance> MutableInstance = System->GetPrivate()->BuildInstance(MutableModel, PreviewParameters.get(), SelectedOperationAddress);
-		PreviewInstanceViewer->SetInstance(MutableInstance, MutableModel, PreviewParameters, System);
+		TSharedPtr<const mu::FInstance> MutableInstance = System->GetPrivate()->BuildInstance(MutableModel, PreviewParameters.Get(), SelectedOperationAddress);
+		PreviewInstanceViewer->SetInstance(MutableInstance, MutableModel, PreviewParameters, *System);
 		break;
 	}
 
 	case mu::DT_BOOL:
 	{
 		check(PreviewBoolViewer);
-		const bool MutableBool = System->GetPrivate()->BuildBool(MutableModel, PreviewParameters.get(), SelectedOperationAddress);
+		const bool MutableBool = System->GetPrivate()->BuildBool(MutableModel, PreviewParameters.Get(), SelectedOperationAddress);
 		PreviewBoolViewer->SetBool(MutableBool);
 		break;
 	}
@@ -2693,7 +2693,7 @@ void SMutableCodeViewer::Tick(const FGeometry& AllottedGeometry, const double In
 	case mu::DT_INT:
 	{
 		check(PreviewIntViewer);
-		const int32 MutableInt = System->GetPrivate()->BuildInt(MutableModel, PreviewParameters.get(), SelectedOperationAddress);
+		const int32 MutableInt = System->GetPrivate()->BuildInt(MutableModel, PreviewParameters.Get(), SelectedOperationAddress);
 		PreviewIntViewer->SetInt(MutableInt);
 		break;
 	}
@@ -2701,7 +2701,7 @@ void SMutableCodeViewer::Tick(const FGeometry& AllottedGeometry, const double In
 	case mu::DT_SCALAR:
 	{
 		check(PreviewScalarViewer);
-		const float MutableScalar = System->GetPrivate()->BuildScalar(MutableModel, PreviewParameters.get(), SelectedOperationAddress);
+		const float MutableScalar = System->GetPrivate()->BuildScalar(MutableModel, PreviewParameters.Get(), SelectedOperationAddress);
 		PreviewScalarViewer->SetScalar(MutableScalar);
 		break;
 	}
@@ -2709,7 +2709,7 @@ void SMutableCodeViewer::Tick(const FGeometry& AllottedGeometry, const double In
 	case mu::DT_STRING:
 	{
 		check(PreviewStringViewer);
-		mu::Ptr<const mu::String> MutableString = System->GetPrivate()->BuildString(MutableModel, PreviewParameters.get(), SelectedOperationAddress);
+		TSharedPtr<const mu::String> MutableString = System->GetPrivate()->BuildString(MutableModel, PreviewParameters.Get(), SelectedOperationAddress);
 		const FText MutableText = FText::FromString(FString(MutableString->GetValue()));
 		PreviewStringViewer->SetString(MutableText);
 		break;
@@ -2718,7 +2718,7 @@ void SMutableCodeViewer::Tick(const FGeometry& AllottedGeometry, const double In
 	case mu::DT_COLOUR:
 	{
 		check(PreviewColorViewer);
-		FVector4f Color = System->GetPrivate()->BuildColour(MutableModel, PreviewParameters.get(), SelectedOperationAddress);
+		FVector4f Color = System->GetPrivate()->BuildColour(MutableModel, PreviewParameters.Get(), SelectedOperationAddress);
 		PreviewColorViewer->SetColor(Color);
 		break;
 	}
@@ -2726,7 +2726,7 @@ void SMutableCodeViewer::Tick(const FGeometry& AllottedGeometry, const double In
 	case mu::DT_PROJECTOR:
 	{
 		check (PreviewProjectorViewer);
-		mu::FProjector Projector = System->GetPrivate()->BuildProjector(MutableModel, PreviewParameters.get(), SelectedOperationAddress);
+		mu::FProjector Projector = System->GetPrivate()->BuildProjector(MutableModel, PreviewParameters.Get(), SelectedOperationAddress);
 		PreviewProjectorViewer->SetProjector(Projector);
 	}
 	
@@ -2824,21 +2824,21 @@ void SMutableCodeViewer::PreviewMutableString(const FString& InString)
 }
 
 
-void SMutableCodeViewer::PreviewMutableImage(mu::ImagePtrConst InImagePtr)
+void SMutableCodeViewer::PreviewMutableImage(TSharedPtr<const mu::FImage> InImagePtr)
 {
 	PrepareImageViewer();
 	PreviewImageViewer->SetImage(InImagePtr,0);
 }
 
 
-void SMutableCodeViewer::PreviewMutableMesh(mu::MeshPtrConst InMeshPtr)
+void SMutableCodeViewer::PreviewMutableMesh(TSharedPtr<const mu::FMesh> InMeshPtr)
 {
 	PrepareMeshViewer();
 	PreviewMeshViewer->SetMesh(InMeshPtr);
 }
 
 
-void SMutableCodeViewer::PreviewMutableLayout(mu::Ptr<const mu::Layout> Layout)
+void SMutableCodeViewer::PreviewMutableLayout(TSharedPtr<const mu::FLayout> Layout)
 {
 	PrepareLayoutViewer();
 	PreviewLayoutViewer->SetLayout(Layout);
@@ -2858,7 +2858,7 @@ void SMutableCodeViewer::PreviewMutableProjector(const mu::FProjector* Projector
 }
 
 
-void SMutableCodeViewer::PreviewMutableSkeleton(mu::SkeletonPtrConst Skeleton)
+void SMutableCodeViewer::PreviewMutableSkeleton(TSharedPtr<const mu::FSkeleton> Skeleton)
 {
 	if (!PreviewSkeletonViewer)
 	{
@@ -2896,7 +2896,7 @@ void SMutableCodeViewer::PreviewMutableShape(const mu::FShape* Shape)
 }
 
 
-FMutableCodeTreeElement::FMutableCodeTreeElement(int32 InIndexOnTree, const int32& InMutableStateIndex, const TSharedPtr<mu::Model, ESPMode::ThreadSafe>& InModel, mu::OP::ADDRESS InOperation, const FString& InCaption, const FSlateColor InLabelColor, const TSharedPtr<FMutableCodeTreeElement>* InDuplicatedOf)
+FMutableCodeTreeElement::FMutableCodeTreeElement(int32 InIndexOnTree, const int32& InMutableStateIndex, const TSharedPtr<mu::FModel, ESPMode::ThreadSafe>& InModel, mu::OP::ADDRESS InOperation, const FString& InCaption, const FSlateColor InLabelColor, const TSharedPtr<FMutableCodeTreeElement>* InDuplicatedOf)
 {
 	MutableModel = InModel;
 	MutableOperation = InOperation;
@@ -2926,11 +2926,11 @@ void FMutableCodeTreeElement::SetElementCurrentState(const int32& InStateIndex)
 
 	// Check for an out of bounds value
 	check(MutableModel);
-	mu::FProgram& MutableProgram = MutableModel->GetPrivate()->m_program;
-	check(InStateIndex >= 0 && InStateIndex < MutableProgram.m_states.Num());
+	mu::FProgram& MutableProgram = MutableModel->GetPrivate()->Program;
+	check(InStateIndex >= 0 && InStateIndex < MutableProgram.States.Num());
 
 	CurrentMutableStateIndex = InStateIndex;
-	const mu::FProgram::FState& CurrentState = MutableProgram.m_states[CurrentMutableStateIndex];
+	const mu::FProgram::FState& CurrentState = MutableProgram.States[CurrentMutableStateIndex];
 
 	// Check if it is a dynamic resource
 	for (auto& DynamicResource : CurrentState.m_dynamicResources)
@@ -2955,29 +2955,29 @@ void FMutableCodeTreeElement::SetElementCurrentState(const int32& InStateIndex)
 
 void FMutableCodeTreeElement::GenerateLabelText()
 {
-	const mu::FProgram& Program = MutableModel->GetPrivate()->m_program;
-	const mu::OP_TYPE Type = Program.GetOpType(MutableOperation);
+	const mu::FProgram& Program = MutableModel->GetPrivate()->Program;
+	const mu::EOpType Type = Program.GetOpType(MutableOperation);
 	FString OpName = mu::s_opNames[static_cast<int32>(Type)];
 	OpName.TrimEndInline();
 
 	// See if the operation type accepts additional information in the label
 	switch (Type)
 	{
-	case mu::OP_TYPE::BO_PARAMETER:
-	case mu::OP_TYPE::NU_PARAMETER:
-	case mu::OP_TYPE::SC_PARAMETER:
-	case mu::OP_TYPE::CO_PARAMETER:
-	case mu::OP_TYPE::PR_PARAMETER:
-	case mu::OP_TYPE::IM_PARAMETER:
-	case mu::OP_TYPE::ST_PARAMETER:
+	case mu::EOpType::BO_PARAMETER:
+	case mu::EOpType::NU_PARAMETER:
+	case mu::EOpType::SC_PARAMETER:
+	case mu::EOpType::CO_PARAMETER:
+	case mu::EOpType::PR_PARAMETER:
+	case mu::EOpType::IM_PARAMETER:
+	case mu::EOpType::ST_PARAMETER:
 	{
 		mu::OP::ParameterArgs Args = Program.GetOpArgs<mu::OP::ParameterArgs>(MutableOperation);
 		OpName += TEXT(" ");
-		OpName += Program.m_parameters[int32(Args.variable)].m_name;
+		OpName += Program.Parameters[int32(Args.variable)].Name;
 		break;
 	}
 
-	case mu::OP_TYPE::IM_SWIZZLE:
+	case mu::EOpType::IM_SWIZZLE:
 	{
 		mu::OP::ImageSwizzleArgs Args = Program.GetOpArgs<mu::OP::ImageSwizzleArgs>(MutableOperation);
 		OpName += TEXT(" ");
@@ -2985,7 +2985,7 @@ void FMutableCodeTreeElement::GenerateLabelText()
 		break;
 	}
 
-	case mu::OP_TYPE::IM_PIXELFORMAT:
+	case mu::EOpType::IM_PIXELFORMAT:
 	{
 		mu::OP::ImagePixelFormatArgs Args = Program.GetOpArgs<mu::OP::ImagePixelFormatArgs>(MutableOperation);
 		OpName += TEXT(" ");
@@ -2995,28 +2995,28 @@ void FMutableCodeTreeElement::GenerateLabelText()
 		break;
 	}
 
-	case mu::OP_TYPE::IM_MIPMAP:
+	case mu::EOpType::IM_MIPMAP:
 	{
 		mu::OP::ImageMipmapArgs Args = Program.GetOpArgs<mu::OP::ImageMipmapArgs>(MutableOperation);
 		OpName += FString::Printf(TEXT(" levels: %d-%d tail: %d"), Args.levels, Args.blockLevels, int32(Args.onlyTail));
 		break;
 	}
 
-	case mu::OP_TYPE::IM_RESIZE:
+	case mu::EOpType::IM_RESIZE:
 	{
 		mu::OP::ImageResizeArgs Args = Program.GetOpArgs<mu::OP::ImageResizeArgs>(MutableOperation);
 		OpName += FString::Printf(TEXT(" %d x %d"), int32(Args.size[0]), int32(Args.size[1]));
 		break;
 	}
 
-	case mu::OP_TYPE::IM_RESIZEREL:
+	case mu::EOpType::IM_RESIZEREL:
 	{
 		mu::OP::ImageResizeRelArgs Args = Program.GetOpArgs<mu::OP::ImageResizeRelArgs>(MutableOperation);
 		OpName += FString::Printf(TEXT(" %.3f x %.3f"), Args.factor[0], Args.factor[1]);
 		break;
 	}
 
-	case mu::OP_TYPE::IM_MULTILAYER:
+	case mu::EOpType::IM_MULTILAYER:
 	{
 		mu::OP::ImageMultiLayerArgs Args = Program.GetOpArgs<mu::OP::ImageMultiLayerArgs>(MutableOperation);
 		OpName += TEXT(" rgb: ");
@@ -3029,7 +3029,7 @@ void FMutableCodeTreeElement::GenerateLabelText()
 		break;
 	}
 
-	case mu::OP_TYPE::IM_LAYER:
+	case mu::EOpType::IM_LAYER:
 	{
 		mu::OP::ImageLayerArgs Args = Program.GetOpArgs<mu::OP::ImageLayerArgs>(MutableOperation);
 		OpName += TEXT(" rgb: ");
@@ -3041,7 +3041,7 @@ void FMutableCodeTreeElement::GenerateLabelText()
 		break;
 	}
 
-	case mu::OP_TYPE::IM_LAYERCOLOUR:
+	case mu::EOpType::IM_LAYERCOLOUR:
 	{
 		mu::OP::ImageLayerColourArgs Args = Program.GetOpArgs<mu::OP::ImageLayerColourArgs>(MutableOperation);
 		OpName += TEXT(" rgb: ");
@@ -3054,7 +3054,7 @@ void FMutableCodeTreeElement::GenerateLabelText()
 		break;
 	}
 
-	case mu::OP_TYPE::IM_PLAINCOLOUR:
+	case mu::EOpType::IM_PLAINCOLOUR:
 	{
 		mu::OP::ImagePlainColourArgs Args = Program.GetOpArgs<mu::OP::ImagePlainColourArgs>(MutableOperation);
 		OpName += TEXT(" format: ");
@@ -3064,13 +3064,13 @@ void FMutableCodeTreeElement::GenerateLabelText()
 		break;
 	}
 
-	case mu::OP_TYPE::IN_ADDIMAGE:
+	case mu::EOpType::IN_ADDIMAGE:
 	{
 		mu::OP::InstanceAddArgs Args = Program.GetOpArgs<mu::OP::InstanceAddArgs>(MutableOperation);
-		if (Program.m_constantStrings.IsValidIndex(Args.name))
+		if (Program.ConstantStrings.IsValidIndex(Args.name))
 		{
 			OpName += TEXT(" name: ");
-			OpName += Program.m_constantStrings[Args.name];
+			OpName += Program.ConstantStrings[Args.name];
 		}
 		break;
 	}
diff --git a/Engine/Plugins/Mutable/Source/CustomizableObjectEditor/Private/MuCOE/SMutableCodeViewer.h b/Engine/Plugins/Mutable/Source/CustomizableObjectEditor/Private/MuCOE/SMutableCodeViewer.h
--- a/Engine/Plugins/Mutable/Source/CustomizableObjectEditor/Private/MuCOE/SMutableCodeViewer.h
+++ b/Engine/Plugins/Mutable/Source/CustomizableObjectEditor/Private/MuCOE/SMutableCodeViewer.h
@@ -59,7 +59,7 @@ public:
 		
 	SLATE_END_ARGS()
 
-	void Construct(const FArguments& InArgs, const TSharedPtr<mu::Model, ESPMode::ThreadSafe>& InMutableModel,
+	void Construct(const FArguments& InArgs, const TSharedPtr<mu::FModel, ESPMode::ThreadSafe>& InMutableModel,
 		const TArray<TSoftObjectPtr<const UTexture>>& ReferencedTextures );
 
 	// SWidget interface
@@ -75,7 +75,7 @@ public:
 private:
 
 	/** The Mutable Model that we are showing. */
-	TSharedPtr<mu::Model, ESPMode::ThreadSafe> MutableModel;
+	TSharedPtr<mu::FModel, ESPMode::ThreadSafe> MutableModel;
 	
 	/** Array of external referenced textures in MutableModel, indexed by id. */
 	TArray<TSoftObjectPtr<const UTexture>> ReferencedTextures;
@@ -84,7 +84,7 @@ private:
 	mu::OP::ADDRESS SelectedOperationAddress = 0;
 
 	/** Mutable parameters used in the preview. */
-	mu::ParametersPtr PreviewParameters;
+	TSharedPtr<mu::FParameters> PreviewParameters;
 
 	/** Widget showing the parameters that affect the current preview. */
 	TSharedPtr<SMutableParametersWidget> ParametersWidget;
@@ -186,7 +186,7 @@ private:
 	TMap< mu::OP::ADDRESS, TSharedPtr<FMutableCodeTreeElement>> ExpandedElements;
 
 	/** Prepare the widget for the given model. */
-	void SetCurrentModel(const TSharedPtr<mu::Model, ESPMode::ThreadSafe>&, const TArray<TSoftObjectPtr<const UTexture>>& ReferencedTextures);
+	void SetCurrentModel(const TSharedPtr<mu::FModel, ESPMode::ThreadSafe>&, const TArray<TSoftObjectPtr<const UTexture>>& ReferencedTextures);
 
 	/** Before any UI operation generate all the elements that may be navigable over the tree. No children of duplicated 
 	 * addresses will be generated.
@@ -214,25 +214,25 @@ private:
 	 */
 
 	/** Collection with all very expensive to run operation types */
-	const TArray<mu::OP_TYPE> VeryExpensiveOperationTypes
+	const TArray<mu::EOpType> VeryExpensiveOperationTypes
 	{
-		mu::OP_TYPE::ME_BINDSHAPE,
-		mu::OP_TYPE::ME_MASKCLIPMESH,
-		mu::OP_TYPE::ME_MASKCLIPUVMASK,
-		mu::OP_TYPE::ME_FORMAT,
-		mu::OP_TYPE::ME_MASKDIFF,
-		mu::OP_TYPE::ME_DIFFERENCE,
-		mu::OP_TYPE::ME_EXTRACTLAYOUTBLOCK,
-		mu::OP_TYPE::IM_MAKEGROWMAP,
+		mu::EOpType::ME_BINDSHAPE,
+		mu::EOpType::ME_MASKCLIPMESH,
+		mu::EOpType::ME_MASKCLIPUVMASK,
+		mu::EOpType::ME_FORMAT,
+		mu::EOpType::ME_MASKDIFF,
+		mu::EOpType::ME_DIFFERENCE,
+		mu::EOpType::ME_EXTRACTLAYOUTBLOCK,
+		mu::EOpType::IM_MAKEGROWMAP,
 	};
 
 	/** Collection with all expensive to run operation types */
-	const TArray<mu::OP_TYPE> ExpensiveOperationTypes
+	const TArray<mu::EOpType> ExpensiveOperationTypes
 	{
-		mu::OP_TYPE::IM_PIXELFORMAT,
-		mu::OP_TYPE::IM_SWIZZLE,
-		mu::OP_TYPE::ME_PROJECT,
-		mu::OP_TYPE::ME_TRANSFORMWITHMESH,
+		mu::EOpType::IM_PIXELFORMAT,
+		mu::EOpType::IM_SWIZZLE,
+		mu::EOpType::ME_PROJECT,
+		mu::EOpType::ME_TRANSFORMWITHMESH,
 	};
 	
 	/** Enum designed to be able to notify the row generation of the type of operation being generated */
@@ -257,7 +257,7 @@ private:
 	 * @param OperationType The operation type you want to know the computational cost of
 	 * @return The enum value representing the computational cost of the operation type provided.
 	 */
-	EOperationComputationalCost GetOperationTypeComputationalCost(mu::OP_TYPE OperationType) const;
+	EOperationComputationalCost GetOperationTypeComputationalCost(mu::EOpType OperationType) const;
 	
 	
 	/*
@@ -282,11 +282,11 @@ private:
 	TSharedPtr<const FMutableOperationElement> CurrentlySelectedOperationTypeElement;
 	
 	/** Operation type we are using to search for tree nodes. Driven primarily by the UI */
-	mu::OP_TYPE OperationTypeToSearch = mu::OP_TYPE::NONE;
+	mu::EOpType OperationTypeToSearch = mu::EOpType::NONE;
 
 	/** Operation types present on the currently set mutable model.
 	 * The value of the TPair represents the amount of times operations using it have been found (not duplicates) */
-	TArray< TPair< mu::OP_TYPE,uint32>> ModelOperationTypes;
+	TArray< TPair< mu::EOpType,uint32>> ModelOperationTypes;
 	
 	/** Array with all the names for each of the operations available on mutable. Used by the UI  */
 	TArray<FString> ModelOperationTypeNames;
@@ -521,7 +521,7 @@ private:
 	 * the other. It also stores the found elements and other data.
 	 * @param InProgram The mutable program holding the data to be searched over
 	 */
-	void GetOperationsOfType(const mu::OP_TYPE& TargetOperationType,
+	void GetOperationsOfType(const mu::EOpType& TargetOperationType,
 	                         FElementsSearchCache& InSearchPayload,
 	                         const mu::FProgram& InProgram);
 
@@ -591,7 +591,7 @@ private:
 	/** Last search configuration -> Prevents recurrent searches with the same search parameters */
 	FString LastSearchedString;
 	bool bWasLastSearchRegEx = false;
-	TSharedPtr<mu::Model> LastSearchedModel;
+	TSharedPtr<mu::FModel> LastSearchedModel;
 
 	/** A collection with all the elements we have found to be related with the SearchString. */
 	TArray<TSharedPtr<FMutableCodeTreeElement>> NameBasedNavigationElements;
@@ -792,10 +792,10 @@ private:
 	 */
 
 public:
-	void PreviewMutableImage (mu::ImagePtrConst InImagePtr);
-	void PreviewMutableMesh (mu::MeshPtrConst InMeshPtr);
-	void PreviewMutableLayout(mu::Ptr<const mu::Layout> Layout);
-	void PreviewMutableSkeleton(mu::SkeletonPtrConst Skeleton);
+	void PreviewMutableImage (TSharedPtr<const mu::FImage> InImagePtr);
+	void PreviewMutableMesh (TSharedPtr<const mu::FMesh> InMeshPtr);
+	void PreviewMutableLayout(TSharedPtr<const mu::FLayout> Layout);
+	void PreviewMutableSkeleton(TSharedPtr<const mu::FSkeleton> Skeleton);
 	void PreviewMutableString(const FString& InString);
 	void PreviewMutableProjector(const mu::FProjector* Projector);
 	void PreviewMutableMatrix(const FMatrix44f& Mat);
@@ -815,7 +815,7 @@ private:
 class FMutableOperationElement : public TSharedFromThis<FMutableOperationElement>
 {
 public:
-	FMutableOperationElement(mu::OP_TYPE InOperationType,FText OperationTypeName,uint32 OperationTypeInstanceCount,FSlateColor OperationColor)
+	FMutableOperationElement(mu::EOpType InOperationType,FText OperationTypeName,uint32 OperationTypeInstanceCount,FSlateColor OperationColor)
 	{
 		OperationType = InOperationType;
 		OperationTextColor = OperationColor;
@@ -832,14 +832,14 @@ public:
 	}
 
 public:
-	mu::OP_TYPE OperationType;
+	mu::EOpType OperationType;
 	FText OperationTypeText;
 	FSlateColor OperationTextColor;
 
 	/** Get the type of visibility an slate representing this objects should have. */
 	EVisibility GetEntryVisibility() const
 	{
-		return OperationType == mu::OP_TYPE::NONE ? EVisibility::Collapsed : EVisibility::Visible;
+		return OperationType == mu::EOpType::NONE ? EVisibility::Collapsed : EVisibility::Visible;
 	}
 };
 
@@ -849,7 +849,7 @@ public:
 class FMutableCodeTreeElement : public TSharedFromThis<FMutableCodeTreeElement>
 {
 public:
-	FMutableCodeTreeElement(int32 InIndexOnTree, const int32& InMutableStateIndex, const TSharedPtr<mu::Model, ESPMode::ThreadSafe>& InModel, mu::OP::ADDRESS InOperation, const FString& InCaption, const FSlateColor InLabelColor, const TSharedPtr<FMutableCodeTreeElement>* InDuplicatedOf = nullptr);
+	FMutableCodeTreeElement(int32 InIndexOnTree, const int32& InMutableStateIndex, const TSharedPtr<mu::FModel, ESPMode::ThreadSafe>& InModel, mu::OP::ADDRESS InOperation, const FString& InCaption, const FSlateColor InLabelColor, const TSharedPtr<FMutableCodeTreeElement>* InDuplicatedOf = nullptr);
 	
 	void SetElementCurrentState(const int32& InStateIndex);
 
@@ -860,7 +860,7 @@ public:
 public:
 
 	/** Model this operation is part of */
-	TSharedPtr<mu::Model, ESPMode::ThreadSafe> MutableModel;
+	TSharedPtr<mu::FModel, ESPMode::ThreadSafe> MutableModel;
 	
 	/** Mutable Graph Node represented in this tree row. */
 	mu::OP::ADDRESS MutableOperation;
diff --git a/Engine/Plugins/Mutable/Source/CustomizableObjectEditor/Private/MuCOE/SMutableConstantsWidget.cpp b/Engine/Plugins/Mutable/Source/CustomizableObjectEditor/Private/MuCOE/SMutableConstantsWidget.cpp
--- a/Engine/Plugins/Mutable/Source/CustomizableObjectEditor/Private/MuCOE/SMutableConstantsWidget.cpp
+++ b/Engine/Plugins/Mutable/Source/CustomizableObjectEditor/Private/MuCOE/SMutableConstantsWidget.cpp
@@ -81,7 +81,7 @@ namespace
 	 * @param MeshPtr Pointer to the mesh whose channels we want to count
 	 * @return : the amount of channels found.
 	 */
-	int32 GetMeshChannelCount (mu::Ptr<const mu::Mesh> MeshPtr)
+	int32 GetMeshChannelCount (TSharedPtr<const mu::FMesh> MeshPtr)
 	{
 		check(MeshPtr)
 
@@ -135,7 +135,7 @@ public:
 		// The amount of mesh vertex buffer and image buffer channels
 		if (ColumnName == MeshConstantTitles::MeshChannels)
 		{
-			mu::Ptr<const mu::Mesh> Mesh = RowElement->MeshPtr.get();
+			TSharedPtr<const mu::FMesh> Mesh = RowElement->MeshPtr;
 			const int32 ChannelCount =  GetMeshChannelCount(Mesh);
 			const FText ChannelCountText = FText::AsNumber(ChannelCount);
 			return SNew(STextBlock).Text(ChannelCountText);
@@ -144,7 +144,7 @@ public:
 		// The amount of indices of the mesh
 		if (ColumnName == MeshConstantTitles::MeshIndices)
 		{
-			mu::Ptr<const mu::Mesh> Mesh = RowElement->MeshPtr.get();
+			TSharedPtr<const mu::FMesh> Mesh = RowElement->MeshPtr;
 			const FText IndexCount = FText::AsNumber(Mesh->GetIndexCount());
 			return SNew(STextBlock).Text(IndexCount);
 		}
@@ -152,7 +152,7 @@ public:
 		// The amount of vertices of the mesh
 		if (ColumnName == MeshConstantTitles::MeshVertices)
 		{
-			mu::Ptr<const mu::Mesh> Mesh = RowElement->MeshPtr.get();
+			TSharedPtr<const mu::FMesh> Mesh = RowElement->MeshPtr;
 			const FText VertexCount = FText::AsNumber(Mesh->GetVertexCount());
 			return SNew(STextBlock).Text(VertexCount);
 		}
@@ -160,7 +160,7 @@ public:
 		// Memory used by the mesh 
 		if (ColumnName == MeshConstantTitles::MeshMemory)
 		{
-			mu::Ptr<const mu::Mesh> Mesh = RowElement->MeshPtr.get();
+			TSharedPtr<const mu::FMesh> Mesh = RowElement->MeshPtr;
 			const FText SizeAsText = GenerateTextForSize(Mesh->GetDataSize());
 			return SNew(STextBlock).Text(SizeAsText);
 		}
@@ -244,7 +244,7 @@ public:
 	{
 		check(RowElement);
 
-		const mu::Image* Image = RowElement->ImagePtr.get();
+		const mu::FImage* Image = RowElement->ImagePtr.Get();
 		
 		// Index
 		if (ColumnName == ImageConstantTitles::ImageID)
@@ -1153,14 +1153,14 @@ void SMutableConstantsWidget::LoadConstantStrings()
 {
 	check (MutableProgramPtr);
 
-	const int32 ConstantsCount = MutableProgramPtr->m_constantStrings.Num();
+	const int32 ConstantsCount = MutableProgramPtr->ConstantStrings.Num();
 	ConstantStringElements.Empty(ConstantsCount);
 	uint64 ConstantStringsAccumulatedSize = 0;
 	
 	for (int32 StringAddressIndex = 0; StringAddressIndex < ConstantsCount; StringAddressIndex++)
 	{
 		TSharedPtr<FMutableConstantStringElement> ConstantStringElement = MakeShared<FMutableConstantStringElement>();
-		ConstantStringElement->MutableString = MutableProgramPtr->m_constantStrings[StringAddressIndex];
+		ConstantStringElement->MutableString = MutableProgramPtr->ConstantStrings[StringAddressIndex];
 		ConstantStringElement->IndexOnSourceVector = StringAddressIndex;
 		
 		// Cache resource size
@@ -1196,7 +1196,7 @@ void SMutableConstantsWidget::LoadConstantImages()
 		TSharedPtr<FMutableConstantImageElement> ConstantImageElement = MakeShared<FMutableConstantImageElement>();
 
 		MutableProgramPtr->GetConstant(ImageIndex, ConstantImageElement->ImagePtr, 0,
-			[this](int32 x, int32 y, int32 m, mu::EImageFormat f, mu::EInitializationType i) { return new mu::Image(x, y, m, f, i); });
+			[this](int32 x, int32 y, int32 m, mu::EImageFormat f, mu::EInitializationType i) { return MakeShared<mu::FImage>(x, y, m, f, i); });
 
 		ConstantImageElement->IndexOnSourceVector = ImageIndex;
 		
@@ -1233,8 +1233,8 @@ void SMutableConstantsWidget::LoadConstantMeshes()
 		ConstantMeshElement->IndexOnSourceVector = MeshIndex;
 
 		// Actual core disk size would be:
-		//mu::OutputMemoryStream Stream;
-		//mu::OutputArchive Archive{ &Stream };
+		//mu::FOutputMemoryStream Stream;
+		//mu::FOutputArchive Archive{ &Stream };
 		//ConstantMeshElement->MeshPtr->Serialise(Archive);
 		//uint64 OtherSize = Stream.GetBufferSize();
 
@@ -1251,16 +1251,16 @@ void SMutableConstantsWidget::LoadConstantMeshes()
 void SMutableConstantsWidget::LoadConstantLayouts()
 {
 	check (MutableProgramPtr);
-	const int32 ConstantsCount = MutableProgramPtr->m_constantLayouts.Num();
+	const int32 ConstantsCount = MutableProgramPtr->ConstantLayouts.Num();
 	ConstantLayoutElements.Empty(ConstantsCount);
 	
-	mu::OutputMemoryStream Stream;
-	mu::OutputArchive Archive{&Stream};
+	mu::FOutputMemoryStream Stream;
+	mu::FOutputArchive Archive{&Stream};
 	
 	for (int32 LayoutIndex = 0; LayoutIndex < ConstantsCount; LayoutIndex++)
 	{		
 		TSharedPtr<FMutableConstantLayoutElement> ConstantLayoutElement = MakeShared<FMutableConstantLayoutElement>();
-		ConstantLayoutElement->Layout = MutableProgramPtr->m_constantLayouts[LayoutIndex];
+		ConstantLayoutElement->Layout = MutableProgramPtr->ConstantLayouts[LayoutIndex];
 		ConstantLayoutElement->IndexOnSourceVector = LayoutIndex;
 		
 		// Cache resource for later GetBufferSize() call
@@ -1276,16 +1276,16 @@ void SMutableConstantsWidget::LoadConstantLayouts()
 void SMutableConstantsWidget::LoadConstantSkeletons()
 {
 	check (MutableProgramPtr);
-	const int32 ConstantsCount = MutableProgramPtr->m_constantSkeletons.Num();
+	const int32 ConstantsCount = MutableProgramPtr->ConstantSkeletons.Num();
 	ConstantSkeletonElements.Empty(ConstantsCount);
 
-	mu::OutputMemoryStream Stream;
-	mu::OutputArchive Archive{&Stream};
+	mu::FOutputMemoryStream Stream;
+	mu::FOutputArchive Archive{&Stream};
 	
 	for (int32 SkeletonIndex = 0; SkeletonIndex < ConstantsCount; SkeletonIndex++)
 	{
 		TSharedPtr<FMutableConstantSkeletonElement> ConstantSkeletonElement = MakeShared<FMutableConstantSkeletonElement>();
-		ConstantSkeletonElement->Skeleton = MutableProgramPtr->m_constantSkeletons[SkeletonIndex];
+		ConstantSkeletonElement->Skeleton = MutableProgramPtr->ConstantSkeletons[SkeletonIndex];
 		ConstantSkeletonElement->IndexOnSourceVector = SkeletonIndex;
 		
 		ConstantSkeletonElement->Skeleton->Serialise(Archive);
@@ -1300,16 +1300,16 @@ void SMutableConstantsWidget::LoadConstantSkeletons()
 void SMutableConstantsWidget::LoadConstantProjectors()
 {
 	check (MutableProgramPtr);
-	const int32 ConstantsCount = MutableProgramPtr->m_constantProjectors.Num();
+	const int32 ConstantsCount = MutableProgramPtr->ConstantProjectors.Num();
 	ConstantProjectorElements.Empty(ConstantsCount);
 
-	mu::OutputMemoryStream Stream;
-	mu::OutputArchive Archive{&Stream};
+	mu::FOutputMemoryStream Stream;
+	mu::FOutputArchive Archive{&Stream};
 	
 	for (int32 ProjectorIndex = 0; ProjectorIndex < ConstantsCount; ProjectorIndex++)
 	{
 		TSharedPtr<FMutableConstantProjectorElement> ConstantProjectorElement = MakeShared<FMutableConstantProjectorElement>();
-		ConstantProjectorElement->Projector = &(MutableProgramPtr->m_constantProjectors[ProjectorIndex]);
+		ConstantProjectorElement->Projector = &(MutableProgramPtr->ConstantProjectors[ProjectorIndex]);
 		ConstantProjectorElement->IndexOnSourceVector = ProjectorIndex;
 		
 		ConstantProjectorElement->Projector->Serialise(Archive);
@@ -1324,16 +1324,16 @@ void SMutableConstantsWidget::LoadConstantProjectors()
 void SMutableConstantsWidget::LoadConstantMatrices()
 {
 	check (MutableProgramPtr);
-	const int32 ConstantsCount = MutableProgramPtr->m_constantMatrices.Num();
+	const int32 ConstantsCount = MutableProgramPtr->ConstantMatrices.Num();
 	ConstantMatrixElements.Empty(ConstantsCount);
 
-	mu::OutputMemoryStream Stream;
-	mu::OutputArchive Archive{&Stream};
+	mu::FOutputMemoryStream Stream;
+	mu::FOutputArchive Archive{&Stream};
 	
 	for (int32 MatrixIndex = 0; MatrixIndex < ConstantsCount; MatrixIndex++)
 	{
 		TSharedPtr<FMutableConstantMatrixElement> ConstantMatrixElement = MakeShared<FMutableConstantMatrixElement>();
-		ConstantMatrixElement->Matrix = MutableProgramPtr->m_constantMatrices[MatrixIndex];
+		ConstantMatrixElement->Matrix = MutableProgramPtr->ConstantMatrices[MatrixIndex];
 		ConstantMatrixElement->IndexOnSourceVector = MatrixIndex;
 		
 		Archive << ConstantMatrixElement->Matrix;
@@ -1348,16 +1348,16 @@ void SMutableConstantsWidget::LoadConstantMatrices()
 void SMutableConstantsWidget::LoadConstantShapes()
 {
 	check (MutableProgramPtr);
-	const int32 ConstantsCount = MutableProgramPtr->m_constantShapes.Num();
+	const int32 ConstantsCount = MutableProgramPtr->ConstantShapes.Num();
 	ConstantShapeElements.Empty(ConstantsCount);
 
-	mu::OutputMemoryStream Stream;
-	mu::OutputArchive Archive{&Stream};	
+	mu::FOutputMemoryStream Stream;
+	mu::FOutputArchive Archive{&Stream};	
 	
 	for (int32 ShapeIndex = 0; ShapeIndex < ConstantsCount; ShapeIndex++)
 	{
 		TSharedPtr<FMutableConstantShapeElement> ConstantShapeElement = MakeShared<FMutableConstantShapeElement>();
-		ConstantShapeElement->Shape = &(MutableProgramPtr->m_constantShapes[ShapeIndex]);
+		ConstantShapeElement->Shape = &(MutableProgramPtr->ConstantShapes[ShapeIndex]);
 		ConstantShapeElement->IndexOnSourceVector = ShapeIndex;
 		
 		ConstantShapeElement->Shape->Serialise(Archive);
@@ -1375,8 +1375,8 @@ void SMutableConstantsWidget::LoadConstantCurves()
 	const int32 ConstantsCount = MutableProgramPtr->ConstantCurves.Num();
 	ConstantCurveElements.Empty(ConstantsCount);
 	
-	mu::OutputMemoryStream Stream;
-	mu::OutputArchive Archive{&Stream};
+	mu::FOutputMemoryStream Stream;
+	mu::FOutputArchive Archive{&Stream};
 	
 	for (int32 CurveIndex = 0; CurveIndex < ConstantsCount; CurveIndex++)
 	{
diff --git a/Engine/Plugins/Mutable/Source/CustomizableObjectEditor/Private/MuCOE/SMutableConstantsWidget.h b/Engine/Plugins/Mutable/Source/CustomizableObjectEditor/Private/MuCOE/SMutableConstantsWidget.h
--- a/Engine/Plugins/Mutable/Source/CustomizableObjectEditor/Private/MuCOE/SMutableConstantsWidget.h
+++ b/Engine/Plugins/Mutable/Source/CustomizableObjectEditor/Private/MuCOE/SMutableConstantsWidget.h
@@ -35,7 +35,7 @@ struct FMutableConstantElement
 */
 struct FMutableConstantImageElement : public FMutableConstantElement
 {
-	mu::ImagePtrConst ImagePtr = nullptr;
+	TSharedPtr<const mu::FImage> ImagePtr;
 };
 
 /**
@@ -43,7 +43,7 @@ struct FMutableConstantImageElement : public FMutableConstantElement
 */
 struct FMutableConstantMeshElement : public FMutableConstantElement
 {
-	mu::MeshPtrConst MeshPtr = nullptr;
+	TSharedPtr<const mu::FMesh> MeshPtr;
 };
 
 /**
@@ -59,7 +59,7 @@ struct FMutableConstantStringElement : public FMutableConstantElement
 */
 struct FMutableConstantLayoutElement : public FMutableConstantElement
 {
-	mu::Ptr<const mu::Layout> Layout = nullptr;
+	TSharedPtr<const mu::FLayout> Layout;
 };
 
 /**
@@ -67,7 +67,7 @@ struct FMutableConstantLayoutElement : public FMutableConstantElement
 */
 struct FMutableConstantSkeletonElement : public FMutableConstantElement
 {
-	mu::SkeletonPtrConst Skeleton = nullptr;
+	TSharedPtr<const mu::FSkeleton> Skeleton;
 };
 
 /**
diff --git a/Engine/Plugins/Mutable/Source/CustomizableObjectEditor/Private/MuCOE/SMutableGraphViewer.cpp b/Engine/Plugins/Mutable/Source/CustomizableObjectEditor/Private/MuCOE/SMutableGraphViewer.cpp
--- a/Engine/Plugins/Mutable/Source/CustomizableObjectEditor/Private/MuCOE/SMutableGraphViewer.cpp
+++ b/Engine/Plugins/Mutable/Source/CustomizableObjectEditor/Private/MuCOE/SMutableGraphViewer.cpp
@@ -655,9 +655,9 @@ void SMutableGraphViewer::GetChildrenForInfo(TSharedPtr<FMutableGraphTreeElement
 		mu::NodeImageSwitch* ImageSwitchVar = StaticCast<mu::NodeImageSwitch*>(ParentNode);
 		mu::NodeImageSwitch::Private* Private = ImageSwitchVar->GetPrivate();
 		AddChildFunc(Private->m_pParameter.get(), FString::Printf(TEXT("PARAM")));
-		for (int32 OptionIndex = 0; OptionIndex < Private->m_options.Num(); OptionIndex++)
+		for (int32 OptionIndex = 0; OptionIndex < Private->Options.Num(); OptionIndex++)
 		{
-			AddChildFunc(Private->m_options[OptionIndex].get(), FString::Printf(TEXT("OPTION [%d]"), OptionIndex));
+			AddChildFunc(Private->Options[OptionIndex].get(), FString::Printf(TEXT("OPTION [%d]"), OptionIndex));
 		}
 	}
 
@@ -731,9 +731,9 @@ void SMutableGraphViewer::GetChildrenForInfo(TSharedPtr<FMutableGraphTreeElement
 	{
 		mu::NodeScalarEnumParameter* ScalarEnumParameterVar = StaticCast<mu::NodeScalarEnumParameter*>(ParentNode);
 		mu::NodeScalarEnumParameter::Private* Private = ScalarEnumParameterVar->GetPrivate();
-		for (int32 RangeIndex = 0; RangeIndex < Private->m_ranges.Num(); RangeIndex++)
+		for (int32 RangeIndex = 0; RangeIndex < Private->Ranges.Num(); RangeIndex++)
 		{
-			AddChildFunc(Private->m_ranges[RangeIndex].get(), FString::Printf(TEXT("RANGE [%d]"), RangeIndex));
+			AddChildFunc(Private->Ranges[RangeIndex].get(), FString::Printf(TEXT("RANGE [%d]"), RangeIndex));
 		}
 	}
 
@@ -771,9 +771,9 @@ void SMutableGraphViewer::GetChildrenForInfo(TSharedPtr<FMutableGraphTreeElement
 	{
 		mu::NodeScalarParameter* ScalarParameterVar = StaticCast<mu::NodeScalarParameter*>(ParentNode);
 		mu::NodeScalarParameter::Private* Private = ScalarParameterVar->GetPrivate();
-		for (int32 RangeIndex = 0; RangeIndex < Private->m_ranges.Num(); RangeIndex++)
+		for (int32 RangeIndex = 0; RangeIndex < Private->Ranges.Num(); RangeIndex++)
 		{
-			AddChildFunc(Private->m_ranges[RangeIndex].get(), FString::Printf(TEXT("RANGE [%d]"), RangeIndex));
+			AddChildFunc(Private->Ranges[RangeIndex].get(), FString::Printf(TEXT("RANGE [%d]"), RangeIndex));
 		}
 	}
 
@@ -816,9 +816,9 @@ void SMutableGraphViewer::GetChildrenForInfo(TSharedPtr<FMutableGraphTreeElement
 	{
 		mu::NodeProjectorParameter* ProjectorParameterVar = StaticCast<mu::NodeProjectorParameter*>(ParentNode);
 		mu::NodeProjectorParameter::Private* Private = ProjectorParameterVar->GetPrivate();
-		for (int32 RangeIndex = 0; RangeIndex < Private->m_ranges.Num(); RangeIndex++)
+		for (int32 RangeIndex = 0; RangeIndex < Private->Ranges.Num(); RangeIndex++)
 		{
-			AddChildFunc(Private->m_ranges[RangeIndex].get(), FString::Printf(TEXT("RANGE [%d]"), RangeIndex));
+			AddChildFunc(Private->Ranges[RangeIndex].get(), FString::Printf(TEXT("RANGE [%d]"), RangeIndex));
 		}
 	}
 
diff --git a/Engine/Plugins/Mutable/Source/CustomizableObjectEditor/Private/MuCOE/SMutableImageViewer.cpp b/Engine/Plugins/Mutable/Source/CustomizableObjectEditor/Private/MuCOE/SMutableImageViewer.cpp
--- a/Engine/Plugins/Mutable/Source/CustomizableObjectEditor/Private/MuCOE/SMutableImageViewer.cpp
+++ b/Engine/Plugins/Mutable/Source/CustomizableObjectEditor/Private/MuCOE/SMutableImageViewer.cpp
@@ -262,7 +262,7 @@ void SMutableImageViewer::Tick(const FGeometry& AllottedGeometry, const double I
 		{
 			// If the image format is not supported in the editor platform (for instance, it is a mobile-specific format)
 			// convert the mutable image to a generic format so that it can be previewed.
-			mu::Ptr<const mu::Image> ImageToConvert = MutableImage;
+			TSharedPtr<const mu::FImage> ImageToConvert = MutableImage;
 
 			// \TODO: This test is a bit weak
 			if (MutableImage->GetFormat() >= mu::EImageFormat::ASTC_4x4_RGB_LDR
@@ -271,7 +271,7 @@ void SMutableImageViewer::Tick(const FGeometry& AllottedGeometry, const double I
 			{
 				mu::FImageOperator ImOp = mu::FImageOperator::GetDefault(mu::FImageOperator::FImagePixelFormatFunc());
 				int32 Quality = 4;
-				ImageToConvert = ImOp.ImagePixelFormat(Quality, MutableImage.get(), mu::GetUncompressedFormat(MutableImage->GetFormat()) );
+				ImageToConvert = ImOp.ImagePixelFormat(Quality, MutableImage.Get(), mu::GetUncompressedFormat(MutableImage->GetFormat()) );
 			}
 
 			// Should we extract a single channel?
@@ -300,7 +300,7 @@ void SMutableImageViewer::Tick(const FGeometry& AllottedGeometry, const double I
 }
 
 
-void SMutableImageViewer::SetImage(const mu::ImagePtrConst& InMutableImage, int32 CurrentLOD)
+void SMutableImageViewer::SetImage(const TSharedPtr<const mu::FImage>& InMutableImage, int32 CurrentLOD)
 {
 	MutableImage = InMutableImage;
 	CurrentVisibleLOD = CurrentLOD;
diff --git a/Engine/Plugins/Mutable/Source/CustomizableObjectEditor/Private/MuCOE/SMutableImageViewer.h b/Engine/Plugins/Mutable/Source/CustomizableObjectEditor/Private/MuCOE/SMutableImageViewer.h
--- a/Engine/Plugins/Mutable/Source/CustomizableObjectEditor/Private/MuCOE/SMutableImageViewer.h
+++ b/Engine/Plugins/Mutable/Source/CustomizableObjectEditor/Private/MuCOE/SMutableImageViewer.h
@@ -78,7 +78,7 @@ public:
 
 	SLATE_BEGIN_ARGS(SMutableImageViewer) {}
 		SLATE_ATTRIBUTE(FIntPoint, GridSize)
-		SLATE_ARGUMENT(mu::ImagePtrConst, Image)
+		SLATE_ARGUMENT(TSharedPtr<const mu::FImage>, Image)
 	SLATE_END_ARGS()
 
 	void Construct(const FArguments& InArgs);
@@ -89,7 +89,7 @@ public:
 	// Own interface
 
 	/** Set the image to show in the widget. */
-	void SetImage(const mu::ImagePtrConst& Image, int32 VisibleLOD );
+	void SetImage(const TSharedPtr<const mu::FImage>& Image, int32 VisibleLOD );
 
 private:
 
@@ -100,7 +100,7 @@ private:
 	TAttribute<FIntPoint> GridSize;
 
 	/** */
-	mu::ImagePtrConst MutableImage;
+	TSharedPtr<const mu::FImage> MutableImage;
 
 	/** Mipmap shown in the widget. 0 is the biggest one. */
 	int32 CurrentVisibleLOD=0;
diff --git a/Engine/Plugins/Mutable/Source/CustomizableObjectEditor/Private/MuCOE/SMutableInstanceViewer.cpp b/Engine/Plugins/Mutable/Source/CustomizableObjectEditor/Private/MuCOE/SMutableInstanceViewer.cpp
--- a/Engine/Plugins/Mutable/Source/CustomizableObjectEditor/Private/MuCOE/SMutableInstanceViewer.cpp
+++ b/Engine/Plugins/Mutable/Source/CustomizableObjectEditor/Private/MuCOE/SMutableInstanceViewer.cpp
@@ -277,7 +277,7 @@ void SMutableInstanceViewer::Construct(const FArguments& InArgs)
 }
 
 
-void SMutableInstanceViewer::SetInstance(const mu::InstancePtrConst& InInstance, const TSharedPtr<const mu::Model>& Model, const mu::Ptr<const mu::Parameters>& Parameters, const mu::Ptr<mu::System>& System)
+void SMutableInstanceViewer::SetInstance(const TSharedPtr<const mu::FInstance>& InInstance, const TSharedPtr<const mu::FModel>& Model, const TSharedPtr<const mu::FParameters>& Parameters, const mu::FSystem& System)
 {
 	if (InInstance != MutableInstance)
 	{
@@ -312,7 +312,7 @@ void SMutableInstanceViewer::SetInstance(const mu::InstancePtrConst& InInstance,
 					// Generate the mesh to get some data
 					{
 						mu::OP::ADDRESS MeshRootAddress = mu::GetResourceIDRoot(MeshId);
-						mu::Ptr<const mu::Mesh> MutableMesh = System->GetPrivate()->BuildMesh(Model, Parameters.get(), MeshRootAddress);
+						TSharedPtr<const mu::FMesh> MutableMesh = System.GetPrivate()->BuildMesh(Model, Parameters.Get(), MeshRootAddress);
 						if (MutableMesh && MutableMesh->IsReference())
 						{
 							NewLOD->MeshId = FText::Format(
diff --git a/Engine/Plugins/Mutable/Source/CustomizableObjectEditor/Private/MuCOE/SMutableInstanceViewer.h b/Engine/Plugins/Mutable/Source/CustomizableObjectEditor/Private/MuCOE/SMutableInstanceViewer.h
--- a/Engine/Plugins/Mutable/Source/CustomizableObjectEditor/Private/MuCOE/SMutableInstanceViewer.h
+++ b/Engine/Plugins/Mutable/Source/CustomizableObjectEditor/Private/MuCOE/SMutableInstanceViewer.h
@@ -59,7 +59,7 @@ public:
 	void Construct(const FArguments& InArgs);
 
 	/** Set the Mutable Instance to be used for this widget */
-	void SetInstance(const mu::Ptr<const mu::Instance>&, const TSharedPtr<const mu::Model>&, const mu::Ptr<const mu::Parameters>&, const mu::Ptr<mu::System>&);
+	void SetInstance(const TSharedPtr<const mu::FInstance>&, const TSharedPtr<const mu::FModel>&, const TSharedPtr<const mu::FParameters>&, const mu::FSystem&);
 
 	
 	/** */
@@ -69,7 +69,7 @@ public:
 private:
 
 	/** Data backend for the widget. It represents the Instance that is being "displayed" */
-	mu::Ptr<const mu::Instance> MutableInstance = nullptr;
+	TSharedPtr<const mu::FInstance> MutableInstance = nullptr;
 
 	/** Splitter used to separate the two sides of the slate (tables and viewport) */
 	TSharedPtr<SSplitter> SpaceSplitter;
diff --git a/Engine/Plugins/Mutable/Source/CustomizableObjectEditor/Private/MuCOE/SMutableLayoutViewer.cpp b/Engine/Plugins/Mutable/Source/CustomizableObjectEditor/Private/MuCOE/SMutableLayoutViewer.cpp
--- a/Engine/Plugins/Mutable/Source/CustomizableObjectEditor/Private/MuCOE/SMutableLayoutViewer.cpp
+++ b/Engine/Plugins/Mutable/Source/CustomizableObjectEditor/Private/MuCOE/SMutableLayoutViewer.cpp
@@ -120,7 +120,7 @@ void SMutableLayoutViewer::Tick(const FGeometry& AllottedGeometry, const double
 }
 
 
-void SMutableLayoutViewer::SetLayout(const mu::Ptr<const mu::Layout>& InMutableLayout)
+void SMutableLayoutViewer::SetLayout(const TSharedPtr<const mu::FLayout>& InMutableLayout)
 {
 	MutableLayout = InMutableLayout;
 	bIsPendingUpdate = true;
diff --git a/Engine/Plugins/Mutable/Source/CustomizableObjectEditor/Private/MuCOE/SMutableLayoutViewer.h b/Engine/Plugins/Mutable/Source/CustomizableObjectEditor/Private/MuCOE/SMutableLayoutViewer.h
--- a/Engine/Plugins/Mutable/Source/CustomizableObjectEditor/Private/MuCOE/SMutableLayoutViewer.h
+++ b/Engine/Plugins/Mutable/Source/CustomizableObjectEditor/Private/MuCOE/SMutableLayoutViewer.h
@@ -17,7 +17,7 @@ public:
 
 	SLATE_BEGIN_ARGS(SMutableLayoutViewer) {}
 		SLATE_ATTRIBUTE(FIntPoint, GridSize)
-		SLATE_ARGUMENT(mu::Ptr<const mu::Layout>, Layout)
+		SLATE_ARGUMENT(TSharedPtr<const mu::FLayout>, Layout)
 	SLATE_END_ARGS()
 
 	void Construct(const FArguments& InArgs);
@@ -28,7 +28,7 @@ public:
 	// Own interface
 
 	/** Set the image to show in the widget. */
-	void SetLayout(const mu::Ptr<const mu::Layout>& Image );
+	void SetLayout(const TSharedPtr<const mu::FLayout>& Image );
 
 private:
 
@@ -36,7 +36,7 @@ private:
 	TSharedPtr<SCustomizableObjectLayoutGrid> LayoutViewer;
 
 	/** Mutable Layout object being previewed */
-	mu::Ptr<const mu::Layout> MutableLayout;
+	TSharedPtr<const mu::FLayout> MutableLayout;
 
 	/** Is true, the image or the visible LOD have changed and we need to update. */
 	bool bIsPendingUpdate = false;
diff --git a/Engine/Plugins/Mutable/Source/CustomizableObjectEditor/Private/MuCOE/SMutableMeshViewer.cpp b/Engine/Plugins/Mutable/Source/CustomizableObjectEditor/Private/MuCOE/SMutableMeshViewer.cpp
--- a/Engine/Plugins/Mutable/Source/CustomizableObjectEditor/Private/MuCOE/SMutableMeshViewer.cpp
+++ b/Engine/Plugins/Mutable/Source/CustomizableObjectEditor/Private/MuCOE/SMutableMeshViewer.cpp
@@ -198,7 +198,7 @@ void SMutableMeshViewer::Construct(const FArguments& InArgs)
 }
 
 
-void SMutableMeshViewer::SetMesh(const mu::MeshPtrConst& InMesh)
+void SMutableMeshViewer::SetMesh(const TSharedPtr<const mu::FMesh>& InMesh)
 {
 	if (InMesh != MutableMesh)
 	{
@@ -217,7 +217,7 @@ void SMutableMeshViewer::SetMesh(const mu::MeshPtrConst& InMesh)
 			OnMeshChanged();
 
 			// Skeleton slate viewer update
-			if (mu::SkeletonPtrConst MutableSkeleton = MutableMesh->GetSkeleton())
+			if (TSharedPtr<const mu::FSkeleton> MutableSkeleton = MutableMesh->GetSkeleton())
 			{
 				MutableSkeletonViewer->SetSkeleton(MutableSkeleton);
 
@@ -590,7 +590,7 @@ TSharedRef<SWidget> SMutableMeshViewer::GenerateBufferChannelsListView(
 
 void SMutableMeshViewer::OnMeshChanged()
 {
-	// Cache the data accessible from the mu::Mesh to be later used by the UI
+	// Cache the data accessible from the mu::FMesh to be later used by the UI
 	FillTargetBufferSetDataArray(MutableMesh->GetVertexBuffers(), VertexBuffers, VertexBuffersSlateView);
 	FillTargetBufferSetDataArray(MutableMesh->GetIndexBuffers(), IndexBuffers, IndexBuffersSlateView);
 
diff --git a/Engine/Plugins/Mutable/Source/CustomizableObjectEditor/Private/MuCOE/SMutableMeshViewer.h b/Engine/Plugins/Mutable/Source/CustomizableObjectEditor/Private/MuCOE/SMutableMeshViewer.h
--- a/Engine/Plugins/Mutable/Source/CustomizableObjectEditor/Private/MuCOE/SMutableMeshViewer.h
+++ b/Engine/Plugins/Mutable/Source/CustomizableObjectEditor/Private/MuCOE/SMutableMeshViewer.h
@@ -41,19 +41,19 @@ struct FBufferElement
 	TSharedPtr<TArray<TSharedPtr<FBufferChannelElement>>> BufferChannels;
 };
 
-/** Widget designed to show the statistical data from a Mutable Mesh*/
+/** Widget designed to show the statistical data from a Mutable FMesh*/
 class SMutableMeshViewer : public SCompoundWidget
 {
 public:
 	SLATE_BEGIN_ARGS(SMutableMeshViewer) {}
-		SLATE_ARGUMENT_DEFAULT(mu::MeshPtrConst, Mesh){nullptr};
+		SLATE_ARGUMENT_DEFAULT(TSharedPtr<const mu::FMesh>, Mesh){nullptr};
 	SLATE_END_ARGS()
 
 	/** Builds the widget */
 	void Construct(const FArguments& InArgs);
 
 	/** Set the Mutable Mesh to be used for this widget */
-	void SetMesh(const mu::MeshPtrConst& InMesh);
+	void SetMesh(const TSharedPtr<const mu::FMesh>& InMesh);
 
 	
 	/** Generate a new SListView for all the ChannelElements provided.
@@ -73,7 +73,7 @@ private:
 	TSharedRef<ITableRow> GenerateTagRow(TSharedPtr<FString> InItem, const TSharedRef<STableViewBase>& OwnerTable);
 
 	/** Data backend for the widget. It represents the mesh that is being "displayed" */
-	mu::MeshPtrConst MutableMesh = nullptr;
+	TSharedPtr<const mu::FMesh> MutableMesh = nullptr;
 
 	/** Splitter used to separate the two sides of the slate (tables and viewport) */
 	TSharedPtr<SSplitter> SpaceSplitter;
diff --git a/Engine/Plugins/Mutable/Source/CustomizableObjectEditor/Private/MuCOE/SMutableMeshViewport.cpp b/Engine/Plugins/Mutable/Source/CustomizableObjectEditor/Private/MuCOE/SMutableMeshViewport.cpp
--- a/Engine/Plugins/Mutable/Source/CustomizableObjectEditor/Private/MuCOE/SMutableMeshViewport.cpp
+++ b/Engine/Plugins/Mutable/Source/CustomizableObjectEditor/Private/MuCOE/SMutableMeshViewport.cpp
@@ -33,7 +33,7 @@ void SMutableMeshViewport::Tick(const FGeometry& AllottedGeometry, const double
 	SEditorViewport::Tick(AllottedGeometry, InCurrentTime, InDeltaTime);
 }
 
-void SMutableMeshViewport::SetMesh(const mu::MeshPtrConst& InMesh)
+void SMutableMeshViewport::SetMesh(const TSharedPtr<const mu::FMesh>& InMesh)
 {
 	if (InMesh == MutableMesh)
 	{
@@ -111,7 +111,7 @@ void SMutableMeshViewport::SendMeshToViewport()
 bool SMutableMeshViewport::GenerateUnrealMesh()
 {
 	// Convert the provided mutable mesh to unreal with heavy aid of the Reference Skeletal Mesh
-	USkeletalMesh* GeneratedSkeletalMesh = MutableMeshPreviewUtils::GenerateSkeletalMeshFromMutableMesh(MutableMesh);
+	USkeletalMesh* GeneratedSkeletalMesh = MutableMeshPreviewUtils::GenerateSkeletalMeshFromMutableMesh(MutableMesh.Get());
 	
 	if (!GeneratedSkeletalMesh)
 	{
diff --git a/Engine/Plugins/Mutable/Source/CustomizableObjectEditor/Private/MuCOE/SMutableMeshViewport.h b/Engine/Plugins/Mutable/Source/CustomizableObjectEditor/Private/MuCOE/SMutableMeshViewport.h
--- a/Engine/Plugins/Mutable/Source/CustomizableObjectEditor/Private/MuCOE/SMutableMeshViewport.h
+++ b/Engine/Plugins/Mutable/Source/CustomizableObjectEditor/Private/MuCOE/SMutableMeshViewport.h
@@ -22,7 +22,7 @@ struct FGeometry;
 class SMutableMeshViewport : public SEditorViewport , public FGCObject
 {
 	SLATE_BEGIN_ARGS(SMutableMeshViewport) {}
-    	SLATE_ARGUMENT_DEFAULT(mu::MeshPtrConst, Mesh) { nullptr };
+    	SLATE_ARGUMENT_DEFAULT(TSharedPtr<const mu::FMesh>, Mesh) { nullptr };
     SLATE_END_ARGS()
     
 public:
@@ -32,7 +32,7 @@ public:
     /** Set the Mutable Mesh to be used for this widget
      * @param InMesh - The mutable mesh to be displayed
      */
-    void SetMesh(const mu::MeshPtrConst& InMesh);
+    void SetMesh(const TSharedPtr<const mu::FMesh>& InMesh);
 
 	// FGCObject interface
 	virtual void AddReferencedObjects(FReferenceCollector& Collector) override;
@@ -52,7 +52,7 @@ protected:
 private:
 
 	/** The mutable mesh being displayed */
-	mu::MeshPtrConst MutableMesh = nullptr;
+	TSharedPtr<const mu::FMesh> MutableMesh = nullptr;
 
 	/*
 	 *	UI  
