[mutable] More enforcement of coding standards.

diff --git a/Engine/Plugins/Mutable/Source/CustomizableObjectEditor/Private/MuCOE/SMutableParametersWidget.cpp b/Engine/Plugins/Mutable/Source/CustomizableObjectEditor/Private/MuCOE/SMutableParametersWidget.cpp
--- a/Engine/Plugins/Mutable/Source/CustomizableObjectEditor/Private/MuCOE/SMutableParametersWidget.cpp
+++ b/Engine/Plugins/Mutable/Source/CustomizableObjectEditor/Private/MuCOE/SMutableParametersWidget.cpp
@@ -101,9 +101,8 @@ void SMutableParametersWidget::Tick(const FGeometry& AllottedGeometry, const dou
 		
 		
 		// Use the parameter descriptor to know what parameter value determines the amount of values this one has
-		mu::FParameterDesc& ParameterDescription = MutableParameters->GetPrivate()->m_pModel->GetPrivate()->m_program.
-															m_parameters[ParamIndex];
-		TArray<uint32> Ranges = ParameterDescription.m_ranges;
+		mu::FParameterDesc& ParameterDescription = MutableParameters->GetPrivate()->Model->GetPrivate()->Program.Parameters[ParamIndex];
+		TArray<uint32> Ranges = ParameterDescription.Ranges;
 		if (Ranges.Num() >= 1)
 		{
 			if (Ranges.Num() > 1)
@@ -122,7 +121,7 @@ void SMutableParametersWidget::Tick(const FGeometry& AllottedGeometry, const dou
 			// grab only first value (first dimension)
 			const uint32 RangeIdx = Ranges[Dimension];
 			// Get the parameter that controls the amount of values of current parameter
-			const int32 ParameterControllingValueCount = MutableParameters->GetPrivate()->m_pModel->GetPrivate()->m_program.m_ranges[RangeIdx].m_dimensionParameter;
+			const int32 ParameterControllingValueCount = MutableParameters->GetPrivate()->Model->GetPrivate()->Program.Ranges[RangeIdx].DimensionParameter;
 
 			// If a parameter index has been reported then we know that another parameter's value determines the amount of values we have
 			if (ParameterControllingValueCount >= 0)
@@ -159,7 +158,7 @@ void SMutableParametersWidget::Tick(const FGeometry& AllottedGeometry, const dou
 			// Iterate over all the values for the parameter and show them. Updating their value is managed by the UI
 			for (int32 Position = 0; Position < PositionCount; Position++)
 			{
-				mu::RangeIndexPtr RangeIndex = MutableParameters->NewRangeIndex(ParamIndex);
+				TSharedPtr<mu::FRangeIndex> RangeIndex = MutableParameters->NewRangeIndex(ParamIndex);
 				RangeIndex->SetPosition(Dimension,Position);
 			
 				// Show position data if more than one is set to avoid UI clutter
@@ -216,8 +215,8 @@ void SMutableParametersWidget::ScheduleUpdateIfRequired(const int32& InParameter
 {
 	// Update the ui to display the new amount of values we want to show for the multidimensional values if this parameter's value
 	// controls the amount of values another parameter has.
-	const bool bIsRangeSize = MutableParameters->GetPrivate()->m_pModel->GetPrivate()->m_program.m_ranges.ContainsByPredicate([InParameterIndex](const mu::FRangeDesc& r)
-		{ return r.m_dimensionParameter == InParameterIndex; });
+	const bool bIsRangeSize = MutableParameters->GetPrivate()->Model->GetPrivate()->Program.Ranges.ContainsByPredicate([InParameterIndex](const mu::FRangeDesc& r)
+		{ return r.DimensionParameter == InParameterIndex; });
 
 	// Schedule update since we know the value does change the amount of values shown by another parameter (multivalue parameter)
 	if (bIsRangeSize && !bIsPendingUpdate)
@@ -227,7 +226,7 @@ void SMutableParametersWidget::ScheduleUpdateIfRequired(const int32& InParameter
 }
 
 
-void SMutableParametersWidget::GenerateAndAttachParameterSlate(const int32 ParamIndex, TSharedPtr<SHorizontalBox> ParameterHorizontalBox,  mu::RangeIndexPtrConst RangeIndex)
+void SMutableParametersWidget::GenerateAndAttachParameterSlate(const int32 ParamIndex, TSharedPtr<SHorizontalBox> ParameterHorizontalBox, const TSharedPtr<const mu::FRangeIndex>& RangeIndex)
 {
 	switch (MutableParameters->GetType(ParamIndex) )
 	{
@@ -290,7 +289,7 @@ void SMutableParametersWidget::GenerateAndAttachParameterSlate(const int32 Param
 
 			OptionNamesAttribute.Empty();
 
-			const int32 Value = MutableParameters->GetIntValue( ParamIndex , RangeIndex);
+			const int32 Value = MutableParameters->GetIntValue( ParamIndex, RangeIndex.Get());
 			const int32 ValueIndex = MutableParameters->GetIntValueIndex( ParamIndex, Value );
 			for (int32 i = 0; i < ValueCount; ++i)
 			{
@@ -315,7 +314,7 @@ void SMutableParametersWidget::GenerateAndAttachParameterSlate(const int32 Param
 						SNew(STextBlock)
 						.Text_Lambda([this, ParamIndex, RangeIndex, OptionNamesAttribute]() -> FText
 						{
-							const int32 Value = MutableParameters ? MutableParameters->GetIntValue(ParamIndex, RangeIndex) : 0;
+							const int32 Value = MutableParameters ? MutableParameters->GetIntValue(ParamIndex, RangeIndex.Get()) : 0;
 							const int32 ValueIndex = MutableParameters ? MutableParameters->GetIntValueIndex(ParamIndex, Value) : 0;
 
 							return FText::FromString(*OptionNamesAttribute[ValueIndex]);
@@ -428,7 +427,7 @@ EVisibility SMutableParametersWidget::GetParameterVisibility(int32 ParamIndex) c
 }
 
 
-ECheckBoxState SMutableParametersWidget::GetBoolParameterValue(int32 ParamIndex,  mu::RangeIndexPtrConst RangeIndex) const
+ECheckBoxState SMutableParametersWidget::GetBoolParameterValue(int32 ParamIndex, TSharedPtr<const mu::FRangeIndex> RangeIndex) const
 {
 	if ( !MutableParameters 
 		|| ParamIndex>=MutableParameters->GetCount()
@@ -437,11 +436,11 @@ ECheckBoxState SMutableParametersWidget::GetBoolParameterValue(int32 ParamIndex,
 		return ECheckBoxState::Undetermined;
 	}	
 	
-	return MutableParameters->GetBoolValue(ParamIndex,RangeIndex) ? ECheckBoxState::Checked : ECheckBoxState::Unchecked;
+	return MutableParameters->GetBoolValue(ParamIndex,RangeIndex.Get()) ? ECheckBoxState::Checked : ECheckBoxState::Unchecked;
 }
 
 
-void SMutableParametersWidget::OnBoolParameterChanged(ECheckBoxState InCheckboxState, int32 ParamIndex,  mu::RangeIndexPtrConst RangeIndex)
+void SMutableParametersWidget::OnBoolParameterChanged(ECheckBoxState InCheckboxState, int32 ParamIndex,  TSharedPtr<const mu::FRangeIndex> RangeIndex)
 {
 	if (!MutableParameters
 		|| ParamIndex >= MutableParameters->GetCount()
@@ -450,13 +449,13 @@ void SMutableParametersWidget::OnBoolParameterChanged(ECheckBoxState InCheckboxS
 		return;
 	}
 
-	MutableParameters->SetBoolValue( ParamIndex, InCheckboxState==ECheckBoxState::Checked, RangeIndex );
+	MutableParameters->SetBoolValue( ParamIndex, InCheckboxState==ECheckBoxState::Checked, RangeIndex.Get() );
 
 	OnParametersValueChanged.ExecuteIfBound(ParamIndex);
 }
 
 
-TOptional<float> SMutableParametersWidget::GetFloatParameterValue(int32 ParamIndex, mu::RangeIndexPtrConst RangeIndex) const
+TOptional<float> SMutableParametersWidget::GetFloatParameterValue(int32 ParamIndex, TSharedPtr<const mu::FRangeIndex> RangeIndex) const
 {
 	if (!MutableParameters
 		|| ParamIndex >= MutableParameters->GetCount()
@@ -465,11 +464,11 @@ TOptional<float> SMutableParametersWidget::GetFloatParameterValue(int32 ParamInd
 		return TOptional<float>();
 	}
 
-	return MutableParameters->GetFloatValue(ParamIndex,RangeIndex);
+	return MutableParameters->GetFloatValue(ParamIndex, RangeIndex.Get());
 }
 
 
-void SMutableParametersWidget::OnFloatParameterChanged(float InValue, int32 ParamIndex, mu::RangeIndexPtrConst RangeIndex)
+void SMutableParametersWidget::OnFloatParameterChanged(float InValue, int32 ParamIndex, TSharedPtr<const mu::FRangeIndex> RangeIndex)
 {
 	if (!MutableParameters
 		|| ParamIndex >= MutableParameters->GetCount()
@@ -478,7 +477,7 @@ void SMutableParametersWidget::OnFloatParameterChanged(float InValue, int32 Para
 		return;
 	}
 
-	MutableParameters->SetFloatValue(ParamIndex, InValue, RangeIndex);
+	MutableParameters->SetFloatValue(ParamIndex, InValue, RangeIndex.Get());
 	
 	OnParametersValueChanged.ExecuteIfBound(ParamIndex);
 }
@@ -493,7 +492,7 @@ void SMutableParametersWidget::OnFloatParameterCommitted(float Value, ETextCommi
 }
 
 
-TOptional<int32> SMutableParametersWidget::GetIntParameterValue (int32 ParamIndex, mu::RangeIndexPtrConst RangeIndex) const
+TOptional<int32> SMutableParametersWidget::GetIntParameterValue (int32 ParamIndex, TSharedPtr<const mu::FRangeIndex> RangeIndex) const
 {
 	if (!MutableParameters
 		|| ParamIndex >= MutableParameters->GetCount()
@@ -502,7 +501,7 @@ TOptional<int32> SMutableParametersWidget::GetIntParameterValue (int32 ParamInde
 		return TOptional<int32>();
 	}
 
-	int32 Value = MutableParameters->GetIntValue(ParamIndex, RangeIndex);
+	int32 Value = MutableParameters->GetIntValue(ParamIndex, RangeIndex.Get());
 	int32 ValueIndex = MutableParameters->GetIntValueIndex(ParamIndex, Value);
 
 	return ValueIndex;
@@ -529,7 +528,7 @@ TOptional<int32> SMutableParametersWidget::GetIntParameterValueMax(int32 ParamIn
 }
 
 
-void SMutableParametersWidget::OnIntParameterChanged(int32 InValue, int32 ParamIndex, TSharedPtr<SSearchableComboBox> Combo, mu::RangeIndexPtrConst RangeIndex )
+void SMutableParametersWidget::OnIntParameterChanged(int32 InValue, int32 ParamIndex, TSharedPtr<SSearchableComboBox> Combo, TSharedPtr<const mu::FRangeIndex> RangeIndex )
 {
 	if (!MutableParameters
 		|| ParamIndex >= MutableParameters->GetCount()
@@ -542,7 +541,7 @@ void SMutableParametersWidget::OnIntParameterChanged(int32 InValue, int32 ParamI
 	if (InValue >= 0 && InValue < MutableParameters->GetIntPossibleValueCount(ParamIndex))
 	{
 		int32 RealValue = MutableParameters->GetIntPossibleValue(ParamIndex, InValue);
-		MutableParameters->SetIntValue(ParamIndex, RealValue,RangeIndex);
+		MutableParameters->SetIntValue(ParamIndex, RealValue, RangeIndex.Get());
 		
 		// Update the text combo if any
 		if (Combo)
@@ -559,7 +558,7 @@ void SMutableParametersWidget::OnIntParameterChanged(int32 InValue, int32 ParamI
 }
 
 
-void SMutableParametersWidget::OnIntParameterTextChanged(TSharedPtr<FString> Selection, ESelectInfo::Type, int32 ParamIndex, mu::RangeIndexPtrConst RangeIndex)
+void SMutableParametersWidget::OnIntParameterTextChanged(TSharedPtr<FString> Selection, ESelectInfo::Type, int32 ParamIndex, TSharedPtr<const mu::FRangeIndex> RangeIndex)
 {
 	if (!Selection
 		|| !MutableParameters
@@ -573,7 +572,7 @@ void SMutableParametersWidget::OnIntParameterTextChanged(TSharedPtr<FString> Sel
 	if (ValueIndex >= 0 && ValueIndex < MutableParameters->GetIntPossibleValueCount(ParamIndex))
 	{
 		int32 RealValue = MutableParameters->GetIntPossibleValue(ParamIndex, ValueIndex);
-		MutableParameters->SetIntValue(ParamIndex, RealValue,RangeIndex);
+		MutableParameters->SetIntValue(ParamIndex, RealValue, RangeIndex.Get());
 		
 		OnParametersValueChanged.ExecuteIfBound(ParamIndex);
 		
@@ -588,7 +587,7 @@ TSharedRef<SWidget> SMutableParametersWidget::OnGenerateWidgetIntParameter(TShar
 }
 
 
-FLinearColor SMutableParametersWidget::GetColorParameterValue(int32 ParamIndex, mu::RangeIndexPtrConst RangeIndex) const
+FLinearColor SMutableParametersWidget::GetColorParameterValue(int32 ParamIndex, TSharedPtr<const mu::FRangeIndex> RangeIndex) const
 {
 	if (!MutableParameters
 		|| ParamIndex >= MutableParameters->GetCount()
@@ -598,20 +597,20 @@ FLinearColor SMutableParametersWidget::GetColorParameterValue(int32 ParamIndex,
 	}
 
 	FVector4f Result;
-	MutableParameters->GetColourValue(ParamIndex, Result, RangeIndex);
+	MutableParameters->GetColourValue(ParamIndex, Result, RangeIndex.Get());
 
 	return Result;
 }
 
 
-FReply SMutableParametersWidget::OnColorBlockMouseButtonDown(const FGeometry& MyGeometry, const FPointerEvent& MouseEvent, int32 ParamIndex, mu::RangeIndexPtrConst RangeIndex)
+FReply SMutableParametersWidget::OnColorBlockMouseButtonDown(const FGeometry& MyGeometry, const FPointerEvent& MouseEvent, int32 ParamIndex, TSharedPtr<const mu::FRangeIndex> RangeIndex)
 {
 	if (MouseEvent.GetEffectingButton() != EKeys::LeftMouseButton)
 	{
 		return FReply::Unhandled();
 	}
 
-	FLinearColor Col = GetColorParameterValue(ParamIndex,RangeIndex);
+	FLinearColor Col = GetColorParameterValue(ParamIndex, RangeIndex);
 
 	FColorPickerArgs args;
 	args.bIsModal = true;
@@ -625,7 +624,7 @@ FReply SMutableParametersWidget::OnColorBlockMouseButtonDown(const FGeometry& My
 }
 
 
-void SMutableParametersWidget::OnSetColorFromColorPicker(FLinearColor NewColor, int32 ParamIndex,  mu::RangeIndexPtrConst RangeIndex)
+void SMutableParametersWidget::OnSetColorFromColorPicker(FLinearColor NewColor, int32 ParamIndex, TSharedPtr<const mu::FRangeIndex> RangeIndex)
 {
 	if (!MutableParameters
 		|| ParamIndex >= MutableParameters->GetCount()
@@ -634,13 +633,13 @@ void SMutableParametersWidget::OnSetColorFromColorPicker(FLinearColor NewColor,
 		return;
 	}
 	
-	MutableParameters->SetColourValue(ParamIndex, NewColor, RangeIndex);
+	MutableParameters->SetColourValue(ParamIndex, NewColor, RangeIndex.Get());
 
 	OnParametersValueChanged.ExecuteIfBound(ParamIndex);
 }
 
 
-TOptional<FVector::FReal> SMutableParametersWidget::GetProjectorLocation(EAxis::Type Axis, int32 ParamIndex, mu::RangeIndexPtrConst RangeIndex) const
+TOptional<FVector::FReal> SMutableParametersWidget::GetProjectorLocation(EAxis::Type Axis, int32 ParamIndex, TSharedPtr<const mu::FRangeIndex> RangeIndex) const
 { 
 	if (!MutableParameters
 		|| ParamIndex >= MutableParameters->GetCount()
@@ -655,13 +654,13 @@ TOptional<FVector::FReal> SMutableParametersWidget::GetProjectorLocation(EAxis::
 	FVector3f Up;
 	FVector3f Scale;
 	float Angle;
-	MutableParameters->GetProjectorValue(ParamIndex, &Type, &Pos, &Dir, &Up, &Scale, &Angle , RangeIndex);
+	MutableParameters->GetProjectorValue(ParamIndex, &Type, &Pos, &Dir, &Up, &Scale, &Angle , RangeIndex.Get());
 
 	return Pos[ Axis-EAxis::X ];
 }
 
 
-void SMutableParametersWidget::SetProjectorLocation(FVector::FReal NewValue, ETextCommit::Type, EAxis::Type Axis, bool bCommitted, int32 ParamIndex, mu::RangeIndexPtrConst RangeIndex)
+void SMutableParametersWidget::SetProjectorLocation(FVector::FReal NewValue, ETextCommit::Type, EAxis::Type Axis, bool bCommitted, int32 ParamIndex, TSharedPtr<const mu::FRangeIndex> RangeIndex)
 {
 	if (!MutableParameters
 		|| ParamIndex >= MutableParameters->GetCount()
@@ -670,7 +669,7 @@ void SMutableParametersWidget::SetProjectorLocation(FVector::FReal NewValue, ETe
 		return;
 	}
 	
-	//mu::RangeIndexPtr UpdatedRange = RangeIndexPayload.GetRangeIndex();
+	//TSharedPtr<mu::FRangeIndex> UpdatedRange = RangeIndexPayload.GetRangeIndex();
 	
 	mu::EProjectorType Type;
 	FVector3f Pos;
@@ -678,17 +677,17 @@ void SMutableParametersWidget::SetProjectorLocation(FVector::FReal NewValue, ETe
 	FVector3f Up;
 	FVector3f Scale;
 	float Angle;
-	MutableParameters->GetProjectorValue(ParamIndex, &Type, &Pos, &Dir, &Up, &Scale, &Angle, RangeIndex);
+	MutableParameters->GetProjectorValue(ParamIndex, &Type, &Pos, &Dir, &Up, &Scale, &Angle, RangeIndex.Get());
 
 	Pos[Axis - EAxis::X] = NewValue;
 
-	MutableParameters->SetProjectorValue(ParamIndex, Pos, Dir, Up, Scale, Angle, RangeIndex);
+	MutableParameters->SetProjectorValue(ParamIndex, Pos, Dir, Up, Scale, Angle, RangeIndex.Get());
 
 	OnParametersValueChanged.ExecuteIfBound(ParamIndex);
 }
 
 
-TOptional<FVector::FReal> SMutableParametersWidget::GetProjectorScale(EAxis::Type Axis, int32 ParamIndex, mu::RangeIndexPtrConst RangeIndex) const
+TOptional<FVector::FReal> SMutableParametersWidget::GetProjectorScale(EAxis::Type Axis, int32 ParamIndex, TSharedPtr<const mu::FRangeIndex> RangeIndex) const
 {
 	if (!MutableParameters
 		|| ParamIndex >= MutableParameters->GetCount()
@@ -703,13 +702,13 @@ TOptional<FVector::FReal> SMutableParametersWidget::GetProjectorScale(EAxis::Typ
 	FVector3f Up;
 	FVector3f Scale;
 	float Angle;
-	MutableParameters->GetProjectorValue(ParamIndex, &Type, &Pos, &Dir, &Up, &Scale, &Angle, RangeIndex);
+	MutableParameters->GetProjectorValue(ParamIndex, &Type, &Pos, &Dir, &Up, &Scale, &Angle, RangeIndex.Get());
 
 	return Scale[Axis - EAxis::X];
 }
 
 
-void SMutableParametersWidget::SetProjectorScale(FVector::FReal NewValue, ETextCommit::Type, EAxis::Type Axis, bool bCommitted, int32 ParamIndex,  mu::RangeIndexPtrConst RangeIndex)
+void SMutableParametersWidget::SetProjectorScale(FVector::FReal NewValue, ETextCommit::Type, EAxis::Type Axis, bool bCommitted, int32 ParamIndex, TSharedPtr<const mu::FRangeIndex> RangeIndex)
 {
 	if (!MutableParameters
 		|| ParamIndex >= MutableParameters->GetCount()
@@ -725,11 +724,11 @@ void SMutableParametersWidget::SetProjectorScale(FVector::FReal NewValue, ETextC
 	FVector3f Up;
 	FVector3f Scale;
 	float Angle;
-	MutableParameters->GetProjectorValue(ParamIndex, &Type, &Pos, &Dir, &Up, &Scale, &Angle, RangeIndex);
+	MutableParameters->GetProjectorValue(ParamIndex, &Type, &Pos, &Dir, &Up, &Scale, &Angle, RangeIndex.Get());
 
 	Scale[Axis - EAxis::X] = NewValue;
 
-	MutableParameters->SetProjectorValue(ParamIndex, Pos, Dir, Up, Scale, Angle, RangeIndex);
+	MutableParameters->SetProjectorValue(ParamIndex, Pos, Dir, Up, Scale, Angle, RangeIndex.Get());
 
 	OnParametersValueChanged.ExecuteIfBound(ParamIndex);
 }
diff --git a/Engine/Plugins/Mutable/Source/CustomizableObjectEditor/Private/MuCOE/SMutableParametersWidget.h b/Engine/Plugins/Mutable/Source/CustomizableObjectEditor/Private/MuCOE/SMutableParametersWidget.h
--- a/Engine/Plugins/Mutable/Source/CustomizableObjectEditor/Private/MuCOE/SMutableParametersWidget.h
+++ b/Engine/Plugins/Mutable/Source/CustomizableObjectEditor/Private/MuCOE/SMutableParametersWidget.h
@@ -23,7 +23,7 @@ struct FPointerEvent;
 DECLARE_DELEGATE_OneParam(FOnMutableParameterValueChanged, int32);
 
 
-/**  This widget displays and allows edition of a set of parameter values for the lower level Mutable Parameters 
+/**  This widget displays and allows edition of a set of parameter values for the lower level Mutable FParameters 
 * object. It is meant to be used for tools and debugging and it doesn't have any knowledge of Customizable Objects
 * and the additional parameter details stored there.
 */
@@ -33,8 +33,8 @@ public:
 
 	SLATE_BEGIN_ARGS(SMutableParametersWidget) {}
 
-	/** Parameters to show and edit. */
-	SLATE_ARGUMENT_DEFAULT(mu::ParametersPtr, Parameters) { nullptr };
+	/** FParameters to show and edit. */
+	SLATE_ARGUMENT_DEFAULT(TSharedPtr<mu::FParameters>, Parameters) { nullptr };
 
 	/** Called when any parameter value has changed, with the parameter index as argument.  */
 	SLATE_EVENT(FOnMutableParameterValueChanged, OnParametersValueChanged)
@@ -49,7 +49,7 @@ public:
 private:
 
 	/** Reference to the mutable parameters. */
-	mu::ParametersPtr MutableParameters;
+	TSharedPtr<mu::FParameters> MutableParameters;
 
 	/** Is true, the parameter object has changed and we need to update. */
 	bool bIsPendingUpdate = false;
@@ -76,32 +76,32 @@ private:
 	/** Generate a new slate for the provided parameter index and attaches it to the also provided Horizontal box slate.
 	 * @param ParamIndex The index of the mutable parameter to represent using a slate object.
 	 * @param ParameterHorizontalBox The Horizontal box slate object to contain the newly generated slate object representing the parameter.
-	 * @param RangeIndex Range object pointing at the parameters' target dimension and position. Providing a nullptr is expected for parameters with
+	 * @param FRangeIndex Range object pointing at the parameters' target dimension and position. Providing a nullptr is expected for parameters with
 	 * a single value (not a multivalue parameter)
 	 * This object can also represent single value parameters but must be provided in either case.
 	 */
-	void GenerateAndAttachParameterSlate(const int32 ParamIndex, TSharedPtr<SHorizontalBox> ParameterHorizontalBox, mu::RangeIndexPtrConst RangeIndex);
+	void GenerateAndAttachParameterSlate(const int32 ParamIndex, TSharedPtr<SHorizontalBox> ParameterHorizontalBox, const TSharedPtr<const mu::FRangeIndex>&);
 	
 	/** Internal UI callbacks. */
 	EVisibility GetParameterVisibility(int32 ParamIndex) const;
-	ECheckBoxState GetBoolParameterValue(int32 ParamIndex, mu::RangeIndexPtrConst RangeIndex) const;
-	void OnBoolParameterChanged(ECheckBoxState CheckboxState, int32 ParamIndex, mu::RangeIndexPtrConst RangeIndex);
-	TOptional<float> GetFloatParameterValue(int32 ParamIndex, mu::RangeIndexPtrConst RangeIndex) const;
-	void OnFloatParameterChanged(float Value, int32 ParamIndex, mu::RangeIndexPtrConst RangeIndex);
+	ECheckBoxState GetBoolParameterValue(int32 ParamIndex, TSharedPtr<const mu::FRangeIndex> ) const;
+	void OnBoolParameterChanged(ECheckBoxState CheckboxState, int32 ParamIndex, TSharedPtr<const mu::FRangeIndex>);
+	TOptional<float> GetFloatParameterValue(int32 ParamIndex, TSharedPtr<const mu::FRangeIndex>) const;
+	void OnFloatParameterChanged(float Value, int32 ParamIndex, TSharedPtr<const mu::FRangeIndex> );
 	void OnFloatParameterCommitted(float Value, ETextCommit::Type CommitType, int32 ParamIndex);
-	FLinearColor GetColorParameterValue(int32 ParamIndex, mu::RangeIndexPtrConst RangeIndex) const;
-	FReply OnColorBlockMouseButtonDown(const FGeometry& MyGeometry, const FPointerEvent& MouseEvent, int32 ParamIndex, mu::RangeIndexPtrConst RangeIndex);
-	void OnSetColorFromColorPicker(FLinearColor NewColor, int32 ParamIndex, mu::RangeIndexPtrConst RangeIndex);
-	TOptional<int32> GetIntParameterValue(int32 ParamIndex,  mu::RangeIndexPtrConst RangeIndex) const;
+	FLinearColor GetColorParameterValue(int32 ParamIndex, TSharedPtr<const mu::FRangeIndex>) const;
+	FReply OnColorBlockMouseButtonDown(const FGeometry& MyGeometry, const FPointerEvent& MouseEvent, int32 ParamIndex, TSharedPtr<const mu::FRangeIndex>);
+	void OnSetColorFromColorPicker(FLinearColor NewColor, int32 ParamIndex, TSharedPtr<const mu::FRangeIndex>);
+	TOptional<int32> GetIntParameterValue(int32 ParamIndex, TSharedPtr<const mu::FRangeIndex>) const;
 	TOptional<int32> GetIntParameterValueMax(int32 ParamIndex) const;
-	void OnIntParameterChanged(int32 Value, int32 ParamIndex, TSharedPtr<SSearchableComboBox> Combo,  mu::RangeIndexPtrConst RangeIndex);
-	void OnIntParameterTextChanged(TSharedPtr<FString> Selection, ESelectInfo::Type SelectInfo, int32 ParamIndex, mu::RangeIndexPtrConst RangeIndex);
+	void OnIntParameterChanged(int32 Value, int32 ParamIndex, TSharedPtr<SSearchableComboBox> Combo, TSharedPtr<const mu::FRangeIndex>);
+	void OnIntParameterTextChanged(TSharedPtr<FString> Selection, ESelectInfo::Type SelectInfo, int32 ParamIndex, TSharedPtr<const mu::FRangeIndex>);
 	TSharedRef<SWidget> OnGenerateWidgetIntParameter(TSharedPtr<FString> InItem) const;
 
 	/** Projector UI callbacks */
-	TOptional<FVector::FReal> GetProjectorLocation(EAxis::Type Axis, int32 ParamIndex,  mu::RangeIndexPtrConst RangeIndex) const;
-	void SetProjectorLocation(FVector::FReal NewValue, ETextCommit::Type, EAxis::Type, bool bCommitted, int32 ParamIndex,  mu::RangeIndexPtrConst RangeIndex);
-	TOptional<FVector::FReal> GetProjectorScale(EAxis::Type Axis, int32 ParamIndex, mu::RangeIndexPtrConst RangeIndex) const;
-	void SetProjectorScale(FVector::FReal NewValue, ETextCommit::Type, EAxis::Type, bool bCommitted, int32 ParamIndex, mu::RangeIndexPtrConst RangeIndex);
+	TOptional<FVector::FReal> GetProjectorLocation(EAxis::Type Axis, int32 ParamIndex, TSharedPtr<const mu::FRangeIndex>) const;
+	void SetProjectorLocation(FVector::FReal NewValue, ETextCommit::Type, EAxis::Type, bool bCommitted, int32 ParamIndex, TSharedPtr<const mu::FRangeIndex>);
+	TOptional<FVector::FReal> GetProjectorScale(EAxis::Type Axis, int32 ParamIndex, TSharedPtr<const mu::FRangeIndex>) const;
+	void SetProjectorScale(FVector::FReal NewValue, ETextCommit::Type, EAxis::Type, bool bCommitted, int32 ParamIndex, TSharedPtr<const mu::FRangeIndex>);
 
 };
diff --git a/Engine/Plugins/Mutable/Source/CustomizableObjectEditor/Private/MuCOE/SMutableSkeletonViewer.cpp b/Engine/Plugins/Mutable/Source/CustomizableObjectEditor/Private/MuCOE/SMutableSkeletonViewer.cpp
--- a/Engine/Plugins/Mutable/Source/CustomizableObjectEditor/Private/MuCOE/SMutableSkeletonViewer.cpp
+++ b/Engine/Plugins/Mutable/Source/CustomizableObjectEditor/Private/MuCOE/SMutableSkeletonViewer.cpp
@@ -10,7 +10,7 @@ class ITableRow;
 
 #define LOCTEXT_NAMESPACE "MutableSkeletonViewer"
 
-/** Row on the tree representing the bone structure of the mu::Mesh. It represents the UI side of the BoneDefinition */
+/** Row on the tree representing the bone structure of the mu::FMesh. It represents the UI side of the BoneDefinition */
 class SMutableMeshBoneTreeRow : public STableRow<TSharedPtr<FBoneDefinition>>
 {
 public:
@@ -91,7 +91,7 @@ void SMutableSkeletonViewer::Construct(const FArguments& InArgs)
 	SetSkeleton(InArgs._Skeleton);
 }
 
-void SMutableSkeletonViewer::SetSkeleton(mu::SkeletonPtrConst InSkeleton)
+void SMutableSkeletonViewer::SetSkeleton(TSharedPtr<const mu::FSkeleton> InSkeleton)
 {
 	if (InSkeleton != Skeleton)
 	{
diff --git a/Engine/Plugins/Mutable/Source/CustomizableObjectEditor/Private/MuCOE/SMutableSkeletonViewer.h b/Engine/Plugins/Mutable/Source/CustomizableObjectEditor/Private/MuCOE/SMutableSkeletonViewer.h
--- a/Engine/Plugins/Mutable/Source/CustomizableObjectEditor/Private/MuCOE/SMutableSkeletonViewer.h
+++ b/Engine/Plugins/Mutable/Source/CustomizableObjectEditor/Private/MuCOE/SMutableSkeletonViewer.h
@@ -24,7 +24,7 @@ struct FBoneDefinition
 class SMutableSkeletonViewer : public SCompoundWidget
 {
 	SLATE_BEGIN_ARGS(SMutableSkeletonViewer) {}
-		SLATE_ARGUMENT_DEFAULT(mu::SkeletonPtrConst, Skeleton) { nullptr };
+		SLATE_ARGUMENT_DEFAULT(TSharedPtr<const mu::FSkeleton>, Skeleton) { nullptr };
 	SLATE_END_ARGS()
 
 public:
@@ -32,19 +32,19 @@ public:
 	void Construct(const FArguments& InArgs);
 
 	/** Sets the mutable skeleton object to be inspected */
-	void SetSkeleton(mu::SkeletonPtrConst InSkeleton);
+	void SetSkeleton(TSharedPtr<const mu::FSkeleton> InSkeleton);
 	
 private:
 
 	/** The mutable skeleton whose data is being inspected */
-	mu::SkeletonPtrConst Skeleton = nullptr;
+	TSharedPtr<const mu::FSkeleton> Skeleton;
 	
 	/** Bones data: Bone definitions used to generate the bone tree structure */
 	TArray<TSharedPtr<FBoneDefinition>> RootBones;
 	TArray<TSharedPtr< FBoneDefinition>> ChildBones;
 
 	/** UI Tree showing the bone data */
-	TSharedPtr < STreeView<TSharedPtr<FBoneDefinition>>> BoneTree;
+	TSharedPtr<STreeView<TSharedPtr<FBoneDefinition>>> BoneTree;
 
 	/** Fills the bone TArrays that hold the root and child bone definitions */
 	void FillBoneDefinitionArrays();
diff --git a/Engine/Plugins/Mutable/Source/CustomizableObjectEditor/Private/MuCOE/UnrealBakeHelpers.cpp b/Engine/Plugins/Mutable/Source/CustomizableObjectEditor/Private/MuCOE/UnrealBakeHelpers.cpp
--- a/Engine/Plugins/Mutable/Source/CustomizableObjectEditor/Private/MuCOE/UnrealBakeHelpers.cpp
+++ b/Engine/Plugins/Mutable/Source/CustomizableObjectEditor/Private/MuCOE/UnrealBakeHelpers.cpp
@@ -175,7 +175,7 @@ UTexture2D* FUnrealBakeHelpers::BakeHelper_CreateAssetTexture(UTexture2D* SrcTex
 
 	// Create a mutable image from the platform data.
 	mu::EImageFormat MutableFormat = UnrealToMutablePixelFormat(SrcTex->GetPlatformData()->PixelFormat, SrcTex->HasAlphaChannel());
-	mu::Ptr<mu::Image> PlatformImage = new mu::Image(sx, sy, 1, MutableFormat, mu::EInitializationType::NotInitialized);
+	TSharedPtr<mu::FImage> PlatformImage = MakeShared<mu::FImage>(sx, sy, 1, MutableFormat, mu::EInitializationType::NotInitialized);
 
 	constexpr int32 MipIndex = 0;
 	const uint8* SourceData = reinterpret_cast<const uint8*>(SrcTex->GetPlatformData()->Mips[MipIndex].BulkData.LockReadOnly());
@@ -203,7 +203,7 @@ UTexture2D* FUnrealBakeHelpers::BakeHelper_CreateAssetTexture(UTexture2D* SrcTex
 
 	mu::FImageOperator ImOp = mu::FImageOperator::GetDefault(mu::FImageOperator::FImagePixelFormatFunc());
 	int32 Quality = 4; // Doesn't matter for decompression.
-	mu::Ptr<mu::Image> UncompressedImage = ImOp.ImagePixelFormat(Quality, PlatformImage.get(), UncompressedMutableFormat);
+	TSharedPtr<mu::FImage> UncompressedImage = ImOp.ImagePixelFormat(Quality, PlatformImage.Get(), UncompressedMutableFormat);
 
 	// Copy the decompressed data to the texture source data
 	int32 SourceDataSize = DupTex->Source.CalcMipSize(MipIndex);
diff --git a/Engine/Plugins/Mutable/Source/CustomizableObjectEditor/Public/MuCOE/ExtensionDataCompilerInterface.h b/Engine/Plugins/Mutable/Source/CustomizableObjectEditor/Public/MuCOE/ExtensionDataCompilerInterface.h
--- a/Engine/Plugins/Mutable/Source/CustomizableObjectEditor/Public/MuCOE/ExtensionDataCompilerInterface.h
+++ b/Engine/Plugins/Mutable/Source/CustomizableObjectEditor/Public/MuCOE/ExtensionDataCompilerInterface.h
@@ -40,7 +40,7 @@ public:
 	 * TSoftObjectPtr, because the container will be moved to a different package during cooking
 	 * and your paths won't be automatically updated.
 	 */
-	mu::ExtensionDataPtrConst MakeStreamedExtensionData(UCustomizableObjectResourceDataContainer*& OutContainer);
+	TSharedPtr<const mu::FExtensionData> MakeStreamedExtensionData(UCustomizableObjectResourceDataContainer*& OutContainer);
 
 	/**
 	 * Register a new Extension Data compile-time constant that will always be loaded.
@@ -51,12 +51,12 @@ public:
 	 * Any UObjects referenced by the provided instanced struct that aren't in an asset package
 	 * should be created with the return value of GetOuterForAlwaysLoadedObjects as their Outer.
 	 *
-	 * Move the constant data into this function and use the resulting mu::ExtensionData in the
+	 * Move the constant data into this function and use the resulting mu::FExtensionData in the
 	 * mu::Node graph, e.g. set it as the value of a mu::NodeExtensionDataConstant.
 	 * 
 	 * As with MakeStreamedExtensionData, the provided instanced struct will be visible to the GC.
 	 */
-	mu::ExtensionDataPtrConst MakeAlwaysLoadedExtensionData(FInstancedStruct&& Data);
+	TSharedPtr<const mu::FExtensionData> MakeAlwaysLoadedExtensionData(FInstancedStruct&& Data);
 
 	/** The Outer to use for objects owned by always-loaded Extension Data constants. */
 	const UObject* GetOuterForAlwaysLoadedObjects();
diff --git a/Engine/Plugins/Mutable/Source/CustomizableObjectEditor/Public/MuCOE/ICustomizableObjectExtensionNode.h b/Engine/Plugins/Mutable/Source/CustomizableObjectEditor/Public/MuCOE/ICustomizableObjectExtensionNode.h
--- a/Engine/Plugins/Mutable/Source/CustomizableObjectEditor/Public/MuCOE/ICustomizableObjectExtensionNode.h
+++ b/Engine/Plugins/Mutable/Source/CustomizableObjectEditor/Public/MuCOE/ICustomizableObjectExtensionNode.h
@@ -9,7 +9,7 @@
 
 class FExtensionDataCompilerInterface;
 
-/** Customizable Object graph nodes that output ExtensionData should implement this interface */
+/** Customizable Object graph nodes that output FExtensionData should implement this interface */
 UINTERFACE()
 class CUSTOMIZABLEOBJECTEDITOR_API UCustomizableObjectExtensionNode : public UInterface
 {
@@ -21,6 +21,6 @@ class ICustomizableObjectExtensionNode
 	GENERATED_BODY()
 
 public:
-	/** Generate a mu::Node that produces ExtensionData */
+	/** Generate a mu::Node that produces FExtensionData */
 	virtual mu::Ptr<mu::NodeExtensionData> GenerateMutableNode(FExtensionDataCompilerInterface& CompilerInterface) const = 0;
 };
diff --git a/Engine/Plugins/Mutable/Source/CustomizableObjectEditor/Public/MuCOE/MutableMeshPreviewUtils.h b/Engine/Plugins/Mutable/Source/CustomizableObjectEditor/Public/MuCOE/MutableMeshPreviewUtils.h
--- a/Engine/Plugins/Mutable/Source/CustomizableObjectEditor/Public/MuCOE/MutableMeshPreviewUtils.h
+++ b/Engine/Plugins/Mutable/Source/CustomizableObjectEditor/Public/MuCOE/MutableMeshPreviewUtils.h
@@ -18,11 +18,6 @@ namespace MutableMeshPreviewUtils
 	 * @param InMutableMesh - The mutable mesh used as base for the generation of the USkeletalMesh
 	 * @return - A new Skeletal mesh based on the provided data. It may be a nullptr if the conversion fails.
 	 */
-	CUSTOMIZABLEOBJECTEDITOR_API USkeletalMesh* GenerateSkeletalMeshFromMutableMesh(
-		mu::MeshPtrConst InMutableMesh);
+	CUSTOMIZABLEOBJECTEDITOR_API USkeletalMesh* GenerateSkeletalMeshFromMutableMesh(const mu::FMesh*);
 
 }
-
-#if UE_ENABLE_INCLUDE_ORDER_DEPRECATED_IN_5_2
-#include "MuR/System.h"
-#endif
diff --git a/Engine/Plugins/Mutable/Source/CustomizableObjectEditor/Public/MuCOE/Nodes/CustomizableObjectNodeExtensionDataConstant.h b/Engine/Plugins/Mutable/Source/CustomizableObjectEditor/Public/MuCOE/Nodes/CustomizableObjectNodeExtensionDataConstant.h
--- a/Engine/Plugins/Mutable/Source/CustomizableObjectEditor/Public/MuCOE/Nodes/CustomizableObjectNodeExtensionDataConstant.h
+++ b/Engine/Plugins/Mutable/Source/CustomizableObjectEditor/Public/MuCOE/Nodes/CustomizableObjectNodeExtensionDataConstant.h
@@ -7,7 +7,7 @@
 #include "CustomizableObjectNodeExtensionDataConstant.generated.h"
 
 /**
- * A base class for nodes that output a constant ExtensionData.
+ * A base class for nodes that output a constant FExtensionData.
  *
  * This is typically used to import an Unreal asset into the Customizable Object graph.
  */
diff --git a/Engine/Plugins/Mutable/Source/CustomizableObjectEditor/Public/MuCOE/Nodes/CustomizableObjectNodeMaterial.h b/Engine/Plugins/Mutable/Source/CustomizableObjectEditor/Public/MuCOE/Nodes/CustomizableObjectNodeMaterial.h
--- a/Engine/Plugins/Mutable/Source/CustomizableObjectEditor/Public/MuCOE/Nodes/CustomizableObjectNodeMaterial.h
+++ b/Engine/Plugins/Mutable/Source/CustomizableObjectEditor/Public/MuCOE/Nodes/CustomizableObjectNodeMaterial.h
@@ -68,8 +68,8 @@ public:
 UENUM()
 enum class ENodePinMode
 {
-	Mutable UMETA(ToolTip = "All Material Texture Parameters go through Mutable."),
-	Passthrough UMETA(ToolTip = "All Material Texture Parameters are not modified by Mutable.")
+	Mutable UMETA(ToolTip = "All Material Texture FParameters go through Mutable."),
+	Passthrough UMETA(ToolTip = "All Material Texture FParameters are not modified by Mutable.")
 };
 
 
@@ -78,8 +78,8 @@ UENUM()
 enum class EPinMode
 {
 	Default UMETA(DisplayName = "Node Defined", ToolTip = "Use node's \"Default Texture Parameter Mode\"."),
-	Mutable UMETA(ToolTip = "The Material Texture Parameters goes through Mutable."),
-	Passthrough UMETA(ToolTip = "The Material Texture Parameters is not modified by Mutable.")
+	Mutable UMETA(ToolTip = "The Material Texture FParameters goes through Mutable."),
+	Passthrough UMETA(ToolTip = "The Material Texture FParameters is not modified by Mutable.")
 };
 
 
@@ -178,7 +178,7 @@ private:
 	UPROPERTY(EditAnywhere, Category=CustomizableObject)
 	TObjectPtr<UMaterialInterface> Material = nullptr;
 
-	UPROPERTY(EditAnywhere, Category=CustomizableObject, DisplayName = "Default Texture Parameter Mode", Meta = (ToolTip = "All Mateiral Texture Parameters set to \"Node Defined\" will use this mode."))
+	UPROPERTY(EditAnywhere, Category=CustomizableObject, DisplayName = "Default Texture Parameter Mode", Meta = (ToolTip = "All Mateiral Texture FParameters set to \"Node Defined\" will use this mode."))
 	ENodePinMode TextureParametersMode = ENodePinMode::Passthrough;
 
 	UPROPERTY()
diff --git a/Engine/Plugins/Mutable/Source/MutableRuntime/Internal/MuR/CodeVisitor.h b/Engine/Plugins/Mutable/Source/MutableRuntime/Internal/MuR/CodeVisitor.h
--- a/Engine/Plugins/Mutable/Source/MutableRuntime/Internal/MuR/CodeVisitor.h
+++ b/Engine/Plugins/Mutable/Source/MutableRuntime/Internal/MuR/CodeVisitor.h
@@ -7,7 +7,6 @@
 #include "Misc/AssertionMacros.h"
 #include "MuR/Model.h"
 #include "MuR/ModelPrivate.h"
-#include "MuR/MutableMemory.h"
 #include "MuR/Operations.h"
 #include "MuR/ParametersPrivate.h"
 #include "MuR/Ptr.h"
@@ -20,10 +19,10 @@ namespace mu
     //---------------------------------------------------------------------------------------------
     //! Decide what operations are an "add resource" since they are handled differently sometimes.
     //---------------------------------------------------------------------------------------------
-    inline bool VisitorIsAddResource(const OP_TYPE& type)
+    inline bool VisitorIsAddResource(const EOpType& type)
     {
-        return  type == OP_TYPE::IN_ADDIMAGE ||
-                type == OP_TYPE::IN_ADDMESH;
+        return  type == EOpType::IN_ADDIMAGE ||
+                type == EOpType::IN_ADDMESH;
     }
 
 
@@ -34,7 +33,7 @@ namespace mu
     //! - will not visit twice the same instruction with the same state.
     //! - Its iterative
     //---------------------------------------------------------------------------------------------
-    template<typename STATE=int>
+    template<typename STATE=int32>
     class UniqueConstCodeVisitorIterative
     {
     public:
@@ -42,9 +41,9 @@ namespace mu
         UniqueConstCodeVisitorIterative( bool skipResources=false )
         {
             // Default state
-            m_states.Add(STATE());
-            m_currentState = 0;
-            m_skipResources = skipResources;
+            States.Add(STATE());
+            CurrentState = 0;
+            bSkipResources = skipResources;
         }
 
         //! Ensure virtual destruction
@@ -55,68 +54,68 @@ namespace mu
         //!
         void SetDefaultState(const STATE& s)
         {
-            m_states[0] = s;
+            States[0] = s;
         }
 
         //!
         const STATE& GetDefaultState() const
         {
-            return m_states[0];
+            return States[0];
         }
 
         //! Use this from visit to access the state at the time of processing the current
         //! instruction.
         STATE GetCurrentState() const
         {
-            return m_states[m_currentState];
+            return States[CurrentState];
         }
 
         //! For manual recursion that changes the state for a specific path.
         void RecurseWithState(OP::ADDRESS at, const STATE& newState)
         {
-			int32 it = m_states.Find(newState);
+			int32 it = States.Find(newState);
             if (it==INDEX_NONE)
             {
-                m_states.Add(newState);
+                States.Add(newState);
             }
-            int stateIndex = m_states.IndexOfByKey(newState);
+            int32 StateIndex = States.IndexOfByKey(newState);
 
-            m_pending.Add( PENDING(at,stateIndex) );
+            Pending.Add( FPending(at,StateIndex) );
         }
 
         //! For manual recursion that doesn't change the state for a specific path.
         void RecurseWithCurrentState(OP::ADDRESS at)
         {
-            m_pending.Add( PENDING(at,m_currentState) );
+            Pending.Add( FPending(at,CurrentState) );
         }
 
         //! Can be called from visit to set the state to visit all children ops
         void SetCurrentState(const STATE& newState)
         {
-            int32 it = m_states.Find(newState);
+            int32 it = States.Find(newState);
             if (it==INDEX_NONE)
             {
-                m_states.Add(newState);
+                States.Add(newState);
             }
-            m_currentState = m_states.Find(newState);
+            CurrentState = States.Find(newState);
         }
 
 
         void Traverse( OP::ADDRESS root, FProgram& program )
         {
-            m_pending.Reserve( program.m_opAddress.Num() );
+            Pending.Reserve( program.OpAddress.Num() );
 
             // Visit the given root
-            m_pending.Add( PENDING(root,0) );
+            Pending.Add( FPending(root,0) );
             Recurse( program );
          }
 
         void FullTraverse( FProgram& program )
         {
             // Visit all the state roots
-            for ( int32 p=0; p<program.m_states.Num(); ++p )
+            for ( int32 p=0; p<program.States.Num(); ++p )
             {
-                m_pending.Add(PENDING(program.m_states[p].m_root,0) );
+                Pending.Add(FPending(program.States[p].Root,0) );
                 Recurse( program );
             }
         }
@@ -129,68 +128,68 @@ namespace mu
         virtual bool Visit( OP::ADDRESS at, FProgram& program ) = 0;
 
         //! Operations to be processed
-        struct PENDING
+        struct FPending
         {
-            PENDING()
+            FPending()
             {
                 at = 0;
-                stateIndex = 0;
+                StateIndex = 0;
             }
-            PENDING(OP::ADDRESS _at, int _stateIndex)
+            FPending(OP::ADDRESS _at, int32 _stateIndex)
             {
                 at = _at;
-                stateIndex = _stateIndex;
+                StateIndex = _stateIndex;
             }
 
             OP::ADDRESS at;
-            int stateIndex;
+            int32 StateIndex;
         };
-		TArray<PENDING> m_pending;
+		TArray<FPending> Pending;
 
         //! States found so far
-		TArray<STATE> m_states;
+		TArray<STATE> States;
 
-        //! Index of the current state, from the m_states array.
-        int m_currentState;
+        //! Index of the current state, from the States array.
+        int32 CurrentState;
 
         //! If true, operations adding resources (meshes or images) will only
         //! have the base operation recursed, but not the resources.
-        bool m_skipResources;
+        bool bSkipResources;
 
         //! Array of states visited for each operation.
         //! Empty array means operation not visited at all.
-		TArray<TArray<int>> m_visited;
+		TArray<TArray<int32>> m_visited;
 
         //! Process all the pending operations and visit all children if necessary
         void Recurse( FProgram& program )
         {
 			m_visited.Empty();
-			m_visited.SetNum(program.m_opAddress.Num());
+			m_visited.SetNum(program.OpAddress.Num());
 
-            while ( m_pending.Num() )
+            while ( Pending.Num() )
             {
-                OP::ADDRESS at = m_pending.Last().at;
-                m_currentState = m_pending.Last().stateIndex;
-                m_pending.Pop();
+                OP::ADDRESS at = Pending.Last().at;
+                CurrentState = Pending.Last().StateIndex;
+                Pending.Pop();
 
-                bool recurse = false;
+                bool bRecurse = false;
 
-                bool visitedInThisState = m_visited[at].Contains(m_currentState);
+                bool visitedInThisState = m_visited[at].Contains(CurrentState);
                 if (!visitedInThisState)
                 {
-                    m_visited[at].Add(m_currentState);
+                    m_visited[at].Add(CurrentState);
 
                     // Visit may change current state
-                    recurse = Visit( at, program );
+                    bRecurse = Visit( at, program );
                 }
 
-                if (recurse)
+                if (bRecurse)
                 {
                     ForEachReference( program, at, [&](OP::ADDRESS ref)
                     {
                         if (ref)
                         {
-                            m_pending.Add( PENDING(ref,m_currentState) );
+                            Pending.Add( FPending(ref,CurrentState) );
                         }
                     });
                 }
@@ -207,7 +206,7 @@ namespace mu
     //! - will repeat visits to the instructions that have multiple references.
     //! - Its iterative
     //---------------------------------------------------------------------------------------------
-    template<class STATE=int>
+    template<class STATE=int32>
     class RepeatConstCodeVisitorIterative
     {
     public:
@@ -215,9 +214,9 @@ namespace mu
         RepeatConstCodeVisitorIterative( bool skipResources=false )
         {
             // Default state
-            m_states.Add(STATE());
-            m_currentState = 0;
-            m_skipResources = skipResources;
+            States.Add(STATE());
+            CurrentState = 0;
+            bSkipResources = skipResources;
         }
 
         //! Ensure virtual destruction
@@ -228,70 +227,70 @@ namespace mu
         //!
         void SetDefaultState(const STATE& s)
         {
-            m_states[0] = s;
+            States[0] = s;
         }
 
         //!
         const STATE& GetDefaultState() const
         {
-            return m_states[0];
+            return States[0];
         }
 
         //! Use this from visit to access the state at the time of processing the current
         //! instruction.
         const STATE& GetCurrentState() const
         {
-            return m_states[m_currentState];
+            return States[CurrentState];
         }
 
         //! For manual recursion that changes the state for a specific path.
         void RecurseWithState(OP::ADDRESS At, const STATE& NewState)
         {
-            int32 StateIndex = m_states.Find(NewState);
+            int32 StateIndex = States.Find(NewState);
             if (StateIndex ==INDEX_NONE)
             {
-				StateIndex = m_states.Add(NewState);
+				StateIndex = States.Add(NewState);
             }
 
             //check(at<1000000);
-            m_pending.Add( PENDING(At,StateIndex) );
+            Pending.Add( FPending(At,StateIndex) );
         }
 
         //! For manual recursion that doesn't change the state for a specific path.
         void RecurseWithCurrentState(OP::ADDRESS at)
         {
             //check(at<1000000);
-            m_pending.Add( PENDING(at,m_currentState) );
+            Pending.Add( FPending(at,CurrentState) );
         }
 
         //! Can be called from visit to set the state to visit all children ops
         void SetCurrentState(const STATE& NewState)
         {
-			int32 StateIndex = m_states.Find(NewState);
+			int32 StateIndex = States.Find(NewState);
 			if (StateIndex == INDEX_NONE)
 			{
-				StateIndex = m_states.Add(NewState);
+				StateIndex = States.Add(NewState);
 			}
 
-			m_currentState = StateIndex;
+			CurrentState = StateIndex;
         }
 
 
         void Traverse( OP::ADDRESS Root, FProgram& Program )
         {
-            m_pending.reserve(Program.m_opAddress.Num() );
+            Pending.reserve(Program.OpAddress.Num() );
 
             // Visit the given root
-            m_pending.Add( PENDING(Root,0) );
+            Pending.Add( FPending(Root,0) );
             Recurse(Program);
         }
 
         void FullTraverse( FProgram& Program)
         {
             // Visit all the state roots
-            for ( int32 p=0; p<Program.m_states.Num(); ++p )
+            for ( int32 p=0; p<Program.States.Num(); ++p )
             {
-                m_pending.Add(PENDING(Program.m_states[p].m_root,0) );
+                Pending.Add(FPending(Program.States[p].Root,0) );
                 Recurse(Program);
             }
         }
@@ -301,55 +300,55 @@ namespace mu
 
         //! Do the actual work by overriding this in the derived classes.
         //! Return true if the traverse has to continue with the children of "at"
-        virtual bool Visit( OP::ADDRESS at, FProgram& program ) = 0;
+        virtual bool Visit( OP::ADDRESS, FProgram& ) = 0;
 
         //! Operations to be processed
-        struct PENDING
+        struct FPending
         {
-            PENDING()
+            FPending()
             {
                 at = 0;
-                stateIndex = 0;
+                StateIndex = 0;
             }
-            PENDING(OP::ADDRESS _at, int _stateIndex)
+            FPending(OP::ADDRESS InAt, int32 _stateIndex)
             {
-                at = _at;
-                stateIndex = _stateIndex;
+                at = InAt;
+                StateIndex = _stateIndex;
             }
 
             OP::ADDRESS at;
-            int stateIndex;
+            int32 StateIndex;
         };
-        TArray<PENDING> m_pending;
+        TArray<FPending> Pending;
 
         //! States found so far
-		TArray<STATE> m_states;
+		TArray<STATE> States;
 
-        //! Index of the current state, from the m_states array.
-        int m_currentState;
+        //! Index of the current state, from the States array.
+        int32 CurrentState;
 
         //! If true, operations adding resources (meshes or images) will only
         //! have the base operation recursed, but not the resources.
-        bool m_skipResources;
+        bool bSkipResources;
 
 
         //! Process all the pending operations and visit all children if necessary
         void Recurse( FProgram& program )
         {
-            while ( m_pending.Num() )
+            while ( Pending.Num() )
             {
-                OP::ADDRESS at = m_pending.back().at;
-                m_currentState = m_pending.back().stateIndex;
-                m_pending.pop_back();
+                OP::ADDRESS at = Pending.back().at;
+                CurrentState = Pending.back().StateIndex;
+                Pending.pop_back();
 
-                bool recurse = false;
+                bool bRecurse = false;
 
                 // Visit may change current state
-                recurse = Visit( at, program );
+                bRecurse = Visit( at, program );
 
-                if (recurse)
+                if (bRecurse)
                 {
-                    if (m_skipResources && VisitorIsAddResource(program.GetOpType(at)))
+                    if (bSkipResources && VisitorIsAddResource(program.GetOpType(at)))
                     {
 						OP::InstanceAddArgs args = program.GetOpArgs<OP::InstanceAddArgs>(at);
 
@@ -357,8 +356,8 @@ namespace mu
                         OP::ADDRESS base = args.instance;
                         if (base)
                         {
-							check(base<program.m_opAddress.Num());
-                            m_pending.Add( PENDING(base,m_currentState) );
+							check(base<program.OpAddress.Num());
+                            Pending.Add( FPending(base,CurrentState) );
                         }
                     }
                     else
@@ -367,8 +366,8 @@ namespace mu
                         {
                             if (ref)
                             {
-								check(ref<program.m_opAddress.Num());
-                                m_pending.Add( PENDING(ref,m_currentState) );
+								check(ref<program.OpAddress.Num());
+                                Pending.Add( FPending(ref,CurrentState) );
                             }
                         });
                     }
@@ -404,18 +403,18 @@ namespace mu
 
         DiscreteCoveredCodeVisitorBase
             (
-                System::Private* pSystem,
-                const TSharedPtr<const Model>& pModel,
-                const ParametersPtrConst& pParams,
-                unsigned lodMask,
-                bool skipResources=false
+                FSystem::Private* InSystem,
+                const TSharedPtr<const FModel>& InModel,
+                const TSharedPtr<const FParameters>& InParams,
+                unsigned InLodMask,
+                bool bSkipResources=false
             )
-            : PARENT(skipResources)
+            : PARENT(bSkipResources)
         {
-            m_pSystem = pSystem;
-            m_pModel = pModel;
-            m_pParams = pParams.get();
-            m_lodMask = lodMask;
+            System = InSystem;
+            Model = InModel;
+            Params = InParams.Get();
+            LODMask = InLodMask;
 
             // Visiting state
             PARENT::SetDefaultState( STATE() );
@@ -425,31 +424,31 @@ namespace mu
         {
             PARENT::SetDefaultState( STATE() );
 
-            PARENT::Traverse( at, m_pModel->GetPrivate()->m_program );
+            PARENT::Traverse( at, Model->GetPrivate()->Program );
         }
 
     protected:
 
         virtual bool Visit( OP::ADDRESS at, FProgram& program )
         {
-            bool recurse = true;
+            bool bRecurse = true;
 
-            OP_TYPE type = program.GetOpType(at);
+            EOpType type = program.GetOpType(at);
 
             switch ( type )
             {
-            case OP_TYPE::NU_CONDITIONAL:
-            case OP_TYPE::SC_CONDITIONAL:
-            case OP_TYPE::CO_CONDITIONAL:
-            case OP_TYPE::IM_CONDITIONAL:
-            case OP_TYPE::ME_CONDITIONAL:
-            case OP_TYPE::LA_CONDITIONAL:
-            case OP_TYPE::IN_CONDITIONAL:
-			case OP_TYPE::ED_CONDITIONAL:
+            case EOpType::NU_CONDITIONAL:
+            case EOpType::SC_CONDITIONAL:
+            case EOpType::CO_CONDITIONAL:
+            case EOpType::IM_CONDITIONAL:
+            case EOpType::ME_CONDITIONAL:
+            case EOpType::LA_CONDITIONAL:
+            case EOpType::IN_CONDITIONAL:
+			case EOpType::ED_CONDITIONAL:
             {
 				OP::ConditionalArgs args = program.GetOpArgs<OP::ConditionalArgs>(at);
 
-                recurse = false;
+                bRecurse = false;
 
                 PARENT::RecurseWithCurrentState( args.condition );
 
@@ -458,7 +457,7 @@ namespace mu
 
                 if (args.condition)
                 {
-                    value = m_pSystem->BuildBool(m_pModel, m_pParams, args.condition);
+                    value = System->BuildBool(Model, Params, args.condition);
                 }
 
                 if (value)
@@ -472,16 +471,16 @@ namespace mu
                 break;
             }
 
-            case OP_TYPE::NU_SWITCH:
-            case OP_TYPE::SC_SWITCH:
-            case OP_TYPE::CO_SWITCH:
-            case OP_TYPE::IM_SWITCH:
-            case OP_TYPE::ME_SWITCH:
-            case OP_TYPE::LA_SWITCH:
-            case OP_TYPE::IN_SWITCH:
-			case OP_TYPE::ED_SWITCH:
+            case EOpType::NU_SWITCH:
+            case EOpType::SC_SWITCH:
+            case EOpType::CO_SWITCH:
+            case EOpType::IM_SWITCH:
+            case EOpType::ME_SWITCH:
+            case EOpType::LA_SWITCH:
+            case EOpType::IN_SWITCH:
+			case EOpType::ED_SWITCH:
             {
-                recurse = false;
+                bRecurse = false;
 
 				const uint8* data = program.GetOpArgsPointer(at);
 				
@@ -501,7 +500,7 @@ namespace mu
 
                     PARENT::RecurseWithCurrentState( VarAddress );
 
-                    int var = m_pSystem->BuildInt( m_pModel, m_pParams, VarAddress );
+                    int32 var = System->BuildInt( Model, Params, VarAddress );
 
 					OP::ADDRESS valueAt = DefAddress;
 					for (uint32 C = 0; C < CaseCount; ++C)
@@ -514,7 +513,7 @@ namespace mu
 						FMemory::Memcpy( &At, data, sizeof(OP::ADDRESS));
 						data += sizeof(OP::ADDRESS);
 
-						if (At && var == (int)Condition)
+						if (At && var == (int32)Condition)
 						{
 							valueAt = At;
 							break;	
@@ -528,9 +527,9 @@ namespace mu
             }
 
 
-            case OP_TYPE::IN_ADDLOD:
+            case EOpType::IN_ADDLOD:
             {
-                recurse = false;
+                bRecurse = false;
 
 				const uint8* Data = program.GetOpArgsPointer(at);
 
@@ -547,7 +546,7 @@ namespace mu
 					
 					if (LODAddress)
                     {
-                        bool bSelected = ( (1<< LODIndex) & m_lodMask ) != 0;
+                        bool bSelected = ( (1<< LODIndex) & LODMask ) != 0;
                         if (bSelected)
                         {
                             PARENT::RecurseWithState(LODAddress, NewState);
@@ -558,11 +557,11 @@ namespace mu
             }
 
 
-            case OP_TYPE::IN_ADDMESH:
+            case EOpType::IN_ADDMESH:
             {
 				OP::InstanceAddArgs args = program.GetOpArgs<OP::InstanceAddArgs>(at);
 
-                recurse = false;
+                bRecurse = false;
 
                 PARENT::RecurseWithCurrentState(args.instance);
 
@@ -578,11 +577,11 @@ namespace mu
             }
 
 
-            case OP_TYPE::IN_ADDIMAGE:
+            case EOpType::IN_ADDIMAGE:
             {
 				OP::InstanceAddArgs args = program.GetOpArgs<OP::InstanceAddArgs>(at);
 
-                recurse = false;
+                bRecurse = false;
 
                 PARENT::RecurseWithCurrentState(args.instance);
 
@@ -601,15 +600,15 @@ namespace mu
                 break;
             }
 
-            return recurse;
+            return bRecurse;
         }
 
 
     protected:
-        System::Private* m_pSystem = nullptr;
-		TSharedPtr<const Model> m_pModel;
-        const Parameters* m_pParams = nullptr;
-        unsigned m_lodMask = 0;
+        FSystem::Private* System = nullptr;
+		TSharedPtr<const FModel> Model;
+        const FParameters* Params = nullptr;
+        unsigned LODMask = 0;
     };
 
 
@@ -637,9 +636,9 @@ namespace mu
 
         UniqueDiscreteCoveredCodeVisitor
             (
-                System::Private* InSystem,
-				const TSharedPtr<const Model>& InModel,
-                const ParametersPtrConst& InParams,
+                FSystem::Private* InSystem,
+				const TSharedPtr<const FModel>& InModel,
+                const TSharedPtr<const FParameters>& InParams,
                 uint32 InLodMask
             )
             : PARENT(InSystem, InModel, InParams, InLodMask )
diff --git a/Engine/Plugins/Mutable/Source/MutableRuntime/Internal/MuR/MeshPrivate.h b/Engine/Plugins/Mutable/Source/MutableRuntime/Internal/MuR/MeshPrivate.h
--- a/Engine/Plugins/Mutable/Source/MutableRuntime/Internal/MuR/MeshPrivate.h
+++ b/Engine/Plugins/Mutable/Source/MutableRuntime/Internal/MuR/MeshPrivate.h
@@ -11,41 +11,40 @@
 
 #include "CoreFwd.h"
 
-
 namespace mu
 {
 
 	//---------------------------------------------------------------------------------------------
 	inline void GetMeshBuf
 		(
-			Mesh* pMesh,
-			EMeshBufferSemantic semantic,
+			FMesh* Mesh,
+			EMeshBufferSemantic Semantic,
 			EMeshBufferFormat expectedFormat,
-			int expectedComponents,
-            uint8_t*& pBuf,
-			int& elemSize
+			int32 expectedComponents,
+            uint8*& pBuf,
+			int32& elemSize
 		)
 	{
 		// Avoid unreferenced parameter warnings
 		(void)expectedFormat;
 		(void)expectedComponents;
 
-		int buffer = -1;
-		int channel = -1;
-		pMesh->GetVertexBuffers().FindChannel( semantic, 0, &buffer, &channel );
-		check( buffer>=0 && channel>=0 );
+		int32 BufferIndex = -1;
+		int32 ChannelIndex = -1;
+		Mesh->GetVertexBuffers().FindChannel( Semantic, 0, &BufferIndex, &ChannelIndex );
+		check( BufferIndex>=0 && ChannelIndex>=0 );
 
 		EMeshBufferSemantic realSemantic = EMeshBufferSemantic::None;
-		int realSemanticIndex = 0;
+		int32 realSemanticIndex = 0;
 		EMeshBufferFormat format = EMeshBufferFormat::None;
-		int components = 0;
-		int offset = 0;
-		pMesh->GetVertexBuffers().GetChannel( buffer, channel, &realSemantic, &realSemanticIndex, &format, &components, &offset );
-		check( realSemantic == semantic );
+		int32 components = 0;
+		int32 offset = 0;
+		Mesh->GetVertexBuffers().GetChannel( BufferIndex, ChannelIndex, &realSemantic, &realSemanticIndex, &format, &components, &offset );
+		check( realSemantic == Semantic );
 		check( format == expectedFormat );
 		check( components == expectedComponents );
-		elemSize = pMesh->GetVertexBuffers().GetElementSize( buffer );
-		pBuf = pMesh->GetVertexBuffers().GetBufferData( buffer );
+		elemSize = Mesh->GetVertexBuffers().GetElementSize( BufferIndex );
+		pBuf = Mesh->GetVertexBuffers().GetBufferData( BufferIndex );
 		pBuf += offset;
 	}
 
@@ -53,41 +52,41 @@ namespace mu
 	//---------------------------------------------------------------------------------------------
 	inline void GetMeshBuf
 		(
-			const Mesh* pMesh,
-			EMeshBufferSemantic semantic,
+			const FMesh* Mesh,
+			EMeshBufferSemantic Semantic,
 			EMeshBufferFormat expectedFormat,
-			int expectedComponents,
-            const uint8_t*& pBuf,
-			int& elemSize
+			int32 expectedComponents,
+            const uint8*& pBuf,
+			int32& elemSize
 		)
 	{
 		// Avoid unreferenced parameter warnings
 		(void)expectedFormat;
 		(void)expectedComponents;
 
-		int buffer = -1;
-		int channel = -1;
-		pMesh->GetVertexBuffers().FindChannel( semantic, 0, &buffer, &channel );
-		check( buffer>=0 && channel>=0 );
+		int32 BufferIndex = -1;
+		int32 ChannelIndex = -1;
+		Mesh->GetVertexBuffers().FindChannel( Semantic, 0, &BufferIndex, &ChannelIndex );
+		check( BufferIndex>=0 && ChannelIndex>=0 );
 
 		EMeshBufferSemantic realSemantic = EMeshBufferSemantic::None;
-		int realSemanticIndex = 0;
+		int32 realSemanticIndex = 0;
 		EMeshBufferFormat format = EMeshBufferFormat::None;
-		int components = 0;
-		int offset = 0;
-		pMesh->GetVertexBuffers().GetChannel
-				( buffer, channel, &realSemantic, &realSemanticIndex, &format, &components, &offset );
-		check( realSemantic == semantic );
+		int32 components = 0;
+		int32 offset = 0;
+		Mesh->GetVertexBuffers().GetChannel
+				( BufferIndex, ChannelIndex, &realSemantic, &realSemanticIndex, &format, &components, &offset );
+		check( realSemantic == Semantic );
 		check( format == expectedFormat );
 		check( components == expectedComponents );
-		elemSize = pMesh->GetVertexBuffers().GetElementSize( buffer );
-		pBuf = pMesh->GetVertexBuffers().GetBufferData( buffer );
+		elemSize = Mesh->GetVertexBuffers().GetElementSize( BufferIndex );
+		pBuf = Mesh->GetVertexBuffers().GetBufferData( BufferIndex );
 		pBuf += offset;
 	}
 
 
 	//---------------------------------------------------------------------------------------------
-	//! Class to iterate a specific buffer channel of unknown type
+	//! Class to iterate a specific BufferIndex ChannelIndex of unknown type
 	//---------------------------------------------------------------------------------------------
 	class UntypedMeshBufferIteratorConst;
 
@@ -96,76 +95,76 @@ namespace mu
 	public:
 		inline UntypedMeshBufferIterator()
 		{
-			m_format = EMeshBufferFormat::None;
-			m_components = 0;
-			m_elementSize = 0;
-			m_pBuf = 0;
+			Format = EMeshBufferFormat::None;
+			Components = 0;
+			ElementSize = 0;
+			Buffer = 0;
 		}
 
 
 		inline UntypedMeshBufferIterator
 			(
-				FMeshBufferSet& bufferSet,
-				EMeshBufferSemantic semantic,
-				int semanticIndex = 0
+				FMeshBufferSet& BufferSet,
+				EMeshBufferSemantic Semantic,
+				int32 SemanticIndex = 0
 			)
 		{
-			int buffer = -1;
-			int channel = -1;
-			bufferSet.FindChannel( semantic, semanticIndex, &buffer, &channel );
+			int32 BufferIndex = -1;
+			int32 ChannelIndex = -1;
+			BufferSet.FindChannel( Semantic, SemanticIndex, &BufferIndex, &ChannelIndex );
 
-			if ( buffer>=0 && channel>=0 )
+			if ( BufferIndex>=0 && ChannelIndex>=0 )
 			{
 				EMeshBufferSemantic realSemantic = EMeshBufferSemantic::None;
-				int realSemanticIndex = 0;
-				m_format = EMeshBufferFormat::None;
-				m_components = 0;
-				int offset = 0;
-				bufferSet.GetChannel
+				int32 realSemanticIndex = 0;
+				Format = EMeshBufferFormat::None;
+				Components = 0;
+				int32 offset = 0;
+				BufferSet.GetChannel
 						(
-							buffer, channel,
+							BufferIndex, ChannelIndex,
 							&realSemantic, &realSemanticIndex,
-							&m_format, &m_components,
+							&Format, &Components,
 							&offset
 						);
-				check( realSemantic == semantic );
-				check( realSemanticIndex == semanticIndex );
+				check( realSemantic == Semantic );
+				check( realSemanticIndex == SemanticIndex );
 
-				m_elementSize = bufferSet.GetElementSize( buffer );
+				ElementSize = BufferSet.GetElementSize( BufferIndex );
 
-				m_pBuf = bufferSet.GetBufferData( buffer );
-				m_pBuf += offset;
+				Buffer = BufferSet.GetBufferData( BufferIndex );
+				Buffer += offset;
 			}
 			else
 			{
-				m_format = EMeshBufferFormat::None;
-				m_components = 0;
-				m_elementSize = 0;
-				m_pBuf = 0;
+				Format = EMeshBufferFormat::None;
+				Components = 0;
+				ElementSize = 0;
+				Buffer = 0;
 			}
 		}
 
-        inline uint8_t* ptr() const
+        inline uint8* ptr() const
 		{
-			return m_pBuf;
+			return Buffer;
 		}
 
 		inline void operator++()
 		{
-			m_pBuf += m_elementSize;
+			Buffer += ElementSize;
 		}
 
-		inline void operator++(int)
+		inline void operator++(int32)
 		{
-			m_pBuf += m_elementSize;
+			Buffer += ElementSize;
 		}
 
-		inline void operator+=( int c )
+		inline void operator+=( int32 c )
 		{
-			m_pBuf += c*m_elementSize;
+			Buffer += c*ElementSize;
 		}
 
-		inline UntypedMeshBufferIterator operator+(int c) const
+		inline UntypedMeshBufferIterator operator+(int32 c) const
 		{
 			UntypedMeshBufferIterator res = *this;
 			res += c;
@@ -173,40 +172,40 @@ namespace mu
 			return res;
 		}
 
-		inline std::size_t operator-( const UntypedMeshBufferIterator& other ) const
+		inline SIZE_T operator-( const UntypedMeshBufferIterator& Other ) const
 		{
 			// Special degenerate case.
-			if (m_elementSize == 0)
+			if (ElementSize == 0)
 			{
 				return 0;
 			}
-			check( other.m_elementSize == m_elementSize );
-			check( (ptr()-other.ptr()) % m_elementSize == 0 );
-			return (ptr()-other.ptr()) / m_elementSize;
+			check( Other.ElementSize == ElementSize );
+			check( (ptr()-Other.ptr()) % ElementSize == 0 );
+			return (ptr()-Other.ptr()) / ElementSize;
 		}
 
-		inline int GetElementSize() const
+		inline int32 GetElementSize() const
 		{
-			return m_elementSize;
+			return ElementSize;
 		}
 
 		inline EMeshBufferFormat GetFormat() const
 		{
-			return m_format;
+			return Format;
 		}
 
-		inline int GetComponents() const
+		inline int32 GetComponents() const
 		{
-			return m_components;
+			return Components;
 		}
 
         FVector4f GetAsVec4f() const
         {
 			FVector4f res;
 
-            for (int c=0; c< FMath::Min(m_components,4); ++c)
+            for (int32 c=0; c< FMath::Min(Components,4); ++c)
             {
-                ConvertData( c, &res[0], EMeshBufferFormat::Float32, ptr(), m_format );
+                ConvertData( c, &res[0], EMeshBufferFormat::Float32, ptr(), Format );
             }
 
             return res;
@@ -215,9 +214,9 @@ namespace mu
 		FVector3f GetAsVec3f() const
 		{
 			FVector3f res = { 0,0,0 };
-			for (int c = 0; c < FMath::Min(m_components, 3); ++c)
+			for (int32 c = 0; c < FMath::Min(Components, 3); ++c)
 			{
-				ConvertData(c, &res[0], EMeshBufferFormat::Float32, ptr(), m_format);
+				ConvertData(c, &res[0], EMeshBufferFormat::Float32, ptr(), Format);
 			}
 			return res;
 		}
@@ -225,9 +224,9 @@ namespace mu
 		FVector3d GetAsVec3d() const
 		{
 			FVector3d res = { 0,0,0 };
-			for (int c = 0; c < FMath::Min(m_components, 3); ++c)
+			for (int32 c = 0; c < FMath::Min(Components, 3); ++c)
 			{
-				ConvertData(c, &res[0], EMeshBufferFormat::Float64, ptr(), m_format);
+				ConvertData(c, &res[0], EMeshBufferFormat::Float64, ptr(), Format);
 			}
 			return res;
 		}
@@ -235,9 +234,9 @@ namespace mu
 		FVector2f GetAsVec2f() const
 		{
 			FVector2f res = {0.0f, 0.0f};
-			for (int32 c = 0; c < FMath::Min(m_components, 2); ++c)
+			for (int32 c = 0; c < FMath::Min(Components, 2); ++c)
 			{
-				ConvertData(c, &res[0], EMeshBufferFormat::Float32, ptr(), m_format);
+				ConvertData(c, &res[0], EMeshBufferFormat::Float32, ptr(), Format);
 			}
 
 			return res;
@@ -246,82 +245,82 @@ namespace mu
 		uint32 GetAsUINT32() const
 		{
 			uint32 res;
-			ConvertData(0, &res, EMeshBufferFormat::UInt32, ptr(), m_format);
+			ConvertData(0, &res, EMeshBufferFormat::UInt32, ptr(), Format);
 			return res;
 		}
 
 		uint64 GetAsUINT64() const
 		{
 			uint64 res;
-			ConvertData(0, &res, EMeshBufferFormat::UInt64, ptr(), m_format);
+			ConvertData(0, &res, EMeshBufferFormat::UInt64, ptr(), Format);
 			return res;
 		}
 
-		void SetFromUINT32(uint32_t v)
+		void SetFromUINT32(uint32 v)
 		{
-			ConvertData(0, ptr(), m_format, &v, EMeshBufferFormat::UInt32);
+			ConvertData(0, ptr(), Format, &v, EMeshBufferFormat::UInt32);
 		}
 
 		void SetFromVec3f(const FVector3f& v)
 		{
-			for (int c = 0; c < FMath::Min(m_components, 3); ++c)
+			for (int32 c = 0; c < FMath::Min(Components, 3); ++c)
 			{
-				ConvertData(c, ptr(), m_format, &v, EMeshBufferFormat::Float32);
+				ConvertData(c, ptr(), Format, &v, EMeshBufferFormat::Float32);
 			}
 		}
 
 		void SetFromVec3d(const FVector3d& v)
 		{
-			for (int c = 0; c < FMath::Min(m_components, 3); ++c)
+			for (int32 c = 0; c < FMath::Min(Components, 3); ++c)
 			{
-				ConvertData(c, ptr(), m_format, &v, EMeshBufferFormat::Float64);
+				ConvertData(c, ptr(), Format, &v, EMeshBufferFormat::Float64);
 			}
 		}
 
 	protected:
 
-		int m_elementSize;
-        uint8_t* m_pBuf;
-		EMeshBufferFormat m_format;
-		int m_components;
+		int32 ElementSize;
+        uint8* Buffer;
+		EMeshBufferFormat Format;
+		int32 Components;
 	};
 
 
 	//---------------------------------------------------------------------------------------------
-	//! Class to iterate a specific buffer channel with known type
+	//! Class to iterate a specific BufferIndex ChannelIndex with known type
 	//---------------------------------------------------------------------------------------------
-	template<EMeshBufferFormat FORMAT, class CTYPE, int COMPONENTS>
+	template<EMeshBufferFormat FORMAT, class CTYPE, int32 COMPONENTS>
 	class MeshBufferIteratorConst;
 
-	template<EMeshBufferFormat FORMAT, class CTYPE, int COMPONENTS>
+	template<EMeshBufferFormat FORMAT, class CTYPE, int32 COMPONENTS>
 	class MUTABLERUNTIME_API MeshBufferIterator : public UntypedMeshBufferIterator
 	{
 	public:
 		inline MeshBufferIterator
 			(
-				FMeshBufferSet& bufferSet,
-				EMeshBufferSemantic semantic,
-				int semanticIndex = 0
+				FMeshBufferSet& BufferSet,
+				EMeshBufferSemantic Semantic,
+				int32 SemanticIndex = 0
 			)
-			: UntypedMeshBufferIterator( bufferSet, semantic, semanticIndex )
+			: UntypedMeshBufferIterator( BufferSet, Semantic, SemanticIndex )
 		{
-			if (m_pBuf)
+			if (Buffer)
 			{
 				// Extra checks
-				int buffer = -1;
-				int channel = -1;
-				bufferSet.FindChannel(semantic, semanticIndex, &buffer, &channel);
+				int32 BufferIndex = -1;
+				int32 ChannelIndex = -1;
+				BufferSet.FindChannel(Semantic, SemanticIndex, &BufferIndex, &ChannelIndex);
 				
-				if (buffer >= 0 && channel >= 0)
+				if (BufferIndex >= 0 && ChannelIndex >= 0)
 				{
 					EMeshBufferSemantic realSemantic = EMeshBufferSemantic::None;
-					int realSemanticIndex = 0;
+					int32 realSemanticIndex = 0;
 					EMeshBufferFormat format = EMeshBufferFormat::None;
-					int components = 0;
-					int offset = 0;
-					bufferSet.GetChannel
+					int32 components = 0;
+					int32 offset = 0;
+					BufferSet.GetChannel
 					(
-						buffer, channel,
+						BufferIndex, ChannelIndex,
 						&realSemantic, &realSemanticIndex,
 						&format, &components,
 						&offset
@@ -330,19 +329,19 @@ namespace mu
 					if (format!=FORMAT || components!=COMPONENTS)
 					{
 						// Invalidate
-						m_format = EMeshBufferFormat::None;
-						m_components = 0;
-						m_elementSize = 0;
-						m_pBuf = nullptr;
+						Format = EMeshBufferFormat::None;
+						Components = 0;
+						ElementSize = 0;
+						Buffer = nullptr;
 					}
 				}
 				else
 				{
 					// Invalidate
-					m_format = EMeshBufferFormat::None;
-					m_components = 0;
-					m_elementSize = 0;
-					m_pBuf = nullptr;
+					Format = EMeshBufferFormat::None;
+					Components = 0;
+					ElementSize = 0;
+					Buffer = nullptr;
 				}
 			}
 		}
@@ -350,10 +349,10 @@ namespace mu
 		// \TODO: Replace this with safer sized-typed access.
 		inline CTYPE* operator*()
 		{
-			return reinterpret_cast<CTYPE*>(m_pBuf);
+			return reinterpret_cast<CTYPE*>(Buffer);
 		}
 
-		inline MeshBufferIterator<FORMAT,CTYPE,COMPONENTS> operator+(int c) const
+		inline MeshBufferIterator<FORMAT,CTYPE,COMPONENTS> operator+(int32 c) const
 		{
 			MeshBufferIterator<FORMAT,CTYPE,COMPONENTS> res = *this;
 			res += c;
@@ -364,7 +363,7 @@ namespace mu
 
 
 	//---------------------------------------------------------------------------------------------
-	//! Class to iterate a specific buffer channel of unknown type
+	//! Class to iterate a specific BufferIndex ChannelIndex of unknown type
 	//---------------------------------------------------------------------------------------------
 	class MUTABLERUNTIME_API UntypedMeshBufferIteratorConst
 	{
@@ -378,111 +377,111 @@ namespace mu
 
 		inline UntypedMeshBufferIteratorConst
 			(
-				const FMeshBufferSet& bufferSet,
-				EMeshBufferSemantic semantic,
-				int semanticIndex = 0
+				const FMeshBufferSet& BufferSet,
+				EMeshBufferSemantic Semantic,
+				int32 SemanticIndex = 0
 			)
 		{
-			int buffer = -1;
-			int channel = -1;
-			bufferSet.FindChannel( semantic, semanticIndex, &buffer, &channel );
+			int32 BufferIndex = -1;
+			int32 ChannelIndex = -1;
+			BufferSet.FindChannel( Semantic, SemanticIndex, &BufferIndex, &ChannelIndex );
 
-			if ( buffer>=0 && channel>=0 )
+			if ( BufferIndex>=0 && ChannelIndex>=0 )
 			{
-				EMeshBufferSemantic realSemantic = EMeshBufferSemantic::None;
-				int realSemanticIndex = 0;
-				m_format = EMeshBufferFormat::None;
-				m_components = 0;
-				int offset = 0;
-				bufferSet.GetChannel
+				EMeshBufferSemantic RealSemantic = EMeshBufferSemantic::None;
+				int32 RealSemanticIndex = 0;
+				Format = EMeshBufferFormat::None;
+				Components = 0;
+				int32 offset = 0;
+				BufferSet.GetChannel
 						(
-							buffer, channel,
-							&realSemantic, &realSemanticIndex,
-							&m_format, &m_components,
+							BufferIndex, ChannelIndex,
+							&RealSemantic, &RealSemanticIndex,
+							&Format, &Components,
 							&offset
 						);
-				check( realSemantic == semantic );
-				check( realSemanticIndex == semanticIndex );
+				check( RealSemantic == Semantic );
+				check( RealSemanticIndex == SemanticIndex );
 
-				m_elementSize = bufferSet.GetElementSize( buffer );
+				ElementSize = BufferSet.GetElementSize( BufferIndex );
 
-				m_pBuf = bufferSet.GetBufferData( buffer );
-				m_pBuf += offset;
+				Buffer = BufferSet.GetBufferData( BufferIndex );
+				Buffer += offset;
 			}
 			else
 			{
-				m_format = EMeshBufferFormat::None;
-				m_components = 0;
-				m_elementSize = 0;
-				m_pBuf = nullptr;
+				Format = EMeshBufferFormat::None;
+				Components = 0;
+				ElementSize = 0;
+				Buffer = nullptr;
 			}
 		}
 
-        inline const uint8_t* ptr() const
+        inline const uint8* ptr() const
 		{
-			return m_pBuf;
+			return Buffer;
 		}
 
 		inline void operator++()
 		{
-			m_pBuf += m_elementSize;
+			Buffer += ElementSize;
 		}
 
-		inline void operator++(int)
+		inline void operator++(int32)
 		{
-			m_pBuf += m_elementSize;
+			Buffer += ElementSize;
 		}
 
-		inline void operator+=( int c )
+		inline void operator+=( int32 c )
 		{
-			m_pBuf += c*m_elementSize;
+			Buffer += c*ElementSize;
 		}
 
-		inline std::size_t operator-( const UntypedMeshBufferIterator& other ) const
+		inline SIZE_T operator-( const UntypedMeshBufferIterator& Other ) const
 		{
 			// Special degenerate case.
-			if (m_elementSize == 0)
+			if (ElementSize == 0)
 			{
 				return 0;
 			}
-			check( other.GetElementSize() == m_elementSize );
-			check( (ptr()-other.ptr()) % m_elementSize == 0 );
-			return (ptr()-other.ptr()) / m_elementSize;
+			check( Other.GetElementSize() == ElementSize );
+			check( (ptr()-Other.ptr()) % ElementSize == 0 );
+			return (ptr()-Other.ptr()) / ElementSize;
 		}
 
-		inline std::size_t operator-( const UntypedMeshBufferIteratorConst& other ) const
+		inline SIZE_T operator-( const UntypedMeshBufferIteratorConst& Other ) const
 		{
 			// Special degenerate case.
-			if (m_elementSize == 0)
+			if (ElementSize == 0)
 			{
 				return 0;
 			}
-			check( other.GetElementSize() == m_elementSize );
-			check( (ptr()-other.ptr()) % m_elementSize == 0 );
-			return (ptr()-other.ptr()) / m_elementSize;
+			check( Other.GetElementSize() == ElementSize );
+			check( (ptr()-Other.ptr()) % ElementSize == 0 );
+			return (ptr()-Other.ptr()) / ElementSize;
 		}
 
-		inline int GetElementSize() const
+		inline int32 GetElementSize() const
 		{
-			return m_elementSize;
+			return ElementSize;
 		}
 
         inline EMeshBufferFormat GetFormat() const
         {
-            return m_format;
+            return Format;
         }
 
         inline int32 GetComponents() const
         {
-            return m_components;
+            return Components;
         }
 
 		FVector4f GetAsVec4f() const
 		{
 			FVector4f res(0.0f,0.0f,0.0f,0.0f);
-			for (int c = 0; c < FMath::Min(m_components, 4); ++c)
+			for (int32 c = 0; c < FMath::Min(Components, 4); ++c)
 			{
-				ConvertData(c, &res[0], EMeshBufferFormat::Float32, ptr(), m_format);
+				ConvertData(c, &res[0], EMeshBufferFormat::Float32, ptr(), Format);
 			}
 			return res;
 		}
@@ -490,9 +489,9 @@ namespace mu
 		FVector3f GetAsVec3f() const
 		{
 			FVector3f res = { 0,0,0 };
-			for (int c = 0; c < FMath::Min(m_components, 3); ++c)
+			for (int32 c = 0; c < FMath::Min(Components, 3); ++c)
 			{
-				ConvertData(c, &res[0], EMeshBufferFormat::Float32, ptr(), m_format);
+				ConvertData(c, &res[0], EMeshBufferFormat::Float32, ptr(), Format);
 			}
 			return res;
 		}
@@ -500,9 +499,9 @@ namespace mu
 		FVector3d GetAsVec3d() const
 		{
 			FVector3d res = { 0,0,0 };
-			for (int c = 0; c < FMath::Min(m_components, 3); ++c)
+			for (int32 c = 0; c < FMath::Min(Components, 3); ++c)
 			{
-				ConvertData(c, &res[0], EMeshBufferFormat::Float64, ptr(), m_format);
+				ConvertData(c, &res[0], EMeshBufferFormat::Float64, ptr(), Format);
 			}
 			return res;
 		}
@@ -510,9 +509,9 @@ namespace mu
 		FVector2f GetAsVec2f() const
 		{
 			FVector2f res = {0.0f, 0.0f};
-			for (int32 c = 0; c < FMath::Min(m_components, 2); ++c)
+			for (int32 c = 0; c < FMath::Min(Components, 2); ++c)
 			{
-				ConvertData(c, &res[0], EMeshBufferFormat::Float32, ptr(), m_format);
+				ConvertData(c, &res[0], EMeshBufferFormat::Float32, ptr(), Format);
 			}
 
 			return res;
@@ -521,26 +520,26 @@ namespace mu
 		uint32 GetAsUINT32() const
 		{
 			uint32 res = 0;
-			ConvertData(0, &res, EMeshBufferFormat::UInt32, ptr(), m_format);
+			ConvertData(0, &res, EMeshBufferFormat::UInt32, ptr(), Format);
 			return res;
 		}
 
 		uint64 GetAsUINT64() const
 		{
 			uint64 res;
-			ConvertData(0, &res, EMeshBufferFormat::UInt64, ptr(), m_format);
+			ConvertData(0, &res, EMeshBufferFormat::UInt64, ptr(), Format);
 			return res;
 		}
 
 		void GetAsInt32Vec(int32* Data, int32 Count) const
 		{
-			for (int32 c = 0; c < FMath::Min(m_components, Count); ++c)
+			for (int32 c = 0; c < FMath::Min(Components, Count); ++c)
 			{
-				ConvertData(c, Data, EMeshBufferFormat::Int32, ptr(), m_format);
+				ConvertData(c, Data, EMeshBufferFormat::Int32, ptr(), Format);
 			}
 		}
 
-        inline UntypedMeshBufferIteratorConst operator+(int c) const
+        inline UntypedMeshBufferIteratorConst operator+(int32 c) const
         {
             UntypedMeshBufferIteratorConst res = *this;
             res += c;
@@ -549,19 +548,19 @@ namespace mu
 
 	protected:
 
-		int m_elementSize = 0;
-        const uint8_t* m_pBuf = nullptr;
-		EMeshBufferFormat m_format = EMeshBufferFormat::None;
-		int m_components = 0;
+		int32 ElementSize = 0;
+        const uint8* Buffer = nullptr;
+		EMeshBufferFormat Format = EMeshBufferFormat::None;
+		int32 Components = 0;
 	};
 
 
 
 
 	//---------------------------------------------------------------------------------------------
-	//! Class to iterate a specific buffer channel of a constant buffer set
+	//! Class to iterate a specific BufferIndex ChannelIndex of a constant BufferIndex set
 	//---------------------------------------------------------------------------------------------
-	template<EMeshBufferFormat FORMAT, class CTYPE, int COMPONENTS>
+	template<EMeshBufferFormat FORMAT, class CTYPE, int32 COMPONENTS>
 	class MUTABLERUNTIME_API MeshBufferIteratorConst : public UntypedMeshBufferIteratorConst
 	{
 	public:
@@ -573,29 +572,29 @@ namespace mu
 
 		inline MeshBufferIteratorConst
 			(
-				const FMeshBufferSet& bufferSet,
-				EMeshBufferSemantic semantic,
-				int semanticIndex = 0
+				const FMeshBufferSet& BufferSet,
+				EMeshBufferSemantic Semantic,
+				int32 SemanticIndex = 0
 			)
-			: UntypedMeshBufferIteratorConst( bufferSet, semantic, semanticIndex )
+			: UntypedMeshBufferIteratorConst( BufferSet, Semantic, SemanticIndex )
 		{
-			if (m_pBuf)
+			if (Buffer)
 			{
 				// Extra checks
-				int buffer = -1;
-				int channel = -1;
-				bufferSet.FindChannel(semantic, semanticIndex, &buffer, &channel);
+				int32 BufferIndex = -1;
+				int32 ChannelIndex = -1;
+				BufferSet.FindChannel(Semantic, SemanticIndex, &BufferIndex, &ChannelIndex);
 
-				if (buffer >= 0 && channel >= 0)
+				if (BufferIndex >= 0 && ChannelIndex >= 0)
 				{
 					EMeshBufferSemantic realSemantic = EMeshBufferSemantic::None;
-					int realSemanticIndex = 0;
+					int32 realSemanticIndex = 0;
 					EMeshBufferFormat format = EMeshBufferFormat::None;
-					int components = 0;
-					int offset = 0;
-					bufferSet.GetChannel
+					int32 components = 0;
+					int32 offset = 0;
+					BufferSet.GetChannel
 					(
-						buffer, channel,
+						BufferIndex, ChannelIndex,
 						&realSemantic, &realSemanticIndex,
 						&format, &components,
 						&offset
@@ -604,29 +603,29 @@ namespace mu
 					if (format!=FORMAT || components!=COMPONENTS)
 					{
 						// Invalidate
-						m_format = EMeshBufferFormat::None;
-						m_components = 0;
-						m_elementSize = 0;
-						m_pBuf = nullptr;
+						Format = EMeshBufferFormat::None;
+						Components = 0;
+						ElementSize = 0;
+						Buffer = nullptr;
 					}
 				}
 				else
 				{
 					// Invalidate
-					m_format = EMeshBufferFormat::None;
-					m_components = 0;
-					m_elementSize = 0;
-					m_pBuf = nullptr;
+					Format = EMeshBufferFormat::None;
+					Components = 0;
+					ElementSize = 0;
+					Buffer = nullptr;
 				}
 			}
 		}
 
 		inline const CTYPE* operator*() const
 		{
-			return reinterpret_cast<const CTYPE*>(m_pBuf);
+			return reinterpret_cast<const CTYPE*>(Buffer);
 		}
 
-		inline MeshBufferIteratorConst<FORMAT,CTYPE,COMPONENTS> operator+(int c) const
+		inline MeshBufferIteratorConst<FORMAT,CTYPE,COMPONENTS> operator+(int32 c) const
 		{
 			MeshBufferIteratorConst<FORMAT,CTYPE,COMPONENTS> res = *this;
 			res += c;
@@ -636,7 +635,9 @@ namespace mu
 	};
 
 
-	/** */
+	/** Iterator for vertex Ids.
+	* The lifecycle of the iterator cannot exceed the mesh lifecycle. 
+	*/
 	class MUTABLERUNTIME_API MeshVertexIdIteratorConst
 	{
 	private:
@@ -645,9 +646,9 @@ namespace mu
 		int32 CurrentIdIndex=0;
 
 		/** */
-		Ptr<const Mesh> Mesh;
+		const FMesh* Mesh=nullptr;
 
-		/** Buffer iterator in case there is an actual Id buffer. */
+		/** Buffer iterator in case there is an actual Id BufferIndex. */
 		UntypedMeshBufferIteratorConst BufferIterator;
 
 	public:
@@ -656,7 +657,7 @@ namespace mu
 		{
 		}
 
-		inline MeshVertexIdIteratorConst( const mu::Mesh* InMesh )
+		inline MeshVertexIdIteratorConst( const FMesh* InMesh )
 		{
 			if (!InMesh)
 			{
@@ -705,7 +706,7 @@ namespace mu
 			}
 		}
 
-		inline MeshVertexIdIteratorConst operator+(int c) const
+		inline MeshVertexIdIteratorConst operator+(int32 c) const
 		{
 			MeshVertexIdIteratorConst res = *this;
 			res += c;
@@ -726,7 +727,7 @@ namespace mu
 			check(Mesh);
 			check(CurrentIdIndex < Mesh->GetVertexCount());
 
-			uint64 Id=Mesh::InvalidVertexId;
+			uint64 Id=FMesh::InvalidVertexId;
 
 			// Is it implicit?
 			if (!BufferIterator.ptr())
@@ -738,7 +739,7 @@ namespace mu
 			// Is it relative?
 			else if (BufferIterator.GetFormat() == EMeshBufferFormat::UInt32)
 			{
-				// There is a buffer storing IDs without prefix because it is the same for all vertices.
+				// There is a BufferIndex storing IDs without prefix because it is the same for all vertices.
 				uint32 RelativeId = BufferIterator.GetAsUINT32();
 				Id = (uint64(Mesh->MeshIDPrefix) << 32) | uint64(RelativeId);
 			}
