[mutable] More enforcement of coding standards.

diff --git a/Engine/Plugins/Mutable/Source/MutableRuntime/Internal/MuR/ModelPrivate.h b/Engine/Plugins/Mutable/Source/MutableRuntime/Internal/MuR/ModelPrivate.h
--- a/Engine/Plugins/Mutable/Source/MutableRuntime/Internal/MuR/ModelPrivate.h
+++ b/Engine/Plugins/Mutable/Source/MutableRuntime/Internal/MuR/ModelPrivate.h
@@ -56,11 +56,11 @@ namespace mu
 
     //!
     template<typename DATA>
-    inline void AppendCode(TArray<uint8>& code, const DATA& data )
+    inline void AppendCode(TArray<uint8>& Code, const DATA& Data )
     {
-        int32 pos = code.Num();
-        code.SetNum( pos+sizeof(DATA), EAllowShrinking::No);
-		FMemory::Memcpy (&code[pos], &data, sizeof(DATA));
+        int32 Pos = Code.Num();
+        Code.SetNum( Pos+sizeof(DATA), EAllowShrinking::No);
+		FMemory::Memcpy (&Code[Pos], &Data, sizeof(DATA));
     }
 
 
@@ -81,7 +81,7 @@ namespace mu
 		// This should always be valid, but if the state is Unloaded it won't be usable.
 		//
 		// Avoid storing references to this Data in Memory while the state is Unloaded.
-		ExtensionDataPtrConst Data;
+		TSharedPtr<const FExtensionData> Data;
 
 		enum class ELoadState : uint8
 		{
@@ -95,20 +95,20 @@ namespace mu
 		// This should be initialized to a valid load state when Data is set
 		ELoadState LoadState = ELoadState::Invalid;
 
-		inline void Serialise(OutputArchive& arch) const
+		inline void Serialise(FOutputArchive& Arch) const
 		{
-			arch << Data;
+			Arch << Data;
 		}
 		
-		inline void Unserialise(InputArchive& arch)
+		inline void Unserialise(FInputArchive& Arch)
 		{
-			arch >> Data;
+			Arch >> Data;
 
-			check(Data.get());
-			check(Data->Origin == ExtensionData::EOrigin::ConstantAlwaysLoaded
-				|| Data->Origin == ExtensionData::EOrigin::ConstantStreamed);
+			check(Data);
+			check(Data->Origin == FExtensionData::EOrigin::ConstantAlwaysLoaded
+				|| Data->Origin == FExtensionData::EOrigin::ConstantStreamed);
 
-			if (Data->Origin == ExtensionData::EOrigin::ConstantAlwaysLoaded)
+			if (Data->Origin == FExtensionData::EOrigin::ConstantAlwaysLoaded)
 			{
 				LoadState = ELoadState::AlwaysLoaded;
 			}
@@ -127,8 +127,8 @@ namespace mu
         {
             // Add the null instruction at address 0.
             // TODO: Will do it in the linker
-            AppendCode( m_byteCode, OP_TYPE::NONE );
-            m_opAddress.Add(0);
+            AppendCode( ByteCode, EOpType::NONE );
+            OpAddress.Add(0);
         }
 
         struct FState
@@ -137,9 +137,9 @@ namespace mu
             FString Name;
 
             //! First instruction of the full build of an instance in this state
-            OP::ADDRESS m_root = 0;
+            OP::ADDRESS Root = 0;
 
-            //! List of parameters index (to FProgram::m_parameters) of the runtime parameters of
+            //! List of parameters index (to FProgram::Parameters) of the runtime parameters of
             //! this state.
 			TArray<int> m_runtimeParameters;
 
@@ -153,20 +153,20 @@ namespace mu
 			TArray< TPair<OP::ADDRESS,uint64> > m_dynamicResources;
 
             //!
-            inline void Serialise( OutputArchive& arch ) const
+            inline void Serialise( FOutputArchive& arch ) const
             {
                 arch << Name;
-                arch << m_root;
+                arch << Root;
                 arch << m_runtimeParameters;
                 arch << m_updateCache;
                 arch << m_dynamicResources;
             }
 
             //!
-            inline void Unserialise( InputArchive& arch )
+            inline void Unserialise( FInputArchive& arch )
             {
                 arch >> Name;
-                arch >> m_root;
+                arch >> Root;
                 arch >> m_runtimeParameters;
                 arch >> m_updateCache;
                 arch >> m_dynamicResources;
@@ -196,7 +196,7 @@ namespace mu
             {
                 bool res = false;
 
-                for ( SIZE_T i=0; !res && i<m_updateCache.Num(); ++i )
+                for (int32 i=0; !res && i<m_updateCache.Num(); ++i )
                 {
                     if ( m_updateCache[i]==at )
                     {
@@ -217,23 +217,23 @@ namespace mu
             }
         };
 
-        //! Location in the m_byteCode of the beginning of each operation
-		TArray<uint32> m_opAddress;
+        //! Location in the ByteCode of the beginning of each operation
+		TArray<uint32> OpAddress;
 
         //! Byte-coded representation of the program, using variable-sized op data.
-		TArray<uint8> m_byteCode;
+		TArray<uint8> ByteCode;
 
         //!
-		TArray<FState> m_states;
+		TArray<FState> States;
 
 		//! Data for every rom.
-		TArray<FRomData> m_roms;
+		TArray<FRomData> Roms;
 	
 		/** Loaded roms worth tracking (only images and meshes for now). Stores the rom's data type.*/
 		TSparseArray<uint8> LoadedMemTrackedRoms;
 
-		//! Constant image mip data: the first is the index in m_roms for each image mip or -1 if it is always loaded.
-		TArray<TPair<int32, Ptr<const Image>>> ConstantImageLODs;
+		//! Constant image mip data: the first is the index in Roms for each image mip or -1 if it is always loaded.
+		TArray<TPair<int32, TSharedPtr<const FImage>>> ConstantImageLODs;
 
 		//! Constant image mip chain indices: ranges in this array are defined in FImageLODRange and the indices here refer to ConstantImageLODs.
 		TArray<uint32> ConstantImageLODIndices;
@@ -241,47 +241,47 @@ namespace mu
 		//! Constant image data.
 		TArray<FImageLODRange> ConstantImages;
 
-        //! Constant mesh data: the first is the index in m_roms for each mesh or -1 if it is always loaded.
-		TArray<TPair<int32, Ptr<const Mesh>>> ConstantMeshes;
+        //! Constant mesh data: the first is the index in Roms for each mesh or -1 if it is always loaded.
+		TArray<TPair<int32, TSharedPtr<const FMesh>>> ConstantMeshes;
 
-		//! Constant ExtensionData
-		TArray<FExtensionDataConstant> m_constantExtensionData;
+		//! Constant FExtensionData
+		TArray<FExtensionDataConstant> ConstantExtensionData;
 
         //! Constant string data
-		TArray<FString> m_constantStrings;
+		TArray<FString> ConstantStrings;
 
         //! Constant layout data
-		TArray<Ptr<const Layout>> m_constantLayouts;
+		TArray<TSharedPtr<const FLayout>> ConstantLayouts;
 
         //! Constant projectors
-		TArray<FProjector> m_constantProjectors;
+		TArray<FProjector> ConstantProjectors;
 
         //! Constant matrices, usually used for transforms
-		TArray<FMatrix44f> m_constantMatrices;
+		TArray<FMatrix44f> ConstantMatrices;
 
 		//! Constant shapes
-		TArray<FShape> m_constantShapes;
+		TArray<FShape> ConstantShapes;
 
         //! Constant curves
 		TArray<FRichCurve> ConstantCurves;
 
         //! Constant skeletons
-		TArray<Ptr<const Skeleton>> m_constantSkeletons;
+		TArray<TSharedPtr<const FSkeleton>> ConstantSkeletons;
 
 		//! Constant Physics Bodies
-		TArray<Ptr<const PhysicsBody>> m_constantPhysicsBodies;
+		TArray<TSharedPtr<const FPhysicsBody>> ConstantPhysicsBodies;
 
-        //! Parameters of the model.
+        //! FParameters of the model.
         //! The value stored here is the default value.
-		TArray<FParameterDesc> m_parameters;
+		TArray<FParameterDesc> Parameters;
 
         //! Ranges for iteration of the model operations.
-		TArray<FRangeDesc> m_ranges;
+		TArray<FRangeDesc> Ranges;
 
         //! List of parameter lists. These are used in several places, like storing the
         //! pregenerated list of parameters influencing a resource.
         //! The parameter lists are sorted.
-		TArray<TArray<uint16>> m_parameterLists;
+		TArray<TArray<uint16>> ParameterLists;
 
 #if WITH_EDITOR
 		//! State of the program. True unless the streamed resources were destroyed,
@@ -289,53 +289,53 @@ namespace mu
 		bool bIsValid = true;
 #endif
         //!
-        void Serialise( OutputArchive& arch ) const
+        void Serialise( FOutputArchive& arch ) const
         {
-            arch << m_opAddress;
-            arch << m_byteCode;
-            arch << m_states;
-			arch << m_roms;
+            arch << OpAddress;
+            arch << ByteCode;
+            arch << States;
+			arch << Roms;
 			arch << ConstantImageLODs;
 			arch << ConstantImageLODIndices;
 			arch << ConstantImages;
 			arch << ConstantMeshes;
-			arch << m_constantExtensionData;
-			arch << m_constantStrings;
-            arch << m_constantLayouts;
-            arch << m_constantProjectors;
-			arch << m_constantMatrices;
-			arch << m_constantShapes;
+			arch << ConstantExtensionData;
+			arch << ConstantStrings;
+            arch << ConstantLayouts;
+            arch << ConstantProjectors;
+			arch << ConstantMatrices;
+			arch << ConstantShapes;
             arch << ConstantCurves;
-            arch << m_constantSkeletons;
-			arch << m_constantPhysicsBodies;
-            arch << m_parameters;
-            arch << m_ranges;
-            arch << m_parameterLists;
+            arch << ConstantSkeletons;
+			arch << ConstantPhysicsBodies;
+            arch << Parameters;
+            arch << Ranges;
+            arch << ParameterLists;
         }
 
         //!
-        void Unserialise( InputArchive& arch )
+        void Unserialise( FInputArchive& arch )
         {
-            arch >> m_opAddress;
-            arch >> m_byteCode;
-            arch >> m_states;
-			arch >> m_roms;
+            arch >> OpAddress;
+            arch >> ByteCode;
+            arch >> States;
+			arch >> Roms;
 			arch >> ConstantImageLODs;
 			arch >> ConstantImageLODIndices;
 			arch >> ConstantImages;
 			arch >> ConstantMeshes;
-			arch >> m_constantExtensionData;
-			arch >> m_constantStrings;
-            arch >> m_constantLayouts;
-            arch >> m_constantProjectors;
-			arch >> m_constantMatrices;
-			arch >> m_constantShapes;
+			arch >> ConstantExtensionData;
+			arch >> ConstantStrings;
+            arch >> ConstantLayouts;
+            arch >> ConstantProjectors;
+			arch >> ConstantMatrices;
+			arch >> ConstantShapes;
             arch >> ConstantCurves;
-            arch >> m_constantSkeletons;
-			arch >> m_constantPhysicsBodies;
-            arch >> m_parameters;
-            arch >> m_ranges;
-            arch >> m_parameterLists;
+            arch >> ConstantSkeletons;
+			arch >> ConstantPhysicsBodies;
+            arch >> Parameters;
+            arch >> Ranges;
+            arch >> ParameterLists;
         }
 
         //! Debug method that sanity-checks the program with a variety of tests.
@@ -347,12 +347,12 @@ namespace mu
 		//! Return true if the given ROM is loaded.
 		FORCEINLINE bool IsRomLoaded(int32 RomIndex) const
 		{
-			switch (m_roms[RomIndex].ResourceType)
+			switch (Roms[RomIndex].ResourceType)
 			{
 				case DT_IMAGE: 
-					return ConstantImageLODs[m_roms[RomIndex].ResourceIndex].Value.get() != nullptr;
+					return ConstantImageLODs[Roms[RomIndex].ResourceIndex].Value != nullptr;
 				case DT_MESH: 
-					return ConstantMeshes[m_roms[RomIndex].ResourceIndex].Value.get() != nullptr;
+					return ConstantMeshes[Roms[RomIndex].ResourceIndex].Value != nullptr;
 				default:
 					check(false);
 					break;
@@ -374,20 +374,20 @@ namespace mu
 				LoadedMemTrackedRoms.RemoveAt(RomIndex);
 			}
 
-			switch (m_roms[RomIndex].ResourceType)
+			switch (Roms[RomIndex].ResourceType)
 			{
 			case DT_IMAGE:
-				if (ConstantImageLODs[m_roms[RomIndex].ResourceIndex].Value)
+				if (ConstantImageLODs[Roms[RomIndex].ResourceIndex].Value)
 				{
-					RomSize = ConstantImageLODs[m_roms[RomIndex].ResourceIndex].Value->GetDataSize();
-					ConstantImageLODs[m_roms[RomIndex].ResourceIndex].Value = nullptr;
+					RomSize = ConstantImageLODs[Roms[RomIndex].ResourceIndex].Value->GetDataSize();
+					ConstantImageLODs[Roms[RomIndex].ResourceIndex].Value = nullptr;
 				}
 				break;
 			case DT_MESH:
-				if (ConstantMeshes[m_roms[RomIndex].ResourceIndex].Value)
+				if (ConstantMeshes[Roms[RomIndex].ResourceIndex].Value)
 				{
-					RomSize = ConstantMeshes[m_roms[RomIndex].ResourceIndex].Value->GetDataSize();
-					ConstantMeshes[m_roms[RomIndex].ResourceIndex].Value = nullptr;
+					RomSize = ConstantMeshes[Roms[RomIndex].ResourceIndex].Value->GetDataSize();
+					ConstantMeshes[Roms[RomIndex].ResourceIndex].Value = nullptr;
 				}
 				break;
 			default:
@@ -398,168 +398,168 @@ namespace mu
 			return RomSize;
 		}
 
-		FORCEINLINE void SetMeshRomValue(int32 RomIndex, const Ptr<Mesh>& Value)
+		FORCEINLINE void SetMeshRomValue(int32 RomIndex, const TSharedPtr<FMesh>& Value)
 		{
-			check(m_roms[RomIndex].ResourceType == DT_MESH);
+			check(Roms[RomIndex].ResourceType == DT_MESH);
 			
-			LoadedMemTrackedRoms.EmplaceAt(RomIndex, (uint8)m_roms[RomIndex].ResourceType);
-			ConstantMeshes[m_roms[RomIndex].ResourceIndex].Value = Value;
+			LoadedMemTrackedRoms.EmplaceAt(RomIndex, (uint8)Roms[RomIndex].ResourceType);
+			ConstantMeshes[Roms[RomIndex].ResourceIndex].Value = Value;
 		}
 
-		FORCEINLINE void SetImageRomValue(int32 RomIndex, const Ptr<Image>& Value)
+		FORCEINLINE void SetImageRomValue(int32 RomIndex, const TSharedPtr<FImage>& Value)
 		{
-			check(m_roms[RomIndex].ResourceType == DT_IMAGE);
+			check(Roms[RomIndex].ResourceType == DT_IMAGE);
 			
-			LoadedMemTrackedRoms.EmplaceAt(RomIndex, (uint8)m_roms[RomIndex].ResourceType);
-			ConstantImageLODs[m_roms[RomIndex].ResourceIndex].Value = Value;
+			LoadedMemTrackedRoms.EmplaceAt(RomIndex, (uint8)Roms[RomIndex].ResourceType);
+			ConstantImageLODs[Roms[RomIndex].ResourceIndex].Value = Value;
 		}
 
 
-		int32 AddConstant(Ptr<const Mesh> pMesh)
+		int32 AddConstant(TSharedPtr<const FMesh> pMesh)
 		{
 			// Uniques needs to be ensured outside
-			return ConstantMeshes.Add(TPair<int32, Ptr<const Mesh>>( -1, pMesh.get() ));
+			return ConstantMeshes.Add(TPair<int32, TSharedPtr<const FMesh>>( -1, pMesh ));
 		}
 
-		OP::ADDRESS AddConstant(Ptr<const ExtensionData> Data)
+		OP::ADDRESS AddConstant(TSharedPtr<const FExtensionData> Data)
 		{
 			// Ensure unique
-			for (int32 Index = 0; Index < m_constantExtensionData.Num(); Index++)
+			for (int32 Index = 0; Index < ConstantExtensionData.Num(); Index++)
 			{
-				const ExtensionData* Candidate = m_constantExtensionData[Index].Data.get();
+				const FExtensionData* Candidate = ConstantExtensionData[Index].Data.Get();
 				if (*Candidate == *Data)
 				{
 					return Index;
 				}
 			}
 
-			FExtensionDataConstant& NewConstant = m_constantExtensionData.AddDefaulted_GetRef();
+			FExtensionDataConstant& NewConstant = ConstantExtensionData.AddDefaulted_GetRef();
 			NewConstant.Data = Data;
 
 			// The data is assumed to be loaded during compilation
 			NewConstant.LoadState =
-				Data->Origin == ExtensionData::EOrigin::ConstantAlwaysLoaded
+				Data->Origin == FExtensionData::EOrigin::ConstantAlwaysLoaded
 				? FExtensionDataConstant::ELoadState::AlwaysLoaded
 				: FExtensionDataConstant::ELoadState::CurrentlyLoaded;
 
-			return m_constantExtensionData.Num() - 1;
+			return ConstantExtensionData.Num() - 1;
 		}
 
-		OP::ADDRESS AddConstant( Ptr<const Layout> pLayout )
+		OP::ADDRESS AddConstant( TSharedPtr<const FLayout> pLayout )
         {
             // Ensure unique
-            for ( SIZE_T i=0; i<m_constantLayouts.Num(); ++i)
+            for (int32 i=0; i<ConstantLayouts.Num(); ++i)
             {
-                if (m_constantLayouts[i]==pLayout)
+                if (ConstantLayouts[i]==pLayout)
                 {
                     return (OP::ADDRESS)i;
                 }
             }
 
-            OP::ADDRESS index = OP::ADDRESS( m_constantLayouts.Num() );
-            m_constantLayouts.Add( pLayout );
+            OP::ADDRESS index = OP::ADDRESS( ConstantLayouts.Num() );
+            ConstantLayouts.Add( pLayout );
             return index;
         }
 
-        OP::ADDRESS AddConstant( Ptr<const Skeleton> pSkeleton )
+        OP::ADDRESS AddConstant( TSharedPtr<const FSkeleton> Skeleton )
         {
             // Ensure unique
-            for ( SIZE_T i=0; i<m_constantSkeletons.Num(); ++i)
+            for ( int32 i=0; i<ConstantSkeletons.Num(); ++i)
             {
-                if ( m_constantSkeletons[i]==pSkeleton
+                if ( ConstantSkeletons[i]== Skeleton
                      ||
-                     *m_constantSkeletons[i]==*pSkeleton
+                     *ConstantSkeletons[i]==*Skeleton
                      )
                 {
                     return (OP::ADDRESS)i;
                 }
             }
 
-            OP::ADDRESS index = OP::ADDRESS( m_constantSkeletons.Num() );
-            m_constantSkeletons.Add( pSkeleton );
+            OP::ADDRESS index = OP::ADDRESS( ConstantSkeletons.Num() );
+            ConstantSkeletons.Add(Skeleton);
             return index;
         }
 
-        OP::ADDRESS AddConstant( Ptr<const PhysicsBody> pPhysicsBody )
+        OP::ADDRESS AddConstant(TSharedPtr<const FPhysicsBody> pPhysicsBody )
         {
             // Ensure unique
-            for ( SIZE_T i=0; i<m_constantPhysicsBodies.Num(); ++i)
+            for (int32 i=0; i<ConstantPhysicsBodies.Num(); ++i)
             {
-                if ( m_constantPhysicsBodies[i]==pPhysicsBody
+                if ( ConstantPhysicsBodies[i]==pPhysicsBody
                      ||
-                     *m_constantPhysicsBodies[i]==*pPhysicsBody
+                     *ConstantPhysicsBodies[i]==*pPhysicsBody
                      )
                 {
                     return (OP::ADDRESS)i;
                 }
             }
 
-            OP::ADDRESS index = OP::ADDRESS( m_constantPhysicsBodies.Num() );
-            m_constantPhysicsBodies.Add( pPhysicsBody );
+            OP::ADDRESS index = OP::ADDRESS( ConstantPhysicsBodies.Num() );
+            ConstantPhysicsBodies.Add( pPhysicsBody );
             return index;
         }
 
         OP::ADDRESS AddConstant( const FString& str )
         {            
             // Ensure unique
-            for ( SIZE_T i=0; i<m_constantStrings.Num(); ++i)
+            for (int32 i=0; i<ConstantStrings.Num(); ++i)
             {
-                if (m_constantStrings[i]==str)
+                if (ConstantStrings[i]==str)
                 {
                     return (OP::ADDRESS)i;
                 }
             }
 
-            OP::ADDRESS index = OP::ADDRESS( m_constantStrings.Num() );
-            m_constantStrings.Add( str );
+            OP::ADDRESS index = OP::ADDRESS( ConstantStrings.Num() );
+            ConstantStrings.Add( str );
             return index;
         }
 
         OP::ADDRESS AddConstant( const FMatrix44f& m )
         {
             // Ensure unique
-            for ( SIZE_T i=0; i<m_constantMatrices.Num(); ++i)
+            for (int32 i=0; i<ConstantMatrices.Num(); ++i)
             {
-                if (m_constantMatrices[i]==m)
+                if (ConstantMatrices[i]==m)
                 {
                     return (OP::ADDRESS)i;
                 }
             }
 
-            OP::ADDRESS index = OP::ADDRESS( m_constantMatrices.Num() );
-            m_constantMatrices.Add( m );
+            OP::ADDRESS index = OP::ADDRESS( ConstantMatrices.Num() );
+            ConstantMatrices.Add( m );
             return index;
         }
 
         OP::ADDRESS AddConstant( const FShape& m )
         {
             // Ensure unique
-            for ( SIZE_T i=0; i<m_constantShapes.Num(); ++i)
+            for (int32 i=0; i<ConstantShapes.Num(); ++i)
             {
-                if (m_constantShapes[i]==m)
+                if (ConstantShapes[i]==m)
                 {
                     return (OP::ADDRESS)i;
                 }
             }
 
-            OP::ADDRESS index = OP::ADDRESS( m_constantShapes.Num() );
-            m_constantShapes.Add( m );
+            OP::ADDRESS index = OP::ADDRESS( ConstantShapes.Num() );
+            ConstantShapes.Add( m );
             return index;
         }
 
         OP::ADDRESS AddConstant( const FProjector& m )
         {
             // Ensure unique
-            for ( SIZE_T i=0; i<m_constantProjectors.Num(); ++i)
+            for (int32 i=0; i<ConstantProjectors.Num(); ++i)
             {
-                if (m_constantProjectors[i]==m)
+                if (ConstantProjectors[i]==m)
                 {
                     return (OP::ADDRESS)i;
                 }
             }
 
-            OP::ADDRESS index = OP::ADDRESS( m_constantProjectors.Num() );
-            m_constantProjectors.Add( m );
+            OP::ADDRESS index = OP::ADDRESS( ConstantProjectors.Num() );
+            ConstantProjectors.Add( m );
             return index;
         }
 
@@ -571,7 +571,7 @@ namespace mu
 
         //! Get a constant image, assuming at least some mips are loaded. The image constant will be composed with lodaded mips if necessary.
 		template <typename CreateImageFunc>
-        void GetConstant( int32 ConstantIndex, Ptr<const Image>& res, int32 MipsToSkip, const CreateImageFunc& CreateImage) const
+        void GetConstant( int32 ConstantIndex, TSharedPtr<const FImage>& res, int32 MipsToSkip, const CreateImageFunc& CreateImage) const
         {
 			int32 ReallySkippedLODs = FMath::Min(ConstantImages[ConstantIndex].LODCount - 1, MipsToSkip);
 			int32 FirstLODIndexIndex = ConstantImages[ConstantIndex].FirstIndex;
@@ -579,7 +579,7 @@ namespace mu
 			// Get the first mip
 			int32 ResultLODIndexIndex = FirstLODIndexIndex + ReallySkippedLODs;
 			int32 ResultLODIndex = ConstantImageLODIndices[ResultLODIndexIndex];
-			Ptr<const Image> CurrentMip = ConstantImageLODs[ResultLODIndex].Value;
+			TSharedPtr<const FImage> CurrentMip = ConstantImageLODs[ResultLODIndex].Value;
 
 			// We may need to skip more LODs if they are not loaded
 			while (!CurrentMip)
@@ -612,8 +612,8 @@ namespace mu
 			{
 				MUTABLE_CPUPROFILER_SCOPE(ComposeConstantImage);
 
-				Ptr<Image> Result = CreateImage(CurrentMip->GetSizeX(), CurrentMip->GetSizeY(), FinalLODs, CurrentMip->GetFormat(), EInitializationType::NotInitialized);
-				Result->m_flags = CurrentMip->m_flags;
+				TSharedPtr<FImage> Result = CreateImage(CurrentMip->GetSizeX(), CurrentMip->GetSizeY(), FinalLODs, CurrentMip->GetFormat(), EInitializationType::NotInitialized);
+				Result->Flags = CurrentMip->Flags;
 
 				// Some non-block pixel formats require separate memory size calculation
 				if (Result->DataStorage.IsEmpty())
@@ -650,29 +650,29 @@ namespace mu
 			}
 		}
 
-        void GetConstant(int32 ConstantIndex, MeshPtrConst& res) const
+        void GetConstant(int32 ConstantIndex, TSharedPtr<const FMesh>& res) const
         {
 			res = ConstantMeshes[ConstantIndex].Value;
 		}
 
-		void GetExtensionDataConstant(int32 ConstantIndex, ExtensionDataPtrConst& Result) const
+		void GetExtensionDataConstant(int32 ConstantIndex, TSharedPtr<const FExtensionData>& Result) const
 		{
-			const FExtensionDataConstant& Constant = m_constantExtensionData[ConstantIndex];
+			const FExtensionDataConstant& Constant = ConstantExtensionData[ConstantIndex];
 
 			check(Constant.LoadState != FExtensionDataConstant::ELoadState::Unloaded);
-			check(Constant.Data.get());
+			check(Constant.Data);
 
 			Result = Constant.Data;
 		}
 
-        inline OP_TYPE GetOpType( OP::ADDRESS at ) const
+        inline EOpType GetOpType( OP::ADDRESS at ) const
         {
-            if (at>=OP::ADDRESS(m_opAddress.Num())) return OP_TYPE::NONE;
+            if (at>=OP::ADDRESS(OpAddress.Num())) return EOpType::NONE;
 
-            OP_TYPE result;
-            uint64 byteCodeAddress = m_opAddress[at];
-			FMemory::Memcpy( &result, &m_byteCode[byteCodeAddress], sizeof(OP_TYPE) );
-            check( result<OP_TYPE::COUNT);
+            EOpType result;
+            uint64 byteCodeAddress = OpAddress[at];
+			FMemory::Memcpy( &result, &ByteCode[byteCodeAddress], sizeof(EOpType) );
+            check( result<EOpType::COUNT);
             return result;
         }
 
@@ -680,80 +680,49 @@ namespace mu
         inline const ARGS GetOpArgs( OP::ADDRESS at ) const
         {
             ARGS result;
-            uint64 byteCodeAddress = m_opAddress[at];
-            byteCodeAddress += sizeof(OP_TYPE);
-			FMemory::Memcpy( &result, &m_byteCode[byteCodeAddress], sizeof(ARGS));
+            uint64 byteCodeAddress = OpAddress[at];
+            byteCodeAddress += sizeof(EOpType);
+			FMemory::Memcpy( &result, &ByteCode[byteCodeAddress], sizeof(ARGS));
             return result;
         }
 
         inline const uint8* GetOpArgsPointer( OP::ADDRESS at ) const
         {
-            uint64 byteCodeAddress = m_opAddress[at];
-            byteCodeAddress += sizeof(OP_TYPE);
-            const uint8* pData = (const uint8*)&m_byteCode[byteCodeAddress];
+            uint64 byteCodeAddress = OpAddress[at];
+            byteCodeAddress += sizeof(EOpType);
+            const uint8* pData = (const uint8*)&ByteCode[byteCodeAddress];
             return pData;
         }
 
         inline uint8* GetOpArgsPointer( OP::ADDRESS at )
         {
-            uint64 byteCodeAddress = m_opAddress[at];
-            byteCodeAddress += sizeof(OP_TYPE);
-            uint8* pData = (uint8*)&m_byteCode[byteCodeAddress];
+            uint64 byteCodeAddress = OpAddress[at];
+            byteCodeAddress += sizeof(EOpType);
+            uint8* pData = (uint8*)&ByteCode[byteCodeAddress];
             return pData;
         }
     };
 
 
     //!
-    class MUTABLERUNTIME_API Model::Private
+    class MUTABLERUNTIME_API FModel::Private
     {
     public:
 
-        //!
-        FProgram m_program;
-
+        FProgram Program;
 
     	void UnloadRoms();
 
-
-        //!
-        void Serialise( OutputArchive& arch ) const
+        void Serialise( FOutputArchive& arch ) const
         {
-            uint32 ver = 0;
-            arch << ver;
-
-            arch << m_program;
+             arch << Program;
         }
 
-        //!
-        void Unserialise( InputArchive& arch )
+        void Unserialise( FInputArchive& arch )
         {
-            uint32 ver;
-            arch >> ver;
-            check(ver==0);
-
-            arch >> m_program;
+            arch >> Program;
         }
 
     };
 
-
-    //!
-    class MUTABLERUNTIME_API ModelParametersGenerator::Private
-    {
-    public:
-
-        //!
-        TSharedPtr<const Model> m_pModel;
-        Ptr<System> m_pSystem;
-
-        //! Number of possible instances
-        int64 m_instanceCount;
-
-		/** Value used for scalar paramters that control multidemnsional sizes. 
-		* TODO: Make it an option.
-		*/
-		int32 DefaultRangeDimension = 8;
-    };
-
 }
diff --git a/Engine/Plugins/Mutable/Source/MutableRuntime/Internal/MuR/OpImageGrow.h b/Engine/Plugins/Mutable/Source/MutableRuntime/Internal/MuR/OpImageGrow.h
--- a/Engine/Plugins/Mutable/Source/MutableRuntime/Internal/MuR/OpImageGrow.h
+++ b/Engine/Plugins/Mutable/Source/MutableRuntime/Internal/MuR/OpImageGrow.h
@@ -13,7 +13,7 @@ namespace mu
 	//---------------------------------------------------------------------------------------------
 	//! Grow the image as if it was a bitmap
 	//---------------------------------------------------------------------------------------------
-	inline void ImageGrow( Image* pImage )
+	inline void ImageGrow( FImage* pImage )
 	{
 		check(pImage && pImage->GetFormat()== EImageFormat::L_UByte );
 
@@ -23,7 +23,7 @@ namespace mu
 		int sizeY = pImage->GetSizeY();
 		check( sizeX>0 && sizeY>0 );
 
-		ImagePtrConst pSource = pImage->Clone();
+		TSharedPtr<const FImage> pSource = pImage->Clone();
 
 		int rowSize = sizeX*GetImageFormatData( pImage->GetFormat() ).BytesPerBlock;
 
@@ -65,7 +65,7 @@ namespace mu
 	//---------------------------------------------------------------------------------------------
 	//! Grow the image copying border colours, by using an external mask as reference
 	//---------------------------------------------------------------------------------------------
-	inline void ImageGrow( Image* pImage, const Image* pMask )
+	inline void ImageGrow( FImage* pImage, const FImage* pMask )
 	{
 		MUTABLE_CPUPROFILER_SCOPE(ImageGrow2)
 
diff --git a/Engine/Plugins/Mutable/Source/MutableRuntime/Internal/MuR/OpImageProject.h b/Engine/Plugins/Mutable/Source/MutableRuntime/Internal/MuR/OpImageProject.h
--- a/Engine/Plugins/Mutable/Source/MutableRuntime/Internal/MuR/OpImageProject.h
+++ b/Engine/Plugins/Mutable/Source/MutableRuntime/Internal/MuR/OpImageProject.h
@@ -7,8 +7,8 @@
 
 namespace mu
 {
-	class Image;
-	class Mesh;
+	class FImage;
+	class FMesh;
 	struct FProjector;
 	enum class ESamplingMethod : uint8;
 	template <int NUM_INTERPOLATORS> class RasterVertex;
@@ -19,8 +19,8 @@ namespace mu
 		TArray<uint8> CulledVertex;
 	};
 
-    extern void ImageRasterProjectedPlanar( const Mesh* pMesh, Image* pTargetImage,
-		const Image* pSource, const Image* pMask,
+    extern void ImageRasterProjectedPlanar( const FMesh* pMesh, FImage* pTargetImage,
+		const FImage* pSource, const FImage* pMask,
 		bool bIsRGBFadingEnabled, bool bIsAlphaFadingEnabled,
 		ESamplingMethod SamplingMethod,
 		float FadeStart, float FadeEnd, float MipInterpolationFactor,
@@ -28,8 +28,8 @@ namespace mu
 		UE::Math::TIntVector2<uint16> CropMin, UE::Math::TIntVector2<uint16> UncroppedSize,
 		FScratchImageProject* scratch, bool bUseVectorImplementation = false);
 
-    extern void ImageRasterProjectedCylindrical( const Mesh* pMesh, Image* pTargetImage,
-		const Image* pSource, const Image* pMask,
+    extern void ImageRasterProjectedCylindrical( const FMesh* pMesh, FImage* pTargetImage,
+		const FImage* pSource, const FImage* pMask,
 		bool bIsRGBFadingEnabled, bool bIsAlphaFadingEnabled,
 		ESamplingMethod SamplingMethod,
 		float FadeStart, float FadeEnd, float MipInterpolationFactor,
@@ -38,8 +38,8 @@ namespace mu
 		UE::Math::TIntVector2<uint16> CropMin, UE::Math::TIntVector2<uint16> UncroppedSize,
 		FScratchImageProject* Scratch, bool bUseVectorImplementation = false);
 
-    extern void ImageRasterProjectedWrapping( const Mesh* pMesh, Image* pTargetImage,
-		const Image* pSource, const Image* pMask,
+    extern void ImageRasterProjectedWrapping( const FMesh* pMesh, FImage* pTargetImage,
+		const FImage* pSource, const FImage* pMask,
 		bool bIsRGBFadingEnabled, bool bIsAlphaFadingEnabled,
 		ESamplingMethod SamplingMethod,
 		float FadeStart, float FadeEnd, float MipInterpolationFactor,
@@ -48,11 +48,11 @@ namespace mu
 		FScratchImageProject* Scratch, bool bUseVectorImplementation = false);
 
 	extern float ComputeProjectedFootprintBestMip(
-			const Mesh* pMesh, const FProjector& Projector, const FVector2f& TargetSize, const FVector2f& SourceSize);
+			const FMesh* pMesh, const FProjector& Projector, const FVector2f& TargetSize, const FVector2f& SourceSize);
 
-    extern void MeshProject(Mesh* Result, const Mesh* pMesh, const FProjector& Projector, bool& bOutSuccess);
+    extern void MeshProject(FMesh* Result, const FMesh* pMesh, const FProjector& Projector, bool& bOutSuccess);
 
-	MUTABLERUNTIME_API extern void CreateMeshOptimisedForProjection(Mesh* Result, int32 LayoutIndex);
-	MUTABLERUNTIME_API extern void CreateMeshOptimisedForWrappingProjection(Mesh* Result, int32 LayoutIndex);
+	MUTABLERUNTIME_API extern void CreateMeshOptimisedForProjection(FMesh* Result, int32 LayoutIndex);
+	MUTABLERUNTIME_API extern void CreateMeshOptimisedForWrappingProjection(FMesh* Result, int32 LayoutIndex);
 
 }
diff --git a/Engine/Plugins/Mutable/Source/MutableRuntime/Internal/MuR/OpMeshFormat.h b/Engine/Plugins/Mutable/Source/MutableRuntime/Internal/MuR/OpMeshFormat.h
--- a/Engine/Plugins/Mutable/Source/MutableRuntime/Internal/MuR/OpMeshFormat.h
+++ b/Engine/Plugins/Mutable/Source/MutableRuntime/Internal/MuR/OpMeshFormat.h
@@ -6,7 +6,7 @@
 
 namespace mu
 {
-	class Mesh;
+	class FMesh;
 	class FMeshBufferSet;
 	struct FMeshBuffer;
 
@@ -17,9 +17,9 @@ namespace mu
 	*/
 	MUTABLERUNTIME_API extern void MeshFormat
 		(
-			Mesh* Result,
-			const Mesh* Source,
-			const Mesh* Format,
+			FMesh* Result,
+			const FMesh* Source,
+			const FMesh* Format,
 			bool bKeepSystemBuffers,
 			bool bFormatVertices,
 			bool bFormatIndices,
@@ -32,6 +32,6 @@ namespace mu
     void MeshFormatBuffer( const FMeshBufferSet& Source, FMeshBuffer& Result, int32 ResultElementOffset, bool bHasSpecialSemantics, uint32 IDPrefix);
 
 	/** Try to reduce the mesh size by reducing the component count, and data type of some buffers. */
-	void MeshOptimizeBuffers( Mesh* );
+	void MeshOptimizeBuffers( FMesh* );
 
 }
diff --git a/Engine/Plugins/Mutable/Source/MutableRuntime/Internal/MuR/OpMeshMerge.h b/Engine/Plugins/Mutable/Source/MutableRuntime/Internal/MuR/OpMeshMerge.h
--- a/Engine/Plugins/Mutable/Source/MutableRuntime/Internal/MuR/OpMeshMerge.h
+++ b/Engine/Plugins/Mutable/Source/MutableRuntime/Internal/MuR/OpMeshMerge.h
@@ -14,14 +14,14 @@ namespace mu
 
 	struct FMeshMergeScratchMeshes
 	{
-		Ptr<Mesh> FirstReformat;
-		Ptr<Mesh> SecondReformat;
+		TSharedPtr<FMesh> FirstReformat;
+		TSharedPtr<FMesh> SecondReformat;
 	};
 
 	//---------------------------------------------------------------------------------------------
 	//! Merge two meshes into one new mesh
 	//---------------------------------------------------------------------------------------------
-	inline void MeshMerge(Mesh* Result, const Mesh* pFirst, const Mesh* pSecond, bool bMergeSurfaces, FMeshMergeScratchMeshes& ScratchMeshes)
+	inline void MeshMerge(FMesh* Result, const FMesh* pFirst, const FMesh* pSecond, bool bMergeSurfaces, FMeshMergeScratchMeshes& ScratchMeshes)
 	{
 		MUTABLE_CPUPROFILER_SCOPE(MeshMerge);
 
@@ -128,7 +128,7 @@ namespace mu
 							{
 							case EMeshBufferFormat::UInt16:
 							{
-								for (int v = 0; v < FirstCount; ++v)
+								for (int32 v = 0; v < FirstCount; ++v)
 								{
 									*(uint32_t*)pDest = *(const uint16*)pSource;
 									pSource += FirstIndexBuffer.ElementSize;
@@ -139,7 +139,7 @@ namespace mu
 
 							case EMeshBufferFormat::UInt8:
 							{
-								for (int v = 0; v < FirstCount; ++v)
+								for (int32 v = 0; v < FirstCount; ++v)
 								{
 									*(uint32_t*)pDest = *(const uint8_t*)pSource;
 									pSource += FirstIndexBuffer.ElementSize;
@@ -162,7 +162,7 @@ namespace mu
 
 							case EMeshBufferFormat::UInt8:
 							{
-								for (int v = 0; v < FirstCount; ++v)
+								for (int32 v = 0; v < FirstCount; ++v)
 								{
 									*(uint16*)pDest = *(const uint8_t*)pSource;
 									pSource += FirstIndexBuffer.ElementSize;
@@ -202,7 +202,7 @@ namespace mu
 						case EMeshBufferFormat::NInt32:
 						case EMeshBufferFormat::NUInt32:
 						{
-							for (int v = 0; v < SecondCount; ++v)
+							for (int32 v = 0; v < SecondCount; ++v)
 							{
 								*(uint32_t*)pDest = firstVertexCount + *(const uint32_t*)pSource;
 								pSource += SecondIndexBuffer.ElementSize;
@@ -216,7 +216,7 @@ namespace mu
 						case EMeshBufferFormat::NInt16:
 						case EMeshBufferFormat::NUInt16:
 						{
-							for (int v = 0; v < SecondCount; ++v)
+							for (int32 v = 0; v < SecondCount; ++v)
 							{
 								*(uint16*)pDest = uint16(firstVertexCount) + *(const uint16*)pSource;
 								pSource += SecondIndexBuffer.ElementSize;
@@ -230,7 +230,7 @@ namespace mu
 						case EMeshBufferFormat::NInt8:
 						case EMeshBufferFormat::NUInt8:
 						{
-							for (int v = 0; v < SecondCount; ++v)
+							for (int32 v = 0; v < SecondCount; ++v)
 							{
 								*(uint8_t*)pDest = uint8_t(firstVertexCount) + *(const uint8_t*)pSource;
 								pSource += SecondIndexBuffer.ElementSize;
@@ -259,7 +259,7 @@ namespace mu
 							case EMeshBufferFormat::NInt16:
 							case EMeshBufferFormat::NUInt16:
 							{
-								for (int v = 0; v < SecondCount; ++v)
+								for (int32 v = 0; v < SecondCount; ++v)
 								{
 									*(uint32_t*)pDest = uint32_t(firstVertexCount) + *(const uint16*)pSource;
 									pSource += SecondIndexBuffer.ElementSize;
@@ -273,7 +273,7 @@ namespace mu
 							case EMeshBufferFormat::NInt8:
 							case EMeshBufferFormat::NUInt8:
 							{
-								for (int v = 0; v < SecondCount; ++v)
+								for (int32 v = 0; v < SecondCount; ++v)
 								{
 									*(uint32_t*)pDest = uint32_t(firstVertexCount) + *(const uint8_t*)pSource;
 									pSource += SecondIndexBuffer.ElementSize;
@@ -299,7 +299,7 @@ namespace mu
 							case EMeshBufferFormat::NInt8:
 							case EMeshBufferFormat::NUInt8:
 							{
-								for (int v = 0; v < SecondCount; ++v)
+								for (int32 v = 0; v < SecondCount; ++v)
 								{
 									*(uint16*)pDest = uint16(firstVertexCount) + *(const uint8_t*)pSource;
 									pSource += SecondIndexBuffer.ElementSize;
@@ -336,17 +336,17 @@ namespace mu
 			Result->Layouts.SetNum(ResultLayoutCount);
 			for (int32 LayoutIndex = 0; LayoutIndex < ResultLayoutCount; ++LayoutIndex)
 			{
-				Ptr<Layout> pR;
+				TSharedPtr<FLayout> pR;
 				
 				if (LayoutIndex < pFirst->Layouts.Num())
 				{
-					const Layout* pF = pFirst->Layouts[LayoutIndex].get();
+					const FLayout* pF = pFirst->Layouts[LayoutIndex].Get();
 					pR = pF->Clone();
 				}
 
 				if (LayoutIndex < pSecond->Layouts.Num())
 				{
-					const Layout* pS = pSecond->Layouts[LayoutIndex].get();
+					const FLayout* pS = pSecond->Layouts[LayoutIndex].Get();
 					if (!pR)
 					{
 						pR = pS->Clone();
@@ -386,16 +386,16 @@ namespace mu
 		if (bMergeSkeletons)
 		{
 			MUTABLE_CPUPROFILER_SCOPE(MergeSkeleton);
-			Ptr<Skeleton> pResultSkeleton;
+			TSharedPtr<FSkeleton> ResultSkeleton;
 
-			mu::SkeletonPtrConst pFirstSkeleton = pFirst->GetSkeleton();
-			mu::SkeletonPtrConst pSecondSkeleton = pSecond->GetSkeleton();
+			TSharedPtr<const FSkeleton> FirstSkeleton = pFirst->GetSkeleton();
+			TSharedPtr<const FSkeleton> SecondSkeleton = pSecond->GetSkeleton();
 
-			const int32 NumBonesFirst = pFirstSkeleton ? pFirstSkeleton->GetBoneCount() : 0;
-			const int32 NumBonesSecond = pSecondSkeleton ? pSecondSkeleton->GetBoneCount() : 0;
+			const int32 NumBonesFirst = FirstSkeleton ? FirstSkeleton->GetBoneCount() : 0;
+			const int32 NumBonesSecond = SecondSkeleton ? SecondSkeleton->GetBoneCount() : 0;
 
-			pResultSkeleton = pFirstSkeleton ? pFirstSkeleton->Clone() : new Skeleton;
-			Result->SetSkeleton(pResultSkeleton);
+			ResultSkeleton = FirstSkeleton ? FirstSkeleton->Clone() : MakeShared<FSkeleton>();
+			Result->SetSkeleton(ResultSkeleton);
 
 			TArray<uint16> SecondToResultBoneIndices;
 			SecondToResultBoneIndices.SetNumUninitialized(NumBonesSecond);
@@ -403,20 +403,20 @@ namespace mu
 			// Merge pSecond and build the remap array 
 			for (int32 SecondBoneIndex = 0; SecondBoneIndex < NumBonesSecond; ++SecondBoneIndex)
 			{
-				const FBoneName& BoneNameId = pSecondSkeleton->BoneIds[SecondBoneIndex];
-				int32 Index = pResultSkeleton->FindBone(BoneNameId);
+				const FBoneName& BoneNameId = SecondSkeleton->BoneIds[SecondBoneIndex];
+				int32 Index = ResultSkeleton->FindBone(BoneNameId);
 
 				// Add a new bone
 				if (Index == INDEX_NONE)
 				{
-					Index = pResultSkeleton->BoneIds.Add(BoneNameId);
+					Index = ResultSkeleton->BoneIds.Add(BoneNameId);
 
 					// Add an incorrect index, to be fixed below in case the parent index is later in the bone array.
-					pResultSkeleton->BoneParents.Add(pSecondSkeleton->BoneParents[SecondBoneIndex]);
+					ResultSkeleton->BoneParents.Add(SecondSkeleton->BoneParents[SecondBoneIndex]);
 #if WITH_EDITOR
-					if (pSecondSkeleton->DebugBoneNames.IsValidIndex(SecondBoneIndex))
+					if (SecondSkeleton->DebugBoneNames.IsValidIndex(SecondBoneIndex))
 					{
-						pResultSkeleton->DebugBoneNames.Add(pSecondSkeleton->DebugBoneNames[SecondBoneIndex]);
+						ResultSkeleton->DebugBoneNames.Add(SecondSkeleton->DebugBoneNames[SecondBoneIndex]);
 					}
 #endif
 				}
@@ -425,12 +425,12 @@ namespace mu
 			}
 
 			// Fix second mesh bone parents
-			for (int32 ob = NumBonesFirst; ob < pResultSkeleton->BoneParents.Num(); ++ob)
+			for (int32 ob = NumBonesFirst; ob < ResultSkeleton->BoneParents.Num(); ++ob)
 			{
-				int16 secondMeshIndex = pResultSkeleton->BoneParents[ob];
+				int16 secondMeshIndex = ResultSkeleton->BoneParents[ob];
 				if (secondMeshIndex != INDEX_NONE)
 				{
-					pResultSkeleton->BoneParents[ob] = SecondToResultBoneIndices[secondMeshIndex];
+					ResultSkeleton->BoneParents[ob] = SecondToResultBoneIndices[secondMeshIndex];
 				}
 			}
 		}
@@ -549,17 +549,17 @@ namespace mu
 			Result->BonePoses = pFirst->BonePoses;
 
 			// Add or override bone poses
-			for (const Mesh::FBonePose& SecondBonePose : pSecond->BonePoses)
+			for (const FMesh::FBonePose& SecondBonePose : pSecond->BonePoses)
 			{
 				const int32 ResultBoneIndex = Result->FindBonePose(SecondBonePose.BoneId);
 
 				if (ResultBoneIndex != INDEX_NONE)
 				{
-					Mesh::FBonePose& ResultBonePose = Result->BonePoses[ResultBoneIndex];
+					FMesh::FBonePose& ResultBonePose = Result->BonePoses[ResultBoneIndex];
 
 					// TODO: Not sure how to tune this priority, review it.
 					// For now use a similar strategy as before. 
-					auto ComputeBoneMergePriority = [](const Mesh::FBonePose& BonePose)
+					auto ComputeBoneMergePriority = [](const FMesh::FBonePose& BonePose)
 					{
 						return (EnumHasAnyFlags(BonePose.BoneUsageFlags, EBoneUsageFlags::Skinning) ? 1 : 0) +
 							(EnumHasAnyFlags(BonePose.BoneUsageFlags, EBoneUsageFlags::Reshaped) ? 1 : 0);
@@ -589,7 +589,7 @@ namespace mu
 			MUTABLE_CPUPROFILER_SCOPE(PhysicsBodies);
 
 			// Appends InPhysicsBody to OutPhysicsBody removing Bodies that are equal, have same bone and customId and its properies are identical.	
-			auto AppendPhysicsBodiesUnique = [](PhysicsBody& OutPhysicsBody, const PhysicsBody& InPhysicsBody) -> bool
+			auto AppendPhysicsBodiesUnique = [](FPhysicsBody& OutPhysicsBody, const FPhysicsBody& InPhysicsBody) -> bool
 			{
 				TArray<FBoneName>& OutBones = OutPhysicsBody.BoneIds;
 				TArray<int32>& OutCustomIds = OutPhysicsBody.BodiesCustomIds;
@@ -661,22 +661,22 @@ namespace mu
 				return bModified;
 			};
 
-			TTuple<const PhysicsBody*, bool> SharedResultPhysicsBody = Invoke([&]()
-				-> TTuple<const PhysicsBody*, bool>
+			TTuple<TSharedPtr<const FPhysicsBody>, bool> SharedResultPhysicsBody = Invoke([&]()
+				-> TTuple<TSharedPtr<const FPhysicsBody>, bool>
 			{
 				if (pFirst->GetPhysicsBody() == pSecond->GetPhysicsBody())
 				{
-					return MakeTuple(pFirst->GetPhysicsBody().get(), true);
+					return MakeTuple(pFirst->GetPhysicsBody(), true);
 				}
 
 				if (pFirst->GetPhysicsBody() && !pSecond->GetPhysicsBody())
 				{
-					return MakeTuple(pFirst->GetPhysicsBody().get(), true);
+					return MakeTuple(pFirst->GetPhysicsBody(), true);
 				}
 
 				if (!pFirst->GetPhysicsBody() && pSecond->GetPhysicsBody())
 				{
-					return MakeTuple(pSecond->GetPhysicsBody().get(), true);
+					return MakeTuple(pSecond->GetPhysicsBody(), true);
 				}
 
 				return MakeTuple(nullptr, false);
@@ -691,7 +691,7 @@ namespace mu
 			{
 				check(pFirst->GetPhysicsBody() && pSecond->GetPhysicsBody());
 
-				Ptr<PhysicsBody> MergedResultPhysicsBody = pFirst->GetPhysicsBody()->Clone();
+				TSharedPtr<FPhysicsBody> MergedResultPhysicsBody = pFirst->GetPhysicsBody()->Clone();
 
 				MergedResultPhysicsBody->bBodiesModified =
 					AppendPhysicsBodiesUnique(*MergedResultPhysicsBody, *pSecond->GetPhysicsBody()) ||
@@ -712,8 +712,8 @@ namespace mu
 			{
 				const int32 CustomIdToMerge = pSecond->AdditionalPhysicsBodies[Index]->CustomId;
 
-				const mu::Ptr<const PhysicsBody>* Found = pFirst->AdditionalPhysicsBodies.FindByPredicate(
-					[CustomIdToMerge](const Ptr<const mu::PhysicsBody>& Body) { return CustomIdToMerge == Body->CustomId; });
+				const TSharedPtr<const FPhysicsBody>* Found = pFirst->AdditionalPhysicsBodies.FindByPredicate(
+					[CustomIdToMerge](const TSharedPtr<const FPhysicsBody>& Body) { return CustomIdToMerge == Body->CustomId; });
 
 				// TODO: current usages do not expect collisions, but same Id collision with bodies modified in differnet ways
 				// may need to be contemplated at some point.
@@ -1098,12 +1098,12 @@ namespace mu
     //---------------------------------------------------------------------------------------------
     //!
     //---------------------------------------------------------------------------------------------
-    inline void ExtendSkeleton( Skeleton* pBase, const Skeleton* pOther )
+    inline void ExtendSkeleton( FSkeleton* pBase, const FSkeleton* pOther )
     {
-        TMap<int,int> otherToResult;
+        TMap<int32,int32> otherToResult;
 
-        int initialBones = pBase->GetBoneCount();
-		for (int b = 0; pOther && b < pOther->GetBoneCount(); ++b)
+        int32 initialBones = pBase->GetBoneCount();
+		for (int32 b = 0; pOther && b < pOther->GetBoneCount(); ++b)
 		{
 			int32 resultBoneIndex = pBase->FindBone(pOther->GetBoneName(b));
 			if (resultBoneIndex < 0)
@@ -1129,7 +1129,7 @@ namespace mu
 		}
 
         // Fix bone parent indices of the bones added from pOther
-        for ( int b=initialBones;b<pBase->GetBoneCount(); ++b)
+        for ( int32 b=initialBones;b<pBase->GetBoneCount(); ++b)
         {
             int16_t sourceIndex = pBase->BoneParents[b];
             if (sourceIndex>=0)
@@ -1143,7 +1143,7 @@ namespace mu
     //---------------------------------------------------------------------------------------------
     //! Return null if there is no need to remap the mesh
     //---------------------------------------------------------------------------------------------
-    inline void MeshRemapSkeleton(Mesh* Result, const Mesh* SourceMesh, const Skeleton* Skeleton, bool& bOutSuccess)
+    inline void MeshRemapSkeleton(FMesh* Result, const FMesh* SourceMesh, TSharedPtr<const FSkeleton> Skeleton, bool& bOutSuccess)
     {
 		bOutSuccess = true;
 
diff --git a/Engine/Plugins/Mutable/Source/MutableRuntime/Internal/MuR/Operations.h b/Engine/Plugins/Mutable/Source/MutableRuntime/Internal/MuR/Operations.h
--- a/Engine/Plugins/Mutable/Source/MutableRuntime/Internal/MuR/Operations.h
+++ b/Engine/Plugins/Mutable/Source/MutableRuntime/Internal/MuR/Operations.h
@@ -18,7 +18,7 @@ namespace mu
 {
 
     //!
-    enum class OP_TYPE : uint16
+    enum class EOpType : uint16
     {
         //-----------------------------------------------------------------------------------------
         // No operation
@@ -1021,7 +1021,7 @@ namespace mu
 
             ADDRESS name;
 
-            // Index in the FProgram::m_parameterLists with the parameters that are relevant
+            // Index in the FProgram::ParameterLists with the parameters that are relevant
             // for whatever is added by this operation. This is used only for resources like
             // images or meshes.
             ADDRESS relevantParametersListIndex;
@@ -1029,12 +1029,12 @@ namespace mu
 
 		struct InstanceAddExtensionDataArgs
 		{
-			// This is a reference to an op that produces the Instance that the ExtensionData will
+			// This is a reference to an op that produces the Instance that the FExtensionData will
 			// be added to.
 			ADDRESS Instance;
-			// An op that produces the ExtensionData to add to the Instance
+			// An op that produces the FExtensionData to add to the Instance
 			ADDRESS ExtensionData;
-			// The name to associate with the ExtensionData
+			// The name to associate with the FExtensionData
 			//
 			// This is an index into the string table
 			ADDRESS ExtensionDataName;
@@ -1073,7 +1073,7 @@ namespace mu
         //-------------------------------------------------------------------------------------
         t_OP()
         {
-            type = OP_TYPE::NONE;
+            type = EOpType::NONE;
             unused = 0;
 			FMemory::Memzero( &args, sizeof(args) );
         }
@@ -1096,7 +1096,7 @@ namespace mu
             return false;
         }
 
-        OP_TYPE type;
+        EOpType type;
 
         typedef enum
         {
@@ -1225,11 +1225,11 @@ namespace mu
 		bool supportedBasePixelFormats[uint8(EImageFormat::Count)];
     };
 
-	MUTABLERUNTIME_API extern const FOpDesc& GetOpDesc( OP_TYPE type );
+	MUTABLERUNTIME_API extern const FOpDesc& GetOpDesc( EOpType type );
 
 
     //!
-    inline static DATATYPE GetOpDataType( OP_TYPE type )
+    inline static DATATYPE GetOpDataType( EOpType type )
     {
 		return GetOpDesc(type).type;
     }
@@ -1241,24 +1241,24 @@ namespace mu
 	MUTABLERUNTIME_API extern void ForEachReference( const struct FProgram& program, OP::ADDRESS at, const TFunctionRef<void(OP::ADDRESS)> );
 
 	//!
-	MUTABLERUNTIME_API inline OP_TYPE GetSwitchForType( DATATYPE d )
+	MUTABLERUNTIME_API inline EOpType GetSwitchForType( DATATYPE d )
     {
         switch (d)
         {
-        case DT_INSTANCE: return OP_TYPE::IN_SWITCH;
-        case DT_MESH: return OP_TYPE::ME_SWITCH;
-		case DT_IMAGE: return OP_TYPE::IM_SWITCH;
-		case DT_LAYOUT: return OP_TYPE::LA_SWITCH;
-        case DT_COLOUR: return OP_TYPE::CO_SWITCH;
-        case DT_SCALAR: return OP_TYPE::SC_SWITCH;
-        case DT_INT: return OP_TYPE::NU_SWITCH;
-		case DT_EXTENSION_DATA: return OP_TYPE::ED_SWITCH;
+        case DT_INSTANCE: return EOpType::IN_SWITCH;
+        case DT_MESH: return EOpType::ME_SWITCH;
+		case DT_IMAGE: return EOpType::IM_SWITCH;
+		case DT_LAYOUT: return EOpType::LA_SWITCH;
+        case DT_COLOUR: return EOpType::CO_SWITCH;
+        case DT_SCALAR: return EOpType::SC_SWITCH;
+        case DT_INT: return EOpType::NU_SWITCH;
+		case DT_EXTENSION_DATA: return EOpType::ED_SWITCH;
 
         default:
 			check(false);
 			break;
         }
-        return OP_TYPE::NONE;
+        return EOpType::NONE;
     }
 }
 
diff --git a/Engine/Plugins/Mutable/Source/MutableRuntime/Internal/MuR/ParametersPrivate.h b/Engine/Plugins/Mutable/Source/MutableRuntime/Internal/MuR/ParametersPrivate.h
--- a/Engine/Plugins/Mutable/Source/MutableRuntime/Internal/MuR/ParametersPrivate.h
+++ b/Engine/Plugins/Mutable/Source/MutableRuntime/Internal/MuR/ParametersPrivate.h
@@ -9,11 +9,8 @@
 #include "MuR/SerialisationPrivate.h"
 #include "MuR/Operations.h"
 #include "MuR/Image.h"
-
 #include "MuR/Model.h"
 
-#include <string>
-
 namespace mu
 {
 
@@ -123,13 +120,13 @@ namespace Private
         }
 
 		//!
-		void Serialise(OutputArchive& Arch) const
+		void Serialise(FOutputArchive& Arch) const
 		{
 			Arch << Get();
 		}
 
 		//!
-		void Unserialise(InputArchive& Arch)
+		void Unserialise(FInputArchive& Arch)
 		{
 			Arch >> Get();
 		} 
@@ -162,7 +159,7 @@ namespace Private
 
 
         //!
-        void Serialise( OutputArchive& arch ) const
+        void Serialise( FOutputArchive& arch ) const
         {
             arch << type;
             arch << position;
@@ -173,7 +170,7 @@ namespace Private
         }
 
         //!
-        void Unserialise( InputArchive& arch )
+        void Unserialise( FInputArchive& arch )
         {
             arch >> type;
             arch >> position;
@@ -218,7 +215,7 @@ namespace Private
 		
 
 		//!
-		void Serialise(OutputArchive& arch) const
+		void Serialise(FOutputArchive& arch) const
 		{
 			arch << position;
 			arch << up;
@@ -228,7 +225,7 @@ namespace Private
 		}
 
 		//!
-		void Unserialise(InputArchive& arch)
+		void Unserialise(FInputArchive& arch)
 		{
 			arch >> position;
 			arch >> up;
@@ -284,205 +281,169 @@ namespace Private
 
     struct FParameterDesc
     {
-        FString m_name;
+        FString Name;
 
         //! Unique id (provided externally, so no actual guarantee that it is unique.)
-		FGuid m_uid;
+		FGuid UID;
 
-        EParameterType m_type = EParameterType::None;
+        EParameterType Type = EParameterType::None;
 
-        FParameterValue m_defaultValue;
+        FParameterValue DefaultValue;
 
         //! Ranges, if the parameter is multi-dimensional. The indices refer to the Model's program
         //! vector of range descriptors.
-		TArray<uint32> m_ranges;
+		TArray<uint32> Ranges;
 
         //! Possible values of the parameter in case of being an integer, and its names
         struct FIntValueDesc
         {
-            int16 m_value;
-			FString m_name;
+            int16 Value;
+			FString Name;
 
             //!
-            bool operator==( const FIntValueDesc& other ) const
+            bool operator==( const FIntValueDesc& Other ) const
             {
-                return m_value==other.m_value &&
-                        m_name==other.m_name;
+                return Value== Other.Value &&
+						Name== Other.Name;
             }
 
             //!
-            void Serialise( OutputArchive& arch ) const
+            void Serialise( FOutputArchive& Arch ) const
             {
-                arch << m_value;
-                arch << m_name;
+				Arch << Value;
+				Arch << Name;
             }
 
             //!
-            void Unserialise( InputArchive& arch )
+            void Unserialise( FInputArchive& Arch)
             {
-                arch >> m_value;
-                arch >> m_name;
+				Arch >> Value;
+				Arch >> Name;
             }
         };
 
         //! For integer parameters, this contains the description of the possible values.
         //! If empty, the integer may have any value.
-		TArray<FIntValueDesc> m_possibleValues;
+		TArray<FIntValueDesc> PossibleValues;
 
         //!
         bool operator==( const FParameterDesc& other ) const
         {
-            return m_name == other.m_name && m_uid == other.m_uid && m_type == other.m_type &&
-                   m_defaultValue == other.m_defaultValue &&
-                   m_ranges == other.m_ranges &&
-                   m_possibleValues == other.m_possibleValues;
+            return Name == other.Name && UID == other.UID && Type == other.Type &&
+                   DefaultValue == other.DefaultValue &&
+                   Ranges == other.Ranges &&
+                   PossibleValues == other.PossibleValues;
         }
 
         //!
-        void Serialise( OutputArchive& arch ) const
+        void Serialise( FOutputArchive& arch ) const
         {
-            const int32 ver = 10;
-            arch << ver;
-
-			arch << m_name;
-			arch << m_uid;
-            arch << m_type;
-            arch << m_defaultValue;
-            arch << m_ranges;
-            arch << m_possibleValues;
+			arch << Name;
+			arch << UID;
+            arch << Type;
+            arch << DefaultValue;
+            arch << Ranges;
+            arch << PossibleValues;
         }
 
         //!
-        void Unserialise( InputArchive& arch )
+        void Unserialise( FInputArchive& arch )
         {
-            int32 ver;
-            arch >> ver;
-            check( ver == 10 );
-
-			arch >> m_name;
-			arch >> m_uid;
-            arch >> m_type;
-			arch >> m_defaultValue;
-            arch >> m_ranges;
-			arch >> m_possibleValues;
+			arch >> Name;
+			arch >> UID;
+            arch >> Type;
+			arch >> DefaultValue;
+            arch >> Ranges;
+			arch >> PossibleValues;
         }
     };
 
 
     struct FRangeDesc
     {
-		FString m_name;
-		FString m_uid;
+		FString Name;
+		FString UID;
 
 		/** Parameter that controls the size of this range, if any. */
-		int32 m_dimensionParameter = -1;
+		int32 DimensionParameter = -1;
 
         //!
         bool operator==( const FRangeDesc& other ) const
         {
-            return m_name==other.m_name
+            return Name==other.Name
 				&&
-				m_uid == other.m_uid 
+				UID == other.UID 
 				&&
-				m_dimensionParameter == other.m_dimensionParameter;
+				DimensionParameter == other.DimensionParameter;
         }
 
         //!
-        void Serialise( OutputArchive& arch ) const
+        void Serialise( FOutputArchive& Arch ) const
         {
-            const uint32 ver = 3;
-            arch << ver;
-
-            arch << m_name;
-            arch << m_uid;
-			arch << m_dimensionParameter;
+			Arch << Name;
+			Arch << UID;
+			Arch << DimensionParameter;
         }
 
         //!
-        void Unserialise( InputArchive& arch )
+        void Unserialise( FInputArchive& Arch)
         {
-            uint32 ver;
-            arch >> ver;
-            check( ver == 3 );
-
-			arch >> m_name;
-			arch >> m_uid;
-			arch >> m_dimensionParameter;
+			Arch >> Name;
+			Arch >> UID;
+			Arch >> DimensionParameter;
         }
     };
 
 
-    class MUTABLERUNTIME_API RangeIndex::Private
-    {
-    public:
-
-        //! Run-time data
-        Ptr<const Parameters> m_pParameters;
-
-        //! Index of the parameter for which we are a range index
-        int m_parameter = -1;
-
-        //! Position in the several dimension of the range, as defined in m_pParameters
-		TArray<int> m_values;
-    };
-
-
-    class MUTABLERUNTIME_API Parameters::Private
+    class MUTABLERUNTIME_API FParameters::Private
     {
     public:
 
-        //! Warning: update Parameters::Clone method if this members change.
+        //! Warning: update FParameters::Clone method if this members change.
 
         //! Run-time data
-        TSharedPtr<const Model> m_pModel;
+        TSharedPtr<const FModel> Model;
  
         //! Values for the parameters if they are not multidimensional.
-		TArray<FParameterValue> m_values;
+		TArray<FParameterValue> Values;
 
 
         //! If the parameter is multidemensional, the values are stored here.
-        //! The key of the map is the vector of values stored in a RangeIndex
-		TArray< TMap< TArray<int32>, FParameterValue > > m_multiValues;
+        //! The key of the map is the vector of values stored in a FRangeIndex
+		TArray< TMap< TArray<int32>, FParameterValue > > MultiValues;
 
 
         //!
-        void Serialise( OutputArchive& arch ) const
+        void Serialise( FOutputArchive& arch ) const
         {
-            const uint32 ver = 3;
-            arch << ver;
-
-            arch << m_values;
-            arch << m_multiValues;
+            arch << Values;
+            arch << MultiValues;
         }
 
         //!
-        void Unserialise( InputArchive& arch )
+        void Unserialise( FInputArchive& arch )
         {
-            uint32 ver;
-            arch >> ver;
-			check(ver == 3);
-       
-		    arch >> m_values;
-			arch >> m_multiValues;
+		    arch >> Values;
+			arch >> MultiValues;
         }
 
         //!
         int32 Find( const FString& Name ) const;
 
         //!
-        FProjector GetProjectorValue( int index, const Ptr<const RangeIndex>& pos ) const;
+        FProjector GetProjectorValue( int32 index, const FRangeIndex* ) const;
 
 		/** Return true if the parameter has any multi-dimensional values set. This is independent to if the model
 		* accepts multi-dimensional parameters for this particular parameter.
 		*/
 		inline bool HasMultipleValues(int32 ParamIndex) const
 		{
-			if (ParamIndex >= m_multiValues.Num())
+			if (ParamIndex >= MultiValues.Num())
 			{
 				return false;
 			}
 
-			return m_multiValues[ParamIndex].Num()>0;
+			return MultiValues[ParamIndex].Num()>0;
 		}
     };
 
diff --git a/Engine/Plugins/Mutable/Source/MutableRuntime/Internal/MuR/SerialisationPrivate.h b/Engine/Plugins/Mutable/Source/MutableRuntime/Internal/MuR/SerialisationPrivate.h
--- a/Engine/Plugins/Mutable/Source/MutableRuntime/Internal/MuR/SerialisationPrivate.h
+++ b/Engine/Plugins/Mutable/Source/MutableRuntime/Internal/MuR/SerialisationPrivate.h
@@ -14,18 +14,15 @@
 #include "Math/Transform.h"
 #include "MuR/Ptr.h"
 #include "MuR/RefCounted.h"
-#include "MuR/MutableMemory.h"
 #include "MuR/MutableMath.h"
 #include "MuR/Types.h"
 
-#include <string>
-
 
 namespace mu
 {	
 
     //!
-    class MUTABLERUNTIME_API InputMemoryStream : public InputStream
+    class MUTABLERUNTIME_API FInputMemoryStream : public FInputStream
     {
     public:
 
@@ -36,11 +33,11 @@ namespace mu
         //! Create the stream using an external buffer.
         //! The buffer will not be owned by this object, so it cannot be deallocated while this
         //! objects is in use.
-        InputMemoryStream( const void* pBuffer, uint64 size );
+        FInputMemoryStream( const void* pBuffer, uint64 size );
 
 
         //-----------------------------------------------------------------------------------------
-        // InputStream interface
+        // FInputStream interface
         //-----------------------------------------------------------------------------------------
         void Read( void* pData, uint64 size ) override;
 
@@ -55,12 +52,12 @@ namespace mu
 
 
 #define MUTABLE_IMPLEMENT_POD_SERIALISABLE(Type)				     \
-    void DLLEXPORT operator<<(OutputArchive& Arch, const Type& T)    \
+    void DLLEXPORT operator<<(FOutputArchive& Arch, const Type& T)    \
     {																 \
         Arch.Stream->Write(&T, sizeof(Type));						 \
     }																 \
                                                                      \
-    void DLLEXPORT operator>>(InputArchive& Arch, Type& T)		     \
+    void DLLEXPORT operator>>(FInputArchive& Arch, Type& T)		     \
     {																 \
         Arch.Stream->Read(&T, sizeof(Type));						 \
     }																 \
@@ -68,7 +65,7 @@ namespace mu
 
 #define MUTABLE_IMPLEMENT_POD_VECTOR_SERIALISABLE(Type)                        \
     template<class Alloc>                                                      \
-	void operator<<(OutputArchive& Arch, const TArray<Type, Alloc>& V)         \
+	void operator<<(FOutputArchive& Arch, const TArray<Type, Alloc>& V)         \
 	{                                                                          \
 		uint32 Num = uint32(V.Num());                                          \
 		Arch << Num;                                                           \
@@ -79,7 +76,7 @@ namespace mu
 	}                                                                          \
                                                                                \
     template<class Alloc>                                                      \
-	void operator>>(InputArchive& Arch, TArray<Type, Alloc>& V)                \
+	void operator>>(FInputArchive& Arch, TArray<Type, Alloc>& V)                \
 	{                                                                          \
 		uint32 Num;                                                            \
 		Arch >> Num;                                                           \
@@ -92,7 +89,7 @@ namespace mu
 
 	/** TVariant custom serialize. Based on the default serialization. */
 	template <typename... Ts>
-	void operator<<(OutputArchive& Ar, const TVariant<Ts...>& Variant)
+	void operator<<(FOutputArchive& Ar, const TVariant<Ts...>& Variant)
 	{
 		const uint8 Index = static_cast<uint8>(Variant.GetIndex());
 		Ar << Index;
@@ -108,7 +105,7 @@ namespace mu
 	struct TVariantLoadFromInputArchiveCaller
 	{
 		/** Default construct the type and load it from the FArchive */
-		static void Load(InputArchive& Ar, VariantType& OutVariant)
+		static void Load(FInputArchive& Ar, VariantType& OutVariant)
 		{
 			OutVariant.template Emplace<T>();
 			Ar >> OutVariant.template Get<T>();
@@ -123,9 +120,9 @@ namespace mu
 		static_assert((std::is_default_constructible<Ts>::value && ...), "Each type in TVariant template parameter pack must be default constructible in order to use FArchive serialization");
 
 		/** Load the type at the specified index from the FArchive and emplace it into the TVariant */
-		static void Load(SIZE_T TypeIndex, InputArchive& Ar, VariantType& OutVariant)
+		static void Load(SIZE_T TypeIndex, FInputArchive& Ar, VariantType& OutVariant)
 		{
-			static constexpr void(*Loaders[])(InputArchive&, VariantType&) = { &TVariantLoadFromInputArchiveCaller<Ts, VariantType>::Load... };
+			static constexpr void(*Loaders[])(FInputArchive&, VariantType&) = { &TVariantLoadFromInputArchiveCaller<Ts, VariantType>::Load... };
 			check(TypeIndex < UE_ARRAY_COUNT(Loaders));
 			Loaders[TypeIndex](Ar, OutVariant);
 		}
@@ -133,7 +130,7 @@ namespace mu
 
 	
 	template <typename... Ts>
-	void operator>>(InputArchive& Ar, TVariant<Ts...>& Variant)
+	void operator>>(FInputArchive& Ar, TVariant<Ts...>& Variant)
 	{
 		uint8 Index;
 		Ar >> Index;
@@ -143,13 +140,13 @@ namespace mu
 	}
 
 #define MUTABLE_IMPLEMENT_ENUM_SERIALISABLE(Type)						\
-        void DLLEXPORT operator<<(OutputArchive& Arch, const Type& T)   \
+        void DLLEXPORT operator<<(FOutputArchive& Arch, const Type& T)   \
 		{																\
             uint32 V = (uint32)T;                                       \
             Arch.Stream->Write(&V, sizeof(uint32));					\
 		}																\
 																		\
-        void DLLEXPORT operator>>(InputArchive& Arch, Type& T)   		\
+        void DLLEXPORT operator>>(FInputArchive& Arch, Type& T)   		\
 		{																\
             uint32 V;													\
             Arch.Stream->Read(&V, sizeof(uint32));					\
@@ -157,21 +154,21 @@ namespace mu
 		}																\
 
     template<typename T0, typename T1>
-    inline void operator<<(OutputArchive& Arch, const std::pair<T0, T1>& V)
+    inline void operator<<(FOutputArchive& Arch, const std::pair<T0, T1>& V)
     {
         Arch << V.first;
         Arch << V.second;
     }
 
     template<typename T0, typename T1>
-    inline void operator>>(InputArchive& Arch, std::pair<T0, T1>& V)
+    inline void operator>>(FInputArchive& Arch, std::pair<T0, T1>& V)
     {
         Arch >> V.first;
         Arch >> V.second;
     }
 	
 	template<typename T, uint32 Size, uint32 Align> 
-    void operator<<(OutputArchive& Arch, const TStaticArray<T, Size, Align>& V)
+    void operator<<(FOutputArchive& Arch, const TStaticArray<T, Size, Align>& V)
 	{
 		for (int32 I = 0; I < Size; ++I)
 		{
@@ -180,7 +177,7 @@ namespace mu
 	}
 
 	template<typename T, uint32 Size, uint32 Align> 
-    void operator>>(InputArchive& Arch, TStaticArray<T, Size, Align>& V)
+    void operator>>(FInputArchive& Arch, TStaticArray<T, Size, Align>& V)
 	{
 		for (uint32 I = 0; I < Size; ++I)
 		{
@@ -191,7 +188,7 @@ namespace mu
 
 	//---------------------------------------------------------------------------------------------
 	template< typename K, typename T >
-	inline void operator<< (OutputArchive& Arch, const TMap<K, T>& V)
+	inline void operator<< (FOutputArchive& Arch, const TMap<K, T>& V)
 	{
 		Arch << (uint32)V.Num();
 		for (const TPair<K, T>& Element : V)
@@ -202,7 +199,7 @@ namespace mu
 	}
 
 	template< typename K, typename T >
-	inline void operator>> (InputArchive& Arch, TMap<K, T>& V)
+	inline void operator>> (FInputArchive& Arch, TMap<K, T>& V)
 	{
 		uint32 Num;
 		Arch >> Num;
@@ -240,65 +237,65 @@ namespace mu
 	MUTABLE_IMPLEMENT_POD_VECTOR_SERIALISABLE(FMatrix44f);
 	MUTABLE_IMPLEMENT_POD_VECTOR_SERIALISABLE(FRichCurveKey);
 
-	
+
 	//---------------------------------------------------------------------------------------------
 	template< typename T >
-	inline void operator<< ( OutputArchive& arch, const Ptr<const T>& p )
+	inline void operator<< (FOutputArchive& Arch, const TSharedPtr<const T>& Ptr)
 	{
-		if (!p.get())
+		if (!Ptr)
 		{
-            arch << (int32)-1;
+			Arch << (int32)-1;
 		}
 		else
 		{
-            int32* it = arch.History.Find(p.get());
+			int32* it = Arch.History.Find(Ptr.Get());
 
-			if ( !it )
+			if (!it)
 			{
-                int32 id = arch.History.Num();
-                arch.History.Add( p.get(), id );
-				arch << id;
-				T::Serialise( p.get(), arch );
+				int32 Id = Arch.History.Num();
+				Arch.History.Add(Ptr.Get(), Id);
+				Arch << Id;
+				T::Serialise(Ptr.Get(), Arch);
 			}
 			else
 			{
-				arch << *it;
+				Arch << *it;
 			}
 		}
 	}
 
 	template< typename T >
-	inline void operator>> ( InputArchive& arch, Ptr<const T>& p )
+	inline void operator>> (FInputArchive& Arch, TSharedPtr<const T>& Ptr)
 	{
-        int32 id;
-		arch >> id;
+		int32 Id;
+		Arch >> Id;
 
-		if ( id == -1 )
+		if (Id == -1)
 		{
-			p = 0;
+			Ptr.Reset();
 		}
 		else
 		{
-			if ( id < arch.History.Num() )
+			if (Id < Arch.History.Num())
 			{
-				p = static_cast<T*>( arch.History[id].get() );
+				Ptr = StaticCastSharedPtr<T>(Arch.History[Id]);
 
-				// If the pointer was 0 it means the position in history is used, but not set yet
+				// If the pointer was null it means the position in history is used, but not set yet
 				// option 1: we have a smart pointer loop which is very bad.
 				// option 2: the resource in this Ptr is also pointed by a Proxy that has absorbed it
-				//			 and this reference should also be a proxy instead of a ptr.
-                check( p );
+				//			 and this reference should also be a proxy instead of a pointer.
+				check(Ptr);
 			}
 			else
 			{
-                // Ids come in order, but they may have been absorbed outside in some serialisations
-                // like proxies.
-                //check( id == (int)arch.>History.size() );
-                arch.History.SetNum( id+1 );
-
-				Ptr<T> t = T::StaticUnserialise( arch );
-				p = t;
-				arch.History[id] = t;
+				// Ids come in order, but they may have been absorbed outside in some serialisations
+				// like proxies.
+				//check( Id == Arch.>History.Num() );
+				Arch.History.SetNum(Id + 1, EAllowShrinking::No);
+
+				TSharedPtr<T> Temp = T::StaticUnserialise(Arch);
+				Ptr = Temp;
+				Arch.History[Id] = Temp;
 			}
 		}
 	}
@@ -306,13 +303,13 @@ namespace mu
 
 	//---------------------------------------------------------------------------------------------
 	template< typename T >
-	inline void operator<< ( OutputArchive& arch, const Ptr<T>& p )
+	inline void operator<< ( FOutputArchive& arch, const Ptr<T>& p )
 	{
 		operator<<( arch, (const Ptr<const T>&) p );
 	}
 
 	template< typename T >
-	inline void operator>> ( InputArchive& arch, Ptr<T>& p )
+	inline void operator>> ( FInputArchive& arch, Ptr<T>& p )
 	{
 		operator>>( arch, (Ptr<const T>&) p );
 	}
@@ -321,13 +318,13 @@ namespace mu
 	//---------------------------------------------------------------------------------------------
 	//---------------------------------------------------------------------------------------------
 	//---------------------------------------------------------------------------------------------
-	template<typename T0, typename T1> inline void operator<<(OutputArchive& arch, const TPair<T0,T1>& v)
+	template<typename T0, typename T1> inline void operator<<(FOutputArchive& arch, const TPair<T0,T1>& v)
 	{
 		arch << v.Key;
 		arch << v.Value;
 	}
 
-	template<typename T0, typename T1> inline void operator>>(InputArchive& arch, TPair<T0, T1>& v)
+	template<typename T0, typename T1> inline void operator>>(FInputArchive& arch, TPair<T0, T1>& v)
 	{
 		arch >> v.Key;
 		arch >> v.Value;
@@ -337,7 +334,7 @@ namespace mu
 	//---------------------------------------------------------------------------------------------
 	// TODO: As POD?
 	template< typename T >
-	inline void operator<< (OutputArchive& arch, const UE::Math::TQuat<T>& v)
+	inline void operator<< (FOutputArchive& arch, const UE::Math::TQuat<T>& v)
 	{		
 		arch << v.X;
 		arch << v.Y;
@@ -346,7 +343,7 @@ namespace mu
 	}
 
 	template< typename T >
-	inline void operator>> (InputArchive& arch, UE::Math::TQuat<T>& v)
+	inline void operator>> (FInputArchive& arch, UE::Math::TQuat<T>& v)
 	{
 		arch >> v.X;
 		arch >> v.Y;
@@ -358,7 +355,7 @@ namespace mu
 	//---------------------------------------------------------------------------------------------
 	// TODO: As POD?
 	template< typename T >
-	inline void operator<< (OutputArchive& arch, const UE::Math::TVector<T>& v)
+	inline void operator<< (FOutputArchive& arch, const UE::Math::TVector<T>& v)
 	{
 		arch << v.X;
 		arch << v.Y;
@@ -366,7 +363,7 @@ namespace mu
 	}
 
 	template< typename T >
-	inline void operator>> (InputArchive& arch, UE::Math::TVector<T>& v)
+	inline void operator>> (FInputArchive& arch, UE::Math::TVector<T>& v)
 	{
 		arch >> v.X;
 		arch >> v.Y;
@@ -376,7 +373,7 @@ namespace mu
 
 	//---------------------------------------------------------------------------------------------
 	template< typename T >
-	inline void operator<< (OutputArchive& arch, const UE::Math::TTransform<T>& v)
+	inline void operator<< (FOutputArchive& arch, const UE::Math::TTransform<T>& v)
 	{
 		arch << v.GetRotation();
 		arch << v.GetTranslation();
@@ -384,7 +381,7 @@ namespace mu
 	}
 
 	template< typename T >
-	inline void operator>> (InputArchive& arch, UE::Math::TTransform<T>& v)
+	inline void operator>> (FInputArchive& arch, UE::Math::TTransform<T>& v)
 	{
 		UE::Math::TQuat<T> Rot;
 		UE::Math::TVector<T> Trans;
