[mutable] More enforcement of coding standards.

diff --git a/Engine/Plugins/Mutable/Source/MutableRuntime/Internal/MuR/SystemPrivate.h b/Engine/Plugins/Mutable/Source/MutableRuntime/Internal/MuR/SystemPrivate.h
--- a/Engine/Plugins/Mutable/Source/MutableRuntime/Internal/MuR/SystemPrivate.h
+++ b/Engine/Plugins/Mutable/Source/MutableRuntime/Internal/MuR/SystemPrivate.h
@@ -36,7 +36,7 @@ namespace mu
 		return uint32(Id >> 32);
 	}
 
-	class ExtensionDataStreamer;
+	class FExtensionDataStreamer;
 
 	// Call the tick of the LLM system (we do this to simulate a frame since the LLM system is not entirelly designed to run over a program)
 	inline void UpdateLLMStats()
@@ -64,10 +64,10 @@ namespace mu
 	{
 	public:
 		//! Set or add a value to the index
-		void SetFromModelRangeIndex(uint16 RangeIndex, int32 RangeValue)
+		void SetFromModelRangeIndex(uint16 FRangeIndex, int32 RangeValue)
 		{
-			SizeType Index = IndexOfByPredicate([=](const ElementType& v) { return v.Key >= RangeIndex; });
-			if (Index != INDEX_NONE && (*this)[Index].Key == RangeIndex)
+			SizeType Index = IndexOfByPredicate([=](const ElementType& v) { return v.Key >= FRangeIndex; });
+			if (Index != INDEX_NONE && (*this)[Index].Key == FRangeIndex)
 			{
 				// Update
 				(*this)[Index].Value = RangeValue;
@@ -75,7 +75,7 @@ namespace mu
 			else
 			{
 				// Add new
-				Push(ElementType(RangeIndex, RangeValue));
+				Push(ElementType(FRangeIndex, RangeValue));
 			}
 		}
 
@@ -208,7 +208,7 @@ namespace mu
 	}
 
 
-	inline uint32 GetTypeHash(const Ptr<const Resource>& V)
+	inline uint32 GetTypeHash(const Ptr<const FResource>& V)
 	{
 		return ::GetTypeHash(V.get());
 	}
@@ -493,24 +493,24 @@ namespace mu
 		template<class ResourceType>
 		struct TResourceResult
 		{
-			static_assert(TIsDerivedFrom<ResourceType, Resource>::Value);
+			static_assert(TIsDerivedFrom<ResourceType, FResource>::Value);
 
 			FCacheAddress OpAddress;
-			Ptr<const ResourceType> Value = nullptr;
+			TSharedPtr<const ResourceType> Value = nullptr;
 		};
-		static_assert(sizeof(TResourceResult<Image>) == 16);
-		static_assert(sizeof(TResourceResult<Mesh>) == 16);
+		static_assert(sizeof(TResourceResult<FImage>) == 24);
+		static_assert(sizeof(TResourceResult<FMesh>) == 24);
 
 		/** */
 
 		TMemoryTrackedArray<FVector4f> ColorResults;
-		TMemoryTrackedArray<TResourceResult<Image>> ImageResults;
-		TMemoryTrackedArray<TResourceResult<Mesh>> MeshResults;
-		TMemoryTrackedArray<Ptr<const Layout>> LayoutResults;
-		TMemoryTrackedArray<Ptr<const Instance>> InstanceResults;
+		TMemoryTrackedArray<TResourceResult<FImage>> ImageResults;
+		TMemoryTrackedArray<TResourceResult<FMesh>> MeshResults;
+		TMemoryTrackedArray<TSharedPtr<const FLayout>> LayoutResults;
+		TMemoryTrackedArray<TSharedPtr<const FInstance>> InstanceResults;
 		TMemoryTrackedArray<FProjector> ProjectorResults;
-		TMemoryTrackedArray<Ptr<const String>> StringResults;
-		TMemoryTrackedArray<Ptr<const ExtensionData>> ExtensionDataResults;
+		TMemoryTrackedArray<TSharedPtr<const String>> StringResults;
+		TMemoryTrackedArray<TSharedPtr<const FExtensionData>> ExtensionDataResults;
 		TMemoryTrackedArray<FMatrix44f> MatrixResults;
 
 		/** */
@@ -745,14 +745,14 @@ namespace mu
 			return ProjectorResults[Data->DataTypeIndex];
 		}
 
-		Ptr<const Instance> GetInstance(FCacheAddress at)
+		TSharedPtr<const FInstance> GetInstance(FCacheAddress at)
 		{
 			if (!at.At) return nullptr;
 			FOpExecutionData* Data = OpExecutionData.get_ptr(at);
 			if (!Data) return nullptr;
 
 			check(Data->DataType == DATATYPE::DT_INSTANCE);
-			Ptr<const Instance> Result = InstanceResults[Data->DataTypeIndex];
+			TSharedPtr<const FInstance> Result = InstanceResults[Data->DataTypeIndex];
 
 			// We need to decrease the hit-count even if the result is null.
 			check(Data->OpHitCount > 0);
@@ -767,7 +767,7 @@ namespace mu
 		}
 
 
-		Ptr<const Image> GetImage(FCacheAddress at, bool& bIsLastReference)
+		TSharedPtr<const FImage> GetImage(FCacheAddress at, bool& bIsLastReference)
 		{
 			bIsLastReference = false;
 
@@ -777,7 +777,7 @@ namespace mu
 			if (!Data) return nullptr;
 
 			check(Data->DataType == DATATYPE::DT_IMAGE);
-			Ptr<const Image> Result = ImageResults[Data->DataTypeIndex].Value;
+			TSharedPtr<const FImage> Result = ImageResults[Data->DataTypeIndex].Value;
 
 			// We need to decrease the hit-count even if the result is null.
 			check(Data->OpHitCount > 0);
@@ -793,7 +793,7 @@ namespace mu
 		}
 
 
-		Ptr<const Mesh> GetMesh(FCacheAddress at, bool& bIsLastReference)
+		TSharedPtr<const FMesh> GetMesh(FCacheAddress at, bool& bIsLastReference)
 		{
 			bIsLastReference = false;
 
@@ -803,7 +803,7 @@ namespace mu
 			if (!Data) return nullptr;
 
 			check(Data->DataType == DATATYPE::DT_MESH);
-			Ptr<const Mesh> Result = MeshResults[Data->DataTypeIndex].Value;
+			TSharedPtr<const FMesh> Result = MeshResults[Data->DataTypeIndex].Value;
 
 			// We need to decrease the hit-count even if the result is null.
 			check(Data->OpHitCount > 0);
@@ -819,7 +819,7 @@ namespace mu
 		}
 
 
-		Ptr<const Layout> GetLayout(FCacheAddress at)
+		TSharedPtr<const FLayout> GetLayout(FCacheAddress at)
 		{
 			if (!at.At) return nullptr;
 			FOpExecutionData* Data = OpExecutionData.get_ptr(at);
@@ -830,7 +830,7 @@ namespace mu
 		}
 
 
-		Ptr<const String> GetString(FCacheAddress at)
+		TSharedPtr<const String> GetString(FCacheAddress at)
 		{
 			if (!at.At) return nullptr;
 			FOpExecutionData* Data = OpExecutionData.get_ptr(at);
@@ -841,7 +841,7 @@ namespace mu
 		}
 
 
-		Ptr<const ExtensionData> GetExtensionData(FCacheAddress at)
+		TSharedPtr<const FExtensionData> GetExtensionData(FCacheAddress at)
 		{
 			if (!at.At) return nullptr;
 			FOpExecutionData* Data = OpExecutionData.get_ptr(at);
@@ -949,7 +949,7 @@ namespace mu
 			check(Data->DataTypeIndex != 0);
 		}
 
-		void SetInstance(FCacheAddress at, Ptr<const Instance> v)
+		void SetInstance(FCacheAddress at, TSharedPtr<const FInstance> v)
 		{
 			check(at.At < OpExecutionData.size_code());
 			FOpExecutionData* Data = OpExecutionData.get_ptr(at);
@@ -969,7 +969,7 @@ namespace mu
 			check(Data->DataTypeIndex != 0);
 		}
 
-		void SetExtensionData(FCacheAddress at, Ptr<const ExtensionData> v)
+		void SetExtensionData(FCacheAddress at, TSharedPtr<const FExtensionData> v)
 		{
 			check(at.At < OpExecutionData.size_code());
 			FOpExecutionData* Data = OpExecutionData.get_ptr(at);
@@ -989,7 +989,7 @@ namespace mu
 			check(Data->DataTypeIndex != 0);
 		}
 
-		void SetImage(FCacheAddress At, Ptr<const Image> Value)
+		void SetImage(FCacheAddress At, TSharedPtr<const FImage> Value)
 		{
 			check(At.At < OpExecutionData.size_code());
 			FOpExecutionData* Data = OpExecutionData.get_ptr(At);
@@ -1000,7 +1000,7 @@ namespace mu
 			if (!Data->DataTypeIndex)
 			{
 				Data->DataTypeIndex = ImageResults.Num();
-				ImageResults.Add(TResourceResult<Image>{At, Value});
+				ImageResults.Add(TResourceResult<FImage>{At, Value});
 			}
 			else
 			{
@@ -1011,7 +1011,7 @@ namespace mu
 			mu::UpdateLLMStats();
 		}
 
-		void SetMesh(FCacheAddress At, Ptr<const Mesh> Value)
+		void SetMesh(FCacheAddress At, TSharedPtr<const FMesh> Value)
 		{
 			check(At.At < OpExecutionData.size_code());
 
@@ -1024,7 +1024,7 @@ namespace mu
 			if (!Data->DataTypeIndex)
 			{
 				Data->DataTypeIndex = MeshResults.Num();
-				MeshResults.Add(TResourceResult<Mesh>{At, Value});
+				MeshResults.Add(TResourceResult<FMesh>{At, Value});
 			}
 			else
 			{
@@ -1035,7 +1035,7 @@ namespace mu
 			mu::UpdateLLMStats();
 		}
 
-		void SetLayout(FCacheAddress at, Ptr<const Layout> v)
+		void SetLayout(FCacheAddress at, TSharedPtr<const FLayout> v)
 		{
 			check(at.At < OpExecutionData.size_code());
 			FOpExecutionData* Data = OpExecutionData.get_ptr(at);
@@ -1057,7 +1057,7 @@ namespace mu
 			mu::UpdateLLMStats();
 		}
 
-		void SetString(FCacheAddress at, Ptr<const String> v)
+		void SetString(FCacheAddress at, TSharedPtr<const String> v)
 		{
 			check(at.At < OpExecutionData.size_code());
 			FOpExecutionData* Data = OpExecutionData.get_ptr(at);
@@ -1131,12 +1131,12 @@ namespace mu
 	*/
 	struct FLiveInstance
 	{
-		Instance::ID InstanceID;
+		FInstance::FID InstanceID;
 		int32 State = 0;
-		Ptr<const Instance> Instance;
-		TSharedPtr<const Model> Model;
+		TSharedPtr<const FInstance> Instance;
+		TSharedPtr<const FModel> Model;
 
-		Ptr<Parameters> OldParameters;
+		TSharedPtr<FParameters> OldParameters;
 
 		/** Mask of the parameters that have changed since the last update.
 		* Every bit represents a state parameter.
@@ -1173,7 +1173,7 @@ namespace mu
         struct FModelCacheEntry
         {
             /** Model who's data is being tracked. */
-            TWeakPtr<const Model> Model;
+            TWeakPtr<const FModel> Model;
 
             /** For each model rom, the last time its streamed data was used. */
             TMemoryTrackedArray<TPair<uint64, uint64>> RomWeights;
@@ -1189,7 +1189,7 @@ namespace mu
 		struct FGeneratedResourceData
 		{
 			/** Model for this resource. */
-			TWeakPtr<const Model> Model;
+			TWeakPtr<const FModel> Model;
 
 			//! The id assigned to the generated resource.
 			FResourceID Id;
@@ -1213,7 +1213,7 @@ namespace mu
 		TMemoryTrackedArray<FGeneratedResourceData> GeneratedResources;
 
 		/** */
-		FResourceID GetResourceKey(const TSharedPtr<const Model>&, const Parameters*, uint32 ParamListIndex, OP::ADDRESS RootAt);
+		FResourceID GetResourceKey(const TSharedPtr<const FModel>&, const FParameters*, uint32 ParamListIndex, OP::ADDRESS RootAt);
 
 		//! @}
 
@@ -1240,18 +1240,18 @@ namespace mu
 		TSharedPtr<FProgramCache> CurrentInstanceCache;
 
 		/** Resources that have been used in the past, but haven't been deallocated because they still fitted the memory budget and they could be reused. */
-		TMemoryTrackedArray<Ptr<Image>> PooledImages;
+		TMemoryTrackedArray<TSharedPtr<FImage>> PooledImages;
 
 		/** List of intermediate resources that are not soterd anywhere yet. They are still locally referenced by code. */
-		TMemoryTrackedArray<Ptr<const Image>> TempImages;
-		TMemoryTrackedArray<Ptr<const Mesh>> TempMeshes;
+		TMemoryTrackedArray<TSharedPtr<const FImage>> TempImages;
+		TMemoryTrackedArray<TSharedPtr<const FMesh>> TempMeshes;
 
 		/** List of resources that are currently in any cache position, and the number of positions they are in. */
-		TMemoryTrackedMap<Ptr<const Resource>, int32> CacheResources;
+		TMemoryTrackedMap<TSharedPtr<const FResource>, int32> CacheResources;
 
 		/** Given a mutable model, find or create its rom cache. */
-		FModelCacheEntry* FindModelCache(const Model*);
-		FModelCacheEntry& FindOrAddModelCache(const TSharedPtr<const Model>&);
+		FModelCacheEntry* FindModelCache(const FModel*);
+		FModelCacheEntry& FindOrAddModelCache(const TSharedPtr<const FModel>&);
 
         /** Make sure the working memory is below the internal budget, even counting with the passed additional memory. 
 		* An optional function can be passed to "block" the unload of certain roms of data.
@@ -1266,27 +1266,27 @@ namespace mu
 		int64 GetCurrentMemoryBytes() const;
 
         /** Register that a specific rom has been requested and update the heuristics to keep it in memory. */
-        void MarkRomUsed( int32 RomIndex, const TSharedPtr<const Model>& );
+        void MarkRomUsed( int32 RomIndex, const TSharedPtr<const FModel>& );
 
 		/** */
-		[[nodiscard]] Ptr<Image> CreateImage(uint32 SizeX, uint32 SizeY, uint32 Lods, EImageFormat Format, EInitializationType Init)
+		[[nodiscard]] TSharedPtr<FImage> CreateImage(uint32 SizeX, uint32 SizeY, uint32 Lods, EImageFormat Format, EInitializationType Init)
 		{
 			CheckRunnerThread();
 
-			uint32 DataSize = Image::CalculateDataSize(SizeX, SizeY, Lods, Format);
+			uint32 DataSize = FImage::CalculateDataSize(SizeX, SizeY, Lods, Format);
 
 			// Look for an unused image in the pool that can be reused
 			int32 PooledImageCount = PooledImages.Num();
 			for (int Index = 0; DataSize>0 && Index<PooledImageCount; ++Index)
 			{
-				Ptr<Image>& Candidate = PooledImages[Index];
+				TSharedPtr<FImage>& Candidate = PooledImages[Index];
 				if (Candidate->GetFormat() == Format
 					&& Candidate->GetSizeX() == SizeX
 					&& Candidate->GetSizeY() == SizeY
 					&& Candidate->GetLODCount() == Lods
 					)
 				{
-					Ptr<Image> Result = Candidate;
+					TSharedPtr<FImage> Result = Candidate;
 					PooledImages.RemoveAtSwap(Index);
 					
 					if (Init == EInitializationType::Black)
@@ -1295,7 +1295,7 @@ namespace mu
 					}
 					else
 					{
-						Result->m_flags = 0;
+						Result->Flags = 0;
 						Result->RelevancyMinY = 0;
 						Result->RelevancyMaxY = 0;
 					}
@@ -1307,14 +1307,14 @@ namespace mu
 			EnsureBudgetBelow(DataSize);
 
 			// Create it
-			Ptr<Image> Result = new Image(SizeX, SizeY, Lods, Format, Init);
+			TSharedPtr<FImage> Result = MakeShared<FImage>(SizeX, SizeY, Lods, Format, Init);
 
 			TempImages.Add(Result);
 			return Result;
 		}
 
 		/** Ref will be nulled and relesed in any case. */
-		[[nodiscard]] Ptr<Image> CloneOrTakeOver(Ptr<const Image>& Resource)
+		[[nodiscard]] TSharedPtr<FImage> CloneOrTakeOver(TSharedPtr<const FImage>& Resource)
 		{
 			CheckRunnerThread();
 
@@ -1323,8 +1323,8 @@ namespace mu
 			check(!TempImages.Contains(Resource));
 			check(!PooledImages.Contains(Resource));
 
-			Ptr<Image> Result;
-			if (!Resource->IsUnique())
+			TSharedPtr<FImage> Result;
+			if (!Resource.IsUnique())
 			{
 				// TODO: try to grab from the pool
 
@@ -1336,7 +1336,7 @@ namespace mu
 			}
 			else
 			{
-				Result = const_cast<Image*>(Resource.get());
+				Result = ConstCastSharedPtr<FImage>(Resource);
 				Resource = nullptr;
 			}
 
@@ -1344,7 +1344,7 @@ namespace mu
 		}
 
 		/** */
-		void Release(Ptr<const Image>& Resource)
+		void Release(TSharedPtr<const FImage>& Resource)
 		{
 			CheckRunnerThread();
 
@@ -1365,7 +1365,7 @@ namespace mu
 				bool bInBudget = EnsureBudgetBelow(ResourceDataSize);
 				if (bInBudget)
 				{
-					PooledImages.Add(const_cast<Image*>(Resource.get()));
+					PooledImages.Add(ConstCastSharedPtr<FImage>(Resource));
 				}
 			}
 			else
@@ -1378,7 +1378,7 @@ namespace mu
 		}
 
 		/** */
-		void Release(Ptr<Image>& Resource)
+		void Release(TSharedPtr<FImage>& Resource)
 		{
 			CheckRunnerThread();
 
@@ -1399,7 +1399,7 @@ namespace mu
 				bool bInBudget = EnsureBudgetBelow(ResourceDataSize);
 				if (bInBudget)
 				{
-					PooledImages.Add(Resource.get());
+					PooledImages.Add(Resource);
 				}
 			}
 			else
@@ -1411,28 +1411,28 @@ namespace mu
 			Resource = nullptr;
 		}
 
-		[[nodiscard]] Ptr<Mesh> CreateMesh(int32 BudgetReserveSize)
+		[[nodiscard]] TSharedPtr<FMesh> CreateMesh(int32 BudgetReserveSize)
 		{
 			CheckRunnerThread();
 
 			EnsureBudgetBelow(BudgetReserveSize);
 
-			Ptr<Mesh> Result = new Mesh();
+			TSharedPtr<FMesh> Result = MakeShared<FMesh>();
 
 			TempMeshes.Add(Result);
 
 			return Result;
 		}
 
-		[[nodiscard]] Ptr<Mesh> CloneOrTakeOver(Ptr<const Mesh>& Resource)
+		[[nodiscard]] TSharedPtr<FMesh> CloneOrTakeOver(TSharedPtr<const FMesh>& Resource)
 		{
 			CheckRunnerThread();
 
 			const int32 ResourceDataSize = Resource->GetDataSize();
 			TempMeshes.RemoveSingle(Resource);
 
-			Ptr<Mesh> Result;
-			if (!Resource->IsUnique())
+			TSharedPtr<FMesh> Result;
+			if (!Resource.IsUnique())
 			{
 				Result = CreateMesh(ResourceDataSize); 
 				Result->CopyFrom(*Resource);
@@ -1440,14 +1440,14 @@ namespace mu
 			}
 			else
 			{
-				Result = Ptr<Mesh>(const_cast<Mesh*>(Resource.get()));
+				Result = TSharedPtr<FMesh>(ConstCastSharedPtr<FMesh>(Resource));
 				Resource = nullptr;
 			}
 
 			return Result;
 		}
 
-		void Release(Ptr<const Mesh>& Resource)
+		void Release(TSharedPtr<const FMesh>& Resource)
 		{
 			CheckRunnerThread();
 
@@ -1464,9 +1464,9 @@ namespace mu
 			Resource = nullptr;
 		}
 
-		void Release(Ptr<Mesh>& Resource)
+		void Release(TSharedPtr<FMesh>& Resource)
 		{
-			Ptr<const Mesh> ConstPtr = Resource;
+			TSharedPtr<const FMesh> ConstPtr = Resource;
 
 			Release(ConstPtr);
 
@@ -1474,10 +1474,10 @@ namespace mu
 		}
 
 		/** */
-		[[nodiscard]] Ptr<const Mesh> LoadMesh(const FCacheAddress& From, bool bTakeOwnership = false)
+		[[nodiscard]] TSharedPtr<const FMesh> LoadMesh(const FCacheAddress& From, bool bTakeOwnership = false)
 		{
 			bool bIsLastReference = false;
-			Ptr<const Mesh> Result = CurrentInstanceCache->GetMesh(From, bIsLastReference);
+			TSharedPtr<const FMesh> Result = CurrentInstanceCache->GetMesh(From, bIsLastReference);
 			if (!Result)
 			{
 				return nullptr;
@@ -1495,7 +1495,7 @@ namespace mu
 				}
 			}
 
-			if (!bTakeOwnership && Result->IsUnique())
+			if (!bTakeOwnership && Result.IsUnique())
 			{
 				TempMeshes.Add(Result);
 			}
@@ -1504,10 +1504,10 @@ namespace mu
 		}
 
 		/** */
-		[[nodiscard]] Ptr<const Image> LoadImage(const FCacheAddress& From, bool bTakeOwnership = false)
+		[[nodiscard]] TSharedPtr<const FImage> LoadImage(const FCacheAddress& From, bool bTakeOwnership = false)
 		{
 			bool bIsLastReference = false;
-			Ptr<const Image> Result = CurrentInstanceCache->GetImage(From, bIsLastReference);
+			TSharedPtr<const FImage> Result = CurrentInstanceCache->GetImage(From, bIsLastReference);
 			if (!Result)
 			{
 				return nullptr;
@@ -1526,7 +1526,7 @@ namespace mu
 			}
 
 
-			if (!bTakeOwnership && Result->IsUnique())
+			if (!bTakeOwnership && Result.IsUnique())
 			{
 				TempImages.Add(Result);
 			}
@@ -1535,7 +1535,7 @@ namespace mu
 		}
 
 		/** */
-		void StoreImage(const FCacheAddress& To, Ptr<const Image> Resource)
+		void StoreImage(const FCacheAddress& To, TSharedPtr<const FImage> Resource)
 		{
 			if (Resource)
 			{
@@ -1551,7 +1551,7 @@ namespace mu
 			CurrentInstanceCache->SetImage(To, Resource);
 		}
 
-		void StoreMesh(const FCacheAddress& To, Ptr<const Mesh> Resource)
+		void StoreMesh(const FCacheAddress& To, TSharedPtr<const FMesh> Resource)
 		{
 			if (Resource)
 			{
@@ -1567,21 +1567,21 @@ namespace mu
 		}
 
 		/** Return true if the resource is not in any cache (0,1,rom). */
-		bool IsBudgetTemp(const Ptr<const Resource>& Resource)
+		bool IsBudgetTemp(const TSharedPtr<const FResource>& Resource)
 		{
 			if (!Resource)
 			{
 				return false;
 			}
 
-			bool bIsTemp = Resource->IsUnique();
+			bool bIsTemp = Resource.IsUnique();
 			return bIsTemp;
 		}
 
 		int32 GetPooledBytes() const
 		{
 			int32 Result = 0;
-			for (const Ptr<Image>& Value : PooledImages)
+			for (const TSharedPtr<FImage>& Value : PooledImages)
 			{
 				Result += Value->GetDataSize();
 			}
@@ -1592,12 +1592,12 @@ namespace mu
 		int32 GetTempBytes() const
 		{
 			int32 Result = 0;
-			for (const Ptr<const Image>& Value : TempImages)
+			for (const TSharedPtr<const FImage>& Value : TempImages)
 			{
 				Result += Value->GetDataSize();
 			}
 
-			for (const Ptr<const Mesh>& Value : TempMeshes)
+			for (const TSharedPtr<const FMesh>& Value : TempMeshes)
 			{
 				Result += Value->GetDataSize();
 			}
@@ -1610,25 +1610,25 @@ namespace mu
 		{
 			int32 Result = 0;
 
-			TArray<const Model*> Models;
+			TArray<const FModel*> Models;
 			for (const FLiveInstance& Instance : LiveInstances)
 			{
 				Models.AddUnique(Instance.Model.Get());
 			}
 
 			// Data stored per-model, but related to instance construction
-			for (const Model* Model : Models)
+			for (const FModel* Model : Models)
 			{
 				// Count streamable and currently-loaded resources
-				const FProgram& Program = Model->GetPrivate()->m_program;
-				for (const TPair<int32, Ptr<const Image>>& Rom : Program.ConstantImageLODs)
+				const FProgram& Program = Model->GetPrivate()->Program;
+				for (const TPair<int32, TSharedPtr<const FImage>>& Rom : Program.ConstantImageLODs)
 				{
 					if (Rom.Value && Rom.Key >= 0)
 					{
 						Result += Rom.Value->GetDataSize();
 					}
 				}
-				for (const TPair<int32, Ptr<const Mesh>>& Rom : Program.ConstantMeshes)
+				for (const TPair<int32, TSharedPtr<const FMesh>>& Rom : Program.ConstantMeshes)
 				{
 					if (Rom.Value && Rom.Key >= 0)
 					{
@@ -1644,7 +1644,7 @@ namespace mu
 		int32 GetTrackedCacheBytes() const
 		{
 			int32 Result = 0;
-			for (const TTuple<Ptr<const Resource>,int32>& It : CacheResources)
+			for (const TTuple<TSharedPtr<const FResource>,int32>& It : CacheResources)
 			{
 				Result += It.Key->GetDataSize();
 			}
@@ -1656,8 +1656,8 @@ namespace mu
 		int32 GetCacheBytes() const
 		{
 			int32 Result = 0;
-			TSet<const Resource*> Cache0Unique;
-			TSet<const Resource*> Cache1Unique;
+			TSet<const FResource*> Cache0Unique;
+			TSet<const FResource*> Cache1Unique;
 
 			for (const FLiveInstance& Instance : LiveInstances)
 			{
@@ -1669,17 +1669,17 @@ namespace mu
 						continue;
 					}
 
-					TSet<const Resource*>* TargetSet = &Cache0Unique;
+					TSet<const FResource*>* TargetSet = &Cache0Unique;
 					if ((*it).IsCacheLocked)
 					{
 						TargetSet = &Cache1Unique;
 					}
 
-					const Resource* Value = nullptr;
+					const FResource* Value = nullptr;
 					switch ((*it).DataType)
 					{
 					case DATATYPE::DT_IMAGE:
-						Value = Instance.Cache->ImageResults[(*it).DataTypeIndex].Value.get();
+						Value = Instance.Cache->ImageResults[(*it).DataTypeIndex].Value.Get();
 						if (Value)
 						{
 							TargetSet->Add(Value);
@@ -1687,7 +1687,7 @@ namespace mu
 						break;
 
 					case DATATYPE::DT_MESH:
-						Value = Instance.Cache->MeshResults[(*it).DataTypeIndex].Value.get();
+						Value = Instance.Cache->MeshResults[(*it).DataTypeIndex].Value.Get();
 						if (Value)
 						{
 							TargetSet->Add(Value);
@@ -1704,7 +1704,7 @@ namespace mu
 			Cache0Unique.Append(Cache1Unique);
 
 			// Count
-			for (const Resource* Value : Cache0Unique)
+			for (const FResource* Value : Cache0Unique)
 			{
 				Result += Value->GetDataSize();
 			}
@@ -1732,37 +1732,43 @@ namespace mu
 					continue;
 				}
 
-				const Resource* Value = nullptr;
-
 				switch (Data.DataType)
 				{
 				case DATATYPE::DT_IMAGE:
-					Value = CurrentInstanceCache->ImageResults[Data.DataTypeIndex].Value.get();
+				{
+					TSharedPtr<const FImage> Value = CurrentInstanceCache->ImageResults[Data.DataTypeIndex].Value;
 					if (Value)
 					{
 						CacheResources.Remove(Value);
 						CurrentInstanceCache->ImageResults[Data.DataTypeIndex].Value = nullptr;
 					}
 					break;
+				}
 
 				case DATATYPE::DT_MESH:
-					Value = CurrentInstanceCache->MeshResults[Data.DataTypeIndex].Value.get();
+				{
+					TSharedPtr<const FMesh> Value = CurrentInstanceCache->MeshResults[Data.DataTypeIndex].Value;
 					if (Value)
 					{
 						CacheResources.Remove(Value);
 						CurrentInstanceCache->MeshResults[Data.DataTypeIndex].Value = nullptr;
 					}
 					break;
+				}
 
 				case DATATYPE::DT_LAYOUT:
-					Value = CurrentInstanceCache->LayoutResults[Data.DataTypeIndex].get();
+				{
+					TSharedPtr<const FLayout> Value = CurrentInstanceCache->LayoutResults[Data.DataTypeIndex];
 					CurrentInstanceCache->LayoutResults[Data.DataTypeIndex] = nullptr;
 					break;
+				}
 
 				case DATATYPE::DT_INSTANCE:
-					Value = CurrentInstanceCache->InstanceResults[Data.DataTypeIndex].get();
+				{
+					TSharedPtr<const FInstance> Value = CurrentInstanceCache->InstanceResults[Data.DataTypeIndex];
 					CurrentInstanceCache->InstanceResults[Data.DataTypeIndex] = nullptr;
 					break;
+				}
 
 				default:
 					break;
@@ -1790,31 +1796,37 @@ namespace mu
 					continue;
 				}
 
-				const Resource* Value = nullptr;
-
 				switch (Data.DataType)
 				{
 				case DATATYPE::DT_IMAGE:
-					Value = CurrentInstanceCache->ImageResults[Data.DataTypeIndex].Value.get();
+				{
+					TSharedPtr<const FImage> Value = CurrentInstanceCache->ImageResults[Data.DataTypeIndex].Value;
 					CacheResources.Remove(Value);
 					CurrentInstanceCache->ImageResults[Data.DataTypeIndex].Value = nullptr;
 					break;
+				}
 
 				case DATATYPE::DT_MESH:
-					Value = CurrentInstanceCache->MeshResults[Data.DataTypeIndex].Value.get();
+				{
+					TSharedPtr<const FMesh> Value = CurrentInstanceCache->MeshResults[Data.DataTypeIndex].Value;
 					CacheResources.Remove(Value);
 					CurrentInstanceCache->MeshResults[Data.DataTypeIndex].Value = nullptr;
 					break;
+				}
 
 				case DATATYPE::DT_LAYOUT:
-					Value = CurrentInstanceCache->LayoutResults[Data.DataTypeIndex].get();
+				{
+					TSharedPtr<const FLayout> Value = CurrentInstanceCache->LayoutResults[Data.DataTypeIndex];
 					CurrentInstanceCache->LayoutResults[Data.DataTypeIndex] = nullptr;
 					break;
+				}
 
 				case DATATYPE::DT_INSTANCE:
-					Value = CurrentInstanceCache->InstanceResults[Data.DataTypeIndex].get();
+				{
+					TSharedPtr<const FInstance> Value = CurrentInstanceCache->InstanceResults[Data.DataTypeIndex];
 					CurrentInstanceCache->InstanceResults[Data.DataTypeIndex] = nullptr;
 					break;
+				}
 
 				default:
 					break;
@@ -1916,12 +1928,12 @@ namespace mu
 
 
 	/** */
-    class System::Private
+    class FSystem::Private
     {
 		friend class System;
     public:
 
-        Private( Ptr<Settings>, const TSharedPtr<ExtensionDataStreamer>& );
+        Private(const FSettings&, const TSharedPtr<FExtensionDataStreamer>& );
         ~Private();
 
         //-----------------------------------------------------------------------------------------
@@ -1929,40 +1941,38 @@ namespace mu
         //-----------------------------------------------------------------------------------------
 
 		/** This method can be used to internally prepare for code execution. */
-		MUTABLERUNTIME_API void BeginBuild(const TSharedPtr<const Model>&);
+		MUTABLERUNTIME_API void BeginBuild(const TSharedPtr<const FModel>&);
 		MUTABLERUNTIME_API void EndBuild();
 
-		MUTABLERUNTIME_API bool BuildBool(const TSharedPtr<const Model>&, const Parameters*, OP::ADDRESS) ;
-		MUTABLERUNTIME_API int32 BuildInt(const TSharedPtr<const Model>&, const Parameters*, OP::ADDRESS) ;
-		MUTABLERUNTIME_API float BuildScalar(const TSharedPtr<const Model>&, const Parameters*, OP::ADDRESS) ;
-		MUTABLERUNTIME_API FVector4f BuildColour(const TSharedPtr<const Model>&, const Parameters*, OP::ADDRESS) ;
-		MUTABLERUNTIME_API Ptr<const String> BuildString(const TSharedPtr<const Model>&, const Parameters*, OP::ADDRESS) ;
-		MUTABLERUNTIME_API Ptr<const Image> BuildImage(const TSharedPtr<const Model>&, const Parameters*, OP::ADDRESS, int32 MipsToSkip, int32 LOD) ;
-		MUTABLERUNTIME_API Ptr<const Mesh> BuildMesh(const TSharedPtr<const Model>&, const Parameters*, OP::ADDRESS) ;
-		MUTABLERUNTIME_API Ptr<const Instance> BuildInstance(const TSharedPtr<const Model>&, const Parameters*, OP::ADDRESS);
-		MUTABLERUNTIME_API Ptr<const Layout> BuildLayout(const TSharedPtr<const Model>&, const Parameters*, OP::ADDRESS) ;
-    	MUTABLERUNTIME_API FProjector BuildProjector(const TSharedPtr<const Model>&, const Parameters*, OP::ADDRESS) ;
+		MUTABLERUNTIME_API bool BuildBool(const TSharedPtr<const FModel>&, const FParameters*, OP::ADDRESS) ;
+		MUTABLERUNTIME_API int32 BuildInt(const TSharedPtr<const FModel>&, const FParameters*, OP::ADDRESS) ;
+		MUTABLERUNTIME_API float BuildScalar(const TSharedPtr<const FModel>&, const FParameters*, OP::ADDRESS) ;
+		MUTABLERUNTIME_API FVector4f BuildColour(const TSharedPtr<const FModel>&, const FParameters*, OP::ADDRESS) ;
+		MUTABLERUNTIME_API TSharedPtr<const String> BuildString(const TSharedPtr<const FModel>&, const FParameters*, OP::ADDRESS) ;
+		MUTABLERUNTIME_API TSharedPtr<const FImage> BuildImage(const TSharedPtr<const FModel>&, const FParameters*, OP::ADDRESS, int32 MipsToSkip, int32 LOD) ;
+		MUTABLERUNTIME_API TSharedPtr<const FMesh> BuildMesh(const TSharedPtr<const FModel>&, const FParameters*, OP::ADDRESS) ;
+		MUTABLERUNTIME_API TSharedPtr<const FInstance> BuildInstance(const TSharedPtr<const FModel>&, const FParameters*, OP::ADDRESS);
+		MUTABLERUNTIME_API TSharedPtr<const FLayout> BuildLayout(const TSharedPtr<const FModel>&, const FParameters*, OP::ADDRESS) ;
+    	MUTABLERUNTIME_API FProjector BuildProjector(const TSharedPtr<const FModel>&, const FParameters*, OP::ADDRESS) ;
 
-		ExtensionDataStreamer* GetExtensionDataStreamer() const { return ExtensionDataStreamer.Get(); }
+		FExtensionDataStreamer* GetExtensionDataStreamer() const { return ExtensionDataStreamer.Get(); }
 
         //!
-        Ptr<const Settings> Settings;
+        FSettings Settings;
 
         //! Data streaming interface, if any.
-		TSharedPtr<ModelReader> StreamInterface;
+		TSharedPtr<FModelReader> StreamInterface;
 
-		TSharedPtr<ExternalResourceProvider> ExternalResourceProvider;
+		TSharedPtr<FExternalResourceProvider> ExternalResourceProvider;
 
 		/** Counter used to generate unique IDs for every new instance created in the system. */
-        Instance::ID LastInstanceID = 0;
+        FInstance::FID LastInstanceID = 0;
 
-	private:
 		/** This flag is turned on when a streaming error or similar happens. Results are not usable.
 		* This should only happen in-editor.
 		*/
 		bool bUnrecoverableError = false;
 
-	public:
 		/** If this is set, it will be tried first instead of the internal formatting function. */
 		FImageOperator::FImagePixelFormatFunc ImagePixelFormatOverride;
 
@@ -1970,16 +1980,16 @@ namespace mu
 		FWorkingMemoryManager WorkingMemoryManager;
 
 		/** The pointer returned by this function is only valid for the duration of the current mutable operation. */
-		inline FLiveInstance* FindLiveInstance(Instance::ID id);
+		inline FLiveInstance* FindLiveInstance(FInstance::FID);
 
         //!
-        bool CheckUpdatedParameters( const FLiveInstance*, const Ptr<const Parameters>&, uint64& OutUpdatedParameters );
+        bool CheckUpdatedParameters( const FLiveInstance*, const TSharedPtr<const FParameters>&, uint64& OutUpdatedParameters );
 
 
-		void RunCode(const TSharedPtr<const Model>&, const Parameters*, OP::ADDRESS at, uint32 LODs = System::AllLODs, uint8 executionOptions = 0, int32 LOD = 0);
+		void RunCode(const TSharedPtr<const FModel>&, const FParameters*, OP::ADDRESS, uint32 LODs = FSystem::AllLODs, uint8 ExecutionOptions = 0, int32 LOD = 0);
 
 		//!
-		void PrepareCache(const Model*, int32 State);
+		void PrepareCache(const FModel*, int32 State);
 
 		//! Update some mutable core unreal stats.
 		void UpdateStats();
@@ -1987,7 +1997,7 @@ namespace mu
 	private:
 
 		/** Owned by this system. */
-		TSharedPtr<ExtensionDataStreamer> ExtensionDataStreamer = nullptr;
+		TSharedPtr<FExtensionDataStreamer> ExtensionDataStreamer;
     };
 
 }
