[mutable] More enforcement of coding standards.

diff --git a/Engine/Plugins/Mutable/Source/MutableRuntime/Private/MuR/CodeRunner.cpp b/Engine/Plugins/Mutable/Source/MutableRuntime/Private/MuR/CodeRunner.cpp
--- a/Engine/Plugins/Mutable/Source/MutableRuntime/Private/MuR/CodeRunner.cpp
+++ b/Engine/Plugins/Mutable/Source/MutableRuntime/Private/MuR/CodeRunner.cpp
@@ -107,11 +107,11 @@ namespace mu
 {
 
 	TSharedRef<CodeRunner> CodeRunner::Create(
-		const Ptr<const Settings>& InSettings,
-		class System::Private* InSystem,
+		const FSettings& InSettings,
+		class FSystem::Private* InSystem,
 		EExecutionStrategy InExecutionStrategy,
-		const TSharedPtr<const Model>& InModel,
-		const Parameters* InParams,
+		const TSharedPtr<const FModel>& InModel,
+		const FParameters* InParams,
 		OP::ADDRESS At,
 		uint32 InLODMask, uint8 ExecutionOptions, int32 InImageLOD, FScheduledOp::EType Type)
 	{
@@ -120,23 +120,23 @@ namespace mu
 	}
 
     CodeRunner::CodeRunner(FPrivateToken PrivateToken, 
-		const Ptr<const Settings>& InSettings,
-		class System::Private* InSystem,
+		const FSettings& InSettings,
+		class FSystem::Private* InSystem,
 		EExecutionStrategy InExecutionStrategy,
-		const TSharedPtr<const Model>& InModel,
-		const Parameters* InParams,
+		const TSharedPtr<const FModel>& InModel,
+		const FParameters* InParams,
 		OP::ADDRESS at,
 		uint32 InLodMask, uint8 executionOptions, int32 InImageLOD, FScheduledOp::EType Type )
-		: m_pSettings(InSettings)
+		: Settings(InSettings)
 		, RunnerCompletionEvent(TEXT("CodeRunnerCompletioneEventInit"))
 		, ExecutionStrategy(InExecutionStrategy)
-		, m_pSystem(InSystem)
-		, m_pModel(InModel)
-		, m_pParams(InParams)
-		, m_lodMask(InLodMask)
+		, System(InSystem)
+		, Model(InModel)
+		, Params(InParams)
+		, LODMask(InLodMask)
 	{
-		const FProgram& program = m_pModel->GetPrivate()->m_program;
-		ScheduledStagePerOp.resize(program.m_opAddress.Num());
+		const FProgram& program = Model->GetPrivate()->Program;
+		ScheduledStagePerOp.resize(program.OpAddress.Num());
 
 		// We will read this in the end, so make sure we keep it.
    		if (Type == FScheduledOp::EType::Full)
@@ -162,27 +162,27 @@ namespace mu
     //---------------------------------------------------------------------------------------------
 	FProgramCache& CodeRunner::GetMemory()
     {
-		return *m_pSystem->WorkingMemoryManager.CurrentInstanceCache;
+		return *System->WorkingMemoryManager.CurrentInstanceCache;
 	}
 
 
-	TTuple<UE::Tasks::FTask, TFunction<void()>> CodeRunner::LoadExternalImageAsync(FExternalResourceId Id, uint8 MipmapsToSkip, TFunction<void(Ptr<Image>)>& ResultCallback)
+	TTuple<UE::Tasks::FTask, TFunction<void()>> CodeRunner::LoadExternalImageAsync(FExternalResourceId Id, uint8 MipmapsToSkip, TFunction<void(TSharedPtr<FImage>)>& ResultCallback)
     {
 		MUTABLE_CPUPROFILER_SCOPE(LoadExternalImageAsync);
 
-		check(m_pSystem);
+		check(System);
 
-		if (m_pSystem->ExternalResourceProvider)
+		if (System->ExternalResourceProvider)
 		{
 			if (Id.ReferenceResourceId < 0)
 			{
 				// It's a parameter image
-				return m_pSystem->ExternalResourceProvider->GetImageAsync(Id.ParameterId, MipmapsToSkip, ResultCallback);
+				return System->ExternalResourceProvider->GetImageAsync(Id.ParameterId, MipmapsToSkip, ResultCallback);
 			}
 			else
 			{
 				// It's an image reference
-				return m_pSystem->ExternalResourceProvider->GetReferencedImageAsync(m_pModel.Get(), Id.ReferenceResourceId, MipmapsToSkip, ResultCallback);
+				return System->ExternalResourceProvider->GetReferencedImageAsync(Model.Get(), Id.ReferenceResourceId, MipmapsToSkip, ResultCallback);
 			}
 		}
 		else
@@ -195,23 +195,23 @@ namespace mu
 	}
 
 
-	TTuple<UE::Tasks::FTask, TFunction<void()>> CodeRunner::LoadExternalMeshAsync(FExternalResourceId Id, TFunction<void(Ptr<Mesh>)>& ResultCallback)
+	TTuple<UE::Tasks::FTask, TFunction<void()>> CodeRunner::LoadExternalMeshAsync(FExternalResourceId Id, TFunction<void(TSharedPtr<FMesh>)>& ResultCallback)
 	{
 		MUTABLE_CPUPROFILER_SCOPE(LoadExternalImageAsync);
 
-		check(m_pSystem);
+		check(System);
 
-		if (m_pSystem->ExternalResourceProvider)
+		if (System->ExternalResourceProvider)
 		{
 			if (Id.ReferenceResourceId < 0)
 			{
 				// It's a parameter mesh
-				return m_pSystem->ExternalResourceProvider->GetMeshAsync(Id.ParameterId, ResultCallback);
+				return System->ExternalResourceProvider->GetMeshAsync(Id.ParameterId, ResultCallback);
 			}
 			else
 			{
 				// It's a mesh reference
-				return m_pSystem->ExternalResourceProvider->GetReferencedMeshAsync(m_pModel.Get(), Id.ReferenceResourceId, ResultCallback);
+				return System->ExternalResourceProvider->GetReferencedMeshAsync(Model.Get(), Id.ReferenceResourceId, ResultCallback);
 			}
 		}
 		else
@@ -229,11 +229,11 @@ namespace mu
 	{
 		MUTABLE_CPUPROFILER_SCOPE(GetExternalImageDesc);
 
-		check(m_pSystem);
+		check(System);
 
-		if (m_pSystem->ExternalResourceProvider)
+		if (System->ExternalResourceProvider)
 		{
-			return m_pSystem->ExternalResourceProvider->GetImageDesc(Id, MipmapsToSkip);
+			return System->ExternalResourceProvider->GetImageDesc(Id, MipmapsToSkip);
 		}
 		else
 		{
@@ -246,13 +246,13 @@ namespace mu
 
 	
     //---------------------------------------------------------------------------------------------
-    void CodeRunner::RunCode_Conditional( const FScheduledOp& item, const Model* pModel )
+    void CodeRunner::RunCode_Conditional( const FScheduledOp& item, const FModel* pModel )
     {
         MUTABLE_CPUPROFILER_SCOPE(RunCode_Conditional);
 
-		const FProgram& Program = m_pModel->GetPrivate()->m_program;
+		const FProgram& Program = Model->GetPrivate()->Program;
 
-		OP_TYPE type = Program.GetOpType(item.At);
+		EOpType type = Program.GetOpType(item.At);
 		OP::ConditionalArgs args = Program.GetOpArgs<OP::ConditionalArgs>(item.At);
 
         // Conditionals have the following execution stages:
@@ -320,11 +320,11 @@ namespace mu
 
 
 	//---------------------------------------------------------------------------------------------
-	void CodeRunner::RunCode_Switch(const FScheduledOp& item, const Model* pModel )
+	void CodeRunner::RunCode_Switch(const FScheduledOp& item, const FModel* pModel )
 	{
-		const FProgram& Program = m_pModel->GetPrivate()->m_program;
+		const FProgram& Program = Model->GetPrivate()->Program;
 
-		OP_TYPE type = Program.GetOpType(item.At);
+		EOpType type = Program.GetOpType(item.At);
 
 		const uint8* data = Program.GetOpArgsPointer(item.At);
 
@@ -362,7 +362,7 @@ namespace mu
 				case DT_IMAGE:      StoreImage(item, nullptr); break;
 				case DT_LAYOUT:     StoreLayout(item, nullptr); break;
 				case DT_INSTANCE:   StoreInstance(item, nullptr); break;
-				case DT_EXTENSION_DATA: StoreExtensionData(item, new ExtensionData); break;
+				case DT_EXTENSION_DATA: StoreExtensionData(item, MakeShared<FExtensionData>()); break;
 				default:
 					// Not implemented
 					check(false);
@@ -374,7 +374,7 @@ namespace mu
 		case 1:
 		{
 			// Get the variable result
-			int var = LoadInt(FCacheAddress(VarAddress, item));
+			int32 var = LoadInt(FCacheAddress(VarAddress, item));
 
 			OP::ADDRESS valueAt = DefAddress;
 			for (uint32 C = 0; C < CaseCount; ++C)
@@ -387,7 +387,7 @@ namespace mu
 				FMemory::Memcpy(&At, data, sizeof(OP::ADDRESS));
 				data += sizeof(OP::ADDRESS);
 
-				if (At && var == (int)Condition)
+				if (At && var == (int32)Condition)
 				{
 					valueAt = At;
 					break; 
@@ -436,17 +436,17 @@ namespace mu
 
 
     //---------------------------------------------------------------------------------------------
-    void CodeRunner::RunCode_Instance(const FScheduledOp& item, const Model* pModel, uint32 lodMask )
+    void CodeRunner::RunCode_Instance(const FScheduledOp& item, const FModel* pModel, uint32 lodMask )
     {
         MUTABLE_CPUPROFILER_SCOPE(RunCode_Instance);
 
-		const FProgram& Program = m_pModel->GetPrivate()->m_program;
+		const FProgram& Program = Model->GetPrivate()->Program;
 
-		OP_TYPE type = Program.GetOpType(item.At);
+		EOpType type = Program.GetOpType(item.At);
         switch (type)
         {
 
-        case OP_TYPE::IN_ADDVECTOR:
+        case EOpType::IN_ADDVECTOR:
         {
 			OP::InstanceAddArgs args = Program.GetOpArgs<OP::InstanceAddArgs>(item.At);
 
@@ -460,15 +460,15 @@ namespace mu
 
             case 1:
             {
-                InstancePtrConst pBase = LoadInstance( FCacheAddress(args.instance,item) );
-                InstancePtr pResult;
-                if (!pBase)
+                TSharedPtr<const FInstance> Base = LoadInstance( FCacheAddress(args.instance,item) );
+				TSharedPtr<FInstance> Result;
+                if (!Base)
                 {
-                    pResult = new Instance();
+                    Result = MakeShared<FInstance>();
                 }
                 else
                 {
-					pResult = mu::CloneOrTakeOver<Instance>(pBase.get());
+					Result = mu::CloneOrTakeOver<FInstance>(Base);
                 }
 
                 if ( args.value )
@@ -476,12 +476,12 @@ namespace mu
 					FVector4f value = LoadColor( FCacheAddress(args.value,item) );
 
                     OP::ADDRESS nameAd = args.name;
-                    check(  nameAd < (uint32)Program.m_constantStrings.Num() );
-                    const FString& Name = Program.m_constantStrings[ nameAd ];
+                    check(  nameAd < (uint32)Program.ConstantStrings.Num() );
+                    const FString& Name = Program.ConstantStrings[ nameAd ];
 
-                    pResult->GetPrivate()->AddVector( 0, 0, 0, value, FName(Name) );
+                    Result->GetPrivate()->AddVector( 0, 0, 0, value, FName(Name) );
                 }
-                StoreInstance( item, pResult );
+                StoreInstance( item, Result );
                 break;
             }
 
@@ -492,7 +492,7 @@ namespace mu
             break;
         }
 
-        case OP_TYPE::IN_ADDSCALAR:
+        case EOpType::IN_ADDSCALAR:
         {
 			OP::InstanceAddArgs args = Program.GetOpArgs<OP::InstanceAddArgs>(item.At);
             switch (item.Stage)
@@ -505,15 +505,15 @@ namespace mu
 
             case 1:
             {
-                InstancePtrConst pBase = LoadInstance( FCacheAddress(args.instance,item) );
-                InstancePtr pResult;
-                if (!pBase)
+                TSharedPtr<const FInstance> Base = LoadInstance( FCacheAddress(args.instance,item) );
+                TSharedPtr<FInstance> Result;
+                if (!Base)
                 {
-                    pResult = new Instance();
+                    Result = MakeShared<FInstance>();
                 }
                 else
                 {
-                    pResult = mu::CloneOrTakeOver<Instance>(pBase.get());
+                    Result = mu::CloneOrTakeOver<FInstance>(Base);
                 }
 
                 if ( args.value )
@@ -521,12 +521,12 @@ namespace mu
                     float value = LoadScalar( FCacheAddress(args.value,item) );
 
                     OP::ADDRESS nameAd = args.name;
-                    check(  nameAd < (uint32)Program.m_constantStrings.Num() );
-                    const FString& Name = Program.m_constantStrings[ nameAd ];
+                    check(  nameAd < (uint32)Program.ConstantStrings.Num() );
+                    const FString& Name = Program.ConstantStrings[ nameAd ];
 
-                    pResult->GetPrivate()->AddScalar( 0, 0, 0, value, FName(Name));
+                    Result->GetPrivate()->AddScalar( 0, 0, 0, value, FName(Name));
                 }
-                StoreInstance( item, pResult );
+                StoreInstance( item, Result );
                 break;
             }
 
@@ -537,7 +537,7 @@ namespace mu
             break;
         }
 
-        case OP_TYPE::IN_ADDSTRING:
+        case EOpType::IN_ADDSTRING:
         {
 			OP::InstanceAddArgs args = Program.GetOpArgs<OP::InstanceAddArgs>( item.At );
             switch ( item.Stage )
@@ -549,29 +549,29 @@ namespace mu
 
             case 1:
             {
-                InstancePtrConst pBase =
+                TSharedPtr<const FInstance> Base =
                     LoadInstance( FCacheAddress( args.instance, item ) );
-                InstancePtr pResult;
-                if ( !pBase )
+                TSharedPtr<FInstance> Result;
+                if ( !Base )
                 {
-                    pResult = new Instance();
+                    Result = MakeShared<FInstance>();
                 }
                 else
                 {
-					pResult = mu::CloneOrTakeOver<Instance>(pBase.get());
+					Result = mu::CloneOrTakeOver<FInstance>(Base);
 				}
 
                 if ( args.value )
                 {
-                    Ptr<const String> value = LoadString( FCacheAddress( args.value, item ) );
+                    TSharedPtr<const String> value = LoadString( FCacheAddress( args.value, item ) );
 
                     OP::ADDRESS nameAd = args.name;
-                    check( nameAd < (uint32)Program.m_constantStrings.Num() );
-                    const FString& Name = Program.m_constantStrings[nameAd];
+                    check( nameAd < (uint32)Program.ConstantStrings.Num() );
+                    const FString& Name = Program.ConstantStrings[nameAd];
 
-                    pResult->GetPrivate()->AddString( 0, 0, 0, value->GetValue(), FName(Name) );
+                    Result->GetPrivate()->AddString( 0, 0, 0, value->GetValue(), FName(Name) );
                 }
-                StoreInstance( item, pResult );
+                StoreInstance( item, Result );
                 break;
             }
 
@@ -582,7 +582,7 @@ namespace mu
             break;
         }
 
-        case OP_TYPE::IN_ADDCOMPONENT:
+        case EOpType::IN_ADDCOMPONENT:
         {
 			OP::InstanceAddArgs args = Program.GetOpArgs<OP::InstanceAddArgs>(item.At);
             switch (item.Stage)
@@ -595,32 +595,32 @@ namespace mu
 
             case 1:
             {
-				Ptr<const Instance> pBase = LoadInstance( FCacheAddress(args.instance,item) );
-				Ptr<Instance> pResult;
-                if (!pBase)
+				TSharedPtr<const FInstance> Base = LoadInstance( FCacheAddress(args.instance,item) );
+				TSharedPtr<FInstance> Result;
+                if (!Base)
                 {
-                    pResult = new Instance();
+                    Result = MakeShared<FInstance>();
                 }
                 else
                 {
-					pResult = mu::CloneOrTakeOver<Instance>(pBase.get());
+					Result = mu::CloneOrTakeOver<FInstance>(Base);
 				}
 
                 if ( args.value )
                 {
-                    Ptr<const Instance> pComp = LoadInstance( FCacheAddress(args.value,item) );
+                    TSharedPtr<const FInstance> pComp = LoadInstance( FCacheAddress(args.value,item) );
 
-                    int32 NewComponentIndex = pResult->GetPrivate()->AddComponent();
+                    int32 NewComponentIndex = Result->GetPrivate()->AddComponent();
 
                     if ( !pComp->GetPrivate()->Components.IsEmpty() )
                     {
-                        pResult->GetPrivate()->Components[NewComponentIndex] = pComp->GetPrivate()->Components[0];
+                        Result->GetPrivate()->Components[NewComponentIndex] = pComp->GetPrivate()->Components[0];
 
 						// Id
-                    	pResult->GetPrivate()->Components[NewComponentIndex].Id = args.ExternalId;
+                    	Result->GetPrivate()->Components[NewComponentIndex].Id = args.ExternalId;
                     }
                 }
-                StoreInstance( item, pResult );
+                StoreInstance( item, Result );
                 break;
             }
 
@@ -631,7 +631,7 @@ namespace mu
             break;
         }
 
-        case OP_TYPE::IN_ADDSURFACE:
+        case EOpType::IN_ADDSURFACE:
         {
 			OP::InstanceAddArgs args = Program.GetOpArgs<OP::InstanceAddArgs>(item.At);
             switch (item.Stage)
@@ -644,25 +644,25 @@ namespace mu
 
             case 1:
             {
-				Ptr<const Instance> pBase = LoadInstance( FCacheAddress(args.instance,item) );
+				TSharedPtr<const FInstance> Base = LoadInstance( FCacheAddress(args.instance,item) );
 
-				Ptr<Instance> pResult;
-				if (pBase)
+				TSharedPtr<FInstance> Result;
+				if (Base)
 				{
-					pResult = mu::CloneOrTakeOver<Instance>(pBase.get());
+					Result = mu::CloneOrTakeOver<FInstance>(Base);
 				}
 				else
 				{
-					pResult = new Instance();
+					Result = MakeShared<FInstance>();
 				}
 
                 // Empty surfaces are ok, they still need to be created, because they may contain
                 // additional information like internal or external IDs
                 //if ( args.value )
                 {
-					Ptr<const Instance> pSurf = LoadInstance( FCacheAddress(args.value,item) );
+					TSharedPtr<const FInstance> pSurf = LoadInstance( FCacheAddress(args.value,item) );
 
-                    int sindex = pResult->GetPrivate()->AddSurface( 0, 0 );
+                    int32 sindex = Result->GetPrivate()->AddSurface( 0, 0 );
 
                     // Surface data
                     if (pSurf
@@ -673,22 +673,22 @@ namespace mu
                             &&
                             pSurf->GetPrivate()->Components[0].LODs[0].Surfaces.Num())
                     {
-                        pResult->GetPrivate()->Components[0].LODs[0].Surfaces[sindex] =
+                        Result->GetPrivate()->Components[0].LODs[0].Surfaces[sindex] =
                             pSurf->GetPrivate()->Components[0].LODs[0].Surfaces[0];
                     }
 
                     // Name
                     OP::ADDRESS nameAd = args.name;
-                    check( nameAd < (uint32)Program.m_constantStrings.Num() );
-                    const FString& Name = Program.m_constantStrings[ nameAd ];
-                    pResult->GetPrivate()->SetSurfaceName( 0, 0, sindex, FName(Name) );
+                    check( nameAd < (uint32)Program.ConstantStrings.Num() );
+                    const FString& Name = Program.ConstantStrings[ nameAd ];
+                    Result->GetPrivate()->SetSurfaceName( 0, 0, sindex, FName(Name) );
 
                     // IDs
-                    pResult->GetPrivate()->Components[0].LODs[0].Surfaces[sindex].InternalId = args.id;
-                    pResult->GetPrivate()->Components[0].LODs[0].Surfaces[sindex].ExternalId = args.ExternalId;
-                    pResult->GetPrivate()->Components[0].LODs[0].Surfaces[sindex].SharedId = args.SharedSurfaceId;
+                    Result->GetPrivate()->Components[0].LODs[0].Surfaces[sindex].InternalId = args.id;
+                    Result->GetPrivate()->Components[0].LODs[0].Surfaces[sindex].ExternalId = args.ExternalId;
+                    Result->GetPrivate()->Components[0].LODs[0].Surfaces[sindex].SharedId = args.SharedSurfaceId;
                 }
-                StoreInstance( item, pResult );
+                StoreInstance( item, Result );
                 break;
             }
 
@@ -699,7 +699,7 @@ namespace mu
             break;
         }
 
-        case OP_TYPE::IN_ADDLOD:
+        case EOpType::IN_ADDLOD:
         {
 			const uint8* Data = Program.GetOpArgsPointer(item.At);
 
@@ -737,8 +737,8 @@ namespace mu
             case 1:
             {
                 // Assemble result
-				Ptr<Instance> pResult = new Instance();
-				int32 ComponentIndex = pResult->GetPrivate()->AddComponent();
+				TSharedPtr<FInstance> Result = MakeShared<FInstance>();
+				int32 ComponentIndex = Result->GetPrivate()->AddComponent();
 
                 for (uint8 LODIndex = 0; LODIndex < LODCount; ++LODIndex)
                 {
@@ -751,24 +751,24 @@ namespace mu
 						bool bIsSelectedLod = ( (1<<LODIndex) & lodMask ) != 0;
 
 						// Add an empty LOD even if not selected.
-						int32 InstanceLODIndex = pResult->GetPrivate()->AddLOD(ComponentIndex);
+						int32 InstanceLODIndex = Result->GetPrivate()->AddLOD(ComponentIndex);
 						
 						if (bIsSelectedLod)
                         {
-							Ptr<const Instance> pLOD = LoadInstance( FCacheAddress(LODAddress,item) );
+							TSharedPtr<const FInstance> pLOD = LoadInstance( FCacheAddress(LODAddress,item) );
 
                             // In a degenerated case, the returned pLOD may not have an LOD inside
  							if (!pLOD->GetPrivate()->Components.IsEmpty()
 								&&
 								!pLOD->GetPrivate()->Components[0].LODs.IsEmpty())
 							{
-								pResult->GetPrivate()->Components[ComponentIndex].LODs[InstanceLODIndex] = pLOD->GetPrivate()->Components[0].LODs[0];
+								Result->GetPrivate()->Components[ComponentIndex].LODs[InstanceLODIndex] = pLOD->GetPrivate()->Components[0].LODs[0];
 							}
 						}
                     }
                 }
 
-                StoreInstance( item, pResult );
+                StoreInstance( item, Result );
                 break;
             }
 
@@ -779,14 +779,14 @@ namespace mu
             break;
         }
 
-		case OP_TYPE::IN_ADDEXTENSIONDATA:
+		case EOpType::IN_ADDEXTENSIONDATA:
 		{
 			OP::InstanceAddExtensionDataArgs Args = Program.GetOpArgs<OP::InstanceAddExtensionDataArgs>(item.At);
 			switch (item.Stage)
 			{
 				case 0:
 				{
-					// Must pass in an Instance op and ExtensionData op
+					// Must pass in an Instance op and FExtensionData op
 					check(Args.Instance);
 					check(Args.ExtensionData);
 
@@ -802,15 +802,15 @@ namespace mu
 				case 1:
 				{
 					// Assemble result
-					InstancePtrConst InstanceOpResult = LoadInstance(FCacheAddress(Args.Instance, item));
+					TSharedPtr<const FInstance> InstanceOpResult = LoadInstance(FCacheAddress(Args.Instance, item));
 
-					InstancePtr Result = mu::CloneOrTakeOver<Instance>(InstanceOpResult.get());
+					TSharedPtr<FInstance> Result = mu::CloneOrTakeOver<FInstance>(InstanceOpResult);
 
-					if (ExtensionDataPtrConst ExtensionData = LoadExtensionData(FCacheAddress(Args.ExtensionData, item)))
+					if (TSharedPtr<const FExtensionData> ExtensionData = LoadExtensionData(FCacheAddress(Args.ExtensionData, item)))
 					{
 						const OP::ADDRESS NameAddress = Args.ExtensionDataName;
-						check(NameAddress < (uint32)Program.m_constantStrings.Num());
-						const FString& NameString = Program.m_constantStrings[NameAddress];
+						check(NameAddress < (uint32)Program.ConstantStrings.Num());
+						const FString& NameString = Program.ConstantStrings[NameAddress];
 
 						Result->GetPrivate()->AddExtensionData(ExtensionData, FName(NameString) );
 					}
@@ -833,22 +833,22 @@ namespace mu
 
 
     //---------------------------------------------------------------------------------------------
-    void CodeRunner::RunCode_InstanceAddResource(const FScheduledOp& item, const TSharedPtr<const Model>& InModel, const Parameters* InParams )
+    void CodeRunner::RunCode_InstanceAddResource(const FScheduledOp& item, const TSharedPtr<const FModel>& InModel, const FParameters* InParams )
     {
 		MUTABLE_CPUPROFILER_SCOPE(RunCode_InstanceAddResource);
 
-		if (!InModel || !m_pSystem)
+		if (!InModel || !System)
 		{
 			return;
 		}
 
-		const FProgram& Program = m_pModel->GetPrivate()->m_program;
+		const FProgram& Program = Model->GetPrivate()->Program;
 
-		OP_TYPE type = Program.GetOpType(item.At);
+		EOpType type = Program.GetOpType(item.At);
         switch (type)
         {
 
-        case OP_TYPE::IN_ADDMESH:
+        case EOpType::IN_ADDMESH:
         {
 			OP::InstanceAddArgs args = Program.GetOpArgs<OP::InstanceAddArgs>(item.At);
             switch (item.Stage)
@@ -860,26 +860,26 @@ namespace mu
 
             case 1:
             {
-                InstancePtrConst pBase = LoadInstance( FCacheAddress(args.instance,item) );
-                InstancePtr pResult;
-                if (!pBase)
+                TSharedPtr<const FInstance> Base = LoadInstance( FCacheAddress(args.instance,item) );
+                TSharedPtr<FInstance> Result;
+                if (!Base)
                 {
-                    pResult = new Instance();
+                    Result = MakeShared<FInstance>();
                 }
                 else
                 {
-					pResult = mu::CloneOrTakeOver<Instance>(pBase.get());
+					Result = mu::CloneOrTakeOver<FInstance>(Base);
 				}
 
                 if ( args.value )
                 {
-					FResourceID MeshId = m_pSystem->WorkingMemoryManager.GetResourceKey(InModel,InParams,args.relevantParametersListIndex, args.value);
+					FResourceID MeshId = System->WorkingMemoryManager.GetResourceKey(InModel,InParams,args.relevantParametersListIndex, args.value);
 					OP::ADDRESS NameAd = args.name;
-					check(NameAd < (uint32)Program.m_constantStrings.Num());
-					const FString& Name = Program.m_constantStrings[NameAd];
-					pResult->GetPrivate()->SetMesh(0, 0, MeshId, FName(Name));
+					check(NameAd < (uint32)Program.ConstantStrings.Num());
+					const FString& Name = Program.ConstantStrings[NameAd];
+					Result->GetPrivate()->SetMesh(0, 0, MeshId, FName(Name));
                 }
-                StoreInstance( item, pResult );
+                StoreInstance( item, Result );
                 break;
             }
 
@@ -889,7 +889,7 @@ namespace mu
             break;
         }
 
-        case OP_TYPE::IN_ADDIMAGE:
+        case EOpType::IN_ADDIMAGE:
         {
 			OP::InstanceAddArgs args = Program.GetOpArgs<OP::InstanceAddArgs>(item.At);
             switch (item.Stage)
@@ -901,26 +901,26 @@ namespace mu
 
             case 1:
             {
-                InstancePtrConst pBase = LoadInstance( FCacheAddress(args.instance,item) );
-                InstancePtr pResult;
-                if (!pBase)
+                TSharedPtr<const FInstance> Base = LoadInstance( FCacheAddress(args.instance,item) );
+                TSharedPtr<FInstance> Result;
+                if (!Base)
                 {
-                    pResult = new Instance();
+                    Result = MakeShared<FInstance>();
                 }
                 else
                 {
-					pResult = mu::CloneOrTakeOver<Instance>(pBase.get());
+					Result = mu::CloneOrTakeOver<FInstance>(Base);
 				}
 
                 if ( args.value )
                 {
-					FResourceID ImageId = m_pSystem->WorkingMemoryManager.GetResourceKey(InModel, InParams, args.relevantParametersListIndex, args.value);
+					FResourceID ImageId = System->WorkingMemoryManager.GetResourceKey(InModel, InParams, args.relevantParametersListIndex, args.value);
 					OP::ADDRESS NameAd = args.name;
-					check(NameAd < (uint32)Program.m_constantStrings.Num());
-					const FString& Name = Program.m_constantStrings[NameAd];
-					pResult->GetPrivate()->AddImage(0, 0, 0, ImageId, FName(Name) );
+					check(NameAd < (uint32)Program.ConstantStrings.Num());
+					const FString& Name = Program.ConstantStrings[NameAd];
+					Result->GetPrivate()->AddImage(0, 0, 0, ImageId, FName(Name) );
                 }
-                StoreInstance( item, pResult );
+                StoreInstance( item, Result );
                 break;
             }
 
@@ -938,17 +938,17 @@ namespace mu
 
 
     //---------------------------------------------------------------------------------------------
-    bool CodeRunner::RunCode_ConstantResource(const FScheduledOp& item, const Model* pModel )
+    bool CodeRunner::RunCode_ConstantResource(const FScheduledOp& item, const FModel* pModel )
     {
 		MUTABLE_CPUPROFILER_SCOPE(RunCode_Constant);
 
-		const FProgram& Program = m_pModel->GetPrivate()->m_program;
+		const FProgram& Program = Model->GetPrivate()->Program;
 
-		OP_TYPE type = Program.GetOpType(item.At);
+		EOpType type = Program.GetOpType(item.At);
         switch (type)
         {
 
-        case OP_TYPE::ME_CONSTANT:
+        case EOpType::ME_CONSTANT:
         {
 			OP::MeshConstantArgs args = Program.GetOpArgs<OP::MeshConstantArgs>(item.At);
 
@@ -957,25 +957,25 @@ namespace mu
             // Assume the ROM has been loaded previously
             check(Program.ConstantMeshes[cat].Value)
 
-            Ptr<const Mesh> SourceConst;
+            TSharedPtr<const FMesh> SourceConst;
 			Program.GetConstant(cat, SourceConst);
 
 			check(SourceConst);
-			Ptr<Mesh> Source = CreateMesh(SourceConst->GetDataSize());
+			TSharedPtr<FMesh> Source = CreateMesh(SourceConst->GetDataSize());
 			Source->CopyFrom(*SourceConst);
 
             // Set the separate skeleton if necessary
             if (args.skeleton >= 0)
             {
-                check(Program.m_constantSkeletons.Num() > size_t(args.skeleton));
-                Ptr<const Skeleton> pSkeleton = Program.m_constantSkeletons[args.skeleton];
+                check(Program.ConstantSkeletons.Num() > args.skeleton);
+                TSharedPtr<const FSkeleton> pSkeleton = Program.ConstantSkeletons[args.skeleton];
                 Source->SetSkeleton(pSkeleton);
             }
 
 			if (args.physicsBody >= 0)
 			{
-                check(Program.m_constantPhysicsBodies.Num() > size_t(args.physicsBody));
-                Ptr<const PhysicsBody> pPhysicsBody = Program.m_constantPhysicsBodies[args.physicsBody];
+                check(Program.ConstantPhysicsBodies.Num() > args.physicsBody);
+                TSharedPtr<const FPhysicsBody> pPhysicsBody = Program.ConstantPhysicsBodies[args.physicsBody];
                 Source->SetPhysicsBody(pPhysicsBody);
 			}
 
@@ -984,13 +984,13 @@ namespace mu
             break;
         }
 
-        case OP_TYPE::IM_CONSTANT:
+        case EOpType::IM_CONSTANT:
         {
 			OP::ResourceConstantArgs args = Program.GetOpArgs<OP::ResourceConstantArgs>(item.At);
             OP::ADDRESS cat = args.value;
 
 			int32 MipsToSkip = item.ExecutionOptions;
-            Ptr<const Image> Source;
+            TSharedPtr<const FImage> Source;
 			Program.GetConstant(cat, Source, MipsToSkip, [this](int32 x, int32 y, int32 m, EImageFormat f, EInitializationType i)
 				{
 					return CreateImage(x, y, m, f, i);
@@ -1007,12 +1007,12 @@ namespace mu
 			break;
         }
 
-		case OP_TYPE::ED_CONSTANT:
+		case EOpType::ED_CONSTANT:
 		{
 			OP::ResourceConstantArgs Args = Program.GetOpArgs<OP::ResourceConstantArgs>(item.At);
 
 			// Assume the ROM has been loaded previously
-			ExtensionDataPtrConst SourceConst;
+			TSharedPtr<const mu::FExtensionData> SourceConst;
 			Program.GetExtensionDataConstant(Args.value, SourceConst);
 
 			check(SourceConst);
@@ -1022,7 +1022,7 @@ namespace mu
 		}
 
         default:
-            if (type!=OP_TYPE::NONE)
+            if (type!=EOpType::NONE)
             {
                 // Operation not implemented
                 check( false );
@@ -1035,25 +1035,25 @@ namespace mu
     }
 
     //---------------------------------------------------------------------------------------------
-	void CodeRunner::RunCode_Mesh(const FScheduledOp& item, const Model* pModel )
+	void CodeRunner::RunCode_Mesh(const FScheduledOp& item, const FModel* pModel )
     {
 		MUTABLE_CPUPROFILER_SCOPE(RunCode_Mesh);
 
-		const FProgram& Program = m_pModel->GetPrivate()->m_program;
+		const FProgram& Program = Model->GetPrivate()->Program;
 
-		OP_TYPE type = Program.GetOpType(item.At);
+		EOpType type = Program.GetOpType(item.At);
 
         switch (type)
         {
 
-		case OP_TYPE::ME_REFERENCE:
+		case EOpType::ME_REFERENCE:
 		{
 			OP::ResourceReferenceArgs Args = Program.GetOpArgs<OP::ResourceReferenceArgs>(item.At);
 			switch (item.Stage)
 			{
 			case 0:
 			{
-				Ptr<Mesh> Result;
+				TSharedPtr<FMesh> Result;
 				if (Args.ForceLoad)
 				{
 					// This should never be reached because it should have been caught as a Task in IssueOp
@@ -1061,7 +1061,7 @@ namespace mu
 				}
 				else
 				{
-					Result = Mesh::CreateAsReference(Args.ID, false);
+					Result = FMesh::CreateAsReference(Args.ID, false);
 				}
 				StoreMesh(item, Result);
 				break;
@@ -1074,7 +1074,7 @@ namespace mu
 			break;
 		}
 
-        case OP_TYPE::ME_APPLYLAYOUT:
+        case EOpType::ME_APPLYLAYOUT:
         {
 			OP::MeshApplyLayoutArgs args = Program.GetOpArgs<OP::MeshApplyLayoutArgs>(item.At);
 			switch (item.Stage)
@@ -1090,16 +1090,16 @@ namespace mu
             {
            		MUTABLE_CPUPROFILER_SCOPE(ME_APPLYLAYOUT)
             		
-                Ptr<const Mesh> pBase = LoadMesh(FCacheAddress(args.mesh, item));
+                TSharedPtr<const FMesh> Base = LoadMesh(FCacheAddress(args.mesh, item));
 
-                if (pBase)
+                if (Base)
                 {
-					Ptr<Mesh> Result = CloneOrTakeOver(pBase);
+					TSharedPtr<FMesh> Result = CloneOrTakeOver(Base);
 
-                    Ptr<const Layout> pLayout = LoadLayout(FCacheAddress(args.layout, item));
-                    int texCoordsSet = args.channel;
+                    TSharedPtr<const FLayout> Layout = LoadLayout(FCacheAddress(args.layout, item));
+                    int32 texCoordsSet = args.channel;
 
-                    MeshApplyLayout(Result.get(), pLayout.get(), texCoordsSet);
+                    MeshApplyLayout(Result.Get(), Layout.Get(), texCoordsSet);
 					
 					StoreMesh(item, Result);
                 }
@@ -1119,7 +1119,7 @@ namespace mu
             break;
         }
 
-        case OP_TYPE::ME_DIFFERENCE:
+        case EOpType::ME_DIFFERENCE:
         {
 			const uint8* data = Program.GetOpArgsPointer(item.At);
 
@@ -1151,8 +1151,8 @@ namespace mu
             {
        	        MUTABLE_CPUPROFILER_SCOPE(ME_DIFFERENCE)
 
-                Ptr<const Mesh> pBase = LoadMesh(FCacheAddress(BaseAt,item));
-                Ptr<const Mesh> pTarget = LoadMesh(FCacheAddress(TargetAt,item));
+                TSharedPtr<const FMesh> Base = LoadMesh(FCacheAddress(BaseAt,item));
+                TSharedPtr<const FMesh> pTarget = LoadMesh(FCacheAddress(TargetAt,item));
 
 				TArray<EMeshBufferSemantic, TInlineAllocator<8>> Semantics;
 				TArray<int32, TInlineAllocator<8>> SemanticIndices;
@@ -1179,12 +1179,12 @@ namespace mu
 					SemanticIndices.Add(SemanticIndex);
                 }
 
-				Ptr<Mesh> Result = CreateMesh();
+				TSharedPtr<FMesh> Result = CreateMesh();
 				bool bOutSuccess = false;
-                MeshDifference(Result.get(), pBase.get(), pTarget.get(),
+                MeshDifference(Result.Get(), Base.Get(), pTarget.Get(),
                                NumChannels, Semantics.GetData(), SemanticIndices.GetData(),
                                bIgnoreTextureCoords != 0, bOutSuccess);
-				Release(pBase);
+				Release(Base);
 				Release(pTarget);
 
                 StoreMesh(item, Result);
@@ -1198,7 +1198,7 @@ namespace mu
             break;
         }
 
-        case OP_TYPE::ME_MORPH:
+        case EOpType::ME_MORPH:
         {
 			const uint8* Data = Program.GetOpArgsPointer(item.At);
 
@@ -1237,9 +1237,9 @@ namespace mu
                 // Factor goes from -1 to 1 across all targets. [0 - 1] represents positive morphs, while [-1, 0) represent negative morphs.
 				Factor = FMath::Clamp(Factor, -1.0f, 1.0f); // Is the factor not in range [-1, 1], it will index a non existing morph.
 
-                FScheduledOpData HeapData;
-				HeapData.Interpolate.Bifactor = Factor;
-				uint32 DataAddress = uint32(m_heapData.Add(HeapData));
+                FScheduledOpData OpHeapData;
+				OpHeapData.Interpolate.Bifactor = Factor;
+				uint32 DataAddress = uint32(HeapData.Add(OpHeapData));
 
                 // No morph
 				if (FMath::IsNearlyZero(Factor))
@@ -1263,11 +1263,11 @@ namespace mu
             {
        		    MUTABLE_CPUPROFILER_SCOPE(ME_MORPH_2)
 
-                Ptr<const Mesh> BaseMesh = LoadMesh(FCacheAddress(BaseAt, item));
+                TSharedPtr<const FMesh> BaseMesh = LoadMesh(FCacheAddress(BaseAt, item));
 
                 // Factor from 0 to 1 between the two targets
-                const FScheduledOpData& HeapData = m_heapData[item.CustomState];
-                float Factor = HeapData.Interpolate.Bifactor;
+                const FScheduledOpData& OpHeapData = HeapData[item.CustomState];
+                float Factor = OpHeapData.Interpolate.Bifactor;
 
                 if (BaseMesh)
                 {
@@ -1279,12 +1279,12 @@ namespace mu
 					// The Morph, partial or full
                     else 
                     {
-                        Ptr<const Mesh> MorphMesh = LoadMesh(FCacheAddress(TargetAt, item));
+                        TSharedPtr<const FMesh> MorphMesh = LoadMesh(FCacheAddress(TargetAt, item));
 						
 						if (MorphMesh)
 						{
-							Ptr<Mesh> Result = CloneOrTakeOver(BaseMesh);
-							MeshMorph(Result.get(), MorphMesh.get(), Factor);
+							TSharedPtr<FMesh> Result = CloneOrTakeOver(BaseMesh);
+							MeshMorph(Result.Get(), MorphMesh.Get(), Factor);
 
 							Release(MorphMesh);
 							StoreMesh(item, Result);
@@ -1310,7 +1310,7 @@ namespace mu
             break;
         }
 
-        case OP_TYPE::ME_MERGE:
+        case EOpType::ME_MERGE:
         {
 			OP::MeshMergeArgs args = Program.GetOpArgs<OP::MeshMergeArgs>(item.At);
             switch (item.Stage)
@@ -1326,8 +1326,8 @@ namespace mu
             {
 				MUTABLE_CPUPROFILER_SCOPE(ME_MERGE_1)
 
-                Ptr<const Mesh> pA = LoadMesh(FCacheAddress(args.base, item));
-                Ptr<const Mesh> pB = LoadMesh(FCacheAddress(args.added, item));
+                TSharedPtr<const FMesh> pA = LoadMesh(FCacheAddress(args.base, item));
+                TSharedPtr<const FMesh> pB = LoadMesh(FCacheAddress(args.added, item));
 
                 if (pA && pB && pA->GetVertexCount() && pB->GetVertexCount())
                 {
@@ -1337,9 +1337,9 @@ namespace mu
 					Scratch.FirstReformat = CreateMesh();
 					Scratch.SecondReformat = CreateMesh();
 
-					Ptr<Mesh> Result = CreateMesh(pA->GetDataSize() + pB->GetDataSize());
+					TSharedPtr<FMesh> Result = CreateMesh(pA->GetDataSize() + pB->GetDataSize());
 
-					MeshMerge(Result.get(), pA.get(), pB.get(), !args.newSurfaceID, Scratch);
+					MeshMerge(Result.Get(), pA.Get(), pB.Get(), !args.newSurfaceID, Scratch);
 
 					Release(Scratch.FirstReformat);
 					Release(Scratch.SecondReformat);
@@ -1361,7 +1361,7 @@ namespace mu
                 }
                 else if (pB && (pB->GetVertexCount() || pB->IsReference()))
                 {
-					Ptr<Mesh> Result = CloneOrTakeOver(pB);
+					TSharedPtr<FMesh> Result = CloneOrTakeOver(pB);
 
                     check(Result->IsReference() || (Result->GetSurfaceCount() == 1) );
 
@@ -1390,7 +1390,7 @@ namespace mu
             break;
         }
 
-        case OP_TYPE::ME_INTERPOLATE:
+        case EOpType::ME_INTERPOLATE:
         {
 			OP::MeshInterpolateArgs Args = Program.GetOpArgs<OP::MeshInterpolateArgs>(item.At);
             switch (item.Stage)
@@ -1433,7 +1433,7 @@ namespace mu
                 Data.Interpolate.Bifactor = Bifactor;
 				Data.Interpolate.Min = FMath::Clamp(Min, 0, Count - 1);
 				Data.Interpolate.Max = FMath::Clamp(Max, 0, Count - 1);
-				uint32 DataAddress = uint32(m_heapData.Num());
+				uint32 DataAddress = uint32(HeapData.Num());
 
                 // Just the first of the targets
                 if (Bifactor < UE_SMALL_NUMBER)
@@ -1441,13 +1441,13 @@ namespace mu
                     if (Min == 0)
                     {
                         // Just the base
-						Ptr<const Mesh> BaseMesh = LoadMesh(FCacheAddress(Args.base, item));
+						TSharedPtr<const FMesh> BaseMesh = LoadMesh(FCacheAddress(Args.base, item));
 						StoreMesh(item, BaseMesh);
 					}
                     else
                     {
                         // Base with one full morph
-                        m_heapData.Add(Data);
+                        HeapData.Add(Data);
 						AddOp(FScheduledOp(item.At, item, 2, DataAddress),
 							FScheduledOp(Args.base, item),
 							FScheduledOp(Args.targets[Min-1], item));
@@ -1456,7 +1456,7 @@ namespace mu
                 // Just the second of the targets
                 else if (Bifactor > 1.0f-UE_SMALL_NUMBER)
                 {
-                    m_heapData.Add(Data);
+                    HeapData.Add(Data);
 					AddOp(FScheduledOp(item.At, item, 2, DataAddress),
 						FScheduledOp(Args.base, item),
 						FScheduledOp(Args.targets[Max-1], item));
@@ -1464,7 +1464,7 @@ namespace mu
                 // Mix the first target on the base
                 else if (Min == 0)
                 {
-                    m_heapData.Add(Data);
+                    HeapData.Add(Data);
 					AddOp(FScheduledOp(item.At, item, 2, DataAddress),
 						FScheduledOp(Args.base, item),
 						FScheduledOp(Args.targets[0], item));
@@ -1472,7 +1472,7 @@ namespace mu
                 // Mix two targets on the base
                 else
                 {
-                    m_heapData.Add(Data);
+                    HeapData.Add(Data);
 					AddOp(FScheduledOp(item.At, item, 2, DataAddress),
 						FScheduledOp(Args.base, item),
 						FScheduledOp(Args.targets[Min-1], item),
@@ -1492,14 +1492,14 @@ namespace mu
                     ++Count;
                 }
 
-                const FScheduledOpData& Data = m_heapData[ (size_t)item.CustomState ];
+                const FScheduledOpData& Data = HeapData[ item.CustomState ];
 
                 // Factor from 0 to 1 between the two targets
                 float Bifactor = Data.Interpolate.Bifactor;
                 int32 Min = Data.Interpolate.Min;
                 int32 Max = Data.Interpolate.Max;
 
-                Ptr<const Mesh> BaseMesh = LoadMesh(FCacheAddress(Args.base, item));
+                TSharedPtr<const FMesh> BaseMesh = LoadMesh(FCacheAddress(Args.base, item));
 
                 if (BaseMesh)
                 {
@@ -1514,13 +1514,13 @@ namespace mu
                         else
                         {
                             // Base with one full morph
-                            Ptr<const Mesh> MorphMesh = LoadMesh(FCacheAddress(Args.targets[Min-1], item));
+                            TSharedPtr<const FMesh> MorphMesh = LoadMesh(FCacheAddress(Args.targets[Min-1], item));
 						
 							if (MorphMesh)
 							{
-								Ptr<Mesh> Result = CloneOrTakeOver(BaseMesh);
+								TSharedPtr<FMesh> Result = CloneOrTakeOver(BaseMesh);
 
-								MeshMorph(Result.get(), MorphMesh.get());
+								MeshMorph(Result.Get(), MorphMesh.Get());
 							
 								StoreMesh(item, Result);
 							}
@@ -1536,13 +1536,13 @@ namespace mu
                     else if (Bifactor > 1.0f-UE_SMALL_NUMBER)
                     {
                         check(Max > 0);
-                        Ptr<const Mesh> MorphMesh = LoadMesh(FCacheAddress(Args.targets[Max-1], item));
+                        TSharedPtr<const FMesh> MorphMesh = LoadMesh(FCacheAddress(Args.targets[Max-1], item));
 
                         if (MorphMesh)
                         {
-							Ptr<Mesh> Result = CloneOrTakeOver(BaseMesh);
+							TSharedPtr<FMesh> Result = CloneOrTakeOver(BaseMesh);
 							
-                            MeshMorph(Result.get(), MorphMesh.get());
+                            MeshMorph(Result.Get(), MorphMesh.Get());
 
 							StoreMesh(item, Result);
                         }
@@ -1556,12 +1556,12 @@ namespace mu
                     // Mix the first target on the base
                     else if (Min == 0)
                     {
-                        Ptr<const Mesh> MorphMesh = LoadMesh(FCacheAddress(Args.targets[0], item));
+                        TSharedPtr<const FMesh> MorphMesh = LoadMesh(FCacheAddress(Args.targets[0], item));
                         if (MorphMesh)
                         {
-							Ptr<Mesh> Result = CloneOrTakeOver(BaseMesh);
+							TSharedPtr<FMesh> Result = CloneOrTakeOver(BaseMesh);
 
-                            MeshMorph(Result.get(), MorphMesh.get(), Bifactor);
+                            MeshMorph(Result.Get(), MorphMesh.Get(), Bifactor);
 
 							StoreMesh(item, Result);
                         }
@@ -1575,14 +1575,14 @@ namespace mu
                     // Mix two targets on the base
                     else
                     {
-                        Ptr<const Mesh> MinMesh = LoadMesh(FCacheAddress(Args.targets[Min-1], item));
-                        Ptr<const Mesh> MaxMesh = LoadMesh(FCacheAddress(Args.targets[Max-1], item));
+                        TSharedPtr<const FMesh> MinMesh = LoadMesh(FCacheAddress(Args.targets[Min-1], item));
+                        TSharedPtr<const FMesh> MaxMesh = LoadMesh(FCacheAddress(Args.targets[Max-1], item));
 
                         if (MinMesh && MaxMesh)
                         {
-							Ptr<Mesh> Result = CloneOrTakeOver(BaseMesh);
+							TSharedPtr<FMesh> Result = CloneOrTakeOver(BaseMesh);
 
-                            MeshMorph2(Result.get(), MinMesh.get(), MaxMesh.get(), Bifactor);
+                            MeshMorph2(Result.Get(), MinMesh.Get(), MaxMesh.Get(), Bifactor);
 
 							StoreMesh(item, Result);
                         }
@@ -1610,7 +1610,7 @@ namespace mu
             break;
         }
 
-        case OP_TYPE::ME_MASKCLIPMESH:
+        case EOpType::ME_MASKCLIPMESH:
         {
 			OP::MeshMaskClipMeshArgs args = Program.GetOpArgs<OP::MeshMaskClipMeshArgs>(item.At);
             switch (item.Stage)
@@ -1626,16 +1626,16 @@ namespace mu
             {
             	MUTABLE_CPUPROFILER_SCOPE(ME_MASKCLIPMESH_1)
             		
-                Ptr<const Mesh> Source = LoadMesh(FCacheAddress(args.source, item));
-                Ptr<const Mesh> pClip = LoadMesh(FCacheAddress(args.clip, item));
+                TSharedPtr<const FMesh> Source = LoadMesh(FCacheAddress(args.source, item));
+                TSharedPtr<const FMesh> pClip = LoadMesh(FCacheAddress(args.clip, item));
 
                 // Only if both are valid.
-                if (Source.get() && pClip.get())
+                if (Source.Get() && pClip.Get())
                 {
-					Ptr<Mesh> Result = CreateMesh();
+					TSharedPtr<FMesh> Result = CreateMesh();
 
 					bool bOutSuccess = false;
-                    MeshMaskClipMesh(Result.get(), Source.get(), pClip.get(), bOutSuccess);
+                    MeshMaskClipMesh(Result.Get(), Source.Get(), pClip.Get(), bOutSuccess);
 					
 					Release(Source);
 					Release(pClip);
@@ -1666,7 +1666,7 @@ namespace mu
             break;
         }
 
-		case OP_TYPE::ME_MASKCLIPUVMASK:
+		case EOpType::ME_MASKCLIPUVMASK:
 		{
 			OP::MeshMaskClipUVMaskArgs args = Program.GetOpArgs<OP::MeshMaskClipUVMaskArgs>(item.At);
 			switch (item.Stage)
@@ -1684,18 +1684,18 @@ namespace mu
 			{
 				MUTABLE_CPUPROFILER_SCOPE(ME_MASKCLIPUVMASK_1);
 
-				Ptr<const Mesh> Source = LoadMesh(FCacheAddress(args.Source, item));
-				Ptr<const Mesh> UVSource = LoadMesh(FCacheAddress(args.UVSource, item));
-				Ptr<const Image> MaskImage = LoadImage(FCacheAddress(args.MaskImage, item));
-				Ptr<const Layout> MaskLayout = LoadLayout(FCacheAddress(args.MaskLayout, item));
+				TSharedPtr<const FMesh> Source = LoadMesh(FCacheAddress(args.Source, item));
+				TSharedPtr<const FMesh> UVSource = LoadMesh(FCacheAddress(args.UVSource, item));
+				TSharedPtr<const FImage> MaskImage = LoadImage(FCacheAddress(args.MaskImage, item));
+				TSharedPtr<const FLayout> MaskLayout = LoadLayout(FCacheAddress(args.MaskLayout, item));
 
 				// Only if both are valid.
-				if (Source.get() && MaskImage.get())
+				if (Source.Get() && MaskImage.Get())
 				{
-					Ptr<Mesh> Result = CreateMesh();
+					TSharedPtr<FMesh> Result = CreateMesh();
 
 					bool bOutSuccess = false;
-					MakeMeshMaskFromUVMask(Result.get(), Source.get(), UVSource.get(), MaskImage.get(), args.LayoutIndex, bOutSuccess);
+					MakeMeshMaskFromUVMask(Result.Get(), Source.Get(), UVSource.Get(), MaskImage.Get(), args.LayoutIndex, bOutSuccess);
 
 					Release(Source);
 					Release(UVSource);
@@ -1710,12 +1710,12 @@ namespace mu
 						StoreMesh(item, Result);
 					}
 				}
-				else if (Source.get() && MaskLayout.get())
+				else if (Source.Get() && MaskLayout.Get())
 				{					
-					Ptr<Mesh> Result = CreateMesh();
+					TSharedPtr<FMesh> Result = CreateMesh();
 
 					bool bOutSuccess = false;
-					MakeMeshMaskFromLayout(Result.get(), Source.get(), UVSource.get(), MaskLayout.get(), args.LayoutIndex, bOutSuccess);
+					MakeMeshMaskFromLayout(Result.Get(), Source.Get(), UVSource.Get(), MaskLayout.Get(), args.LayoutIndex, bOutSuccess);
 
 					Release(Source);
 					Release(UVSource);
@@ -1748,7 +1748,7 @@ namespace mu
 			break;
 		}
 
-        case OP_TYPE::ME_MASKDIFF:
+        case EOpType::ME_MASKDIFF:
         {
 			OP::MeshMaskDiffArgs args = Program.GetOpArgs<OP::MeshMaskDiffArgs>(item.At);
             switch (item.Stage)
@@ -1764,16 +1764,16 @@ namespace mu
             {
            		MUTABLE_CPUPROFILER_SCOPE(ME_MASKDIFF_1)
             		
-                Ptr<const Mesh> Source = LoadMesh(FCacheAddress(args.source, item));
-                Ptr<const Mesh> pClip = LoadMesh(FCacheAddress(args.fragment, item));
+                TSharedPtr<const FMesh> Source = LoadMesh(FCacheAddress(args.source, item));
+                TSharedPtr<const FMesh> pClip = LoadMesh(FCacheAddress(args.fragment, item));
 
                 // Only if both are valid.
-                if (Source.get() && pClip.get())
+                if (Source.Get() && pClip.Get())
                 {
-					Ptr<Mesh> Result = CreateMesh();
+					TSharedPtr<FMesh> Result = CreateMesh();
 
 					bool bOutSuccess = false;
-                    MeshMaskDiff(Result.get(), Source.get(), pClip.get(), bOutSuccess);
+                    MeshMaskDiff(Result.Get(), Source.Get(), pClip.Get(), bOutSuccess);
 
 					Release(Source);
 					Release(pClip);
@@ -1805,7 +1805,7 @@ namespace mu
             break;
         }
 
-        case OP_TYPE::ME_FORMAT:
+        case EOpType::ME_FORMAT:
         {
 			OP::MeshFormatArgs args = Program.GetOpArgs<OP::MeshFormatArgs>(item.At);
             switch (item.Stage)
@@ -1828,8 +1828,8 @@ namespace mu
             {
             	MUTABLE_CPUPROFILER_SCOPE(ME_FORMAT_1)
             		
-                Ptr<const Mesh> Source = LoadMesh(FCacheAddress(args.source,item));
-                Ptr<const Mesh> Format = LoadMesh(FCacheAddress(args.format,item));
+                TSharedPtr<const FMesh> Source = LoadMesh(FCacheAddress(args.source,item));
+                TSharedPtr<const FMesh> Format = LoadMesh(FCacheAddress(args.format,item));
 
 				if (Source && Source->IsReference())
 				{
@@ -1845,7 +1845,7 @@ namespace mu
 					}
 					else if (!Format)
 					{
-						Ptr<Mesh> Result = CloneOrTakeOver(Source);
+						TSharedPtr<FMesh> Result = CloneOrTakeOver(Source);
 
 						if (Flags & OP::MeshFormatArgs::ResetBufferIndices)
 						{
@@ -1856,10 +1856,10 @@ namespace mu
 					}
 					else
 					{
-						Ptr<Mesh> Result = CreateMesh();
+						TSharedPtr<FMesh> Result = CreateMesh();
 
 						bool bOutSuccess = false;
-						MeshFormat(Result.get(), Source.get(), Format.get(),
+						MeshFormat(Result.Get(), Source.Get(), Format.Get(),
 							true,
 							(Flags & OP::MeshFormatArgs::Vertex) != 0,
 							(Flags & OP::MeshFormatArgs::Index) != 0,
@@ -1876,7 +1876,7 @@ namespace mu
 						if (Flags & OP::MeshFormatArgs::OptimizeBuffers)
 						{
 							MUTABLE_CPUPROFILER_SCOPE(MeshOptimizeBuffers)
-							MeshOptimizeBuffers(Result.get());
+							MeshOptimizeBuffers(Result.Get());
 						}
 
 						Release(Source);
@@ -1899,7 +1899,7 @@ namespace mu
             break;
         }
 
-        case OP_TYPE::ME_EXTRACTLAYOUTBLOCK:
+        case EOpType::ME_EXTRACTLAYOUTBLOCK:
         {
             const uint8* data = Program.GetOpArgsPointer(item.At);
 
@@ -1934,7 +1934,7 @@ namespace mu
             {
 				MUTABLE_CPUPROFILER_SCOPE(ME_EXTRACTLAYOUTBLOCK_1)
 
-                Ptr<const Mesh> Source = LoadMesh(FCacheAddress(source, item));
+                TSharedPtr<const FMesh> Source = LoadMesh(FCacheAddress(source, item));
 
                 // Access with memcpy necessary for unaligned memory access issues.
                 uint64 blocks[512];
@@ -1943,16 +1943,16 @@ namespace mu
 
 				if (Source)
 				{
-					Ptr<Mesh> Result = CreateMesh();
+					TSharedPtr<FMesh> Result = CreateMesh();
 					bool bOutSuccess;
 
 					if (blockCount > 0)
 					{
-						MeshExtractLayoutBlock(Result.get(), Source.get(), layout, blockCount, blocks, bOutSuccess);
+						MeshExtractLayoutBlock(Result.Get(), Source.Get(), layout, blockCount, blocks, bOutSuccess);
 					}
 					else
 					{
-						MeshExtractLayoutBlock(Result.get(), Source.get(), layout, bOutSuccess);
+						MeshExtractLayoutBlock(Result.Get(), Source.Get(), layout, bOutSuccess);
 					}
 
 					if (!bOutSuccess)
@@ -1980,7 +1980,7 @@ namespace mu
             break;
         }
 
-        case OP_TYPE::ME_TRANSFORM:
+        case EOpType::ME_TRANSFORM:
         {
 			OP::MeshTransformArgs args = Program.GetOpArgs<OP::MeshTransformArgs>(item.At);
             switch (item.Stage)
@@ -2002,14 +2002,14 @@ namespace mu
             {
            		MUTABLE_CPUPROFILER_SCOPE(ME_TRANSFORM_1)
             		
-                Ptr<const Mesh> Source = LoadMesh(FCacheAddress(args.source,item));
+                TSharedPtr<const FMesh> Source = LoadMesh(FCacheAddress(args.source,item));
 
-                const FMatrix44f& mat = Program.m_constantMatrices[args.matrix];
+                const FMatrix44f& mat = Program.ConstantMatrices[args.matrix];
 
-				Ptr<Mesh> Result = CreateMesh(Source ? Source->GetDataSize() : 0);
+				TSharedPtr<FMesh> Result = CreateMesh(Source ? Source->GetDataSize() : 0);
 
 				bool bOutSuccess = false;
-                MeshTransform(Result.get(), Source.get(), mat, bOutSuccess);
+                MeshTransform(Result.Get(), Source.Get(), mat, bOutSuccess);
 
 				if (!bOutSuccess)
 				{
@@ -2031,7 +2031,7 @@ namespace mu
             break;
         }
 
-        case OP_TYPE::ME_CLIPMORPHPLANE:
+        case EOpType::ME_CLIPMORPHPLANE:
         {
 			OP::MeshClipMorphPlaneArgs args = Program.GetOpArgs<OP::MeshClipMorphPlaneArgs>(item.At);
             switch (item.Stage)
@@ -2053,12 +2053,12 @@ namespace mu
             {
            		MUTABLE_CPUPROFILER_SCOPE(ME_CLIPMORPHPLANE_1)
             		
-                Ptr<const Mesh> Source = LoadMesh(FCacheAddress(args.source, item));
+                TSharedPtr<const FMesh> Source = LoadMesh(FCacheAddress(args.source, item));
 
-                check(args.morphShape < (uint32)pModel->GetPrivate()->m_program.m_constantShapes.Num());
+                check(args.morphShape < (uint32)pModel->GetPrivate()->Program.ConstantShapes.Num());
 
                 // Should be an ellipse
-                const FShape& morphShape = Program.m_constantShapes[args.morphShape];
+                const FShape& morphShape = Program.ConstantShapes[args.morphShape];
 
                 const FVector3f& origin = morphShape.position;
                 const FVector3f& normal = morphShape.up;
@@ -2067,15 +2067,15 @@ namespace mu
 
                 if (args.VertexSelectionType == EClipVertexSelectionType::Shape)
                 {
-                    check(args.vertexSelectionShapeOrBone < (uint32)pModel->GetPrivate()->m_program.m_constantShapes.Num());
+                    check(args.vertexSelectionShapeOrBone < (uint32)pModel->GetPrivate()->Program.ConstantShapes.Num());
 
                     // Should be None or an axis aligned box
-                    const FShape& selectionShape = Program.m_constantShapes[args.vertexSelectionShapeOrBone];
+                    const FShape& selectionShape = Program.ConstantShapes[args.vertexSelectionShapeOrBone];
 
-					Ptr<Mesh> Result = CreateMesh(Source ? Source->GetDataSize() : 0);
+					TSharedPtr<FMesh> Result = CreateMesh(Source ? Source->GetDataSize() : 0);
 
 					bool bOutSuccess = false;
-					MeshClipMorphPlane(Result.get(), Source.get(), origin, normal, args.dist, args.factor, morphShape.size[0], morphShape.size[1], morphShape.size[2], selectionShape, bRemoveFaceIfAllVerticesCulled, bOutSuccess, nullptr, -1);
+					MeshClipMorphPlane(Result.Get(), Source.Get(), origin, normal, args.dist, args.factor, morphShape.size[0], morphShape.size[1], morphShape.size[2], selectionShape, bRemoveFaceIfAllVerticesCulled, bOutSuccess, nullptr, -1);
 					
 					if (!bOutSuccess)
 					{
@@ -2094,13 +2094,13 @@ namespace mu
 					FShape selectionShape;
 					selectionShape.type = (uint8)FShape::Type::None;
 
-					Ptr<Mesh> Result = CreateMesh(Source->GetDataSize());
+					TSharedPtr<FMesh> Result = CreateMesh(Source->GetDataSize());
 
 					check(args.vertexSelectionShapeOrBone <= MAX_uint32);
 					const FBoneName Bone(args.vertexSelectionShapeOrBone);
 
 					bool bOutSuccess = false;
-					MeshClipMorphPlane(Result.get(), Source.get(), origin, normal, args.dist, args.factor, morphShape.size[0], morphShape.size[1], morphShape.size[2], selectionShape, bRemoveFaceIfAllVerticesCulled, bOutSuccess, &Bone, args.maxBoneRadius);
+					MeshClipMorphPlane(Result.Get(), Source.Get(), origin, normal, args.dist, args.factor, morphShape.size[0], morphShape.size[1], morphShape.size[2], selectionShape, bRemoveFaceIfAllVerticesCulled, bOutSuccess, &Bone, args.maxBoneRadius);
 
 					if (!bOutSuccess)
 					{
@@ -2119,10 +2119,10 @@ namespace mu
                     FShape selectionShape;
                     selectionShape.type = (uint8)FShape::Type::None;
 
-					Ptr<Mesh> Result = CreateMesh(Source ? Source->GetDataSize() : 0);
+					TSharedPtr<FMesh> Result = CreateMesh(Source ? Source->GetDataSize() : 0);
 
 					bool bOutSuccess = false;
-					MeshClipMorphPlane(Result.get(), Source.get(), origin, normal, args.dist, args.factor, morphShape.size[0], morphShape.size[1], morphShape.size[2], selectionShape, bRemoveFaceIfAllVerticesCulled, bOutSuccess, nullptr, -1.0f);
+					MeshClipMorphPlane(Result.Get(), Source.Get(), origin, normal, args.dist, args.factor, morphShape.size[0], morphShape.size[1], morphShape.size[2], selectionShape, bRemoveFaceIfAllVerticesCulled, bOutSuccess, nullptr, -1.0f);
 
 					if (!bOutSuccess)
 					{
@@ -2147,7 +2147,7 @@ namespace mu
         }
 
 
-        case OP_TYPE::ME_CLIPWITHMESH:
+        case EOpType::ME_CLIPWITHMESH:
         {
 			OP::MeshClipWithMeshArgs args = Program.GetOpArgs<OP::MeshClipWithMeshArgs>(item.At);
             switch (item.Stage)
@@ -2171,16 +2171,16 @@ namespace mu
             {
 				MUTABLE_CPUPROFILER_SCOPE(ME_CLIPWITHMESH_1)
 
-                Ptr<const Mesh> Source = LoadMesh(FCacheAddress(args.source, item));
-                Ptr<const Mesh> pClip = LoadMesh(FCacheAddress(args.clipMesh, item));
+                TSharedPtr<const FMesh> Source = LoadMesh(FCacheAddress(args.source, item));
+                TSharedPtr<const FMesh> pClip = LoadMesh(FCacheAddress(args.clipMesh, item));
 
                 // Only if both are valid.
                 if (Source && pClip)
                 {
-					Ptr<Mesh> Result = CreateMesh(Source->GetDataSize());
+					TSharedPtr<FMesh> Result = CreateMesh(Source->GetDataSize());
 
 					bool bOutSuccess = false;
-                    MeshClipWithMesh(Result.get(), Source.get(), pClip.get(), bOutSuccess);
+                    MeshClipWithMesh(Result.Get(), Source.Get(), pClip.Get(), bOutSuccess);
 
 					Release(pClip);
 					if (!bOutSuccess)
@@ -2209,7 +2209,7 @@ namespace mu
 
             break;
         }
-		case OP_TYPE::ME_CLIPDEFORM:
+		case EOpType::ME_CLIPDEFORM:
 		{
 			OP::MeshClipDeformArgs args = Program.GetOpArgs<OP::MeshClipDeformArgs>(item.At);
 			switch (item.Stage)
@@ -2232,17 +2232,17 @@ namespace mu
 			{
 				MUTABLE_CPUPROFILER_SCOPE(ME_CLIPDEFORM_1)
 
-				Ptr<const Mesh> BaseMesh = LoadMesh(FCacheAddress(args.mesh, item));
-				Ptr<const Mesh> ClipShape = LoadMesh(FCacheAddress(args.clipShape, item));
+				TSharedPtr<const FMesh> BaseMesh = LoadMesh(FCacheAddress(args.mesh, item));
+				TSharedPtr<const FMesh> ClipShape = LoadMesh(FCacheAddress(args.clipShape, item));
 
 				if (BaseMesh && ClipShape)
 				{
-					Ptr<Mesh> Result = CreateMesh(BaseMesh->GetDataSize());
+					TSharedPtr<FMesh> Result = CreateMesh(BaseMesh->GetDataSize());
 
 					bool bRemoveIfAllVerticesCulled = args.FaceCullStrategy == EFaceCullStrategy::AllVerticesCulled;
 
 					bool bOutSuccess = false;
-					MeshClipDeform(Result.get(), BaseMesh.get(), ClipShape.get(), args.clipWeightThreshold, bRemoveIfAllVerticesCulled, bOutSuccess);
+					MeshClipDeform(Result.Get(), BaseMesh.Get(), ClipShape.Get(), args.clipWeightThreshold, bRemoveIfAllVerticesCulled, bOutSuccess);
 
 					Release(ClipShape);
 
@@ -2273,7 +2273,7 @@ namespace mu
 			break;
 		}
 
-        case OP_TYPE::ME_APPLYPOSE:
+        case EOpType::ME_APPLYPOSE:
         {
 			OP::MeshApplyPoseArgs args = Program.GetOpArgs<OP::MeshApplyPoseArgs>(item.At);
             switch (item.Stage)
@@ -2296,33 +2296,33 @@ namespace mu
             {
           		MUTABLE_CPUPROFILER_SCOPE(ME_APPLYPOSE_1)
 
-                Ptr<const Mesh> pBase = LoadMesh(FCacheAddress(args.base, item));
-                Ptr<const Mesh> pPose = LoadMesh(FCacheAddress(args.pose, item));
+                TSharedPtr<const FMesh> Base = LoadMesh(FCacheAddress(args.base, item));
+                TSharedPtr<const FMesh> pPose = LoadMesh(FCacheAddress(args.pose, item));
 
                 // Only if both are valid.
-                if (pBase && pPose)
+                if (Base && pPose)
                 {
-					Ptr<Mesh> Result = CreateMesh(pBase->GetSkeleton() ? pBase->GetDataSize() : 0);
+					TSharedPtr<FMesh> Result = CreateMesh(Base->GetSkeleton() ? Base->GetDataSize() : 0);
 
 					bool bOutSuccess = false;
-					MeshApplyPose(Result.get(), pBase.get(), pPose.get(), bOutSuccess);
+					MeshApplyPose(Result.Get(), Base.Get(), pPose.Get(), bOutSuccess);
 
 					Release(pPose);
 					if (!bOutSuccess)
 					{
 						Release(Result);
-						StoreMesh(item, pBase);
+						StoreMesh(item, Base);
 					}
 					else
 					{
-						Release(pBase);
+						Release(Base);
 						StoreMesh(item, Result);
 					}
                 }
                 else
                 {
 					Release(pPose);
-					StoreMesh(item, pBase);
+					StoreMesh(item, Base);
                 }
 
                 break;
@@ -2336,7 +2336,7 @@ namespace mu
         }
 
 
-		case OP_TYPE::ME_GEOMETRYOPERATION:
+		case EOpType::ME_GEOMETRYOPERATION:
 		{
 			OP::MeshGeometryOperationArgs args = Program.GetOpArgs<OP::MeshGeometryOperationArgs>(item.At);
 			switch (item.Stage)
@@ -2361,15 +2361,15 @@ namespace mu
 			{
 				MUTABLE_CPUPROFILER_SCOPE(ME_GEOMETRYOPERATION_1)
 				
-				Ptr<const Mesh> MeshA = LoadMesh(FCacheAddress(args.meshA, item));
-				Ptr<const Mesh> MeshB = LoadMesh(FCacheAddress(args.meshB, item));
+				TSharedPtr<const FMesh> MeshA = LoadMesh(FCacheAddress(args.meshA, item));
+				TSharedPtr<const FMesh> MeshB = LoadMesh(FCacheAddress(args.meshB, item));
 				float ScalarA = LoadScalar(FCacheAddress(args.scalarA, item));
 				float ScalarB = LoadScalar(FCacheAddress(args.scalarB, item));
 
-				Ptr<Mesh> Result = CreateMesh(MeshA ? MeshA->GetDataSize() : 0);
+				TSharedPtr<FMesh> Result = CreateMesh(MeshA ? MeshA->GetDataSize() : 0);
 
 				bool bOutSuccess = false;
-				MeshGeometryOperation(Result.get(), MeshA.get(), MeshB.get(), ScalarA, ScalarB, bOutSuccess);
+				MeshGeometryOperation(Result.Get(), MeshA.Get(), MeshB.Get(), ScalarA, ScalarB, bOutSuccess);
 
 				Release(MeshA);
 				Release(MeshB);
@@ -2395,7 +2395,7 @@ namespace mu
 		}
 
 
-		case OP_TYPE::ME_BINDSHAPE:
+		case EOpType::ME_BINDSHAPE:
 		{
 			OP::MeshBindShapeArgs Args = Program.GetOpArgs<OP::MeshBindShapeArgs>(item.At);
 			const uint8* Data = Program.GetOpArgsPointer(item.At);
@@ -2419,8 +2419,8 @@ namespace mu
 			case 1:
 			{
 				MUTABLE_CPUPROFILER_SCOPE(ME_BINDSHAPE_1)
-				Ptr<const Mesh> BaseMesh = LoadMesh(FCacheAddress(Args.mesh, item));
-				Ptr<const Mesh> Shape = LoadMesh(FCacheAddress(Args.shape, item));
+				TSharedPtr<const FMesh> BaseMesh = LoadMesh(FCacheAddress(Args.mesh, item));
+				TSharedPtr<const FMesh> Shape = LoadMesh(FCacheAddress(Args.shape, item));
 				
 				EShapeBindingMethod BindingMethod = static_cast<EShapeBindingMethod>(Args.bindingMethod); 
 
@@ -2454,10 +2454,10 @@ namespace mu
 					FMemory::Memcpy(&ColorChannelUsages, &Args.ColorUsage, sizeof(ColorChannelUsages));
 					static_assert(sizeof(ColorChannelUsages) == sizeof(Args.ColorUsage));
 
-					Ptr<Mesh> BindMeshResult = CreateMesh();
+					TSharedPtr<FMesh> BindMeshResult = CreateMesh();
 
 					bool bOutSuccess = false;
-					MeshBindShapeReshape(BindMeshResult.get(), BaseMesh.get(), Shape.get(), BonesToDeform, PhysicsToDeform, BindFlags, ColorChannelUsages, bOutSuccess);
+					MeshBindShapeReshape(BindMeshResult.Get(), BaseMesh.Get(), Shape.Get(), BonesToDeform, PhysicsToDeform, BindFlags, ColorChannelUsages, bOutSuccess);
 				
 					Release(Shape);
 					// not success indicates nothing has bond so the base mesh can be reused.
@@ -2470,7 +2470,7 @@ namespace mu
 					{
 						if (!EnumHasAnyFlags(BindFlags, EMeshBindShapeFlags::ReshapeVertices))
 						{
-							Ptr<Mesh> BindMeshNoVertsResult = CloneOrTakeOver(BaseMesh);
+							TSharedPtr<FMesh> BindMeshNoVertsResult = CloneOrTakeOver(BaseMesh);
 							BindMeshNoVertsResult->AdditionalBuffers = MoveTemp(BindMeshResult->AdditionalBuffers);
 
 							Release(BaseMesh);
@@ -2486,10 +2486,10 @@ namespace mu
 				}	
 				else
 				{
-					Ptr<Mesh> Result = CreateMesh(BaseMesh ? BaseMesh->GetDataSize() : 0);
+					TSharedPtr<FMesh> Result = CreateMesh(BaseMesh ? BaseMesh->GetDataSize() : 0);
 
 					bool bOutSuccess = false;
-					MeshBindShapeClipDeform(Result.get(), BaseMesh.get(), Shape.get(), BindingMethod, bOutSuccess);
+					MeshBindShapeClipDeform(Result.Get(), BaseMesh.Get(), Shape.Get(), BindingMethod, bOutSuccess);
 
 					Release(Shape);
 					if (!bOutSuccess)
@@ -2515,7 +2515,7 @@ namespace mu
 		}
 
 
-		case OP_TYPE::ME_APPLYSHAPE:
+		case EOpType::ME_APPLYSHAPE:
 		{
 			OP::MeshApplyShapeArgs args = Program.GetOpArgs<OP::MeshApplyShapeArgs>(item.At);
 			switch (item.Stage)
@@ -2538,16 +2538,16 @@ namespace mu
 			{
 				MUTABLE_CPUPROFILER_SCOPE(ME_APPLYSHAPE_1)
 					
-				Ptr<const Mesh> BaseMesh = LoadMesh(FCacheAddress(args.mesh, item));
-				Ptr<const Mesh> Shape = LoadMesh(FCacheAddress(args.shape, item));
+				TSharedPtr<const FMesh> BaseMesh = LoadMesh(FCacheAddress(args.mesh, item));
+				TSharedPtr<const FMesh> Shape = LoadMesh(FCacheAddress(args.shape, item));
 
 				const EMeshBindShapeFlags ReshapeFlags = static_cast<EMeshBindShapeFlags>(args.flags);
 				const bool bReshapeVertices = EnumHasAnyFlags(ReshapeFlags, EMeshBindShapeFlags::ReshapeVertices);
 
-				Ptr<Mesh> ReshapedMeshResult = CreateMesh(BaseMesh ? BaseMesh->GetDataSize() : 0);
+				TSharedPtr<FMesh> ReshapedMeshResult = CreateMesh(BaseMesh ? BaseMesh->GetDataSize() : 0);
 
 				bool bOutSuccess = false;
-				MeshApplyShape(ReshapedMeshResult.get(), BaseMesh.get(), Shape.get(), ReshapeFlags, bOutSuccess);
+				MeshApplyShape(ReshapedMeshResult.Get(), BaseMesh.Get(), Shape.Get(), ReshapeFlags, bOutSuccess);
 
 				Release(Shape);
 				
@@ -2567,10 +2567,10 @@ namespace mu
 							EMeshCopyFlags::WithAdditionalPhysics |
 							EMeshCopyFlags::WithPoses);
 
-						Ptr<Mesh> NoVerticesReshpedMesh = CloneOrTakeOver(BaseMesh);
+						TSharedPtr<FMesh> NoVerticesReshpedMesh = CloneOrTakeOver(BaseMesh);
 
-						NoVerticesReshpedMesh->SetSkeleton(ReshapedMeshResult->GetSkeleton().get());
-						NoVerticesReshpedMesh->SetPhysicsBody(ReshapedMeshResult->GetPhysicsBody().get());
+						NoVerticesReshpedMesh->SetSkeleton(ReshapedMeshResult->GetSkeleton());
+						NoVerticesReshpedMesh->SetPhysicsBody(ReshapedMeshResult->GetPhysicsBody());
 						NoVerticesReshpedMesh->AdditionalPhysicsBodies = ReshapedMeshResult->AdditionalPhysicsBodies;
 						NoVerticesReshpedMesh->BonePoses = ReshapedMeshResult->BonePoses;
 
@@ -2594,7 +2594,7 @@ namespace mu
 			break;
 		}
 
-		case OP_TYPE::ME_MORPHRESHAPE:
+		case EOpType::ME_MORPHRESHAPE:
 		{
 			OP::MeshMorphReshapeArgs Args = Program.GetOpArgs<OP::MeshMorphReshapeArgs>(item.At);
 			switch(item.Stage)
@@ -2617,8 +2617,8 @@ namespace mu
 			{
 				MUTABLE_CPUPROFILER_SCOPE(ME_MORPHRESHAPE_1)
 					
-				Ptr<const Mesh> MorphedMesh = LoadMesh(FCacheAddress(Args.Morph, item));
-				Ptr<const Mesh> ReshapeMesh = LoadMesh(FCacheAddress(Args.Reshape, item));
+				TSharedPtr<const FMesh> MorphedMesh = LoadMesh(FCacheAddress(Args.Morph, item));
+				TSharedPtr<const FMesh> ReshapeMesh = LoadMesh(FCacheAddress(Args.Reshape, item));
 
 				if (ReshapeMesh && MorphedMesh)
 				{
@@ -2628,11 +2628,11 @@ namespace mu
 							EMeshCopyFlags::WithPhysicsBody | 
 							EMeshCopyFlags::WithPoses);
 
-					Ptr<Mesh> Result = CreateMesh(MorphedMesh->GetDataSize());
+					TSharedPtr<FMesh> Result = CreateMesh(MorphedMesh->GetDataSize());
 					Result->CopyFrom(*MorphedMesh, CopyFlags);
 
-					Result->SetSkeleton(ReshapeMesh->GetSkeleton().get());
-					Result->SetPhysicsBody(ReshapeMesh->GetPhysicsBody().get());
+					Result->SetSkeleton(ReshapeMesh->GetSkeleton());
+					Result->SetPhysicsBody(ReshapeMesh->GetPhysicsBody());
 					Result->BonePoses = ReshapeMesh->BonePoses;
 
 					Release(MorphedMesh);
@@ -2654,7 +2654,7 @@ namespace mu
 			break;
 		}
 
-        case OP_TYPE::ME_SETSKELETON:
+        case EOpType::ME_SETSKELETON:
         {
 			OP::MeshSetSkeletonArgs args = Program.GetOpArgs<OP::MeshSetSkeletonArgs>(item.At);
             switch (item.Stage)
@@ -2677,11 +2677,11 @@ namespace mu
             {
             	MUTABLE_CPUPROFILER_SCOPE(ME_SETSKELETON_1)
             		
-                Ptr<const Mesh> Source = LoadMesh(FCacheAddress(args.source, item));
-                Ptr<const Mesh> pSkeleton = LoadMesh(FCacheAddress(args.skeleton, item));
+                TSharedPtr<const FMesh> Source = LoadMesh(FCacheAddress(args.source, item));
+                TSharedPtr<const FMesh> Skeleton = LoadMesh(FCacheAddress(args.skeleton, item));
 
                 // Only if both are valid.
-                if (Source && pSkeleton)
+                if (Source && Skeleton)
                 {
                     if ( Source->GetSkeleton()
                          &&
@@ -2693,12 +2693,12 @@ namespace mu
                         // \TODO Be defensive, for now.
                         UE_LOG(LogMutableCore, Warning, TEXT("Performing a MeshRemapSkeleton, instead of MeshSetSkeletonData because source mesh already has some skeleton."));
 
-						Ptr<Mesh> Result = CreateMesh(Source->GetDataSize());
+						TSharedPtr<FMesh> Result = CreateMesh(Source->GetDataSize());
 
 						bool bOutSuccess = false;
-                        MeshRemapSkeleton(Result.get(), Source.get(), pSkeleton->GetSkeleton().get(), bOutSuccess);
+                        MeshRemapSkeleton(Result.Get(), Source.Get(), Skeleton->GetSkeleton(), bOutSuccess);
 
-						Release(pSkeleton);
+						Release(Skeleton);
 
                         if (!bOutSuccess)
                         {
@@ -2714,18 +2714,18 @@ namespace mu
                     }
                     else
                     {
-						Ptr<Mesh> Result = CloneOrTakeOver(Source);
+						TSharedPtr<FMesh> Result = CloneOrTakeOver(Source);
 
-                        Result->SetSkeleton(pSkeleton->GetSkeleton().get());
+                        Result->SetSkeleton(Skeleton->GetSkeleton());
 
 						//Result->GetPrivate()->CheckIntegrity();
-						Release(pSkeleton);
+						Release(Skeleton);
 						StoreMesh(item, Result);
                     }
                 }
                 else
                 {
-					Release(pSkeleton);
+					Release(Skeleton);
 					StoreMesh(item, Source);
                 }
 
@@ -2739,7 +2739,7 @@ namespace mu
             break;
         }
 
-        case OP_TYPE::ME_REMOVEMASK:
+        case EOpType::ME_REMOVEMASK:
         {
        		MUTABLE_CPUPROFILER_SCOPE(ME_REMOVEMASK)
         		
@@ -2802,7 +2802,7 @@ namespace mu
                 // \todo: store condition values in heap?
                 TArray<FScheduledOp> deps;
                 deps.Emplace( source, item );
-                for( size_t r=0; source && r<conditions.Num(); ++r )
+                for( int32 r=0; source && r<conditions.Num(); ++r )
                 {
                     // If there is no expression, we'll assume true.
                     bool value = true;
@@ -2829,11 +2829,11 @@ namespace mu
             	MUTABLE_CPUPROFILER_SCOPE(ME_REMOVEMASK_2)
             	
                 // \todo: single remove operation with all masks?
-                Ptr<const Mesh> Source = LoadMesh(FCacheAddress(source, item));
+                TSharedPtr<const FMesh> Source = LoadMesh(FCacheAddress(source, item));
 
 				if (Source)
 				{
-					Ptr<Mesh> Result = CreateMesh(Source->GetDataSize());
+					TSharedPtr<FMesh> Result = CreateMesh(Source->GetDataSize());
 					Result->CopyFrom(*Source);
 
 					Release(Source);
@@ -2849,15 +2849,15 @@ namespace mu
 
 						if (value)
 						{
-							Ptr<const Mesh> Mask = LoadMesh(FCacheAddress(masks[r], item));
+							TSharedPtr<const FMesh> Mask = LoadMesh(FCacheAddress(masks[r], item));
 							if (Mask)
 							{
 								//MeshRemoveMask will make a copy of Result, try to make room for it. 
-								Ptr<Mesh> IterResult = CreateMesh(Result->GetDataSize());
+								TSharedPtr<FMesh> IterResult = CreateMesh(Result->GetDataSize());
 
 								bool bOutSuccess = false;
 								bool bRemoveIfAllVerticesCulled = FaceCullStrategy == EFaceCullStrategy::AllVerticesCulled;
-								MeshRemoveMask(IterResult.get(), Result.get(), Mask.get(), bRemoveIfAllVerticesCulled, bOutSuccess);
+								MeshRemoveMask(IterResult.Get(), Result.Get(), Mask.Get(), bRemoveIfAllVerticesCulled, bOutSuccess);
 
 								Release(Mask);
 
@@ -2891,7 +2891,7 @@ namespace mu
             break;
         }
 
-        case OP_TYPE::ME_ADDTAGS:
+        case EOpType::ME_ADDTAGS:
 		{
 			MUTABLE_CPUPROFILER_SCOPE(ME_ADDTAGS)
 
@@ -2924,7 +2924,7 @@ namespace mu
 			{
 				MUTABLE_CPUPROFILER_SCOPE(ME_ADDTAGS_2)
 
-				Ptr<const Mesh> SourceMesh = LoadMesh(FCacheAddress(Source, item));
+				TSharedPtr<const FMesh> SourceMesh = LoadMesh(FCacheAddress(Source, item));
 
 				if (!SourceMesh)
 				{
@@ -2932,7 +2932,7 @@ namespace mu
 				}
 				else
 				{
-					Ptr<Mesh> Result = CloneOrTakeOver(SourceMesh);
+					TSharedPtr<FMesh> Result = CloneOrTakeOver(SourceMesh);
 
 					// Decode the tags
 					uint16 TagCount;
@@ -2947,8 +2947,8 @@ namespace mu
 						FMemory::Memcpy(&TagConstant, Data, sizeof(OP::ADDRESS));
 						Data += sizeof(OP::ADDRESS);
 
-						check(TagConstant < (uint32)pModel->GetPrivate()->m_program.m_constantStrings.Num());
-						const FString& Name = Program.m_constantStrings[TagConstant];
+						check(TagConstant < (uint32)pModel->GetPrivate()->Program.ConstantStrings.Num());
+						const FString& Name = Program.ConstantStrings[TagConstant];
 						Result->Tags[FirstMeshTagIndex+TagIndex] = Name;
 					}
 
@@ -2965,7 +2965,7 @@ namespace mu
 			break;
 		}
 
-        case OP_TYPE::ME_PROJECT:
+        case EOpType::ME_PROJECT:
         {
 			OP::MeshProjectArgs args = Program.GetOpArgs<OP::MeshProjectArgs>(item.At);
             switch (item.Stage)
@@ -2988,32 +2988,32 @@ namespace mu
             {
 				MUTABLE_CPUPROFILER_SCOPE(ME_PROJECT_1)
 
-                Ptr<const Mesh> pMesh = LoadMesh(FCacheAddress(args.mesh,item));
+                TSharedPtr<const FMesh> Mesh = LoadMesh(FCacheAddress(args.mesh,item));
                 const FProjector Projector = LoadProjector(FCacheAddress(args.projector, item));
 
                 // Only if both are valid.
-                if (pMesh && pMesh->GetVertexBuffers().GetBufferCount() > 0)
+                if (Mesh && Mesh->GetVertexBuffers().GetBufferCount() > 0)
                 {
-					Ptr<Mesh> Result = CreateMesh();
+					TSharedPtr<FMesh> Result = CreateMesh();
 
 					bool bOutSuccess = false;
-					MeshProject(Result.get(), pMesh.get(), Projector, bOutSuccess);
+					MeshProject(Result.Get(), Mesh.Get(), Projector, bOutSuccess);
 
 					if (!bOutSuccess)
 					{
 						Release(Result);
-						StoreMesh(item, pMesh);
+						StoreMesh(item, Mesh);
 					}
 					else
 					{	
 //						Result->GetPrivate()->CheckIntegrity();
-						Release(pMesh);
+						Release(Mesh);
 						StoreMesh(item, Result);
 					}
                 }
 				else
 				{
-					Release(pMesh);
+					Release(Mesh);
 					StoreMesh(item, nullptr);
 				}
 
@@ -3027,7 +3027,7 @@ namespace mu
             break;
         }
 
-		case OP_TYPE::ME_OPTIMIZESKINNING:
+		case EOpType::ME_OPTIMIZESKINNING:
 		{
 			OP::MeshOptimizeSkinningArgs args = Program.GetOpArgs<OP::MeshOptimizeSkinningArgs>(item.At);
 			switch (item.Stage)
@@ -3048,17 +3048,17 @@ namespace mu
 			{
 				MUTABLE_CPUPROFILER_SCOPE(ME_OPTIMIZESKINNING_1)
 
-				Ptr<const Mesh> Source = LoadMesh(FCacheAddress(args.source, item));
+				TSharedPtr<const FMesh> Source = LoadMesh(FCacheAddress(args.source, item));
 
 				if (Source && Source->IsReference())
 				{
 					StoreMesh(item, Source);
 				}
 
-				Ptr<Mesh> Result = CreateMesh();
+				TSharedPtr<FMesh> Result = CreateMesh();
 
 				bool bOutSuccess = false;
-				MeshOptimizeSkinning(Result.get(), Source.get(), bOutSuccess);
+				MeshOptimizeSkinning(Result.Get(), Source.Get(), bOutSuccess);
 
 				if (!bOutSuccess)
 				{
@@ -3081,7 +3081,7 @@ namespace mu
 			break;
 		}
 
-        case OP_TYPE::ME_TRANSFORMWITHMESH:
+        case EOpType::ME_TRANSFORMWITHMESH:
 		{
 			OP::MeshTransformWithinMeshArgs args = Program.GetOpArgs<OP::MeshTransformWithinMeshArgs>(item.At);
 			switch (item.Stage)
@@ -3105,16 +3105,16 @@ namespace mu
 			{
 				MUTABLE_CPUPROFILER_SCOPE(ME_TRANSFORMWITHMESH_1)
             
-				Ptr<const Mesh> SourceMesh = LoadMesh(FCacheAddress(args.sourceMesh,item));
-				Ptr<const Mesh> BoundingMesh = LoadMesh(FCacheAddress(args.boundingMesh, item));
+				TSharedPtr<const FMesh> SourceMesh = LoadMesh(FCacheAddress(args.sourceMesh,item));
+				TSharedPtr<const FMesh> BoundingMesh = LoadMesh(FCacheAddress(args.boundingMesh, item));
 				const FMatrix44f& Transform = LoadMatrix(FCacheAddress(args.matrix, item));
 
 				if (SourceMesh)
 				{
-					Ptr<Mesh> Result = CreateMesh(SourceMesh->GetDataSize());
+					TSharedPtr<FMesh> Result = CreateMesh(SourceMesh->GetDataSize());
 
 					bool bOutSuccess = false;
-					MeshTransformWithMesh(Result.get(), SourceMesh.get(), BoundingMesh.get(), Transform, bOutSuccess);
+					MeshTransformWithMesh(Result.Get(), SourceMesh.Get(), BoundingMesh.Get(), Transform, bOutSuccess);
 					Release(BoundingMesh);
 
 					if (!bOutSuccess)
@@ -3143,7 +3143,7 @@ namespace mu
 		}
 
         default:
-            if (type!=OP_TYPE::NONE)
+            if (type!=EOpType::NONE)
             {
                 // Operation not implemented
                 check( false );
@@ -3153,19 +3153,19 @@ namespace mu
     }
 
     //---------------------------------------------------------------------------------------------
-    void CodeRunner::RunCode_Image(const FScheduledOp& item, const Parameters* pParams, const Model* pModel )
+    void CodeRunner::RunCode_Image(const FScheduledOp& item, const FParameters* pParams, const FModel* pModel )
     {
 		MUTABLE_CPUPROFILER_SCOPE(RunCode_Image);
 
 		FImageOperator ImOp = MakeImageOperator(this);
 
-		const FProgram& Program = m_pModel->GetPrivate()->m_program;
+		const FProgram& Program = Model->GetPrivate()->Program;
 
-		OP_TYPE type = Program.GetOpType(item.At);
+		EOpType type = Program.GetOpType(item.At);
 		switch (type)
         {
 
-        case OP_TYPE::IM_LAYERCOLOUR:
+        case EOpType::IM_LAYERCOLOUR:
         {
 			OP::ImageLayerColourArgs args = Program.GetOpArgs<OP::ImageLayerColourArgs>(item.At);
             switch (item.Stage)
@@ -3189,7 +3189,7 @@ namespace mu
             break;
         }
 
-        case OP_TYPE::IM_LAYER:
+        case EOpType::IM_LAYER:
         {
 			OP::ImageLayerArgs args = Program.GetOpArgs<OP::ImageLayerArgs>(item.At);
 
@@ -3242,7 +3242,7 @@ namespace mu
             break;
         }
 
-        case OP_TYPE::IM_MULTILAYER:
+        case EOpType::IM_MULTILAYER:
         {
 			OP::ImageMultiLayerArgs args = Program.GetOpArgs<OP::ImageMultiLayerArgs>(item.At);
             switch (item.Stage)
@@ -3265,7 +3265,7 @@ namespace mu
 
                     // We support both integers and scalars here, which is not common.
                     // \todo: review if this is necessary or we can enforce it at compile time.
-                    DATATYPE RangeSizeType = GetOpDataType( pModel->GetPrivate()->m_program.GetOpType(args.rangeSize) );
+                    DATATYPE RangeSizeType = GetOpDataType( pModel->GetPrivate()->Program.GetOpType(args.rangeSize) );
                     if (RangeSizeType == DT_INT)
                     {
 						Iterations = LoadInt(RangeAddress);
@@ -3276,7 +3276,7 @@ namespace mu
                     }
                 }
 
-				Ptr<const Image> Base = LoadImage(FCacheAddress(args.base, item));
+				TSharedPtr<const FImage> Base = LoadImage(FCacheAddress(args.base, item));
 
 				if (Iterations <= 0)
 				{
@@ -3286,21 +3286,21 @@ namespace mu
 				else
 				{
 					// Store the base
-					Ptr<Image> New = CloneOrTakeOver(Base);
+					TSharedPtr<FImage> New = CloneOrTakeOver(Base);
 					EImageFormat InitialBaseFormat = New->GetFormat();
 
 					// Reset relevancy map.
-					New->m_flags &= ~Image::EImageFlags::IF_HAS_RELEVANCY_MAP;
+					New->Flags &= ~FImage::EImageFlags::IF_HAS_RELEVANCY_MAP;
 
 					// This shouldn't happen in optimised models, but it could happen in editors, etc.
 					// \todo: raise a performance warning?
 					EImageFormat BaseFormat = GetUncompressedFormat(New->GetFormat());
 					if (New->GetFormat() != BaseFormat)
 					{
-						Ptr<Image> Formatted = CreateImage( New->GetSizeX(), New->GetSizeY(), New->GetLODCount(), BaseFormat, EInitializationType::NotInitialized );
+						TSharedPtr<FImage> Formatted = CreateImage( New->GetSizeX(), New->GetSizeY(), New->GetLODCount(), BaseFormat, EInitializationType::NotInitialized );
 
 						bool bSuccess = false;
-						ImOp.ImagePixelFormat(bSuccess, m_pSettings->ImageCompressionQuality, Formatted.get(), New.get());
+						ImOp.ImagePixelFormat(bSuccess, Settings.ImageCompressionQuality, Formatted.Get(), New.Get());
 						check(bSuccess); // Decompression cannot fail
 
 						Release(New);
@@ -3312,7 +3312,7 @@ namespace mu
 					Data.MultiLayer.Iterations = Iterations;
 					Data.MultiLayer.OriginalBaseFormat = InitialBaseFormat;
 					Data.MultiLayer.bBlendOnlyOneMip = false;
-					int32 DataPos = m_heapData.Add(Data);
+					int32 DataPos = HeapData.Add(Data);
 
 					// Request the first layer
 					int32 CurrentIteration = 0;
@@ -3330,7 +3330,7 @@ namespace mu
             {
             	MUTABLE_CPUPROFILER_SCOPE(IM_MULTILAYER_default)
 
-				FScheduledOpData& Data = m_heapData[item.CustomState];
+				FScheduledOpData& Data = HeapData[item.CustomState];
 
 				int32 Iterations = Data.MultiLayer.Iterations;
 				int32 CurrentIteration = item.Stage - 2;
@@ -3340,7 +3340,7 @@ namespace mu
 
 				// Process the current layer
 
-				Ptr<Image> Base = static_cast<Image*>(Data.Resource.get());
+				TSharedPtr<FImage> Base = StaticCastSharedPtr<FImage>(ConstCastSharedPtr<FResource>(Data.Resource));
  
                 FScheduledOp itemCopy = item;
                 ExecutionIndex index = GetMemory().GetRangeIndex( item.ExecutionIndex );
@@ -3350,7 +3350,7 @@ namespace mu
                     itemCopy.ExecutionIndex = GetMemory().GetRangeIndexIndex(index);
 					itemCopy.CustomState = 0;
 
-                    Ptr<const Image> Blended = LoadImage( FCacheAddress(args.blended,itemCopy) );
+                    TSharedPtr<const FImage> Blended = LoadImage( FCacheAddress(args.blended,itemCopy) );
 
                     // This shouldn't happen in optimised models, but it could happen in editors, etc.
                     // \todo: raise a performance warning?
@@ -3358,10 +3358,10 @@ namespace mu
                     {
 						MUTABLE_CPUPROFILER_SCOPE(ImageResize_BlendedReformat);
 
-						Ptr<Image> Formatted = CreateImage(Blended->GetSizeX(), Blended->GetSizeY(), Blended->GetLODCount(), Base->GetFormat(), EInitializationType::NotInitialized);
+						TSharedPtr<FImage> Formatted = CreateImage(Blended->GetSizeX(), Blended->GetSizeY(), Blended->GetLODCount(), Base->GetFormat(), EInitializationType::NotInitialized);
 
 						bool bSuccess = false;
-						ImOp.ImagePixelFormat(bSuccess, m_pSettings->ImageCompressionQuality, Formatted.get(), Blended.get());
+						ImOp.ImagePixelFormat(bSuccess, Settings.ImageCompressionQuality, Formatted.Get(), Blended.Get());
 						check(bSuccess);
 
 						Release(Blended);
@@ -3374,8 +3374,8 @@ namespace mu
 					{
 						MUTABLE_CPUPROFILER_SCOPE(ImageResize_BlendedFixForMultilayer);
 
-						Ptr<Image> Resized = CreateImage(ResultSize[0], ResultSize[1], Blended->GetLODCount(), Blended->GetFormat(), EInitializationType::NotInitialized);
-						ImOp.ImageResizeLinear(Resized.get(), 0, Blended.get());
+						TSharedPtr<FImage> Resized = CreateImage(ResultSize[0], ResultSize[1], Blended->GetLODCount(), Blended->GetFormat(), EInitializationType::NotInitialized);
+						ImOp.ImageResizeLinear(Resized.Get(), 0, Blended.Get());
 						Release(Blended);
 						Blended = Resized;
 					}
@@ -3406,25 +3406,25 @@ namespace mu
 						constexpr bool bUseVectorImpl = false;
 						if constexpr (bUseVectorImpl)
 						{
-							BufferLayerCompositeVector<VectorBlendChannelMasked, VectorLightenChannel, false>(Base.get(), Blended.get(), bBlendOnlyOneMip, args.BlendAlphaSourceChannel);
+							BufferLayerCompositeVector<VectorBlendChannelMasked, VectorLightenChannel, false>(Base.Get(), Blended.Get(), bBlendOnlyOneMip, args.BlendAlphaSourceChannel);
 						}
 						else
 						{
-							BufferLayerComposite<BlendChannelMasked, LightenChannel, false>(Base.get(), Blended.get(), bBlendOnlyOneMip, args.BlendAlphaSourceChannel);
+							BufferLayerComposite<BlendChannelMasked, LightenChannel, false>(Base.Get(), Blended.Get(), bBlendOnlyOneMip, args.BlendAlphaSourceChannel);
 						}
 					}
 
                     if (!bDone && args.mask)
                     {
-                        Ptr<const Image> Mask = LoadImage( FCacheAddress(args.mask,itemCopy) );
+                        TSharedPtr<const FImage> Mask = LoadImage( FCacheAddress(args.mask,itemCopy) );
 
 						// TODO: This shouldn't happen, but be defensive.
 						if (Mask && Mask->GetSize() != ResultSize)
 						{
 							MUTABLE_CPUPROFILER_SCOPE(ImageResize_MaskFixForMultilayer);
 
-							Ptr<Image> Resized = CreateImage(ResultSize[0], ResultSize[1], Mask->GetLODCount(), Mask->GetFormat(), EInitializationType::NotInitialized);
-							ImOp.ImageResizeLinear(Resized.get(), 0, Mask.get());
+							TSharedPtr<FImage> Resized = CreateImage(ResultSize[0], ResultSize[1], Mask->GetLODCount(), Mask->GetFormat(), EInitializationType::NotInitialized);
+							ImOp.ImageResizeLinear(Resized.Get(), 0, Mask.Get());
 							Release(Mask);
 							Mask = Resized;
 						}
@@ -3435,15 +3435,15 @@ namespace mu
                         switch (EBlendType(args.blendType))
                         {
 						case EBlendType::BT_NORMAL_COMBINE: check(false); break;
-                        case EBlendType::BT_SOFTLIGHT: BufferLayer<SoftLightChannelMasked, SoftLightChannel, false>(Base.get(), Base.get(), Mask.get(), Blended.get(), bApplyColorBlendToAlpha, bBlendOnlyOneMip); break;
-                        case EBlendType::BT_HARDLIGHT: BufferLayer<HardLightChannelMasked, HardLightChannel, false>(Base.get(), Base.get(), Mask.get(), Blended.get(), bApplyColorBlendToAlpha, bBlendOnlyOneMip); break;
-                        case EBlendType::BT_BURN: BufferLayer<BurnChannelMasked, BurnChannel, false>(Base.get(), Base.get(), Mask.get(), Blended.get(), bApplyColorBlendToAlpha, bBlendOnlyOneMip); break;
-                        case EBlendType::BT_DODGE: BufferLayer<DodgeChannelMasked, DodgeChannel, false>(Base.get(), Base.get(), Mask.get(), Blended.get(), bApplyColorBlendToAlpha, bBlendOnlyOneMip); break;
-                        case EBlendType::BT_SCREEN: BufferLayer<ScreenChannelMasked, ScreenChannel, false>(Base.get(), Base.get(), Mask.get(), Blended.get(), bApplyColorBlendToAlpha, bBlendOnlyOneMip); break;
-                        case EBlendType::BT_OVERLAY: BufferLayer<OverlayChannelMasked, OverlayChannel, false>(Base.get(), Base.get(), Mask.get(), Blended.get(), bApplyColorBlendToAlpha, bBlendOnlyOneMip); break;
-                        case EBlendType::BT_LIGHTEN: BufferLayer<LightenChannelMasked, LightenChannel, false>(Base.get(), Base.get(), Mask.get(), Blended.get(), bApplyColorBlendToAlpha, bBlendOnlyOneMip); break;
-                        case EBlendType::BT_MULTIPLY: BufferLayer<MultiplyChannelMasked, MultiplyChannel, false>(Base.get(), Base.get(), Mask.get(), Blended.get(), bApplyColorBlendToAlpha, bBlendOnlyOneMip); break;
-                        case EBlendType::BT_BLEND: BufferLayer<BlendChannelMasked, BlendChannel, false>(Base.get(), Base.get(), Mask.get(), Blended.get(), bApplyColorBlendToAlpha, bBlendOnlyOneMip); break;
+                        case EBlendType::BT_SOFTLIGHT: BufferLayer<SoftLightChannelMasked, SoftLightChannel, false>(Base.Get(), Base.Get(), Mask.Get(), Blended.Get(), bApplyColorBlendToAlpha, bBlendOnlyOneMip); break;
+                        case EBlendType::BT_HARDLIGHT: BufferLayer<HardLightChannelMasked, HardLightChannel, false>(Base.Get(), Base.Get(), Mask.Get(), Blended.Get(), bApplyColorBlendToAlpha, bBlendOnlyOneMip); break;
+                        case EBlendType::BT_BURN: BufferLayer<BurnChannelMasked, BurnChannel, false>(Base.Get(), Base.Get(), Mask.Get(), Blended.Get(), bApplyColorBlendToAlpha, bBlendOnlyOneMip); break;
+                        case EBlendType::BT_DODGE: BufferLayer<DodgeChannelMasked, DodgeChannel, false>(Base.Get(), Base.Get(), Mask.Get(), Blended.Get(), bApplyColorBlendToAlpha, bBlendOnlyOneMip); break;
+                        case EBlendType::BT_SCREEN: BufferLayer<ScreenChannelMasked, ScreenChannel, false>(Base.Get(), Base.Get(), Mask.Get(), Blended.Get(), bApplyColorBlendToAlpha, bBlendOnlyOneMip); break;
+                        case EBlendType::BT_OVERLAY: BufferLayer<OverlayChannelMasked, OverlayChannel, false>(Base.Get(), Base.Get(), Mask.Get(), Blended.Get(), bApplyColorBlendToAlpha, bBlendOnlyOneMip); break;
+                        case EBlendType::BT_LIGHTEN: BufferLayer<LightenChannelMasked, LightenChannel, false>(Base.Get(), Base.Get(), Mask.Get(), Blended.Get(), bApplyColorBlendToAlpha, bBlendOnlyOneMip); break;
+                        case EBlendType::BT_MULTIPLY: BufferLayer<MultiplyChannelMasked, MultiplyChannel, false>(Base.Get(), Base.Get(), Mask.Get(), Blended.Get(), bApplyColorBlendToAlpha, bBlendOnlyOneMip); break;
+                        case EBlendType::BT_BLEND: BufferLayer<BlendChannelMasked, BlendChannel, false>(Base.Get(), Base.Get(), Mask.Get(), Blended.Get(), bApplyColorBlendToAlpha, bBlendOnlyOneMip); break;
                         default: check(false);
                         }
 
@@ -3457,15 +3457,15 @@ namespace mu
 						switch (EBlendType(args.blendType))
 						{
 						case EBlendType::BT_NORMAL_COMBINE: check(false); break;
-						case EBlendType::BT_SOFTLIGHT: BufferLayerEmbeddedMask<SoftLightChannelMasked, SoftLightChannel, false>(Base.get(), Base.get(), Blended.get(), bApplyColorBlendToAlpha, bBlendOnlyOneMip); break;
-						case EBlendType::BT_HARDLIGHT: BufferLayerEmbeddedMask<HardLightChannelMasked, HardLightChannel, false>(Base.get(), Base.get(), Blended.get(), bApplyColorBlendToAlpha, bBlendOnlyOneMip); break;
-						case EBlendType::BT_BURN: BufferLayerEmbeddedMask<BurnChannelMasked, BurnChannel, false>(Base.get(), Base.get(), Blended.get(), bApplyColorBlendToAlpha, bBlendOnlyOneMip); break;
-						case EBlendType::BT_DODGE: BufferLayerEmbeddedMask<DodgeChannelMasked, DodgeChannel, false>(Base.get(), Base.get(), Blended.get(), bApplyColorBlendToAlpha, bBlendOnlyOneMip); break;
-						case EBlendType::BT_SCREEN: BufferLayerEmbeddedMask<ScreenChannelMasked, ScreenChannel, false>(Base.get(), Base.get(), Blended.get(), bApplyColorBlendToAlpha, bBlendOnlyOneMip); break;
-						case EBlendType::BT_OVERLAY: BufferLayerEmbeddedMask<OverlayChannelMasked, OverlayChannel, false>(Base.get(), Base.get(), Blended.get(), bApplyColorBlendToAlpha, bBlendOnlyOneMip); break;
-						case EBlendType::BT_LIGHTEN: BufferLayerEmbeddedMask<LightenChannelMasked, LightenChannel, false>(Base.get(), Base.get(), Blended.get(), bApplyColorBlendToAlpha, bBlendOnlyOneMip); break;
-						case EBlendType::BT_MULTIPLY: BufferLayerEmbeddedMask<MultiplyChannelMasked, MultiplyChannel, false>(Base.get(), Base.get(), Blended.get(), bApplyColorBlendToAlpha, bBlendOnlyOneMip); break;
-						case EBlendType::BT_BLEND: BufferLayerEmbeddedMask<BlendChannelMasked, BlendChannel, false>(Base.get(), Base.get(), Blended.get(), bApplyColorBlendToAlpha, bBlendOnlyOneMip); break;
+						case EBlendType::BT_SOFTLIGHT: BufferLayerEmbeddedMask<SoftLightChannelMasked, SoftLightChannel, false>(Base.Get(), Base.Get(), Blended.Get(), bApplyColorBlendToAlpha, bBlendOnlyOneMip); break;
+						case EBlendType::BT_HARDLIGHT: BufferLayerEmbeddedMask<HardLightChannelMasked, HardLightChannel, false>(Base.Get(), Base.Get(), Blended.Get(), bApplyColorBlendToAlpha, bBlendOnlyOneMip); break;
+						case EBlendType::BT_BURN: BufferLayerEmbeddedMask<BurnChannelMasked, BurnChannel, false>(Base.Get(), Base.Get(), Blended.Get(), bApplyColorBlendToAlpha, bBlendOnlyOneMip); break;
+						case EBlendType::BT_DODGE: BufferLayerEmbeddedMask<DodgeChannelMasked, DodgeChannel, false>(Base.Get(), Base.Get(), Blended.Get(), bApplyColorBlendToAlpha, bBlendOnlyOneMip); break;
+						case EBlendType::BT_SCREEN: BufferLayerEmbeddedMask<ScreenChannelMasked, ScreenChannel, false>(Base.Get(), Base.Get(), Blended.Get(), bApplyColorBlendToAlpha, bBlendOnlyOneMip); break;
+						case EBlendType::BT_OVERLAY: BufferLayerEmbeddedMask<OverlayChannelMasked, OverlayChannel, false>(Base.Get(), Base.Get(), Blended.Get(), bApplyColorBlendToAlpha, bBlendOnlyOneMip); break;
+						case EBlendType::BT_LIGHTEN: BufferLayerEmbeddedMask<LightenChannelMasked, LightenChannel, false>(Base.Get(), Base.Get(), Blended.Get(), bApplyColorBlendToAlpha, bBlendOnlyOneMip); break;
+						case EBlendType::BT_MULTIPLY: BufferLayerEmbeddedMask<MultiplyChannelMasked, MultiplyChannel, false>(Base.Get(), Base.Get(), Blended.Get(), bApplyColorBlendToAlpha, bBlendOnlyOneMip); break;
+						case EBlendType::BT_BLEND: BufferLayerEmbeddedMask<BlendChannelMasked, BlendChannel, false>(Base.Get(), Base.Get(), Blended.Get(), bApplyColorBlendToAlpha, bBlendOnlyOneMip); break;
 						default: check(false);
 						}
 					}
@@ -3474,15 +3474,15 @@ namespace mu
                         switch (EBlendType(args.blendType))
                         {
 						case EBlendType::BT_NORMAL_COMBINE: check(false); break;
-                        case EBlendType::BT_SOFTLIGHT: BufferLayer<SoftLightChannel, false>(Base.get(), Base.get(), Blended.get(), bApplyColorBlendToAlpha, bBlendOnlyOneMip, bUseBlendSourceFromBlendAlpha); break;
-                        case EBlendType::BT_HARDLIGHT: BufferLayer<HardLightChannel, false>(Base.get(), Base.get(), Blended.get(), bApplyColorBlendToAlpha, bBlendOnlyOneMip, bUseBlendSourceFromBlendAlpha); break;
-                        case EBlendType::BT_BURN: BufferLayer<BurnChannel, false>(Base.get(), Base.get(), Blended.get(), bApplyColorBlendToAlpha, bBlendOnlyOneMip, bUseBlendSourceFromBlendAlpha); break;
-                        case EBlendType::BT_DODGE: BufferLayer<DodgeChannel, false>(Base.get(), Base.get(), Blended.get(), bApplyColorBlendToAlpha, bBlendOnlyOneMip, bUseBlendSourceFromBlendAlpha); break;
-                        case EBlendType::BT_SCREEN: BufferLayer<ScreenChannel, false>(Base.get(), Base.get(),  Blended.get(), bApplyColorBlendToAlpha, bBlendOnlyOneMip, bUseBlendSourceFromBlendAlpha); break;
-                        case EBlendType::BT_OVERLAY: BufferLayer<OverlayChannel, false>(Base.get(), Base.get(), Blended.get(), bApplyColorBlendToAlpha, bBlendOnlyOneMip, bUseBlendSourceFromBlendAlpha); break;
-                        case EBlendType::BT_LIGHTEN: BufferLayer<LightenChannel, false>(Base.get(), Base.get(), Blended.get(), bApplyColorBlendToAlpha, bBlendOnlyOneMip, bUseBlendSourceFromBlendAlpha); break;
-                        case EBlendType::BT_MULTIPLY: BufferLayer<MultiplyChannel, false>(Base.get(), Base.get(), Blended.get(), bApplyColorBlendToAlpha, bBlendOnlyOneMip, bUseBlendSourceFromBlendAlpha); break;
-                        case EBlendType::BT_BLEND: BufferLayer<BlendChannel, false>(Base.get(), Base.get(), Blended.get(), bApplyColorBlendToAlpha, bBlendOnlyOneMip, bUseBlendSourceFromBlendAlpha); break;
+                        case EBlendType::BT_SOFTLIGHT: BufferLayer<SoftLightChannel, false>(Base.Get(), Base.Get(), Blended.Get(), bApplyColorBlendToAlpha, bBlendOnlyOneMip, bUseBlendSourceFromBlendAlpha); break;
+                        case EBlendType::BT_HARDLIGHT: BufferLayer<HardLightChannel, false>(Base.Get(), Base.Get(), Blended.Get(), bApplyColorBlendToAlpha, bBlendOnlyOneMip, bUseBlendSourceFromBlendAlpha); break;
+                        case EBlendType::BT_BURN: BufferLayer<BurnChannel, false>(Base.Get(), Base.Get(), Blended.Get(), bApplyColorBlendToAlpha, bBlendOnlyOneMip, bUseBlendSourceFromBlendAlpha); break;
+                        case EBlendType::BT_DODGE: BufferLayer<DodgeChannel, false>(Base.Get(), Base.Get(), Blended.Get(), bApplyColorBlendToAlpha, bBlendOnlyOneMip, bUseBlendSourceFromBlendAlpha); break;
+                        case EBlendType::BT_SCREEN: BufferLayer<ScreenChannel, false>(Base.Get(), Base.Get(),  Blended.Get(), bApplyColorBlendToAlpha, bBlendOnlyOneMip, bUseBlendSourceFromBlendAlpha); break;
+                        case EBlendType::BT_OVERLAY: BufferLayer<OverlayChannel, false>(Base.Get(), Base.Get(), Blended.Get(), bApplyColorBlendToAlpha, bBlendOnlyOneMip, bUseBlendSourceFromBlendAlpha); break;
+                        case EBlendType::BT_LIGHTEN: BufferLayer<LightenChannel, false>(Base.Get(), Base.Get(), Blended.Get(), bApplyColorBlendToAlpha, bBlendOnlyOneMip, bUseBlendSourceFromBlendAlpha); break;
+                        case EBlendType::BT_MULTIPLY: BufferLayer<MultiplyChannel, false>(Base.Get(), Base.Get(), Blended.Get(), bApplyColorBlendToAlpha, bBlendOnlyOneMip, bUseBlendSourceFromBlendAlpha); break;
+                        case EBlendType::BT_BLEND: BufferLayer<BlendChannel, false>(Base.Get(), Base.Get(), Blended.Get(), bApplyColorBlendToAlpha, bBlendOnlyOneMip, bUseBlendSourceFromBlendAlpha); break;
                         default: check(false);
                         }
                     }
@@ -3493,15 +3493,15 @@ namespace mu
 						// Separate alpha operation ignores the mask.
 						switch (EBlendType(args.blendTypeAlpha))
 						{
-						case EBlendType::BT_SOFTLIGHT: BufferLayerInPlace<SoftLightChannel, false, 1>(Base.get(), Blended.get(), bBlendOnlyOneMip, 3, args.BlendAlphaSourceChannel); break;
-						case EBlendType::BT_HARDLIGHT: BufferLayerInPlace<HardLightChannel, false, 1>(Base.get(), Blended.get(), bBlendOnlyOneMip, 3, args.BlendAlphaSourceChannel); break;
-						case EBlendType::BT_BURN: BufferLayerInPlace<BurnChannel, false, 1>(Base.get(), Blended.get(), bBlendOnlyOneMip, 3, args.BlendAlphaSourceChannel); break;
-						case EBlendType::BT_DODGE: BufferLayerInPlace<DodgeChannel, false, 1>(Base.get(), Blended.get(), bBlendOnlyOneMip, 3, args.BlendAlphaSourceChannel); break;
-						case EBlendType::BT_SCREEN: BufferLayerInPlace<ScreenChannel, false, 1>(Base.get(), Blended.get(), bBlendOnlyOneMip, 3, args.BlendAlphaSourceChannel); break;
-						case EBlendType::BT_OVERLAY: BufferLayerInPlace<OverlayChannel, false, 1>(Base.get(), Blended.get(), bBlendOnlyOneMip, 3, args.BlendAlphaSourceChannel); break;
-						case EBlendType::BT_LIGHTEN: BufferLayerInPlace<LightenChannel, false, 1>(Base.get(), Blended.get(), bBlendOnlyOneMip, 3, args.BlendAlphaSourceChannel); break;
-						case EBlendType::BT_MULTIPLY: BufferLayerInPlace<MultiplyChannel, false, 1>(Base.get(), Blended.get(), bBlendOnlyOneMip, 3, args.BlendAlphaSourceChannel); break;
-						case EBlendType::BT_BLEND: BufferLayerInPlace<BlendChannel, false, 1>(Base.get(), Blended.get(), bBlendOnlyOneMip, 3, args.BlendAlphaSourceChannel); break;
+						case EBlendType::BT_SOFTLIGHT: BufferLayerInPlace<SoftLightChannel, false, 1>(Base.Get(), Blended.Get(), bBlendOnlyOneMip, 3, args.BlendAlphaSourceChannel); break;
+						case EBlendType::BT_HARDLIGHT: BufferLayerInPlace<HardLightChannel, false, 1>(Base.Get(), Blended.Get(), bBlendOnlyOneMip, 3, args.BlendAlphaSourceChannel); break;
+						case EBlendType::BT_BURN: BufferLayerInPlace<BurnChannel, false, 1>(Base.Get(), Blended.Get(), bBlendOnlyOneMip, 3, args.BlendAlphaSourceChannel); break;
+						case EBlendType::BT_DODGE: BufferLayerInPlace<DodgeChannel, false, 1>(Base.Get(), Blended.Get(), bBlendOnlyOneMip, 3, args.BlendAlphaSourceChannel); break;
+						case EBlendType::BT_SCREEN: BufferLayerInPlace<ScreenChannel, false, 1>(Base.Get(), Blended.Get(), bBlendOnlyOneMip, 3, args.BlendAlphaSourceChannel); break;
+						case EBlendType::BT_OVERLAY: BufferLayerInPlace<OverlayChannel, false, 1>(Base.Get(), Blended.Get(), bBlendOnlyOneMip, 3, args.BlendAlphaSourceChannel); break;
+						case EBlendType::BT_LIGHTEN: BufferLayerInPlace<LightenChannel, false, 1>(Base.Get(), Blended.Get(), bBlendOnlyOneMip, 3, args.BlendAlphaSourceChannel); break;
+						case EBlendType::BT_MULTIPLY: BufferLayerInPlace<MultiplyChannel, false, 1>(Base.Get(), Blended.Get(), bBlendOnlyOneMip, 3, args.BlendAlphaSourceChannel); break;
+						case EBlendType::BT_BLEND: BufferLayerInPlace<BlendChannel, false, 1>(Base.Get(), Blended.Get(), bBlendOnlyOneMip, 3, args.BlendAlphaSourceChannel); break;
 						default: check(false);
 						}
 					}
@@ -3516,7 +3516,7 @@ namespace mu
 					{
 						MUTABLE_CPUPROFILER_SCOPE(ImageLayer_MipFix);
 						FMipmapGenerationSettings DummyMipSettings{};
-						ImageMipmapInPlace(m_pSettings->ImageCompressionQuality, Base.get(), DummyMipSettings);
+						ImageMipmapInPlace(Settings.ImageCompressionQuality, Base.Get(), DummyMipSettings);
 					}
 
 					// TODO: Reconvert to OriginalBaseFormat if necessary?
@@ -3545,7 +3545,7 @@ namespace mu
             break;
         }
 
-		case OP_TYPE::IM_NORMALCOMPOSITE:
+		case EOpType::IM_NORMALCOMPOSITE:
 		{
 			OP::ImageNormalCompositeArgs args = Program.GetOpArgs<OP::ImageNormalCompositeArgs>(item.At);
 			switch (item.Stage)
@@ -3567,8 +3567,8 @@ namespace mu
 			{
 				MUTABLE_CPUPROFILER_SCOPE(IM_NORMALCOMPOSITE_1)
 
-				Ptr<const Image> Base = LoadImage(FCacheAddress(args.base, item));
-				Ptr<const Image> Normal = LoadImage(FCacheAddress(args.normal, item));
+				TSharedPtr<const FImage> Base = LoadImage(FCacheAddress(args.base, item));
+				TSharedPtr<const FImage> Normal = LoadImage(FCacheAddress(args.normal, item));
 
 				if (Normal->GetLODCount() < Base->GetLODCount())
 				{
@@ -3577,17 +3577,17 @@ namespace mu
 					int32 StartLevel = Normal->GetLODCount() - 1;
 					int32 LevelCount = Base->GetLODCount();
 					
-					Ptr<Image> NormalFix = CloneOrTakeOver(Normal);
+					TSharedPtr<FImage> NormalFix = CloneOrTakeOver(Normal);
 
 					FMipmapGenerationSettings MipSettings{};
-					ImOp.ImageMipmap(m_pSettings->ImageCompressionQuality, NormalFix.get(), NormalFix.get(), StartLevel, LevelCount, MipSettings);
+					ImOp.ImageMipmap(Settings.ImageCompressionQuality, NormalFix.Get(), NormalFix.Get(), StartLevel, LevelCount, MipSettings);
 
 					Normal = NormalFix;
 				}
 
 
-                Ptr<Image> Result = CreateImage(Base->GetSizeX(), Base->GetSizeY(), Base->GetLODCount(), Base->GetFormat(), EInitializationType::NotInitialized);
-				ImageNormalComposite(Result.get(), Base.get(), Normal.get(), args.mode, args.power);
+                TSharedPtr<FImage> Result = CreateImage(Base->GetSizeX(), Base->GetSizeY(), Base->GetLODCount(), Base->GetFormat(), EInitializationType::NotInitialized);
+				ImageNormalComposite(Result.Get(), Base.Get(), Normal.Get(), args.mode, args.power);
 
 				Release(Base);
 				Release(Normal);
@@ -3602,7 +3602,7 @@ namespace mu
 			break;
 		}
 
-        case OP_TYPE::IM_PIXELFORMAT:
+        case EOpType::IM_PIXELFORMAT:
         {
 			OP::ImagePixelFormatArgs args = Program.GetOpArgs<OP::ImagePixelFormatArgs>(item.At);
             switch (item.Stage)
@@ -3625,7 +3625,7 @@ namespace mu
             break;
         }
 
-        case OP_TYPE::IM_MIPMAP:
+        case EOpType::IM_MIPMAP:
         {
 			OP::ImageMipmapArgs args = Program.GetOpArgs<OP::ImageMipmapArgs>(item.At);
             switch (item.Stage)
@@ -3646,7 +3646,7 @@ namespace mu
             break;
         }
 
-        case OP_TYPE::IM_RESIZE:
+        case EOpType::IM_RESIZE:
         {
 			OP::ImageResizeArgs args = Program.GetOpArgs<OP::ImageResizeArgs>(item.At);
             switch (item.Stage)
@@ -3668,7 +3668,7 @@ namespace mu
             break;
         }
 
-        case OP_TYPE::IM_RESIZELIKE:
+        case EOpType::IM_RESIZELIKE:
         {
 			OP::ImageResizeLikeArgs args = Program.GetOpArgs<OP::ImageResizeLikeArgs>(item.At);
             switch (item.Stage)
@@ -3683,16 +3683,16 @@ namespace mu
             {
             	MUTABLE_CPUPROFILER_SCOPE(IM_RESIZELIKE_1)
             	
-                Ptr<const Image> Base = LoadImage( FCacheAddress(args.source,item) );
-                Ptr<const Image> SizeBase = LoadImage( FCacheAddress(args.sizeSource,item) );
+                TSharedPtr<const FImage> Base = LoadImage( FCacheAddress(args.source,item) );
+                TSharedPtr<const FImage> SizeBase = LoadImage( FCacheAddress(args.sizeSource,item) );
 				FImageSize DestSize = SizeBase->GetSize();
 				Release(SizeBase);
 
                 if (Base->GetSize() != DestSize)
                 {
 					int32 BaseLODCount = Base->GetLODCount();
-					Ptr<Image> Result = CreateImage(DestSize[0], DestSize[1], BaseLODCount, Base->GetFormat(), EInitializationType::NotInitialized);
-					ImOp.ImageResizeLinear(Result.get(), m_pSettings->ImageCompressionQuality, Base.get());
+					TSharedPtr<FImage> Result = CreateImage(DestSize[0], DestSize[1], BaseLODCount, Base->GetFormat(), EInitializationType::NotInitialized);
+					ImOp.ImageResizeLinear(Result.Get(), Settings.ImageCompressionQuality, Base.Get());
 					Release(Base);
 
                     // If the source image had mips, generate them as well for the resized image.
@@ -3703,11 +3703,11 @@ namespace mu
                     
 					if (bSourceHasMips)
                     {
-						int32 LevelCount = Image::GetMipmapCount(Result->GetSizeX(), Result->GetSizeY());	
+						int32 LevelCount = FImage::GetMipmapCount(Result->GetSizeX(), Result->GetSizeY());	
 						Result->DataStorage.SetNumLODs(LevelCount);
 
 						FMipmapGenerationSettings MipSettings{};
-						ImOp.ImageMipmap(m_pSettings->ImageCompressionQuality, Result.get(), Result.get(), 0, LevelCount, MipSettings);
+						ImOp.ImageMipmap(Settings.ImageCompressionQuality, Result.Get(), Result.Get(), 0, LevelCount, MipSettings);
                     }				
 
 					StoreImage(item, Result);
@@ -3727,7 +3727,7 @@ namespace mu
             break;
         }
 
-        case OP_TYPE::IM_RESIZEREL:
+        case EOpType::IM_RESIZEREL:
         {
 			OP::ImageResizeRelArgs args = Program.GetOpArgs<OP::ImageResizeRelArgs>(item.At);
             switch (item.Stage)
@@ -3749,7 +3749,7 @@ namespace mu
 
             break;
         }
-        case OP_TYPE::IM_BLANKLAYOUT:
+        case EOpType::IM_BLANKLAYOUT:
         {
 			OP::ImageBlankLayoutArgs Args = Program.GetOpArgs<OP::ImageBlankLayoutArgs>(item.At);
             switch (item.Stage)
@@ -3762,7 +3762,7 @@ namespace mu
             {
             	MUTABLE_CPUPROFILER_SCOPE(IM_BLANKLAYOUT_1)
             		
-                Ptr<const Layout> pLayout = LoadLayout(FScheduledOp::FromOpAndOptions(Args.layout, item, 0));
+                TSharedPtr<const FLayout> pLayout = LoadLayout(FScheduledOp::FromOpAndOptions(Args.layout, item, 0));
 
                 FIntPoint SizeInBlocks = pLayout->GetGridSize();
 
@@ -3770,7 +3770,7 @@ namespace mu
 
 				// Image size if we don't skip any mipmap
 				FIntPoint FullImageSizeInPixels = SizeInBlocks * BlockSizeInPixels;
-				int32 FullImageMipCount = Image::GetMipmapCount(FullImageSizeInPixels.X, FullImageSizeInPixels.Y);
+				int32 FullImageMipCount = FImage::GetMipmapCount(FullImageSizeInPixels.X, FullImageSizeInPixels.Y);
 
 				FIntPoint ImageSizeInPixels = FullImageSizeInPixels;
 				int32 MipsToSkip = item.ExecutionOptions;
@@ -3784,7 +3784,7 @@ namespace mu
 					//ReducedBlockSizeInPixels.X = BlockSizeInPixels.X >> MipsToSkip;
 					//ReducedBlockSizeInPixels.Y = BlockSizeInPixels.Y >> MipsToSkip;
 					//const FImageFormatData& FormatData = GetImageFormatData((EImageFormat)args.format);
-					//int MinBlockSize = FMath::Max(FormatData.PixelsPerBlockX, FormatData.PixelsPerBlockY);
+					//int32 MinBlockSize = FMath::Max(FormatData.PixelsPerBlockX, FormatData.PixelsPerBlockY);
 					//ReducedBlockSizeInPixels.X = FMath::Max<int32>(ReducedBlockSizeInPixels.X, FormatData.PixelsPerBlockX);
 					//ReducedBlockSizeInPixels.Y = FMath::Max<int32>(ReducedBlockSizeInPixels.Y, FormatData.PixelsPerBlockY);
 					//FIntPoint ReducedImageSizeInPixels = SizeInBlocks * ReducedBlockSizeInPixels;
@@ -3803,7 +3803,7 @@ namespace mu
                 {
                     if (Args.mipmapCount == 0)
                     {
-						MipsToGenerate = Image::GetMipmapCount(ImageSizeInPixels.X, ImageSizeInPixels.Y);
+						MipsToGenerate = FImage::GetMipmapCount(ImageSizeInPixels.X, ImageSizeInPixels.Y);
                     }
                     else
                     {
@@ -3812,7 +3812,7 @@ namespace mu
                 }
 
 				// It needs to be initialized in case it has gaps.
-                Ptr<Image> New = CreateImage(ImageSizeInPixels.X, ImageSizeInPixels.Y, MipsToGenerate, EImageFormat(Args.format), EInitializationType::Black );
+                TSharedPtr<FImage> New = CreateImage(ImageSizeInPixels.X, ImageSizeInPixels.Y, MipsToGenerate, EImageFormat(Args.format), EInitializationType::Black );
                 StoreImage(item, New);
                 break;
             }
@@ -3825,7 +3825,7 @@ namespace mu
             break;
         }
 
-        case OP_TYPE::IM_COMPOSE:
+        case EOpType::IM_COMPOSE:
         {
 			OP::ImageComposeArgs Args = Program.GetOpArgs<OP::ImageComposeArgs>(item.At);
 
@@ -3838,12 +3838,12 @@ namespace mu
 					break;
 				case 1:
 				{
-					Ptr<const Layout> ComposeLayout = 
+					TSharedPtr<const FLayout> ComposeLayout = 
 							LoadLayout(FCacheAddress(Args.layout, FScheduledOp::FromOpAndOptions(Args.layout, item, 0)));
 
 					FScheduledOpData Data;
-					Data.Resource = const_cast<Layout*>(ComposeLayout.get());
-					int32 DataPos = m_heapData.Add(Data);
+					Data.Resource = ComposeLayout;
+					int32 DataPos = HeapData.Add(Data);
 
 					int32 RelBlockIndex = ComposeLayout->FindBlock(Args.BlockId);
 
@@ -3886,12 +3886,12 @@ namespace mu
 
 				case 1:
 				{	
-					Ptr<const Layout> ComposeLayout = 
+					TSharedPtr<const FLayout> ComposeLayout = 
 							LoadLayout(FCacheAddress(Args.layout, FScheduledOp::FromOpAndOptions(Args.layout, item, 0)));
 
 					FScheduledOpData Data;
-					Data.Resource = const_cast<Layout*>(ComposeLayout.get());
-					int32 DataPos = m_heapData.Add(Data);
+					Data.Resource = ComposeLayout;
+					int32 DataPos = HeapData.Add(Data);
 
 					int32 RelBlockIndex = ComposeLayout->FindBlock(Args.BlockId);
 					if (RelBlockIndex >= 0)
@@ -3921,7 +3921,7 @@ namespace mu
             break;
         }
 
-        case OP_TYPE::IM_INTERPOLATE:
+        case EOpType::IM_INTERPOLATE:
         {
 			OP::ImageInterpolateArgs args = Program.GetOpArgs<OP::ImageInterpolateArgs>(item.At);
             switch (item.Stage)
@@ -3936,8 +3936,8 @@ namespace mu
            		MUTABLE_CPUPROFILER_SCOPE(IM_INTERPOLATE_1)
             	
                 // Targets must be consecutive
-                int count = 0;
-                for ( int i=0
+                int32 count = 0;
+                for ( int32 i=0
                     ; i<MUTABLE_OP_MAX_INTERPOLATE_COUNT && args.targets[i]
                     ; ++i )
                 {
@@ -3947,8 +3947,8 @@ namespace mu
                 float factor = LoadScalar( FCacheAddress(args.factor,item) );
 
                 float delta = 1.0f/(count-1);
-                int min = (int)floorf( factor/delta );
-                int max = (int)ceilf( factor/delta );
+                int32 min = (int32)floorf( factor/delta );
+                int32 max = (int32)ceilf( factor/delta );
 
                 float bifactor = factor/delta - min;
 
@@ -3956,7 +3956,7 @@ namespace mu
                 data.Interpolate.Bifactor = bifactor;
 				data.Interpolate.Min = FMath::Clamp(min, 0, count - 1);
 				data.Interpolate.Max = FMath::Clamp(max, 0, count - 1);
-				uint32 dataPos = uint32(m_heapData.Add(data));
+				uint32 dataPos = uint32(HeapData.Add(data));
 
                 if ( bifactor < UE_SMALL_NUMBER )
                 {
@@ -3982,8 +3982,8 @@ namespace mu
            		MUTABLE_CPUPROFILER_SCOPE(IM_INTERPOLATE_2)
             		
                 // Targets must be consecutive
-                int count = 0;
-                for ( int i=0
+                int32 count = 0;
+                for ( int32 i=0
                     ; i<MUTABLE_OP_MAX_INTERPOLATE_COUNT && args.targets[i]
                     ; ++i )
                 {
@@ -3991,36 +3991,36 @@ namespace mu
                 }
 
                 // Factor from 0 to 1 between the two targets
-                const FScheduledOpData& Data = m_heapData[(size_t)item.CustomState];
+                const FScheduledOpData& Data = HeapData[item.CustomState];
                 float Bifactor = Data.Interpolate.Bifactor;
                 int32 Min = Data.Interpolate.Min;
                 int32 Max = Data.Interpolate.Max;
 
                 if (Bifactor < UE_SMALL_NUMBER)
                 {
-                    Ptr<const Image> Source = LoadImage(FCacheAddress(args.targets[Min], item));
+                    TSharedPtr<const FImage> Source = LoadImage(FCacheAddress(args.targets[Min], item));
 					StoreImage(item, Source);
 				}
                 else if (Bifactor > 1.0f - UE_SMALL_NUMBER)
                 {
-                    Ptr<const Image> Source = LoadImage(FCacheAddress(args.targets[Max], item));
+                    TSharedPtr<const FImage> Source = LoadImage(FCacheAddress(args.targets[Max], item));
 					StoreImage(item, Source);
 				}
                 else
                 {
-					Ptr<const Image> pMin = LoadImage(FCacheAddress(args.targets[Min], item));
-                    Ptr<const Image> pMax = LoadImage(FCacheAddress(args.targets[Max], item));
+					TSharedPtr<const FImage> pMin = LoadImage(FCacheAddress(args.targets[Min], item));
+                    TSharedPtr<const FImage> pMax = LoadImage(FCacheAddress(args.targets[Max], item));
 
                     if (pMin && pMax)
                     {						
-						Ptr<Image> pNew = CloneOrTakeOver(pMin);
+						TSharedPtr<FImage> pNew = CloneOrTakeOver(pMin);
 
 						// Be defensive: ensure image sizes match.
 						if (pNew->GetSize() != pMax->GetSize())
 						{
 							MUTABLE_CPUPROFILER_SCOPE(ImageResize_ForInterpolate);
-							Ptr<Image> Resized = CreateImage(pNew->GetSizeX(), pNew->GetSizeY(), pMax->GetLODCount(), pMax->GetFormat(), EInitializationType::NotInitialized);
-							ImOp.ImageResizeLinear(Resized.get(), 0, pMax.get());
+							TSharedPtr<FImage> Resized = CreateImage(pNew->GetSizeX(), pNew->GetSizeY(), pMax->GetLODCount(), pMax->GetFormat(), EInitializationType::NotInitialized);
+							ImOp.ImageResizeLinear(Resized.Get(), 0, pMax.Get());
 							Release(pMax);
 							pMax = Resized;
 						}
@@ -4030,10 +4030,10 @@ namespace mu
 						{
 							MUTABLE_CPUPROFILER_SCOPE(Format_ForInterpolate);
 
-							Ptr<Image> Formatted = CreateImage(pMax->GetSizeX(), pMax->GetSizeY(), pMax->GetLODCount(), pNew->GetFormat(), EInitializationType::NotInitialized);
+							TSharedPtr<FImage> Formatted = CreateImage(pMax->GetSizeX(), pMax->GetSizeY(), pMax->GetLODCount(), pNew->GetFormat(), EInitializationType::NotInitialized);
 							
 							bool bSuccess = false;
-							ImOp.ImagePixelFormat(bSuccess, m_pSettings->ImageCompressionQuality, Formatted.get(), pMax.get());
+							ImOp.ImagePixelFormat(bSuccess, Settings.ImageCompressionQuality, Formatted.Get(), pMax.Get());
 							check(bSuccess);
 							
 							Release(pMax);
@@ -4050,8 +4050,8 @@ namespace mu
 							// pNew is local owned, no need to CloneOrTakeOver.
 							pNew->DataStorage.SetNumLODs(LevelCount);
 
-							FMipmapGenerationSettings Settings{};
-							ImOp.ImageMipmap(m_pSettings->ImageCompressionQuality, pNew.get(), pNew.get(), StartLevel, LevelCount, Settings);
+							FMipmapGenerationSettings MipSettings{};
+							ImOp.ImageMipmap(Settings.ImageCompressionQuality, pNew.Get(), pNew.Get(), StartLevel, LevelCount, MipSettings);
 
 						}
 
@@ -4061,16 +4061,16 @@ namespace mu
 
 							int32 StartLevel = pMax->GetLODCount() - 1;
 
-							Ptr<Image> MaxFix = CloneOrTakeOver(pMax);
+							TSharedPtr<FImage> MaxFix = CloneOrTakeOver(pMax);
 							MaxFix->DataStorage.SetNumLODs(LevelCount);
 							
-							FMipmapGenerationSettings Settings{};
-							ImOp.ImageMipmap(m_pSettings->ImageCompressionQuality, MaxFix.get(), MaxFix.get(), StartLevel, LevelCount, Settings);
+							FMipmapGenerationSettings MipSettings{};
+							ImOp.ImageMipmap(Settings.ImageCompressionQuality, MaxFix.Get(), MaxFix.Get(), StartLevel, LevelCount, MipSettings);
 
 							pMax = MaxFix;
 						}
 
-                        ImageInterpolate(pNew.get(), pMax.get(), Bifactor);
+                        ImageInterpolate(pNew.Get(), pMax.Get(), Bifactor);
 
 						Release(pMax);
 						StoreImage(item, pNew);
@@ -4095,7 +4095,7 @@ namespace mu
             break;
         }
 
-        case OP_TYPE::IM_SATURATE:
+        case EOpType::IM_SATURATE:
         {
 			OP::ImageSaturateArgs args = Program.GetOpArgs<OP::ImageSaturateArgs>(item.At);
             switch (item.Stage)
@@ -4119,7 +4119,7 @@ namespace mu
             break;
         }
 
-        case OP_TYPE::IM_LUMINANCE:
+        case EOpType::IM_LUMINANCE:
         {
 			OP::ImageLuminanceArgs args = Program.GetOpArgs<OP::ImageLuminanceArgs>(item.At);
             switch (item.Stage)
@@ -4133,10 +4133,10 @@ namespace mu
             {
            		MUTABLE_CPUPROFILER_SCOPE(IM_LUMINANCE_1)
             		
-                Ptr<const Image> Base = LoadImage( FCacheAddress(args.base,item) );
+                TSharedPtr<const FImage> Base = LoadImage( FCacheAddress(args.base,item) );
 
-				Ptr<Image> Result = CreateImage(Base->GetSizeX(), Base->GetSizeY(), Base->GetLODCount(), EImageFormat::L_UByte, EInitializationType::NotInitialized);
-                ImageLuminance( Result.get(),Base.get() );
+				TSharedPtr<FImage> Result = CreateImage(Base->GetSizeX(), Base->GetSizeY(), Base->GetLODCount(), EImageFormat::L_UByte, EInitializationType::NotInitialized);
+                ImageLuminance( Result.Get(),Base.Get() );
 
 				Release(Base);
 				StoreImage( item, Result );
@@ -4150,7 +4150,7 @@ namespace mu
             break;
         }
 
-        case OP_TYPE::IM_SWIZZLE:
+        case EOpType::IM_SWIZZLE:
         {
 			OP::ImageSwizzleArgs args = Program.GetOpArgs<OP::ImageSwizzleArgs>(item.At);
             switch (item.Stage)
@@ -4175,7 +4175,7 @@ namespace mu
             break;
         }
 
-        case OP_TYPE::IM_COLOURMAP:
+        case EOpType::IM_COLOURMAP:
         {
 			OP::ImageColourMapArgs args = Program.GetOpArgs<OP::ImageColourMapArgs>(item.At);
             switch (item.Stage)
@@ -4191,9 +4191,9 @@ namespace mu
             {
             	MUTABLE_CPUPROFILER_SCOPE(IM_COLOURMAP_1)
             		
-                Ptr<const Image> Source = LoadImage( FCacheAddress(args.base,item) );
-                Ptr<const Image> Mask = LoadImage( FCacheAddress(args.mask,item) );
-                Ptr<const Image> Map = LoadImage( FCacheAddress(args.map,item) );
+                TSharedPtr<const FImage> Source = LoadImage( FCacheAddress(args.base,item) );
+                TSharedPtr<const FImage> Mask = LoadImage( FCacheAddress(args.mask,item) );
+                TSharedPtr<const FImage> Map = LoadImage( FCacheAddress(args.map,item) );
 
 				bool bOnlyOneMip = (Mask->GetLODCount() < Source->GetLODCount());
 
@@ -4201,20 +4201,20 @@ namespace mu
 				if (Mask->GetSize() != Source->GetSize())
 				{
 					MUTABLE_CPUPROFILER_SCOPE(ImageResize_ForColourmap);
-					Ptr<Image> Resized = CreateImage(Source->GetSizeX(), Source->GetSizeY(), 1, Mask->GetFormat(), EInitializationType::NotInitialized);
-					ImOp.ImageResizeLinear(Resized.get(), 0, Mask.get());
+					TSharedPtr<FImage> Resized = CreateImage(Source->GetSizeX(), Source->GetSizeY(), 1, Mask->GetFormat(), EInitializationType::NotInitialized);
+					ImOp.ImageResizeLinear(Resized.Get(), 0, Mask.Get());
 					Release(Mask);
 					Mask = Resized;
 				}
 
-				Ptr<Image> Result = CreateImage(Source->GetSizeX(), Source->GetSizeY(), Source->GetLODCount(), Source->GetFormat(), EInitializationType::NotInitialized);
-				ImageColourMap( Result.get(), Source.get(), Mask.get(), Map.get(), bOnlyOneMip);
+				TSharedPtr<FImage> Result = CreateImage(Source->GetSizeX(), Source->GetSizeY(), Source->GetLODCount(), Source->GetFormat(), EInitializationType::NotInitialized);
+				ImageColourMap( Result.Get(), Source.Get(), Mask.Get(), Map.Get(), bOnlyOneMip);
 
 				if (bOnlyOneMip)
 				{
 					MUTABLE_CPUPROFILER_SCOPE(ImageColourMap_MipFix);
 					FMipmapGenerationSettings DummyMipSettings{};
-					ImageMipmapInPlace(m_pSettings->ImageCompressionQuality, Result.get(), DummyMipSettings);
+					ImageMipmapInPlace(Settings.ImageCompressionQuality, Result.Get(), DummyMipSettings);
 				}
 
 				Release(Source);
@@ -4231,7 +4231,7 @@ namespace mu
             break;
         }
 
-        case OP_TYPE::IM_GRADIENT:
+        case EOpType::IM_GRADIENT:
         {
 			OP::ImageGradientArgs args = Program.GetOpArgs<OP::ImageGradientArgs>(item.At);
             switch (item.Stage)
@@ -4249,10 +4249,10 @@ namespace mu
 				FVector4f colour0 = LoadColor(FScheduledOp::FromOpAndOptions(args.colour0, item, 0));
 				FVector4f colour1 = LoadColor(FScheduledOp::FromOpAndOptions(args.colour1, item, 0));
 
-				ImagePtr pResult = CreateImage(args.size[0], args.size[1], 1, EImageFormat::RGB_UByte, EInitializationType::NotInitialized);
-                ImageGradient( pResult.get(), colour0, colour1 );
+				TSharedPtr<FImage> Result = CreateImage(args.size[0], args.size[1], 1, EImageFormat::RGB_UByte, EInitializationType::NotInitialized);
+                ImageGradient( Result.Get(), colour0, colour1 );
 
-				StoreImage( item, pResult );
+				StoreImage( item, Result );
                 break;
             }
 
@@ -4263,7 +4263,7 @@ namespace mu
             break;
         }
 
-        case OP_TYPE::IM_BINARISE:
+        case EOpType::IM_BINARISE:
         {
 			OP::ImageBinariseArgs args = Program.GetOpArgs<OP::ImageBinariseArgs>(item.At);
             switch (item.Stage)
@@ -4278,12 +4278,12 @@ namespace mu
             {
             	MUTABLE_CPUPROFILER_SCOPE(IM_BINARISE_1)
             		
-                Ptr<const Image> pA = LoadImage( FCacheAddress(args.base,item) );
+                TSharedPtr<const FImage> pA = LoadImage( FCacheAddress(args.base,item) );
 
                 float c = LoadScalar(FScheduledOp::FromOpAndOptions(args.threshold, item, 0));
 
-                Ptr<Image> Result = CreateImage(pA->GetSizeX(), pA->GetSizeY(), pA->GetLODCount(), EImageFormat::L_UByte, EInitializationType::NotInitialized);
-				ImageBinarise( Result.get(), pA.get(), c );
+                TSharedPtr<FImage> Result = CreateImage(pA->GetSizeX(), pA->GetSizeY(), pA->GetLODCount(), EImageFormat::L_UByte, EInitializationType::NotInitialized);
+				ImageBinarise( Result.Get(), pA.Get(), c );
 
 				Release(pA);
 				StoreImage( item, Result );
@@ -4297,7 +4297,7 @@ namespace mu
             break;
         }
 
-		case OP_TYPE::IM_INVERT:
+		case EOpType::IM_INVERT:
 		{
 			OP::ImageInvertArgs args = Program.GetOpArgs<OP::ImageInvertArgs>(item.At);
 			switch (item.Stage)
@@ -4319,7 +4319,7 @@ namespace mu
 			break;
 		}
 
-        case OP_TYPE::IM_PLAINCOLOUR:
+        case EOpType::IM_PLAINCOLOUR:
         {
 			OP::ImagePlainColourArgs args = Program.GetOpArgs<OP::ImagePlainColourArgs>(item.At);
             switch (item.Stage)
@@ -4344,16 +4344,16 @@ namespace mu
 					LODs = FMath::CeilLogTwo(FMath::Max(SizeX,SizeY));
 				}
 
-				for (int l=0; l<item.ExecutionOptions; ++l)
+				for (int32 l=0; l<item.ExecutionOptions; ++l)
 				{
 					SizeX = FMath::Max(uint16(1), FMath::DivideAndRoundUp(SizeX, uint16(2)));
 					SizeY = FMath::Max(uint16(1), FMath::DivideAndRoundUp(SizeY, uint16(2)));
 					--LODs;
 				}
 
-                ImagePtr pA = CreateImage( SizeX, SizeY, FMath::Max(LODs,1), EImageFormat(args.format), EInitializationType::NotInitialized );
+                TSharedPtr<FImage> pA = CreateImage( SizeX, SizeY, FMath::Max(LODs,1), EImageFormat(args.format), EInitializationType::NotInitialized );
 
-				ImOp.FillColor(pA.get(), c);
+				ImOp.FillColor(pA.Get(), c);
 
 				StoreImage( item, pA );
                 break;
@@ -4366,14 +4366,14 @@ namespace mu
             break;
         }
 
-		case OP_TYPE::IM_REFERENCE:
+		case EOpType::IM_REFERENCE:
 		{
 			OP::ResourceReferenceArgs Args = Program.GetOpArgs<OP::ResourceReferenceArgs>(item.At);
 			switch (item.Stage)
 			{
 			case 0:
 			{
-				Ptr<Image> Result;
+				TSharedPtr<FImage> Result;
 				if (Args.ForceLoad)
 				{
 					// This should never be reached because it should have been caught as a Task in IssueOp
@@ -4381,7 +4381,7 @@ namespace mu
 				}
 				else
 				{
-					Result = Image::CreateAsReference(Args.ID, Args.ImageDesc, false);
+					Result = FImage::CreateAsReference(Args.ID, Args.ImageDesc, false);
 				}
 				StoreImage(item, Result);
 				break;
@@ -4394,7 +4394,7 @@ namespace mu
 			break;
 		}
 
-        case OP_TYPE::IM_CROP:
+        case EOpType::IM_CROP:
         {
 			OP::ImageCropArgs args = Program.GetOpArgs<OP::ImageCropArgs>(item.At);
             switch (item.Stage)
@@ -4407,7 +4407,7 @@ namespace mu
             {
             	MUTABLE_CPUPROFILER_SCOPE(IM_CROP_1)
             		
-                Ptr<const Image> pA = LoadImage( FCacheAddress(args.source,item) );
+                TSharedPtr<const FImage> pA = LoadImage( FCacheAddress(args.source,item) );
 
                 box< UE::Math::TIntVector2<int32> > rect;
                 rect.min[0] = args.minX;
@@ -4426,15 +4426,15 @@ namespace mu
 					MipsToSkip--;
 				}
 
-				ImagePtr pResult;
+				TSharedPtr<FImage> Result;
 				if (!rect.IsEmpty())
 				{
-					pResult = CreateImage( rect.size[0], rect.size[1], 1, pA->GetFormat(), EInitializationType::NotInitialized);
-					ImOp.ImageCrop(pResult.get(), m_pSettings->ImageCompressionQuality, pA.get(), rect);
+					Result = CreateImage( rect.size[0], rect.size[1], 1, pA->GetFormat(), EInitializationType::NotInitialized);
+					ImOp.ImageCrop(Result, Settings.ImageCompressionQuality, pA.Get(), rect);
 				}
 
 				Release(pA);
-				StoreImage( item, pResult );
+				StoreImage( item, Result );
                 break;
             }
 
@@ -4445,7 +4445,7 @@ namespace mu
             break;
         }
 
-        case OP_TYPE::IM_PATCH:
+        case EOpType::IM_PATCH:
         {
 			// TODO: This is optimized for memory-usage but base and patch could be requested at the same time
 			OP::ImagePatchArgs args = Program.GetOpArgs<OP::ImagePatchArgs>(item.At);
@@ -4463,8 +4463,8 @@ namespace mu
             {
                 MUTABLE_CPUPROFILER_SCOPE(IM_PATCH_1)
 
-                Ptr<const Image> pA = LoadImage( FCacheAddress(args.base,item) );
-                Ptr<const Image> pB = LoadImage( FCacheAddress(args.patch,item) );
+                TSharedPtr<const FImage> pA = LoadImage( FCacheAddress(args.base,item) );
+                TSharedPtr<const FImage> pB = LoadImage( FCacheAddress(args.patch,item) );
 
 				// Failsafe
 				if (!pA || !pB)
@@ -4482,18 +4482,18 @@ namespace mu
 				rect.size[0] = pB->GetSizeX();
 				rect.size[1] = pB->GetSizeY();
 
-                ImagePtr pResult = CloneOrTakeOver(pA);
+                TSharedPtr<FImage> Result = CloneOrTakeOver(pA);
 
 				bool bApplyPatch = !rect.IsEmpty();
 				if (bApplyPatch)
 				{
 					// Change the block image format if it doesn't match the composed image
 					// This is usually enforced at object compilation time.
-					if (pResult->GetFormat() != pB->GetFormat())
+					if (Result->GetFormat() != pB->GetFormat())
 					{
 						MUTABLE_CPUPROFILER_SCOPE(ImagPatchReformat);
 
-						EImageFormat format = GetMostGenericFormat(pResult->GetFormat(), pB->GetFormat());
+						EImageFormat format = GetMostGenericFormat(Result->GetFormat(), pB->GetFormat());
 
 						const FImageFormatData& finfo = GetImageFormatData(format);
 						if (finfo.PixelsPerBlockX == 0)
@@ -4501,20 +4501,20 @@ namespace mu
 							format = GetUncompressedFormat(format);
 						}
 
-						if (pResult->GetFormat() != format)
+						if (Result->GetFormat() != format)
 						{
-							Ptr<Image> Formatted = CreateImage(pResult->GetSizeX(), pResult->GetSizeY(), pResult->GetLODCount(), format, EInitializationType::NotInitialized);
+							TSharedPtr<FImage> Formatted = CreateImage(Result->GetSizeX(), Result->GetSizeY(), Result->GetLODCount(), format, EInitializationType::NotInitialized);
 							bool bSuccess = false;
-							ImOp.ImagePixelFormat(bSuccess, m_pSettings->ImageCompressionQuality, Formatted.get(), pResult.get());
+							ImOp.ImagePixelFormat(bSuccess, Settings.ImageCompressionQuality, Formatted.Get(), Result.Get());
 							check(bSuccess);
-							Release(pResult);
-							pResult = Formatted;
+							Release(Result);
+							Result = Formatted;
 						}
 						if (pB->GetFormat() != format)
 						{
-							Ptr<Image> Formatted = CreateImage(pB->GetSizeX(), pB->GetSizeY(), pB->GetLODCount(), format, EInitializationType::NotInitialized);
+							TSharedPtr<FImage> Formatted = CreateImage(pB->GetSizeX(), pB->GetSizeY(), pB->GetLODCount(), format, EInitializationType::NotInitialized);
 							bool bSuccess = false;
-							ImOp.ImagePixelFormat(bSuccess, m_pSettings->ImageCompressionQuality, Formatted.get(), pB.get());
+							ImOp.ImagePixelFormat(bSuccess, Settings.ImageCompressionQuality, Formatted.Get(), pB.Get());
 							check(bSuccess);
 							Release(pB);
 							pB = Formatted;
@@ -4522,31 +4522,31 @@ namespace mu
 					}
 
 					// Don't patch if below the image compression block size.
-					const FImageFormatData& finfo = GetImageFormatData(pResult->GetFormat());
+					const FImageFormatData& finfo = GetImageFormatData(Result->GetFormat());
 					bApplyPatch =
 						(rect.min[0] % finfo.PixelsPerBlockX == 0) &&
 						(rect.min[1] % finfo.PixelsPerBlockY == 0) &&
 						(rect.size[0] % finfo.PixelsPerBlockX == 0) &&
 						(rect.size[1] % finfo.PixelsPerBlockY == 0) &&
-						(rect.min[0] + rect.size[0]) <= pResult->GetSizeX() &&
-						(rect.min[1] + rect.size[1]) <= pResult->GetSizeY()
+						(rect.min[0] + rect.size[0]) <= Result->GetSizeX() &&
+						(rect.min[1] + rect.size[1]) <= Result->GetSizeY()
 						;
 				}
 
 				if (bApplyPatch)
 				{
-					ImOp.ImageCompose(pResult.get(), pB.get(), rect);
-					pResult->m_flags = 0;
+					ImOp.ImageCompose(Result.Get(), pB.Get(), rect);
+					Result->Flags = 0;
 				}
 				else
 				{
 					// This happens very often when skipping mips, and floods the log.
 					//UE_LOG( LogMutableCore, Verbose, TEXT("Skipped patch operation for image not fitting the block compression size. Small image? Patch rect is (%d, %d), (%d, %d), base is (%d, %d)"),
-					//	rect.min[0], rect.min[1], rect.size[0], rect.size[1], pResult->GetSizeX(), pResult->GetSizeY());
+					//	rect.min[0], rect.min[1], rect.size[0], rect.size[1], Result->GetSizeX(), Result->GetSizeY());
 				}
 
 				Release(pB);
-				StoreImage( item, pResult );
+				StoreImage( item, Result );
                 break;
             }
 
@@ -4557,7 +4557,7 @@ namespace mu
             break;
         }
 
-        case OP_TYPE::IM_RASTERMESH:
+        case EOpType::IM_RASTERMESH:
         {
 			OP::ImageRasterMeshArgs args = Program.GetOpArgs<OP::ImageRasterMeshArgs>(item.At);
             switch (item.Stage)
@@ -4580,7 +4580,7 @@ namespace mu
 			{
 				MUTABLE_CPUPROFILER_SCOPE(IM_RASTERMESH_1)
 
-				Ptr<const Mesh> pMesh = LoadMesh(FScheduledOp::FromOpAndOptions(args.mesh, item, 0));
+				TSharedPtr<const FMesh> Mesh = LoadMesh(FScheduledOp::FromOpAndOptions(args.mesh, item, 0));
 
 				// If no image, we are generating a flat mesh UV raster. This is the final stage in this case.
 				if (!args.image)
@@ -4608,11 +4608,11 @@ namespace mu
 					}
 
                     // Flat mesh UV raster
-					Ptr<Image> ResultImage = CreateImage(SizeX, SizeY, 1, EImageFormat::L_UByte, EInitializationType::Black);
-					if (pMesh)
+					TSharedPtr<FImage> ResultImage = CreateImage(SizeX, SizeY, 1, EImageFormat::L_UByte, EInitializationType::Black);
+					if (Mesh)
 					{
-						ImageRasterMesh(pMesh.get(), ResultImage.get(), args.LayoutIndex, args.BlockId, CropMin, UncroppedSize);
-						Release(pMesh);
+						ImageRasterMesh(Mesh.Get(), ResultImage.Get(), args.LayoutIndex, args.BlockId, CropMin, UncroppedSize);
+						Release(Mesh);
 					}
 
 					// Stop execution.
@@ -4649,19 +4649,19 @@ namespace mu
 						FMath::Max(args.sizeY >> MipsToSkip, 1));
 					FVector2f SourceImageSizeF = FVector2f(args.SourceSizeX, args.SourceSizeY);
 						
-					if (pMesh)
+					if (Mesh)
 					{ 
-						const float ComputedMip = ComputeProjectedFootprintBestMip(pMesh.get(), Projector, TargetImageSizeF, SourceImageSizeF);
+						const float ComputedMip = ComputeProjectedFootprintBestMip(Mesh.Get(), Projector, TargetImageSizeF, SourceImageSizeF);
 
 						Data.RasterMesh.MipValue = FMath::Max(0.0f, ComputedMip + GlobalProjectionLodBias);
 						Data.RasterMesh.Mip = static_cast<uint8>(FMath::FloorToInt32(Data.RasterMesh.MipValue));
 					}
 				}
 		
-				const int32 DataHeapAddress = m_heapData.Add(Data);
+				const int32 DataHeapAddress = HeapData.Add(Data);
 
-				// pMesh is need again in the next stage, store it in the heap.
-				m_heapData[DataHeapAddress].Resource = const_cast<Mesh*>(pMesh.get());
+				// Mesh is need again in the next stage, store it in the heap.
+				HeapData[DataHeapAddress].Resource = Mesh;
 
 				AddOp(FScheduledOp(item.At, item, 2, DataHeapAddress),
 					FScheduledOp::FromOpAndOptions(args.projector, item, 0),
@@ -4684,13 +4684,13 @@ namespace mu
 					break;
 				}
 
-				FScheduledOpData& Data = m_heapData[item.CustomState];
+				FScheduledOpData& Data = HeapData[item.CustomState];
 
 				// Unsafe downcast, should be fine as it is known to be a Mesh.
-				Ptr<const Mesh> pMesh = static_cast<Mesh*>(Data.Resource.get());
+				TSharedPtr<const FMesh> Mesh = StaticCastSharedPtr<FMesh>(ConstCastSharedPtr<FResource>(Data.Resource));
 				Data.Resource = nullptr;
 
-				if (!pMesh)
+				if (!Mesh)
 				{
 					check(false);
 					StoreImage(item, nullptr);
@@ -4720,9 +4720,9 @@ namespace mu
 				}
 
 				// Raster with projection
-				Ptr<const Image> Source = LoadImage(FCacheAddress(args.image, item.ExecutionIndex, Data.RasterMesh.Mip));
+				TSharedPtr<const FImage> Source = LoadImage(FCacheAddress(args.image, item.ExecutionIndex, Data.RasterMesh.Mip));
 
-				Ptr<const Image> Mask = nullptr;
+				TSharedPtr<const FImage> Mask = nullptr;
 				if (args.mask)
 				{
 					Mask = LoadImage(FCacheAddress(args.mask, item));
@@ -4733,8 +4733,8 @@ namespace mu
 					{
 						MUTABLE_CPUPROFILER_SCOPE(ImageResize_MaskFixForProjection);
 
-						Ptr<Image> Resized = CreateImage(SizeX, SizeY, Mask->GetLODCount(), Mask->GetFormat(), EInitializationType::NotInitialized);
-						ImOp.ImageResizeLinear(Resized.get(), 0, Mask.get());
+						TSharedPtr<FImage> Resized = CreateImage(SizeX, SizeY, Mask->GetLODCount(), Mask->GetFormat(), EInitializationType::NotInitialized);
+						ImOp.ImageResizeLinear(Resized.Get(), 0, Mask.Get());
 						Release(Mask);
 						Mask = Resized;
 					}
@@ -4756,9 +4756,9 @@ namespace mu
 				if (Source && Source->GetFormat()!=Format)
 				{
 					MUTABLE_CPUPROFILER_SCOPE(RunCode_RasterMesh_ReformatSource);
-					Ptr<Image> Formatted = CreateImage(Source->GetSizeX(), Source->GetSizeY(), Source->GetLODCount(), Format, EInitializationType::NotInitialized);
+					TSharedPtr<FImage> Formatted = CreateImage(Source->GetSizeX(), Source->GetSizeY(), Source->GetLODCount(), Format, EInitializationType::NotInitialized);
 					bool bSuccess = false;
-					ImOp.ImagePixelFormat(bSuccess, m_pSettings->ImageCompressionQuality, Formatted.get(), Source.get());
+					ImOp.ImagePixelFormat(bSuccess, Settings.ImageCompressionQuality, Formatted.Get(), Source.Get());
 					check(bSuccess); 
 					Release(Source);
 					Source = Formatted;
@@ -4789,8 +4789,8 @@ namespace mu
 					{
 						MUTABLE_CPUPROFILER_SCOPE(RunCode_ImageRasterMesh_SizeFixup);	
 						
-						Ptr<Image> Resized = CreateImage(ExpectedSourceSize.X, ExpectedSourceSize.Y, 1, Format, EInitializationType::NotInitialized);
-						ImOp.ImageResizeLinear(Resized.get(), m_pSettings->ImageCompressionQuality, Source.get());
+						TSharedPtr<FImage> Resized = CreateImage(ExpectedSourceSize.X, ExpectedSourceSize.Y, 1, Format, EInitializationType::NotInitialized);
+						ImOp.ImageResizeLinear(Resized.Get(), Settings.ImageCompressionQuality, Source.Get());
 						
 						Release(Source);
 						Source = Resized;	
@@ -4799,8 +4799,8 @@ namespace mu
 
 				// Allocate memory for the temporary buffers
 				FScratchImageProject Scratch;
-				Scratch.Vertices.SetNum(pMesh->GetVertexCount());
-				Scratch.CulledVertex.SetNum(pMesh->GetVertexCount());
+				Scratch.Vertices.SetNum(Mesh->GetVertexCount());
+				Scratch.CulledVertex.SetNum(Mesh->GetVertexCount());
 
 				ESamplingMethod SamplingMethod = Invoke([&]() -> ESamplingMethod
 				{
@@ -4822,17 +4822,17 @@ namespace mu
 					{
 						MUTABLE_CPUPROFILER_SCOPE(RunCode_RasterMesh_BilinearMipGen);
 
-						Ptr<Image> OwnedSource = CloneOrTakeOver(Source);
+						TSharedPtr<FImage> OwnedSource = CloneOrTakeOver(Source);
 
 						OwnedSource->DataStorage.SetNumLODs(2);
-						ImageMipmapInPlace(0, OwnedSource.get(), FMipmapGenerationSettings{});
+						ImageMipmapInPlace(0, OwnedSource.Get(), FMipmapGenerationSettings{});
 
 						Source = OwnedSource;
 					}
 				}
 
 				// Allocate new image after bilinear mip generation to reduce operation memory peak.
-				Ptr<Image> New = CreateImage(SizeX, SizeY, 1, Format, EInitializationType::Black);
+				TSharedPtr<FImage> New = CreateImage(SizeX, SizeY, 1, Format, EInitializationType::Black);
 
 				if (args.projector && Source && Source->GetSizeX() > 0 && Source->GetSizeY() > 0)
 				{
@@ -4841,8 +4841,8 @@ namespace mu
 					switch (Projector.type)
 					{
 					case EProjectorType::Planar:
-						ImageRasterProjectedPlanar(pMesh.get(), New.get(),
-							Source.get(), Mask.get(),
+						ImageRasterProjectedPlanar(Mesh.Get(), New.Get(),
+							Source.Get(), Mask.Get(),
 							args.bIsRGBFadingEnabled, args.bIsAlphaFadingEnabled,
 							SamplingMethod,
 							FadeStartRad, FadeEndRad, FMath::Frac(Data.RasterMesh.MipValue),
@@ -4852,8 +4852,8 @@ namespace mu
 						break;
 
 					case EProjectorType::Wrapping:
-						ImageRasterProjectedWrapping(pMesh.get(), New.get(),
-							Source.get(), Mask.get(),
+						ImageRasterProjectedWrapping(Mesh.Get(), New.Get(),
+							Source.Get(), Mask.Get(),
 							args.bIsRGBFadingEnabled, args.bIsAlphaFadingEnabled,
 							SamplingMethod,
 							FadeStartRad, FadeEndRad, FMath::Frac(Data.RasterMesh.MipValue),
@@ -4863,8 +4863,8 @@ namespace mu
 						break;
 
 					case EProjectorType::Cylindrical:
-						ImageRasterProjectedCylindrical(pMesh.get(), New.get(),
-							Source.get(), Mask.get(),
+						ImageRasterProjectedCylindrical(Mesh.Get(), New.Get(),
+							Source.Get(), Mask.Get(),
 							args.bIsRGBFadingEnabled, args.bIsAlphaFadingEnabled,
 							SamplingMethod,
 							FadeStartRad, FadeEndRad, FMath::Frac(Data.RasterMesh.MipValue),
@@ -4880,7 +4880,7 @@ namespace mu
 					}
 				}
 
-				Release(pMesh);
+				Release(Mesh);
 				Release(Source);
 				Release(Mask);
 				StoreImage(item, New);
@@ -4895,7 +4895,7 @@ namespace mu
             break;
         }
 
-        case OP_TYPE::IM_MAKEGROWMAP:
+        case EOpType::IM_MAKEGROWMAP:
         {
 			OP::ImageMakeGrowMapArgs args = Program.GetOpArgs<OP::ImageMakeGrowMapArgs>(item.At);
             switch (item.Stage)
@@ -4908,12 +4908,12 @@ namespace mu
             {
                 MUTABLE_CPUPROFILER_SCOPE(IM_MAKEGROWMAP_1)
 
-                Ptr<const Image> Mask = LoadImage( FCacheAddress(args.mask,item) );
+                TSharedPtr<const FImage> Mask = LoadImage( FCacheAddress(args.mask,item) );
 
-                Ptr<Image> Result = CreateImage( Mask->GetSizeX(), Mask->GetSizeY(), Mask->GetLODCount(), EImageFormat::L_UByte, EInitializationType::NotInitialized);
+                TSharedPtr<FImage> Result = CreateImage( Mask->GetSizeX(), Mask->GetSizeY(), Mask->GetLODCount(), EImageFormat::L_UByte, EInitializationType::NotInitialized);
 
-                ImageMakeGrowMap(Result.get(), Mask.get(), args.border );
-				Result->m_flags |= Image::IF_CANNOT_BE_SCALED;
+                ImageMakeGrowMap(Result.Get(), Mask.Get(), args.border );
+				Result->Flags |= FImage::IF_CANNOT_BE_SCALED;
 
 				Release(Mask);
                 StoreImage( item, Result);
@@ -4927,7 +4927,7 @@ namespace mu
             break;
         }
 
-        case OP_TYPE::IM_DISPLACE:
+        case EOpType::IM_DISPLACE:
         {
 			OP::ImageDisplaceArgs args = Program.GetOpArgs<OP::ImageDisplaceArgs>(item.At);
             switch (item.Stage)
@@ -4942,8 +4942,8 @@ namespace mu
             {
                 MUTABLE_CPUPROFILER_SCOPE(IM_DISPLACE_1)
 
-                Ptr<const Image> Source = LoadImage( FCacheAddress(args.source,item) );
-                Ptr<const Image> pMap = LoadImage( FCacheAddress(args.displacementMap,item) );
+                TSharedPtr<const FImage> Source = LoadImage( FCacheAddress(args.source,item) );
+                TSharedPtr<const FImage> pMap = LoadImage( FCacheAddress(args.displacementMap,item) );
 
 				if (!Source)
 				{
@@ -4960,25 +4960,25 @@ namespace mu
 				{
 					MUTABLE_CPUPROFILER_SCOPE(ImageResize_EmergencyHackForDisplacementStep1);
 
-					Ptr<Image> Resized = CreateImage(pMap->GetSizeX(), pMap->GetSizeY(), Source->GetLODCount(), Source->GetFormat(), EInitializationType::NotInitialized);
-					ImOp.ImageResizeLinear(Resized.get(), 0, Source.get());
+					TSharedPtr<FImage> Resized = CreateImage(pMap->GetSizeX(), pMap->GetSizeY(), Source->GetLODCount(), Source->GetFormat(), EInitializationType::NotInitialized);
+					ImOp.ImageResizeLinear(Resized.Get(), 0, Source.Get());
 					Release(Source);
 					Source = Resized;
 				}
 
 				// This works based on the assumption that displacement maps never read from a position they actually write to.
 				// Since they are used for UV border expansion, this should always be the case.
-				Ptr<Image> Result = CloneOrTakeOver(Source);
+				TSharedPtr<FImage> Result = CloneOrTakeOver(Source);
 
 				if (OriginalSourceScale[0] > 0 && OriginalSourceScale[1] > 0)
 				{
-					ImageDisplace(Result.get(), Result.get(), pMap.get());
+					ImageDisplace(Result.Get(), Result.Get(), pMap.Get());
 
 					if (OriginalSourceScale != Result->GetSize())
 					{
 						MUTABLE_CPUPROFILER_SCOPE(ImageResize_EmergencyHackForDisplacementStep2);
-						Ptr<Image> Resized = CreateImage(OriginalSourceScale[0], OriginalSourceScale[1], Result->GetLODCount(), Result->GetFormat(), EInitializationType::NotInitialized);
-						ImOp.ImageResizeLinear(Resized.get(), 0, Result.get());
+						TSharedPtr<FImage> Resized = CreateImage(OriginalSourceScale[0], OriginalSourceScale[1], Result->GetLODCount(), Result->GetFormat(), EInitializationType::NotInitialized);
+						ImOp.ImageResizeLinear(Resized.Get(), 0, Result.Get());
 						Release(Result);
 						Result = Resized;
 					}
@@ -4996,7 +4996,7 @@ namespace mu
             break;
         }
 
-        case OP_TYPE::IM_TRANSFORM:
+        case EOpType::IM_TRANSFORM:
         {
             const OP::ImageTransformArgs Args = Program.GetOpArgs<OP::ImageTransformArgs>(item.At);
 
@@ -5071,16 +5071,16 @@ namespace mu
 					FMath::Log2(FMath::Max(1.0f, FMath::Square(SourceSize.GetMin()))) * 0.5f - 
 				    FMath::Log2(FMath::Max(1.0f, FMath::Square(ScaledSourceSize.GetMin()))) * 0.5f;
 
-				FScheduledOpData HeapData;
-				HeapData.ImageTransform.SizeX = DestSizeI.X;
-				HeapData.ImageTransform.SizeY = DestSizeI.Y;
-				FPlatformMath::StoreHalf(&HeapData.ImageTransform.ScaleXEncodedHalf, Scale.X);
-				FPlatformMath::StoreHalf(&HeapData.ImageTransform.ScaleYEncodedHalf, Scale.Y);
-				HeapData.ImageTransform.MipValue = BestMip + GlobalImageTransformLodBias;
+				FScheduledOpData OpHeapData;
+				OpHeapData.ImageTransform.SizeX = DestSizeI.X;
+				OpHeapData.ImageTransform.SizeY = DestSizeI.Y;
+				FPlatformMath::StoreHalf(&OpHeapData.ImageTransform.ScaleXEncodedHalf, Scale.X);
+				FPlatformMath::StoreHalf(&OpHeapData.ImageTransform.ScaleYEncodedHalf, Scale.Y);
+				OpHeapData.ImageTransform.MipValue = BestMip + GlobalImageTransformLodBias;
 
-				const int32 HeapDataAddress = m_heapData.Add(HeapData);
+				const int32 HeapDataAddress = HeapData.Add(OpHeapData);
 
-				const uint8 Mip = static_cast<uint8>(FMath::Max(0, FMath::FloorToInt(HeapData.ImageTransform.MipValue)));
+				const uint8 Mip = static_cast<uint8>(FMath::Max(0, FMath::FloorToInt(OpHeapData.ImageTransform.MipValue)));
 				const TArray<FScheduledOp, TFixedAllocator<4>> Deps = 
 				{
 					FScheduledOp::FromOpAndOptions(Args.Base, item, Mip),
@@ -5097,23 +5097,23 @@ namespace mu
             {
 				MUTABLE_CPUPROFILER_SCOPE(IM_TRANSFORM_2);
 			
-				const FScheduledOpData HeapData = m_heapData[item.CustomState];
+				const FScheduledOpData OpHeapData = HeapData[item.CustomState];
 
-				const uint8 Mip = static_cast<uint8>(FMath::Max(0, FMath::FloorToInt(HeapData.ImageTransform.MipValue)));
-				Ptr<const Image> Source = LoadImage(FCacheAddress(Args.Base, item.ExecutionIndex, Mip));
+				const uint8 Mip = static_cast<uint8>(FMath::Max(0, FMath::FloorToInt(OpHeapData.ImageTransform.MipValue)));
+				TSharedPtr<const FImage> Source = LoadImage(FCacheAddress(Args.Base, item.ExecutionIndex, Mip));
 
 				const FVector2f Offset = FVector2f(
                         Args.OffsetX ? LoadScalar(FCacheAddress(Args.OffsetX, item)) : 0.0f,
                         Args.OffsetY ? LoadScalar(FCacheAddress(Args.OffsetY, item)) : 0.0f);
 
                 FVector2f Scale = FVector2f(
-						FPlatformMath::LoadHalf(&HeapData.ImageTransform.ScaleXEncodedHalf),
-						FPlatformMath::LoadHalf(&HeapData.ImageTransform.ScaleYEncodedHalf));
+						FPlatformMath::LoadHalf(&OpHeapData.ImageTransform.ScaleXEncodedHalf),
+						FPlatformMath::LoadHalf(&OpHeapData.ImageTransform.ScaleYEncodedHalf));
 
 				FVector2f AspectCorrectionScale = FVector2f(1.0f, 1.0f);
 				if (Args.bKeepAspectRatio)
 				{
-					const FVector2f DestSize   = FVector2f(HeapData.ImageTransform.SizeX, HeapData.ImageTransform.SizeY);
+					const FVector2f DestSize   = FVector2f(OpHeapData.ImageTransform.SizeX, OpHeapData.ImageTransform.SizeY);
 					const FVector2f SourceSize = FVector2f(FMath::Max(Args.SourceSizeX, uint16(1)), FMath::Max(Args.SourceSizeY, uint16(1)));
 					
 					const float DestAspectOverSrcAspect = (DestSize.X * SourceSize.Y) / (DestSize.Y * SourceSize.X);
@@ -5132,9 +5132,9 @@ namespace mu
 				if (Format != SourceFormat)
 				{
 					MUTABLE_CPUPROFILER_SCOPE(RunCode_ImageTransform_FormatFixup);	
-					Ptr<Image> Formatted = CreateImage(Source->GetSizeX(), Source->GetSizeY(), Source->GetLODCount(), Format, EInitializationType::NotInitialized);
+					TSharedPtr<FImage> Formatted = CreateImage(Source->GetSizeX(), Source->GetSizeY(), Source->GetLODCount(), Format, EInitializationType::NotInitialized);
 					bool bSuccess = false;
-					ImOp.ImagePixelFormat(bSuccess, m_pSettings->ImageCompressionQuality, Formatted.get(), Source.get());
+					ImOp.ImagePixelFormat(bSuccess, Settings.ImageCompressionQuality, Formatted.Get(), Source.Get());
 					check(bSuccess); 
 
 					Release(Source);
@@ -5148,8 +5148,8 @@ namespace mu
 				{
 					MUTABLE_CPUPROFILER_SCOPE(RunCode_ImageTransform_SizeFixup);	
 					
-					Ptr<Image> Resized = CreateImage(ExpectedSourceSize.X, ExpectedSourceSize.Y, 1, Format, EInitializationType::NotInitialized);
-					ImOp.ImageResizeLinear(Resized.get(), m_pSettings->ImageCompressionQuality, Source.get());
+					TSharedPtr<FImage> Resized = CreateImage(ExpectedSourceSize.X, ExpectedSourceSize.Y, 1, Format, EInitializationType::NotInitialized);
+					ImOp.ImageResizeLinear(Resized.Get(), Settings.ImageCompressionQuality, Source.Get());
 					
 					Release(Source);
 					Source = Resized;	
@@ -5159,10 +5159,10 @@ namespace mu
 				{
 					MUTABLE_CPUPROFILER_SCOPE(RunCode_ImageTransform_BilinearMipGen);
 
-					Ptr<Image> OwnedSource = CloneOrTakeOver(Source);
+					TSharedPtr<FImage> OwnedSource = CloneOrTakeOver(Source);
 					OwnedSource->DataStorage.SetNumLODs(2);
 
-					ImageMipmapInPlace(0, OwnedSource.get(), FMipmapGenerationSettings{});
+					ImageMipmapInPlace(0, OwnedSource.Get(), FMipmapGenerationSettings{});
 
 					Source = OwnedSource;
 				}
@@ -5190,11 +5190,11 @@ namespace mu
 											       ? EInitializationType::Black
 											       : EInitializationType::NotInitialized;
 
-				Ptr<Image> Result = CreateImage(
-						HeapData.ImageTransform.SizeX, HeapData.ImageTransform.SizeY, 1, Format, InitType);
+				TSharedPtr<FImage> Result = CreateImage(
+						OpHeapData.ImageTransform.SizeX, OpHeapData.ImageTransform.SizeY, 1, Format, InitType);
 
-				const float MipFactor = FMath::Frac(FMath::Max(0.0f, HeapData.ImageTransform.MipValue));
-				ImageTransform(Result.get(), Source.get(), Transform, MipFactor, AddressMode, bUseImageTransformVectorImpl);
+				const float MipFactor = FMath::Frac(FMath::Max(0.0f, OpHeapData.ImageTransform.MipValue));
+				ImageTransform(Result.Get(), Source.Get(), Transform, MipFactor, AddressMode, bUseImageTransformVectorImpl);
 
 				Release(Source);
 				StoreImage(item, Result);
@@ -5210,7 +5210,7 @@ namespace mu
 		}
 
         default:
-            if (type!=OP_TYPE::NONE)
+            if (type!=EOpType::NONE)
             {
                 // Operation not implemented
                 check( false );
@@ -5220,7 +5220,7 @@ namespace mu
     }	
 	
 	//---------------------------------------------------------------------------------------------
-    Ptr<RangeIndex> CodeRunner::BuildCurrentOpRangeIndex( const FScheduledOp& item, const Parameters* pParams, const Model* pModel, int32 parameterIndex )
+    TSharedPtr<FRangeIndex> CodeRunner::BuildCurrentOpRangeIndex( const FScheduledOp& item, const FParameters* pParams, const FModel* pModel, int32 parameterIndex )
     {
         if (!item.ExecutionIndex)
         {
@@ -5229,39 +5229,39 @@ namespace mu
 
         // \todo: optimise to avoid allocating the index here, we could access internal
         // data directly.
-		Ptr<RangeIndex> index = pParams->NewRangeIndex( parameterIndex );
-        if (!index)
+		TSharedPtr<FRangeIndex> Index = pParams->NewRangeIndex( parameterIndex );
+        if (!Index)
         {
             return nullptr;
         }
 
-		const FProgram& Program = m_pModel->GetPrivate()->m_program;
-		const FParameterDesc& paramDesc = Program.m_parameters[ parameterIndex ];
-        for( size_t rangeIndexInParam=0;
-             rangeIndexInParam<paramDesc.m_ranges.Num();
+		const FProgram& Program = Model->GetPrivate()->Program;
+		const FParameterDesc& paramDesc = Program.Parameters[ parameterIndex ];
+        for( int32 rangeIndexInParam=0;
+             rangeIndexInParam<paramDesc.Ranges.Num();
              ++rangeIndexInParam )
         {
-            uint32 rangeIndexInModel = paramDesc.m_ranges[rangeIndexInParam];
+            uint32 rangeIndexInModel = paramDesc.Ranges[rangeIndexInParam];
             const ExecutionIndex& currentIndex = GetMemory().GetRangeIndex( item.ExecutionIndex );
-            int position = currentIndex.GetFromModelRangeIndex(rangeIndexInModel);
-            index->GetPrivate()->m_values[rangeIndexInParam] = position;
+            int32 Position = currentIndex.GetFromModelRangeIndex(rangeIndexInModel);
+			Index->Values[rangeIndexInParam] = Position;
         }
 
-        return index;
+        return Index;
     }
 
 
     //---------------------------------------------------------------------------------------------
-    void CodeRunner::RunCode_Bool(const FScheduledOp& item, const Parameters* pParams, const Model* pModel )
+    void CodeRunner::RunCode_Bool(const FScheduledOp& item, const FParameters* pParams, const FModel* pModel )
     {
         MUTABLE_CPUPROFILER_SCOPE(RunCode_Bool);
 
-		const FProgram& Program = m_pModel->GetPrivate()->m_program;
-		OP_TYPE type = Program.GetOpType(item.At);
+		const FProgram& Program = Model->GetPrivate()->Program;
+		EOpType type = Program.GetOpType(item.At);
         switch (type)
         {
 
-        case OP_TYPE::BO_CONSTANT:
+        case EOpType::BO_CONSTANT:
         {
 			OP::BoolConstantArgs args = Program.GetOpArgs<OP::BoolConstantArgs>(item.At);
             bool result = args.value;
@@ -5269,17 +5269,17 @@ namespace mu
             break;
         }
 
-        case OP_TYPE::BO_PARAMETER:
+        case EOpType::BO_PARAMETER:
         {
 			OP::ParameterArgs args = Program.GetOpArgs<OP::ParameterArgs>(item.At);
             bool result = false;
-			Ptr<RangeIndex> index = BuildCurrentOpRangeIndex( item, pParams, pModel, args.variable );
-            result = pParams->GetBoolValue( args.variable, index );
+			TSharedPtr<FRangeIndex> Index = BuildCurrentOpRangeIndex( item, pParams, pModel, args.variable );
+            result = pParams->GetBoolValue( args.variable, Index.Get() );
             StoreBool( item, result );
             break;
         }
 
-        case OP_TYPE::BO_LESS:
+        case EOpType::BO_LESS:
         {
 			OP::BoolLessArgs args = Program.GetOpArgs<OP::BoolLessArgs>(item.At);
             switch (item.Stage)
@@ -5306,7 +5306,7 @@ namespace mu
             break;
         }
 
-        case OP_TYPE::BO_AND:
+        case EOpType::BO_AND:
         {
 			OP::BoolBinaryArgs args = Program.GetOpArgs<OP::BoolBinaryArgs>(item.At);
             switch (item.Stage)
@@ -5372,7 +5372,7 @@ namespace mu
             break;
         }
 
-        case OP_TYPE::BO_OR:
+        case EOpType::BO_OR:
         {
 			OP::BoolBinaryArgs args = Program.GetOpArgs<OP::BoolBinaryArgs>(item.At);
             switch (item.Stage)
@@ -5438,7 +5438,7 @@ namespace mu
             break;
         }
 
-        case OP_TYPE::BO_NOT:
+        case EOpType::BO_NOT:
         {
 			OP::BoolNotArgs args = Program.GetOpArgs<OP::BoolNotArgs>(item.At);
             switch (item.Stage)
@@ -5461,7 +5461,7 @@ namespace mu
             break;
         }
 
-        case OP_TYPE::BO_EQUAL_INT_CONST:
+        case EOpType::BO_EQUAL_INT_CONST:
         {
 			OP::BoolEqualScalarConstArgs args = Program.GetOpArgs<OP::BoolEqualScalarConstArgs>(item.At);
             switch (item.Stage)
@@ -5473,7 +5473,7 @@ namespace mu
 
             case 1:
             {
-                int a = LoadInt( FCacheAddress(args.value,item) );
+                int32 a = LoadInt( FCacheAddress(args.value,item) );
                 bool result = a == args.constant;
                 StoreBool( item, result );
                 break;
@@ -5492,35 +5492,35 @@ namespace mu
     }
 
     //---------------------------------------------------------------------------------------------
-    void CodeRunner::RunCode_Int(const FScheduledOp& item, const Parameters* pParams, const Model* pModel )
+    void CodeRunner::RunCode_Int(const FScheduledOp& item, const FParameters* pParams, const FModel* pModel )
     {
         MUTABLE_CPUPROFILER_SCOPE(RunCode_Int);
 
-		const FProgram& Program = m_pModel->GetPrivate()->m_program;
+		const FProgram& Program = Model->GetPrivate()->Program;
 
-		OP_TYPE type = Program.GetOpType(item.At);
+		EOpType type = Program.GetOpType(item.At);
         switch (type)
         {
 
-        case OP_TYPE::NU_CONSTANT:
+        case EOpType::NU_CONSTANT:
         {
 			OP::IntConstantArgs args = Program.GetOpArgs<OP::IntConstantArgs>(item.At);
-            int result = args.value;
+            int32 result = args.value;
             StoreInt( item, result );
             break;
         }
 
-        case OP_TYPE::NU_PARAMETER:
+        case EOpType::NU_PARAMETER:
         {
 			OP::ParameterArgs args = Program.GetOpArgs<OP::ParameterArgs>(item.At);
-			Ptr<RangeIndex> index = BuildCurrentOpRangeIndex( item, pParams, pModel, args.variable );
-            int result = pParams->GetIntValue( args.variable, index );
+			TSharedPtr<FRangeIndex> Index = BuildCurrentOpRangeIndex( item, pParams, pModel, args.variable );
+            int32 result = pParams->GetIntValue( args.variable, Index.Get());
 
             // Check that the value is actually valid. Otherwise set the default.
             if ( pParams->GetIntPossibleValueCount( args.variable ) )
             {
                 bool valid = false;
-                for ( int i=0;
+                for ( int32 i=0;
                       (!valid) && i<pParams->GetIntPossibleValueCount( args.variable );
                       ++i )
                 {
@@ -5547,17 +5547,17 @@ namespace mu
     }
 
     //---------------------------------------------------------------------------------------------
-    void CodeRunner::RunCode_Scalar(const FScheduledOp& item, const Parameters* pParams, const Model* pModel )
+    void CodeRunner::RunCode_Scalar(const FScheduledOp& item, const FParameters* pParams, const FModel* pModel )
     {
         MUTABLE_CPUPROFILER_SCOPE(RunCode_Scalar);
 
-		const FProgram& Program = m_pModel->GetPrivate()->m_program;
+		const FProgram& Program = Model->GetPrivate()->Program;
 
-		OP_TYPE type = Program.GetOpType(item.At);
+		EOpType type = Program.GetOpType(item.At);
         switch (type)
         {
 
-        case OP_TYPE::SC_CONSTANT:
+        case EOpType::SC_CONSTANT:
         {
 			OP::ScalarConstantArgs args = Program.GetOpArgs<OP::ScalarConstantArgs>(item.At);
             float result = args.value;
@@ -5565,16 +5565,16 @@ namespace mu
             break;
         }
 
-        case OP_TYPE::SC_PARAMETER:
+        case EOpType::SC_PARAMETER:
         {
 			OP::ParameterArgs args = Program.GetOpArgs<OP::ParameterArgs>(item.At);
-			Ptr<RangeIndex> index = BuildCurrentOpRangeIndex( item, pParams, pModel, args.variable );
-            float result = pParams->GetFloatValue( args.variable, index );
+			TSharedPtr<FRangeIndex> Index = BuildCurrentOpRangeIndex( item, pParams, pModel, args.variable );
+            float result = pParams->GetFloatValue( args.variable, Index.Get());
             StoreScalar( item, result );
             break;
         }
 
-        case OP_TYPE::SC_CURVE:
+        case EOpType::SC_CURVE:
         {
 			OP::ScalarCurveArgs args = Program.GetOpArgs<OP::ScalarCurveArgs>(item.At);
             switch (item.Stage)
@@ -5601,12 +5601,12 @@ namespace mu
             break;
         }
 
-        case OP_TYPE::SC_MULTIPLYADD:
+        case EOpType::SC_MULTIPLYADD:
             // \TODO
             check( false );
             break;
 
-        case OP_TYPE::SC_ARITHMETIC:
+        case EOpType::SC_ARITHMETIC:
         {
 			OP::ArithmeticArgs args = Program.GetOpArgs<OP::ArithmeticArgs>(item.At);
             switch (item.Stage)
@@ -5663,34 +5663,34 @@ namespace mu
     }
 
     //---------------------------------------------------------------------------------------------
-    void CodeRunner::RunCode_String(const FScheduledOp& item, const Parameters* pParams, const Model* pModel )
+    void CodeRunner::RunCode_String(const FScheduledOp& item, const FParameters* pParams, const FModel* pModel )
     {
         MUTABLE_CPUPROFILER_SCOPE(RunCode_String );
 
-		const FProgram& Program = m_pModel->GetPrivate()->m_program;
+		const FProgram& Program = Model->GetPrivate()->Program;
 
-		OP_TYPE type = Program.GetOpType( item.At );
+		EOpType type = Program.GetOpType( item.At );
         switch ( type )
         {
 
-        case OP_TYPE::ST_CONSTANT:
+        case EOpType::ST_CONSTANT:
         {
 			OP::ResourceConstantArgs args = Program.GetOpArgs<OP::ResourceConstantArgs>( item.At );
-            check( args.value < (uint32)pModel->GetPrivate()->m_program.m_constantStrings.Num() );
+            check( args.value < (uint32)pModel->GetPrivate()->Program.ConstantStrings.Num() );
 
-            const FString& result = Program.m_constantStrings[args.value];
-            StoreString( item, new String(result) );
+            const FString& result = Program.ConstantStrings[args.value];
+            StoreString( item, MakeShared<String>(result) );
 
             break;
         }
 
-        case OP_TYPE::ST_PARAMETER:
+        case EOpType::ST_PARAMETER:
         {
 			OP::ParameterArgs args = Program.GetOpArgs<OP::ParameterArgs>( item.At );
-			Ptr<RangeIndex> index = BuildCurrentOpRangeIndex( item, pParams, pModel, args.variable );
+			TSharedPtr<FRangeIndex> Index = BuildCurrentOpRangeIndex( item, pParams, pModel, args.variable );
 			FString result;
-			pParams->GetStringValue(args.variable, result, index);
-            StoreString( item, new String(result) );
+			pParams->GetStringValue(args.variable, result, Index.Get());
+            StoreString( item, MakeShared<String>(result) );
             break;
         }
 
@@ -5702,18 +5702,18 @@ namespace mu
 
 
     //---------------------------------------------------------------------------------------------
-    void CodeRunner::RunCode_Colour(const FScheduledOp& item, const Parameters* pParams, const Model* pModel )
+    void CodeRunner::RunCode_Colour(const FScheduledOp& item, const FParameters* pParams, const FModel* pModel )
     {
 		MUTABLE_CPUPROFILER_SCOPE(RunCode_Colour);
 
-		const FProgram& Program = m_pModel->GetPrivate()->m_program;
+		const FProgram& Program = Model->GetPrivate()->Program;
 
-		OP_TYPE type = Program.GetOpType(item.At);
+		EOpType type = Program.GetOpType(item.At);
 
         switch ( type )
         {
 
-        case OP_TYPE::CO_CONSTANT:
+        case EOpType::CO_CONSTANT:
         {
 			OP::ColourConstantArgs args = Program.GetOpArgs<OP::ColourConstantArgs>(item.At);
 			FVector4f result;
@@ -5725,17 +5725,17 @@ namespace mu
             break;
         }
 
-        case OP_TYPE::CO_PARAMETER:
+        case EOpType::CO_PARAMETER:
         {
 			OP::ParameterArgs args = Program.GetOpArgs<OP::ParameterArgs>(item.At);
-			Ptr<RangeIndex> index = BuildCurrentOpRangeIndex( item, pParams, pModel, args.variable );
+			TSharedPtr<FRangeIndex> Index = BuildCurrentOpRangeIndex( item, pParams, pModel, args.variable );
             FVector4f V;
-            pParams->GetColourValue( args.variable, V, index );
+            pParams->GetColourValue( args.variable, V, Index.Get());
             StoreColor( item, V );
             break;
         }
 
-        case OP_TYPE::CO_SAMPLEIMAGE:
+        case EOpType::CO_SAMPLEIMAGE:
         {
 			OP::ColourSampleImageArgs args = Program.GetOpArgs<OP::ColourSampleImageArgs>(item.At);
             switch (item.Stage)
@@ -5753,7 +5753,7 @@ namespace mu
                 float x = args.x ? LoadScalar( FCacheAddress(args.x,item) ) : 0.5f;
                 float y = args.y ? LoadScalar( FCacheAddress(args.y,item) ) : 0.5f;
 
-                Ptr<const Image> pImage = LoadImage(FScheduledOp::FromOpAndOptions(args.image, item, 0));
+                TSharedPtr<const FImage> pImage = LoadImage(FScheduledOp::FromOpAndOptions(args.image, item, 0));
 
 				FVector4f result;
                 if (pImage)
@@ -5785,7 +5785,7 @@ namespace mu
             break;
         }
 
-        case OP_TYPE::CO_SWIZZLE:
+        case EOpType::CO_SWIZZLE:
         {
 			OP::ColourSwizzleArgs args = Program.GetOpArgs<OP::ColourSwizzleArgs>(item.At);
             switch (item.Stage)
@@ -5802,7 +5802,7 @@ namespace mu
             {
 				FVector4f result;
 
-                for (int t=0;t<MUTABLE_OP_MAX_SWIZZLE_CHANNELS;++t)
+                for (int32 t=0;t<MUTABLE_OP_MAX_SWIZZLE_CHANNELS;++t)
                 {
                     if ( args.sources[t] )
                     {
@@ -5822,7 +5822,7 @@ namespace mu
             break;
         }
 
-        case OP_TYPE::CO_FROMSCALARS:
+        case EOpType::CO_FROMSCALARS:
         {
 			OP::ColourFromScalarsArgs args = Program.GetOpArgs<OP::ColourFromScalarsArgs>(item.At);
             switch (item.Stage)
@@ -5858,7 +5858,7 @@ namespace mu
             break;
         }
 
-        case OP_TYPE::CO_ARITHMETIC:
+        case EOpType::CO_ARITHMETIC:
         {
 			OP::ArithmeticArgs args = Program.GetOpArgs<OP::ArithmeticArgs>(item.At);
             switch (item.Stage)
@@ -5871,7 +5871,7 @@ namespace mu
 
             case 1:
             {
-				OP_TYPE otype = Program.GetOpType( args.a );
+				EOpType otype = Program.GetOpType( args.a );
                 DATATYPE dtype = GetOpDataType( otype );
                 check( dtype == DT_COLOUR );
                 otype = Program.GetOpType( args.b );
@@ -5925,32 +5925,32 @@ namespace mu
 
 
     //---------------------------------------------------------------------------------------------
-    void CodeRunner::RunCode_Projector(const FScheduledOp& item, const Parameters* pParams, const Model* pModel )
+    void CodeRunner::RunCode_Projector(const FScheduledOp& item, const FParameters* pParams, const FModel* pModel )
     {
         MUTABLE_CPUPROFILER_SCOPE(RunCode_Projector);
 
-		const FProgram& Program = m_pModel->GetPrivate()->m_program;
+		const FProgram& Program = Model->GetPrivate()->Program;
 
-		OP_TYPE type = Program.GetOpType(item.At);
+		EOpType type = Program.GetOpType(item.At);
         switch (type)
         {
 
-        case OP_TYPE::PR_CONSTANT:
+        case EOpType::PR_CONSTANT:
         {
 			OP::ResourceConstantArgs args = Program.GetOpArgs<OP::ResourceConstantArgs>(item.At);
-            FProjector Result = Program.m_constantProjectors[args.value];
+            FProjector Result = Program.ConstantProjectors[args.value];
             StoreProjector( item, Result );
             break;
         }
 
-        case OP_TYPE::PR_PARAMETER:
+        case EOpType::PR_PARAMETER:
         {
 			OP::ParameterArgs args = Program.GetOpArgs<OP::ParameterArgs>(item.At);
-			Ptr<RangeIndex> index = BuildCurrentOpRangeIndex( item, pParams, pModel, args.variable );
-            FProjector Result = pParams->GetPrivate()->GetProjectorValue(args.variable,index);
+			TSharedPtr<FRangeIndex> Index = BuildCurrentOpRangeIndex( item, pParams, pModel, args.variable );
+            FProjector Result = pParams->GetPrivate()->GetProjectorValue(args.variable,Index.Get());
 
             // The type cannot be changed, take it from the default value
-            const FProjector& def = Program.m_parameters[args.variable].m_defaultValue.Get<FParamProjectorType>();
+            const FProjector& def = Program.Parameters[args.variable].DefaultValue.Get<FParamProjectorType>();
             Result.type = def.type;
 
             StoreProjector( item, Result );
@@ -5964,29 +5964,29 @@ namespace mu
     }
 
     //---------------------------------------------------------------------------------------------
-    void CodeRunner::RunCode_Matrix(const FScheduledOp& item, const Parameters* pParams, const Model* pModel )
+    void CodeRunner::RunCode_Matrix(const FScheduledOp& item, const FParameters* pParams, const FModel* pModel )
     {
 		MUTABLE_CPUPROFILER_SCOPE(RunCode_Transform);
 
-		const FProgram& Program = m_pModel->GetPrivate()->m_program;
+		const FProgram& Program = Model->GetPrivate()->Program;
 
-		OP_TYPE type = Program.GetOpType(item.At);
+		EOpType type = Program.GetOpType(item.At);
 
 		switch ( type )
 		{
-		case OP_TYPE::MA_CONSTANT:
+		case EOpType::MA_CONSTANT:
 			{
 				OP::MatrixConstantArgs args = Program.GetOpArgs<OP::MatrixConstantArgs>(item.At);
-				StoreMatrix( item, Program.m_constantMatrices[args.value] );
+				StoreMatrix( item, Program.ConstantMatrices[args.value] );
 				break;
 			}
 
-		case OP_TYPE::MA_PARAMETER:
+		case EOpType::MA_PARAMETER:
 			{
 				OP::ParameterArgs args = Program.GetOpArgs<OP::ParameterArgs>(item.At);
-				Ptr<RangeIndex> index = BuildCurrentOpRangeIndex( item, pParams, pModel, args.variable );
+				TSharedPtr<FRangeIndex> Index = BuildCurrentOpRangeIndex( item, pParams, pModel, args.variable );
 				FMatrix44f Value;
-				pParams->GetMatrixValue( args.variable, Value, index );
+				pParams->GetMatrixValue( args.variable, Value, Index.Get());
 				StoreMatrix( item, Value );
 				break;
 			}
@@ -5995,28 +5995,28 @@ namespace mu
 
 
     //---------------------------------------------------------------------------------------------
-    void CodeRunner::RunCode_Layout(const FScheduledOp& item, const Model* pModel )
+    void CodeRunner::RunCode_Layout(const FScheduledOp& item, const FModel* pModel )
     {
         //MUTABLE_CPUPROFILER_SCOPE(RunCode_Layout);
 
-		const FProgram& Program = m_pModel->GetPrivate()->m_program;
+		const FProgram& Program = Model->GetPrivate()->Program;
 
-		OP_TYPE type = Program.GetOpType(item.At);
+		EOpType type = Program.GetOpType(item.At);
         switch (type)
         {
 
-        case OP_TYPE::LA_CONSTANT:
+        case EOpType::LA_CONSTANT:
         {
 			OP::ResourceConstantArgs args = Program.GetOpArgs<OP::ResourceConstantArgs>(item.At);
-            check( args.value < (uint32)pModel->GetPrivate()->m_program.m_constantLayouts.Num() );
+            check( args.value < (uint32)pModel->GetPrivate()->Program.ConstantLayouts.Num() );
 
-            Ptr<const Layout> pResult = Program.m_constantLayouts
+            TSharedPtr<const FLayout> Result = Program.ConstantLayouts
                     [ args.value ];
-            StoreLayout( item, pResult );
+            StoreLayout( item, Result );
             break;
         }
 
-        case OP_TYPE::LA_MERGE:
+        case EOpType::LA_MERGE:
         {
 			OP::LayoutMergeArgs args = Program.GetOpArgs<OP::LayoutMergeArgs>(item.At);
             switch (item.Stage)
@@ -6029,25 +6029,25 @@ namespace mu
 
             case 1:
             {
-                Ptr<const Layout> pA = LoadLayout( FCacheAddress(args.Base,item) );
-                Ptr<const Layout> pB = LoadLayout( FCacheAddress(args.Added,item) );
+                TSharedPtr<const FLayout> pA = LoadLayout( FCacheAddress(args.Base,item) );
+                TSharedPtr<const FLayout> pB = LoadLayout( FCacheAddress(args.Added,item) );
 
-                Ptr<const Layout> pResult;
+                TSharedPtr<const FLayout> Result;
 
                 if (pA && pB)
                 {
-					pResult = LayoutMerge(pA.get(),pB.get());
+					Result = LayoutMerge(pA.Get(),pB.Get());
                 }
                 else if (pA)
                 {
-                    pResult = pA->Clone();
+                    Result = pA->Clone();
                 }
                 else if (pB)
                 {
-                    pResult = pB->Clone();
+                    Result = pB->Clone();
                 }
 
-                StoreLayout( item, pResult );
+                StoreLayout( item, Result );
                 break;
             }
 
@@ -6058,7 +6058,7 @@ namespace mu
             break;
         }
 
-        case OP_TYPE::LA_PACK:
+        case EOpType::LA_PACK:
         {
 			OP::LayoutPackArgs args = Program.GetOpArgs<OP::LayoutPackArgs>(item.At);
             switch (item.Stage)
@@ -6070,15 +6070,15 @@ namespace mu
 
             case 1:
             {
-                Ptr<const Layout> Source = LoadLayout( FCacheAddress(args.Source,item) );
+                TSharedPtr<const FLayout> Source = LoadLayout( FCacheAddress(args.Source,item) );
 
-				Ptr<Layout> Result;
+				TSharedPtr<FLayout> Result;
 
 				if (Source)
 				{
 					Result = Source->Clone();
 
-					LayoutPack3(Result.get(), Source.get() );
+					LayoutPack3(Result.Get(), Source.Get() );
 				}
 
                 StoreLayout( item, Result );
@@ -6092,7 +6092,7 @@ namespace mu
             break;
         }
 
-		case OP_TYPE::LA_FROMMESH:
+		case EOpType::LA_FROMMESH:
 		{
 			OP::LayoutFromMeshArgs args = Program.GetOpArgs<OP::LayoutFromMeshArgs>(item.At);
 			switch (item.Stage)
@@ -6104,9 +6104,9 @@ namespace mu
 
 			case 1:
 			{
-				Ptr<const Mesh> Mesh = LoadMesh(FCacheAddress(args.Mesh, item));
+				TSharedPtr<const FMesh> Mesh = LoadMesh(FCacheAddress(args.Mesh, item));
 
-				Ptr<const Layout> Result = LayoutFromMesh_RemoveBlocks(Mesh.get(), args.LayoutIndex);
+				TSharedPtr<const FLayout> Result = LayoutFromMesh_RemoveBlocks(Mesh.Get(), args.LayoutIndex);
 
 				Release(Mesh);
 				StoreLayout(item, Result);
@@ -6120,7 +6120,7 @@ namespace mu
 			break;
 		}
 
-		case OP_TYPE::LA_REMOVEBLOCKS:
+		case EOpType::LA_REMOVEBLOCKS:
 		{
 			OP::LayoutRemoveBlocksArgs args = Program.GetOpArgs<OP::LayoutRemoveBlocksArgs>(item.At);
 			switch (item.Stage)
@@ -6133,21 +6133,21 @@ namespace mu
 
 			case 1:
 			{
-				Ptr<const Layout> Source = LoadLayout(FCacheAddress(args.Source, item));
-				Ptr<const Layout> ReferenceLayout = LoadLayout(FCacheAddress(args.ReferenceLayout, item));
+				TSharedPtr<const FLayout> Source = LoadLayout(FCacheAddress(args.Source, item));
+				TSharedPtr<const FLayout> ReferenceLayout = LoadLayout(FCacheAddress(args.ReferenceLayout, item));
 
-				Ptr<const Layout> pResult;
+				TSharedPtr<const FLayout> Result;
 
 				if (Source && ReferenceLayout)
 				{
-					pResult = LayoutRemoveBlocks(Source.get(), ReferenceLayout.get());
+					Result = LayoutRemoveBlocks(Source.Get(), ReferenceLayout.Get());
 				}
 				else if (Source)
 				{
-					pResult = Source;
+					Result = Source;
 				}
 
-				StoreLayout(item, pResult);
+				StoreLayout(item, Result);
 				break;
 			}
 
@@ -6167,59 +6167,59 @@ namespace mu
 
 
     //---------------------------------------------------------------------------------------------
-    void CodeRunner::RunCode( const FScheduledOp& item, const Parameters* pParams, const TSharedPtr<const Model>& InModel, uint32 lodMask)
+    void CodeRunner::RunCode( const FScheduledOp& item, const FParameters* pParams, const TSharedPtr<const FModel>& InModel, uint32 lodMask)
     {
 		//UE_LOG( LogMutableCore, Log, TEXT("Running :%5d , %d "), item.At, item.Stage );
 		check( item.Type == FScheduledOp::EType::Full );
 
-		const Model* pModel = InModel.Get();
+		const FModel* pModel = InModel.Get();
 		
-		const FProgram& Program = pModel->GetPrivate()->m_program;
+		const FProgram& Program = pModel->GetPrivate()->Program;
 
-		OP_TYPE type = Program.GetOpType(item.At);
+		EOpType type = Program.GetOpType(item.At);
 		//UE_LOG(LogMutableCore, Log, TEXT("Running :%5d , %d, of type %d "), item.At, item.Stage, type);
 
 		// Very spammy, for debugging purposes.
-		//if (m_pSystem)
+		//if (System)
 		//{
-		//	m_pSystem->WorkingMemoryManager.LogWorkingMemory( this );
+		//	System->WorkingMemoryManager.LogWorkingMemory( this );
 		//}
 
 		switch ( type )
         {
-        case OP_TYPE::NONE:
+        case EOpType::NONE:
             break;
 
-        case OP_TYPE::NU_CONDITIONAL:
-        case OP_TYPE::SC_CONDITIONAL:
-        case OP_TYPE::CO_CONDITIONAL:
-        case OP_TYPE::IM_CONDITIONAL:
-        case OP_TYPE::ME_CONDITIONAL:
-        case OP_TYPE::LA_CONDITIONAL:
-        case OP_TYPE::IN_CONDITIONAL:
-		case OP_TYPE::ED_CONDITIONAL:
+        case EOpType::NU_CONDITIONAL:
+        case EOpType::SC_CONDITIONAL:
+        case EOpType::CO_CONDITIONAL:
+        case EOpType::IM_CONDITIONAL:
+        case EOpType::ME_CONDITIONAL:
+        case EOpType::LA_CONDITIONAL:
+        case EOpType::IN_CONDITIONAL:
+		case EOpType::ED_CONDITIONAL:
             RunCode_Conditional(item, pModel);
             break;
 
-        case OP_TYPE::ME_CONSTANT:
-		case OP_TYPE::IM_CONSTANT:
-		case OP_TYPE::ED_CONSTANT:
+        case EOpType::ME_CONSTANT:
+		case EOpType::IM_CONSTANT:
+		case EOpType::ED_CONSTANT:
             RunCode_ConstantResource(item, pModel);
             break;
 
-        case OP_TYPE::NU_SWITCH:
-        case OP_TYPE::SC_SWITCH:
-        case OP_TYPE::CO_SWITCH:
-        case OP_TYPE::IM_SWITCH:
-        case OP_TYPE::ME_SWITCH:
-        case OP_TYPE::LA_SWITCH:
-        case OP_TYPE::IN_SWITCH:
-		case OP_TYPE::ED_SWITCH:
+        case EOpType::NU_SWITCH:
+        case EOpType::SC_SWITCH:
+        case EOpType::CO_SWITCH:
+        case EOpType::IM_SWITCH:
+        case EOpType::ME_SWITCH:
+        case EOpType::LA_SWITCH:
+        case EOpType::IN_SWITCH:
+		case EOpType::ED_SWITCH:
             RunCode_Switch(item, pModel);
             break;
 
-        case OP_TYPE::IN_ADDMESH:
-        case OP_TYPE::IN_ADDIMAGE:
+        case EOpType::IN_ADDMESH:
+        case EOpType::IN_ADDIMAGE:
             RunCode_InstanceAddResource(item, InModel, pParams);
             break;
 
@@ -6283,32 +6283,32 @@ namespace mu
     }
 
 	//---------------------------------------------------------------------------------------------
-	void CodeRunner::RunCodeImageDesc(const FScheduledOp& item, const Parameters* pParams, const Model* pModel,  uint32 lodMask )
+	void CodeRunner::RunCodeImageDesc(const FScheduledOp& item, const FParameters* pParams, const FModel* pModel,  uint32 lodMask )
 	{
 		MUTABLE_CPUPROFILER_SCOPE(RunCodeImageDesc);
 
 		check(item.Type == FScheduledOp::EType::ImageDesc);
 
 		// Ensure there is room for the result in the heap.
-		if (item.CustomState >= uint32(m_heapData.Num()))
+		if (item.CustomState >= uint32(HeapData.Num()))
 		{
-			m_heapImageDesc.SetNum(item.CustomState+1);
+			HeapImageDesc.SetNum(item.CustomState+1);
 		}
 
 
-		const FProgram& Program = m_pModel->GetPrivate()->m_program;
+		const FProgram& Program = Model->GetPrivate()->Program;
 
-		OP_TYPE type = Program.GetOpType(item.At);
+		EOpType type = Program.GetOpType(item.At);
 		switch (type)
 		{
 
-		case OP_TYPE::IM_CONSTANT:
+		case EOpType::IM_CONSTANT:
 		{
 			check(item.Stage == 0);
 			OP::ResourceConstantArgs args = Program.GetOpArgs<OP::ResourceConstantArgs>(item.At);
 			int32 ImageIndex = args.value;
 
-			FImageDesc& Result = m_heapImageDesc[item.CustomState];
+			FImageDesc& Result = HeapImageDesc[item.CustomState];
 			Result.m_format = Program.ConstantImages[ImageIndex].ImageFormat;
 			Result.m_size[0] = Program.ConstantImages[ImageIndex].ImageSizeX;
 			Result.m_size[1] = Program.ConstantImages[ImageIndex].ImageSizeY;
@@ -6317,28 +6317,28 @@ namespace mu
 			break;
 		}
 
-		case OP_TYPE::IM_PARAMETER:
+		case EOpType::IM_PARAMETER:
 		{
 			check(item.Stage == 0);
 			OP::ParameterArgs args = Program.GetOpArgs<OP::ParameterArgs>(item.At);
 			FName Id = pParams->GetImageValue(args.variable);
 			uint8 MipsToSkip = item.ExecutionOptions;
-			m_heapImageDesc[item.CustomState] = GetExternalImageDesc(Id, MipsToSkip);
+			HeapImageDesc[item.CustomState] = GetExternalImageDesc(Id, MipsToSkip);
 			StoreValidDesc(item);
 			break;
 		}
 
-		case OP_TYPE::IM_REFERENCE:
+		case EOpType::IM_REFERENCE:
 		{
 			check(item.Stage == 0);
 			OP::ResourceReferenceArgs Args = Program.GetOpArgs<OP::ResourceReferenceArgs>(item.At);
-			FImageDesc& Result = m_heapImageDesc[item.CustomState];
+			FImageDesc& Result = HeapImageDesc[item.CustomState];
 			Result = Args.ImageDesc;
 			StoreValidDesc(item);
 			break;
 		}
 
-		case OP_TYPE::IM_CONDITIONAL:
+		case EOpType::IM_CONDITIONAL:
 		{
 			OP::ConditionalArgs args = Program.GetOpArgs<OP::ConditionalArgs>(item.At);
 			switch (item.Stage)
@@ -6366,7 +6366,7 @@ namespace mu
 			break;
 		}
 
-		case OP_TYPE::IM_SWITCH:
+		case EOpType::IM_SWITCH:
 		{
 			const uint8* data = Program.GetOpArgsPointer(item.At);
 		
@@ -6403,7 +6403,7 @@ namespace mu
 			case 1:
 			{
 				// Get the variable result
-				int var = LoadInt(FCacheAddress(VarAddress, item));
+				int32 var = LoadInt(FCacheAddress(VarAddress, item));
 
 				OP::ADDRESS valueAt = DefAddress;
 				for (uint32 C = 0; C < CaseCount; ++C)
@@ -6416,7 +6416,7 @@ namespace mu
 					FMemory::Memcpy( &At, data, sizeof(OP::ADDRESS) );
 					data += sizeof(OP::ADDRESS);
 
-					if (At && var == (int)Condition)
+					if (At && var == (int32)Condition)
 					{
 						valueAt = At;
 						break;
@@ -6435,7 +6435,7 @@ namespace mu
 			break;
 		}
 
-		case OP_TYPE::IM_LAYERCOLOUR:
+		case EOpType::IM_LAYERCOLOUR:
 		{
 			OP::ImageLayerColourArgs args = Program.GetOpArgs<OP::ImageLayerColourArgs>(item.At);
 			switch (item.Stage)
@@ -6447,7 +6447,7 @@ namespace mu
 			break;
 		}
 
-		case OP_TYPE::IM_LAYER:
+		case EOpType::IM_LAYER:
 		{
 			OP::ImageLayerArgs args = Program.GetOpArgs<OP::ImageLayerArgs>(item.At);
 			switch (item.Stage)
@@ -6459,7 +6459,7 @@ namespace mu
 			break;
 		}
 
-		case OP_TYPE::IM_MULTILAYER:
+		case EOpType::IM_MULTILAYER:
 		{
 			OP::ImageMultiLayerArgs args = Program.GetOpArgs<OP::ImageMultiLayerArgs>(item.At);
 			switch (item.Stage)
@@ -6471,7 +6471,7 @@ namespace mu
 			break;
 		}
 
-		case OP_TYPE::IM_NORMALCOMPOSITE:
+		case EOpType::IM_NORMALCOMPOSITE:
 		{
 			OP::ImageNormalCompositeArgs args = Program.GetOpArgs<OP::ImageNormalCompositeArgs>(item.At);
 			switch (item.Stage)
@@ -6483,7 +6483,7 @@ namespace mu
 			break;
 		}
 
-		case OP_TYPE::IM_PIXELFORMAT:
+		case EOpType::IM_PIXELFORMAT:
 		{
 			OP::ImagePixelFormatArgs args = Program.GetOpArgs<OP::ImagePixelFormatArgs>(item.At);
 			switch (item.Stage)
@@ -6495,7 +6495,7 @@ namespace mu
 			case 1:
 			{
 				// Update directly in the heap
-				EImageFormat OldFormat = m_heapImageDesc[item.CustomState].m_format;
+				EImageFormat OldFormat = HeapImageDesc[item.CustomState].m_format;
 				EImageFormat NewFormat = args.format;
 				if (args.formatIfAlpha != EImageFormat::None
 					&&
@@ -6503,7 +6503,7 @@ namespace mu
 				{
 					NewFormat = args.formatIfAlpha;
 				}
-				m_heapImageDesc[item.CustomState].m_format = NewFormat;				
+				HeapImageDesc[item.CustomState].m_format = NewFormat;				
 				StoreValidDesc(item);
 				break;
 			}
@@ -6515,7 +6515,7 @@ namespace mu
 			break;
 		}
 
-		case OP_TYPE::IM_MIPMAP:
+		case EOpType::IM_MIPMAP:
 		{
 			OP::ImageMipmapArgs args = Program.GetOpArgs<OP::ImageMipmapArgs>(item.At);
 			switch (item.Stage)
@@ -6527,9 +6527,9 @@ namespace mu
 			case 1:
 			{
 				// Somewhat synched with Full op execution code.
-				FImageDesc BaseDesc = m_heapImageDesc[item.CustomState];
+				FImageDesc BaseDesc = HeapImageDesc[item.CustomState];
 				int32 LevelCount = args.levels;
-				int32 MaxLevelCount = Image::GetMipmapCount(BaseDesc.m_size[0], BaseDesc.m_size[1]);
+				int32 MaxLevelCount = FImage::GetMipmapCount(BaseDesc.m_size[0], BaseDesc.m_size[1]);
 				if (LevelCount == 0)
 				{
 					LevelCount = MaxLevelCount;
@@ -6546,7 +6546,7 @@ namespace mu
 				LevelCount = FMath::Max(StartLevel, LevelCount);
 
 				// Update result.
-				m_heapImageDesc[item.CustomState].m_lods = LevelCount;
+				HeapImageDesc[item.CustomState].m_lods = LevelCount;
 				StoreValidDesc(item);
 				break;
 			}
@@ -6557,7 +6557,7 @@ namespace mu
 			break;
 		}
 
-		case OP_TYPE::IM_RESIZE:
+		case EOpType::IM_RESIZE:
 		{
 			OP::ImageResizeArgs args = Program.GetOpArgs<OP::ImageResizeArgs>(item.At);
 			switch (item.Stage)
@@ -6567,8 +6567,8 @@ namespace mu
 				break;
 
 			case 1:
-				m_heapImageDesc[item.CustomState].m_size[0] = args.size[0];
-				m_heapImageDesc[item.CustomState].m_size[1] = args.size[1];
+				HeapImageDesc[item.CustomState].m_size[0] = args.size[0];
+				HeapImageDesc[item.CustomState].m_size[1] = args.size[1];
 				StoreValidDesc(item);
 				break;
 
@@ -6578,7 +6578,7 @@ namespace mu
 			break;
 		}
 
-		case OP_TYPE::IM_RESIZELIKE:
+		case EOpType::IM_RESIZELIKE:
 		{
 			OP::ImageResizeLikeArgs args = Program.GetOpArgs<OP::ImageResizeLikeArgs>(item.At);
 			switch (item.Stage)
@@ -6586,11 +6586,11 @@ namespace mu
 			case 0:
 			{
 				int32 ResultAndBaseDesc = item.CustomState;
-				int32 SourceDescAddress = m_heapImageDesc.Add({});
+				int32 SourceDescAddress = HeapImageDesc.Add({});
 				FScheduledOpData Data;
 				Data.ResizeLike.ResultDescAt = ResultAndBaseDesc;
 				Data.ResizeLike.SourceDescAt = SourceDescAddress;
-				int32 SecondStageData = m_heapData.Add(Data);
+				int32 SecondStageData = HeapData.Add(Data);
 				AddOp(FScheduledOp(item.At, item, 1, SecondStageData),
 					FScheduledOp(args.source, item, 0, ResultAndBaseDesc),
 					FScheduledOp(args.sizeSource, item, 0, SourceDescAddress));
@@ -6599,9 +6599,9 @@ namespace mu
 
 			case 1:
 			{
-				const FScheduledOpData& SecondStageData = m_heapData[ item.CustomState ];
-				FImageDesc& ResultAndBaseDesc = m_heapImageDesc[SecondStageData.ResizeLike.ResultDescAt];
-				const FImageDesc& SourceDesc = m_heapImageDesc[SecondStageData.ResizeLike.SourceDescAt];
+				const FScheduledOpData& SecondStageData = HeapData[ item.CustomState ];
+				FImageDesc& ResultAndBaseDesc = HeapImageDesc[SecondStageData.ResizeLike.ResultDescAt];
+				const FImageDesc& SourceDesc = HeapImageDesc[SecondStageData.ResizeLike.SourceDescAt];
 				ResultAndBaseDesc.m_size = SourceDesc.m_size;
 				StoreValidDesc(item);
 				break;
@@ -6614,7 +6614,7 @@ namespace mu
 			break;
 		}
 
-		case OP_TYPE::IM_RESIZEREL:
+		case EOpType::IM_RESIZEREL:
 		{
 			OP::ImageResizeRelArgs args = Program.GetOpArgs<OP::ImageResizeRelArgs>(item.At);
 			switch (item.Stage)
@@ -6625,7 +6625,7 @@ namespace mu
 
 			case 1:
 			{
-				FImageDesc& ResultAndBaseDesc = m_heapImageDesc[item.CustomState];
+				FImageDesc& ResultAndBaseDesc = HeapImageDesc[item.CustomState];
 				FImageSize destSize(
 					uint16(ResultAndBaseDesc.m_size[0] * args.factor[0] + 0.5f),
 					uint16(ResultAndBaseDesc.m_size[1] * args.factor[1] + 0.5f));
@@ -6641,7 +6641,7 @@ namespace mu
 			break;
 		}
 
-		case OP_TYPE::IM_BLANKLAYOUT:
+		case EOpType::IM_BLANKLAYOUT:
 		{
 			OP::ImageBlankLayoutArgs args = Program.GetOpArgs<OP::ImageBlankLayoutArgs>(item.At);
 			switch (item.Stage)
@@ -6657,13 +6657,13 @@ namespace mu
 
 			case 1:
 			{
-				Ptr<const Layout> pLayout = LoadLayout(FCacheAddress(args.layout, item));
+				TSharedPtr<const FLayout> pLayout = LoadLayout(FCacheAddress(args.layout, item));
 
 				FIntPoint SizeInBlocks = pLayout->GetGridSize();
 				FIntPoint BlockSizeInPixels(args.blockSize[0], args.blockSize[1]);
 				FIntPoint ImageSizeInPixels = SizeInBlocks * BlockSizeInPixels;
 
-				FImageDesc& ResultAndBaseDesc = m_heapImageDesc[item.CustomState];
+				FImageDesc& ResultAndBaseDesc = HeapImageDesc[item.CustomState];
 				FImageSize destSize(uint16(ImageSizeInPixels.X), uint16(ImageSizeInPixels.Y));
 				ResultAndBaseDesc.m_size = destSize;
 				ResultAndBaseDesc.m_format = args.format;
@@ -6672,7 +6672,7 @@ namespace mu
 				{
 					if (args.mipmapCount == 0)
 					{
-						ResultAndBaseDesc.m_lods = Image::GetMipmapCount(ImageSizeInPixels.X, ImageSizeInPixels.Y);
+						ResultAndBaseDesc.m_lods = FImage::GetMipmapCount(ImageSizeInPixels.X, ImageSizeInPixels.Y);
 					}
 					else
 					{
@@ -6690,7 +6690,7 @@ namespace mu
 			break;
 		}
 
-		case OP_TYPE::IM_COMPOSE:
+		case EOpType::IM_COMPOSE:
 		{
 			OP::ImageComposeArgs args = Program.GetOpArgs<OP::ImageComposeArgs>(item.At);
 			switch (item.Stage)
@@ -6702,7 +6702,7 @@ namespace mu
 			break;
 		}
 
-		case OP_TYPE::IM_INTERPOLATE:
+		case EOpType::IM_INTERPOLATE:
 		{
 			OP::ImageInterpolateArgs args = Program.GetOpArgs<OP::ImageInterpolateArgs>(item.At);
 			switch (item.Stage)
@@ -6714,7 +6714,7 @@ namespace mu
 			break;
 		}
 
-		case OP_TYPE::IM_SATURATE:
+		case EOpType::IM_SATURATE:
 		{
 			OP::ImageSaturateArgs args = Program.GetOpArgs<OP::ImageSaturateArgs>(item.At);
 			switch (item.Stage)
@@ -6726,7 +6726,7 @@ namespace mu
 			break;
 		}
 
-		case OP_TYPE::IM_LUMINANCE:
+		case EOpType::IM_LUMINANCE:
 		{
 			OP::ImageLuminanceArgs args = Program.GetOpArgs<OP::ImageLuminanceArgs>(item.At);
 			switch (item.Stage)
@@ -6736,7 +6736,7 @@ namespace mu
 				break;
 
 			case 1:
-				m_heapImageDesc[item.CustomState].m_format = EImageFormat::L_UByte;
+				HeapImageDesc[item.CustomState].m_format = EImageFormat::L_UByte;
 				StoreValidDesc(item);
 				break;
 
@@ -6747,7 +6747,7 @@ namespace mu
 			break;
 		}
 
-		case OP_TYPE::IM_SWIZZLE:
+		case EOpType::IM_SWIZZLE:
 		{
 			OP::ImageSwizzleArgs args = Program.GetOpArgs<OP::ImageSwizzleArgs>(item.At);
 			switch (item.Stage)
@@ -6757,7 +6757,7 @@ namespace mu
 				break;
 
 			case 1:
-				m_heapImageDesc[item.CustomState].m_format = args.format;
+				HeapImageDesc[item.CustomState].m_format = args.format;
 				StoreValidDesc(item);
 				break;
 
@@ -6768,7 +6768,7 @@ namespace mu
 			break;
 		}
 
-		case OP_TYPE::IM_COLOURMAP:
+		case EOpType::IM_COLOURMAP:
 		{
 			OP::ImageColourMapArgs args = Program.GetOpArgs<OP::ImageColourMapArgs>(item.At);
 			switch (item.Stage)
@@ -6780,18 +6780,18 @@ namespace mu
 			break;
 		}
 
-		case OP_TYPE::IM_GRADIENT:
+		case EOpType::IM_GRADIENT:
 		{
 			OP::ImageGradientArgs args = Program.GetOpArgs<OP::ImageGradientArgs>(item.At);
-			m_heapImageDesc[item.CustomState].m_size[0] = args.size[0];
-			m_heapImageDesc[item.CustomState].m_size[1] = args.size[1];
-			m_heapImageDesc[item.CustomState].m_lods = 1;
-			m_heapImageDesc[item.CustomState].m_format = EImageFormat::RGB_UByte;
+			HeapImageDesc[item.CustomState].m_size[0] = args.size[0];
+			HeapImageDesc[item.CustomState].m_size[1] = args.size[1];
+			HeapImageDesc[item.CustomState].m_lods = 1;
+			HeapImageDesc[item.CustomState].m_format = EImageFormat::RGB_UByte;
 			StoreValidDesc(item);
 			break;
 		}
 
-		case OP_TYPE::IM_BINARISE:
+		case EOpType::IM_BINARISE:
 		{
 			OP::ImageBinariseArgs args = Program.GetOpArgs<OP::ImageBinariseArgs>(item.At);
 			switch (item.Stage)
@@ -6801,7 +6801,7 @@ namespace mu
 				break;
 
 			case 1:
-				m_heapImageDesc[item.CustomState].m_format = EImageFormat::L_UByte;
+				HeapImageDesc[item.CustomState].m_format = EImageFormat::L_UByte;
 				StoreValidDesc(item);
 				break;
 
@@ -6811,7 +6811,7 @@ namespace mu
 			break;
 		}
 
-		case OP_TYPE::IM_INVERT:
+		case EOpType::IM_INVERT:
 		{
 			OP::ImageInvertArgs args = Program.GetOpArgs<OP::ImageInvertArgs>(item.At);
 			switch (item.Stage)
@@ -6823,18 +6823,18 @@ namespace mu
 			break;
 		}
 
-		case OP_TYPE::IM_PLAINCOLOUR:
+		case EOpType::IM_PLAINCOLOUR:
 		{
 			OP::ImagePlainColourArgs args = Program.GetOpArgs<OP::ImagePlainColourArgs>(item.At);
-			m_heapImageDesc[item.CustomState].m_size[0] = args.size[0];
-			m_heapImageDesc[item.CustomState].m_size[1] = args.size[1];
-			m_heapImageDesc[item.CustomState].m_lods = args.LODs;
-			m_heapImageDesc[item.CustomState].m_format = args.format;
+			HeapImageDesc[item.CustomState].m_size[0] = args.size[0];
+			HeapImageDesc[item.CustomState].m_size[1] = args.size[1];
+			HeapImageDesc[item.CustomState].m_lods = args.LODs;
+			HeapImageDesc[item.CustomState].m_format = args.format;
 			StoreValidDesc(item);
 			break;
 		}
 
-		case OP_TYPE::IM_CROP:
+		case EOpType::IM_CROP:
 		{
 			OP::ImageCropArgs args = Program.GetOpArgs<OP::ImageCropArgs>(item.At);
 			switch (item.Stage)
@@ -6844,9 +6844,9 @@ namespace mu
 				break;
 
 			case 1:
-				m_heapImageDesc[item.CustomState].m_size[0] = args.sizeX;
-				m_heapImageDesc[item.CustomState].m_size[1] = args.sizeY;
-				m_heapImageDesc[item.CustomState].m_lods = 1;
+				HeapImageDesc[item.CustomState].m_size[0] = args.sizeX;
+				HeapImageDesc[item.CustomState].m_size[1] = args.sizeY;
+				HeapImageDesc[item.CustomState].m_lods = 1;
 				StoreValidDesc(item);
 				break;
 
@@ -6856,7 +6856,7 @@ namespace mu
 			break;
 		}
 
-		case OP_TYPE::IM_PATCH:
+		case EOpType::IM_PATCH:
 		{
 			OP::ImagePatchArgs args = Program.GetOpArgs<OP::ImagePatchArgs>(item.At);
 			switch (item.Stage)
@@ -6868,18 +6868,18 @@ namespace mu
 			break;
 		}
 
-		case OP_TYPE::IM_RASTERMESH:
+		case EOpType::IM_RASTERMESH:
 		{
 			OP::ImageRasterMeshArgs args = Program.GetOpArgs<OP::ImageRasterMeshArgs>(item.At);
-			m_heapImageDesc[item.CustomState].m_size[0] = args.sizeX;
-			m_heapImageDesc[item.CustomState].m_size[1] = args.sizeY;
-			m_heapImageDesc[item.CustomState].m_lods = 1;
-			m_heapImageDesc[item.CustomState].m_format = EImageFormat::L_UByte;
+			HeapImageDesc[item.CustomState].m_size[0] = args.sizeX;
+			HeapImageDesc[item.CustomState].m_size[1] = args.sizeY;
+			HeapImageDesc[item.CustomState].m_lods = 1;
+			HeapImageDesc[item.CustomState].m_format = EImageFormat::L_UByte;
 			StoreValidDesc(item);
 			break;
 		}
 
-		case OP_TYPE::IM_MAKEGROWMAP:
+		case EOpType::IM_MAKEGROWMAP:
 		{
 			OP::ImageMakeGrowMapArgs args = Program.GetOpArgs<OP::ImageMakeGrowMapArgs>(item.At);
 			switch (item.Stage)
@@ -6889,8 +6889,8 @@ namespace mu
 				break;
 
 			case 1:
-				m_heapImageDesc[item.CustomState].m_format = EImageFormat::L_UByte;
-				m_heapImageDesc[item.CustomState].m_lods = 1;
+				HeapImageDesc[item.CustomState].m_format = EImageFormat::L_UByte;
+				HeapImageDesc[item.CustomState].m_lods = 1;
 				StoreValidDesc(item);
 				break;
 
@@ -6901,7 +6901,7 @@ namespace mu
 			break;
 		}
 
-		case OP_TYPE::IM_DISPLACE:
+		case EOpType::IM_DISPLACE:
 		{
 			OP::ImageDisplaceArgs args = Program.GetOpArgs<OP::ImageDisplaceArgs>(item.At);
 			switch (item.Stage)
@@ -6913,7 +6913,7 @@ namespace mu
 			break;
 		}
 
-        case OP_TYPE::IM_TRANSFORM:
+        case EOpType::IM_TRANSFORM:
         {
 
 			OP::ImageTransformArgs Args = Program.GetOpArgs<OP::ImageTransformArgs>(item.At);
@@ -6927,13 +6927,13 @@ namespace mu
 			}
             case 1:
             {
-				m_heapImageDesc[item.CustomState].m_lods = 1;
-				m_heapImageDesc[item.CustomState].m_format = GetUncompressedFormat(m_heapImageDesc[item.CustomState].m_format);
+				HeapImageDesc[item.CustomState].m_lods = 1;
+				HeapImageDesc[item.CustomState].m_format = GetUncompressedFormat(HeapImageDesc[item.CustomState].m_format);
 				
 				if (!(Args.SizeX == 0 && Args.SizeY == 0))
 				{
-					m_heapImageDesc[item.CustomState].m_size[0] = Args.SizeX;
-					m_heapImageDesc[item.CustomState].m_size[1] = Args.SizeY;
+					HeapImageDesc[item.CustomState].m_size[0] = Args.SizeX;
+					HeapImageDesc[item.CustomState].m_size[1] = Args.SizeY;
 				}
 
 				StoreValidDesc(item);
@@ -6948,11 +6948,11 @@ namespace mu
 		}
 
 		default:
-			if (type != OP_TYPE::NONE)
+			if (type != EOpType::NONE)
 			{
 				// Operation not implemented
 				check(false);
-				m_heapImageDesc[item.CustomState] = FImageDesc();
+				HeapImageDesc[item.CustomState] = FImageDesc();
 			}
 			break;
 		}
