[mutable] More enforcement of coding standards.

diff --git a/Engine/Plugins/Mutable/Source/MutableRuntime/Private/MuR/CodeRunner.h b/Engine/Plugins/Mutable/Source/MutableRuntime/Private/MuR/CodeRunner.h
--- a/Engine/Plugins/Mutable/Source/MutableRuntime/Private/MuR/CodeRunner.h
+++ b/Engine/Plugins/Mutable/Source/MutableRuntime/Private/MuR/CodeRunner.h
@@ -9,7 +9,6 @@
 #include "MuR/ExtensionDataStreamer.h"
 #include "MuR/Image.h"
 #include "MuR/Layout.h"
-#include "MuR/MutableMemory.h"
 #include "MuR/Operations.h"
 #include "MuR/Ptr.h"
 #include "MuR/RefCounted.h"
@@ -33,9 +32,9 @@ namespace mu::MemoryCounters
 
 namespace  mu
 {
-	class Model;
-	class Parameters;
-	class RangeIndex;
+	class FModel;
+	class FParameters;
+	class FRangeIndex;
 
     /** Code execution of the mutable virtual machine. */
     class CodeRunner : public TSharedFromThis<CodeRunner>
@@ -45,28 +44,28 @@ namespace  mu
 
 	public:
 		static TSharedRef<CodeRunner> Create(
-				const Ptr<const Settings>&, 
-				class System::Private*, 
+				const FSettings&, 
+				class FSystem::Private*, 
 				EExecutionStrategy,
-				const TSharedPtr<const Model>&, 
-				const Parameters* pParams,
+				const TSharedPtr<const FModel>&, 
+				const FParameters* pParams,
 				OP::ADDRESS at, uint32 lodMask, uint8 executionOptions, int32 InImageLOD, FScheduledOp::EType);
 
 		// Private constructor to prevent stack allocation. In general we can not call AsShared() if the lifetime is
 		// bounded.
 		explicit CodeRunner(FPrivateToken, 
-			const Ptr<const Settings>&, 
-			class System::Private*, 
+			const FSettings&, 
+			class FSystem::Private*, 
 			EExecutionStrategy,
-			const TSharedPtr<const Model>&, 
-			const Parameters* pParams,
+			const TSharedPtr<const FModel>&, 
+			const FParameters* pParams,
 			OP::ADDRESS at, uint32 lodMask, uint8 executionOptions, int32 InImageLOD, FScheduledOp::EType);
 
     protected:
 		struct FProfileContext
 		{
 			uint32 NumRunOps = 0;
-			uint32 RunOpsPerType[int32(OP_TYPE::COUNT)] = {};
+			uint32 RunOpsPerType[int32(EOpType::COUNT)] = {};
 		};
 
         /** Type of data sometimes stored in the code runner heap to pass info between operation stages. */
@@ -108,37 +107,38 @@ namespace  mu
 					float MipValue;
 				} ImageTransform;
 			};
-            Ptr<RefCounted> Resource;
+
+			TSharedPtr<const FResource> Resource;
         };
 
 		// Assertion to know when FScheduledOpData size changes. It is ok to modifiy if needed. 
-		static_assert(sizeof(FScheduledOpData) == 4*4 + sizeof(Ptr<RefCounted>), "FScheduledOpData size changed.");
+		static_assert(sizeof(FScheduledOpData) == 4*4+sizeof(TSharedPtr<FResource>), "FScheduledOpData size changed.");
 
 
-        Ptr<RangeIndex> BuildCurrentOpRangeIndex( const FScheduledOp&, const Parameters*, const Model*, int32 ParameterIndex );
+        TSharedPtr<FRangeIndex> BuildCurrentOpRangeIndex( const FScheduledOp&, const FParameters*, const FModel*, int32 ParameterIndex );
 
-        void RunCode( const FScheduledOp&, const Parameters*, const TSharedPtr<const Model>&, uint32 LodMask);
+        void RunCode( const FScheduledOp&, const FParameters*, const TSharedPtr<const FModel>&, uint32 LodMask);
 
-        void RunCode_Conditional(const FScheduledOp&, const Model* );
-        void RunCode_Switch(const FScheduledOp&, const Model* );
-        void RunCode_Instance(const FScheduledOp&, const Model*, uint32 LodMask );
-        void RunCode_InstanceAddResource(const FScheduledOp&, const TSharedPtr<const Model>& Model, const Parameters* );
+        void RunCode_Conditional(const FScheduledOp&, const FModel* );
+        void RunCode_Switch(const FScheduledOp&, const FModel* );
+        void RunCode_Instance(const FScheduledOp&, const FModel*, uint32 LodMask );
+        void RunCode_InstanceAddResource(const FScheduledOp&, const TSharedPtr<const FModel>& Model, const FParameters* );
 
 		/** Return false incase of failure. */
-        bool RunCode_ConstantResource(const FScheduledOp&, const Model* );
+        bool RunCode_ConstantResource(const FScheduledOp&, const FModel* );
 
-        void RunCode_Mesh(const FScheduledOp&, const Model* );
-        void RunCode_Image(const FScheduledOp&, const Parameters*, const Model* );
-        void RunCode_Layout(const FScheduledOp&, const Model* );
-        void RunCode_Bool(const FScheduledOp&, const Parameters*, const Model* );
-        void RunCode_Int(const FScheduledOp&, const Parameters*, const Model* );
-        void RunCode_Scalar(const FScheduledOp&, const Parameters*, const Model* );
-        void RunCode_String(const FScheduledOp&, const Parameters*, const Model* );
-        void RunCode_Colour(const FScheduledOp&, const Parameters*, const Model* );
-        void RunCode_Projector(const FScheduledOp&, const Parameters*, const Model* );
-        void RunCode_Matrix(const FScheduledOp&, const Parameters*, const Model* );
+        void RunCode_Mesh(const FScheduledOp&, const FModel* );
+        void RunCode_Image(const FScheduledOp&, const FParameters*, const FModel* );
+        void RunCode_Layout(const FScheduledOp&, const FModel* );
+        void RunCode_Bool(const FScheduledOp&, const FParameters*, const FModel* );
+        void RunCode_Int(const FScheduledOp&, const FParameters*, const FModel* );
+        void RunCode_Scalar(const FScheduledOp&, const FParameters*, const FModel* );
+        void RunCode_String(const FScheduledOp&, const FParameters*, const FModel* );
+        void RunCode_Colour(const FScheduledOp&, const FParameters*, const FModel* );
+        void RunCode_Projector(const FScheduledOp&, const FParameters*, const FModel* );
+        void RunCode_Matrix(const FScheduledOp&, const FParameters*, const FModel* );
 
-		void RunCodeImageDesc(const FScheduledOp&, const Parameters*, const Model*, uint32 LodMask);
+		void RunCodeImageDesc(const FScheduledOp&, const FParameters*, const FModel*, uint32 LodMask);
 
     public:
 
@@ -152,20 +152,20 @@ namespace  mu
 		};
 
 		/** Load an external image asynchronously, returns an event to wait for complition and a cleanup function that must be called once the event has completed. */
-		TTuple<UE::Tasks::FTask, TFunction<void()>> LoadExternalImageAsync(FExternalResourceId Id, uint8 MipmapsToSkip, TFunction<void(Ptr<Image>)>& ResultCallback);
+		TTuple<UE::Tasks::FTask, TFunction<void()>> LoadExternalImageAsync(FExternalResourceId Id, uint8 MipmapsToSkip, TFunction<void(TSharedPtr<FImage>)>& ResultCallback);
  	    mu::FImageDesc GetExternalImageDesc(FName Id, uint8 MipmapsToSkip);
 
 		/** Load an external mesh asynchronously, returns an event to wait for complition and a cleanup function that must be called once the event has completed. */
-		TTuple<UE::Tasks::FTask, TFunction<void()>> LoadExternalMeshAsync(FExternalResourceId Id, TFunction<void(Ptr<Mesh>)>& ResultCallback);
+		TTuple<UE::Tasks::FTask, TFunction<void()>> LoadExternalMeshAsync(FExternalResourceId Id, TFunction<void(TSharedPtr<FMesh>)>& ResultCallback);
 
 		/** Settings that may affect the execution of some operations, like image conversion quality. */
-		Ptr<const Settings> m_pSettings;
+		FSettings Settings;
 
     protected:
         //! Heap of intermediate data pushed by some instructions and referred by others.
         //! It is not released until no operations are pending.
-		TArray<FScheduledOpData> m_heapData;
-		TArray<FImageDesc> m_heapImageDesc;
+		TArray<FScheduledOpData> HeapData;
+		TArray<FImageDesc> HeapImageDesc;
 
 		/** Only used for correct mip skipping with external images. It is the LOD for which the image is build. */
 		int32 ImageLOD;
@@ -253,7 +253,7 @@ namespace  mu
 			using StreamingDataContainerType = TArray<uint8, FDefaultMemoryTrackingAllocator<MemoryCounters::FStreamingMemoryCounter>>;
 
 			int32 RomIndex = -1;
-			ModelReader::FOperationID m_streamID = -1;
+			FModelReader::FOperationID m_streamID = -1;
 			StreamingDataContainerType m_streamBuffer;
 			UE::Tasks::FTask Event;
 		};
@@ -278,13 +278,13 @@ namespace  mu
 		class FLoadExtensionDataTask : public CodeRunner::FIssuedTask
 		{
 		public:
-			FLoadExtensionDataTask(const FScheduledOp& InOp, const ExtensionDataPtrConst& InExtensionData, int32 InModelConstantIndex)
+			FLoadExtensionDataTask(const FScheduledOp& InOp, const TSharedPtr<const FExtensionData>& InExtensionData, int32 InModelConstantIndex)
 				: FIssuedTask(InOp)
 				, Data(InExtensionData)
 				, ModelConstantIndex(InModelConstantIndex)
 			{
-				check(Data.get());
-				check(Data->Origin == ExtensionData::EOrigin::ConstantStreamed);
+				check(Data.Get());
+				check(Data->Origin == FExtensionData::EOrigin::ConstantStreamed);
 			}
 
 			// FIssuedTask interface
@@ -293,8 +293,8 @@ namespace  mu
 			virtual bool IsComplete(CodeRunner*) override;
 
 		private:
-			ExtensionDataPtrConst Data;
-			// The index of this constant in FProgram::m_constantExtensionData
+			TSharedPtr<const FExtensionData> Data;
+			// The index of this constant in FProgram::ConstantExtensionData
 			int32 ModelConstantIndex;
 
 			TSharedPtr<const FExtensionDataLoadHandle> LoadHandle;
@@ -418,7 +418,7 @@ namespace  mu
 		int32 GetInternalMemoryBytes() const
 		{
 			return sizeof(CodeRunner) 
-				+ m_heapData.GetAllocatedSize() + m_heapImageDesc.GetAllocatedSize()
+				+ HeapData.GetAllocatedSize() + HeapImageDesc.GetAllocatedSize()
 				+ ClosedTasks.GetAllocatedSize() + OpenTasks.GetAllocatedSize() + ScheduledStagePerOp.GetAllocatedSize()
 				// this contains smart pointers, approximate size like this:
 				+ IssuedTasks.Max() * ( sizeof(FIssuedTask) + 16);
@@ -451,176 +451,176 @@ namespace  mu
 
 		inline bool LoadBool(const FCacheAddress& From)
 		{
-			return m_pSystem->WorkingMemoryManager.CurrentInstanceCache->GetBool(From);
+			return System->WorkingMemoryManager.CurrentInstanceCache->GetBool(From);
 		}
 
 		inline float LoadInt(const FCacheAddress& From)
 		{
-			return m_pSystem->WorkingMemoryManager.CurrentInstanceCache->GetInt(From);
+			return System->WorkingMemoryManager.CurrentInstanceCache->GetInt(From);
 		}
 
 		inline float LoadScalar(const FCacheAddress& From)
 		{
-			return m_pSystem->WorkingMemoryManager.CurrentInstanceCache->GetScalar(From);
+			return System->WorkingMemoryManager.CurrentInstanceCache->GetScalar(From);
 		}
 
 		inline FVector4f LoadColor(const FCacheAddress& From)
 		{
-			return m_pSystem->WorkingMemoryManager.CurrentInstanceCache->GetColour(From);
+			return System->WorkingMemoryManager.CurrentInstanceCache->GetColour(From);
 		}
 
     	inline FMatrix44f LoadMatrix(const FCacheAddress& From)
 		{
-			return m_pSystem->WorkingMemoryManager.CurrentInstanceCache->GetMatrix(From);
+			return System->WorkingMemoryManager.CurrentInstanceCache->GetMatrix(From);
 		}
 
-		inline Ptr<const String> LoadString(const FCacheAddress& From)
+		inline TSharedPtr<const String> LoadString(const FCacheAddress& From)
 		{
-			return m_pSystem->WorkingMemoryManager.CurrentInstanceCache->GetString(From);
+			return System->WorkingMemoryManager.CurrentInstanceCache->GetString(From);
 		}
 
 		inline FProjector LoadProjector(const FCacheAddress& From)
 		{
-			return m_pSystem->WorkingMemoryManager.CurrentInstanceCache->GetProjector(From);
+			return System->WorkingMemoryManager.CurrentInstanceCache->GetProjector(From);
 		}
 
-		inline Ptr<const Mesh> LoadMesh(const FCacheAddress& From)
+		inline TSharedPtr<const FMesh> LoadMesh(const FCacheAddress& From)
 		{
-			return m_pSystem->WorkingMemoryManager.LoadMesh(From);
+			return System->WorkingMemoryManager.LoadMesh(From);
 		}
 
-		inline Ptr<const Image> LoadImage(const FCacheAddress& From)
+		inline TSharedPtr<const FImage> LoadImage(const FCacheAddress& From)
 		{
-			return m_pSystem->WorkingMemoryManager.LoadImage(From);
+			return System->WorkingMemoryManager.LoadImage(From);
 		}
 
-		inline Ptr<const Layout> LoadLayout(const FCacheAddress& From)
+		inline TSharedPtr<const FLayout> LoadLayout(const FCacheAddress& From)
 		{
-			return m_pSystem->WorkingMemoryManager.CurrentInstanceCache->GetLayout(From);
+			return System->WorkingMemoryManager.CurrentInstanceCache->GetLayout(From);
 		}
 
-		inline Ptr<const Instance> LoadInstance(const FCacheAddress& From)
+		inline TSharedPtr<const FInstance> LoadInstance(const FCacheAddress& From)
 		{
-			return m_pSystem->WorkingMemoryManager.CurrentInstanceCache->GetInstance(From);
+			return System->WorkingMemoryManager.CurrentInstanceCache->GetInstance(From);
 		}
 
-		inline Ptr<const ExtensionData> LoadExtensionData(const FCacheAddress& From)
+		inline TSharedPtr<const FExtensionData> LoadExtensionData(const FCacheAddress& From)
 		{
-			return m_pSystem->WorkingMemoryManager.CurrentInstanceCache->GetExtensionData(From);
+			return System->WorkingMemoryManager.CurrentInstanceCache->GetExtensionData(From);
 		}
 
 		inline void StoreValidDesc(const FCacheAddress& To)
 		{
-			m_pSystem->WorkingMemoryManager.CurrentInstanceCache->SetValidDesc(To);
+			System->WorkingMemoryManager.CurrentInstanceCache->SetValidDesc(To);
 		}
 
 		inline void StoreBool(const FCacheAddress& To, bool Value)
 		{
-			m_pSystem->WorkingMemoryManager.CurrentInstanceCache->SetBool(To, Value);
+			System->WorkingMemoryManager.CurrentInstanceCache->SetBool(To, Value);
 		}
 
 		inline void StoreInt(const FCacheAddress& To, int32 Value)
 		{
-			m_pSystem->WorkingMemoryManager.CurrentInstanceCache->SetInt(To, Value);
+			System->WorkingMemoryManager.CurrentInstanceCache->SetInt(To, Value);
 		}
 
 		inline void StoreScalar(const FCacheAddress& To, float Value)
 		{
-			m_pSystem->WorkingMemoryManager.CurrentInstanceCache->SetScalar(To, Value);
+			System->WorkingMemoryManager.CurrentInstanceCache->SetScalar(To, Value);
 		}
 
-		inline void StoreString(const FCacheAddress& To, Ptr<const String> Value)
+		inline void StoreString(const FCacheAddress& To, TSharedPtr<const String> Value)
 		{
-			m_pSystem->WorkingMemoryManager.CurrentInstanceCache->SetString(To, Value);
+			System->WorkingMemoryManager.CurrentInstanceCache->SetString(To, Value);
 		}
 
 		inline void StoreColor(const FCacheAddress& To, const FVector4f& Value)
 		{
-			m_pSystem->WorkingMemoryManager.CurrentInstanceCache->SetColour(To, Value);
+			System->WorkingMemoryManager.CurrentInstanceCache->SetColour(To, Value);
 		}
 
     	inline void StoreMatrix(const FCacheAddress& To, const FMatrix44f& Value)
 		{
-			m_pSystem->WorkingMemoryManager.CurrentInstanceCache->SetMatrix(To, Value);
+			System->WorkingMemoryManager.CurrentInstanceCache->SetMatrix(To, Value);
 		}
     	
 		inline void StoreProjector(const FCacheAddress& To, const FProjector& Value)
 		{
-			m_pSystem->WorkingMemoryManager.CurrentInstanceCache->SetProjector(To, Value);
+			System->WorkingMemoryManager.CurrentInstanceCache->SetProjector(To, Value);
 		}
 
-		inline void StoreMesh(const FCacheAddress& To, Ptr<const Mesh> Resource)
+		inline void StoreMesh(const FCacheAddress& To, TSharedPtr<const FMesh> Resource)
 		{
-			m_pSystem->WorkingMemoryManager.StoreMesh(To, Resource);
+			System->WorkingMemoryManager.StoreMesh(To, Resource);
 		}
 
-		inline void StoreImage(const FCacheAddress& To, Ptr<const Image> Resource)
+		inline void StoreImage(const FCacheAddress& To, TSharedPtr<const FImage> Resource)
 		{
-			m_pSystem->WorkingMemoryManager.StoreImage(To, Resource);
+			System->WorkingMemoryManager.StoreImage(To, Resource);
 		}
 
-		inline void StoreLayout(const FCacheAddress& To, Ptr<const Layout> Resource)
+		inline void StoreLayout(const FCacheAddress& To, TSharedPtr<const FLayout> Resource)
 		{
-			m_pSystem->WorkingMemoryManager.CurrentInstanceCache->SetLayout(To, Resource);
+			System->WorkingMemoryManager.CurrentInstanceCache->SetLayout(To, Resource);
 		}
 
-		inline void StoreInstance(const FCacheAddress& To, Ptr<const Instance> Resource)
+		inline void StoreInstance(const FCacheAddress& To, TSharedPtr<const FInstance> Resource)
 		{
-			m_pSystem->WorkingMemoryManager.CurrentInstanceCache->SetInstance(To, Resource);
+			System->WorkingMemoryManager.CurrentInstanceCache->SetInstance(To, Resource);
 		}
 
-		inline void StoreExtensionData(const FCacheAddress& To, Ptr<const ExtensionData> Resource)
+		inline void StoreExtensionData(const FCacheAddress& To, TSharedPtr<const FExtensionData> Resource)
 		{
-			m_pSystem->WorkingMemoryManager.CurrentInstanceCache->SetExtensionData(To, Resource);
+			System->WorkingMemoryManager.CurrentInstanceCache->SetExtensionData(To, Resource);
 		}
 
-		inline Ptr<Image> CreateImage(uint32 SizeX, uint32 SizeY, uint32 Lods, EImageFormat Format, EInitializationType Init)
+		inline TSharedPtr<FImage> CreateImage(uint32 SizeX, uint32 SizeY, uint32 Lods, EImageFormat Format, EInitializationType Init)
 		{
-			Ptr<Image> Result = m_pSystem->WorkingMemoryManager.CreateImage(SizeX, SizeY, Lods, Format, Init);
+			TSharedPtr<FImage> Result = System->WorkingMemoryManager.CreateImage(SizeX, SizeY, Lods, Format, Init);
 			return MoveTemp(Result);
 		}
 
-		Ptr<Image> CreateImageLike(const Image* Ref, EInitializationType Init)
+		TSharedPtr<FImage> CreateImageLike(const FImage* Ref, EInitializationType Init)
 		{
-			Ptr<Image> Result = m_pSystem->WorkingMemoryManager.CreateImage(Ref->GetSizeX(), Ref->GetSizeY(), Ref->GetLODCount(), Ref->GetFormat(), Init);
+			TSharedPtr<FImage> Result = System->WorkingMemoryManager.CreateImage(Ref->GetSizeX(), Ref->GetSizeY(), Ref->GetLODCount(), Ref->GetFormat(), Init);
 			return MoveTemp(Result);
 		}
 
 		/** Ref will be nulled and relesed in any case. */
-		inline Ptr<Image> CloneOrTakeOver(Ptr<const Image>& Ref)
+		inline TSharedPtr<FImage> CloneOrTakeOver(TSharedPtr<const FImage>& Ref)
 		{
-			Ptr<Image> Result = m_pSystem->WorkingMemoryManager.CloneOrTakeOver(Ref);
+			TSharedPtr<FImage> Result = System->WorkingMemoryManager.CloneOrTakeOver(Ref);
 			return MoveTemp(Result);
 		}
 
-		inline void Release(Ptr<const Image>& Resource)
+		inline void Release(TSharedPtr<const FImage>& Resource)
 		{
-			return m_pSystem->WorkingMemoryManager.Release(Resource);
+			return System->WorkingMemoryManager.Release(Resource);
 		}
 
-		inline void Release(Ptr<Image>& Resource)
+		inline void Release(TSharedPtr<FImage>& Resource)
 		{
-			return m_pSystem->WorkingMemoryManager.Release(Resource);
+			return System->WorkingMemoryManager.Release(Resource);
 		}
 
-		[[nodiscard]] inline Ptr<Mesh> CreateMesh(int32 BudgetReserveSize = 0)
+		[[nodiscard]] inline TSharedPtr<FMesh> CreateMesh(int32 BudgetReserveSize = 0)
 		{
-			return m_pSystem->WorkingMemoryManager.CreateMesh(BudgetReserveSize);
+			return System->WorkingMemoryManager.CreateMesh(BudgetReserveSize);
 		}
 
-		[[nodiscard]] inline Ptr<Mesh> CloneOrTakeOver(Ptr<const Mesh>& Ref)
+		[[nodiscard]] inline TSharedPtr<FMesh> CloneOrTakeOver(TSharedPtr<const FMesh>& Ref)
 		{
-			return m_pSystem->WorkingMemoryManager.CloneOrTakeOver(Ref);
+			return System->WorkingMemoryManager.CloneOrTakeOver(Ref);
 		}
 
-		inline void Release(Ptr<const Mesh>& Resource)
+		inline void Release(TSharedPtr<const FMesh>& Resource)
 		{
-			m_pSystem->WorkingMemoryManager.Release(Resource);
+			System->WorkingMemoryManager.Release(Resource);
 		}
 
-		inline void Release(Ptr<Mesh>& Resource)
+		inline void Release(TSharedPtr<FMesh>& Resource)
 		{
-			m_pSystem->WorkingMemoryManager.Release(Resource);
+			System->WorkingMemoryManager.Release(Resource);
 		}
 
 	// TODO: protect this.
@@ -631,10 +631,10 @@ namespace  mu
 		*/
 		bool bUnrecoverableError = false;
 
-		System::Private* m_pSystem = nullptr;
-		TSharedPtr<const Model> m_pModel = nullptr;
-		const Parameters* m_pParams = nullptr;
-		uint32 m_lodMask = 0;
+		FSystem::Private* System = nullptr;
+		TSharedPtr<const FModel> Model = nullptr;
+		const FParameters* Params = nullptr;
+		uint32 LODMask = 0;
 
     private:
     	struct FRomLoadOps
@@ -718,7 +718,7 @@ namespace  mu
 
 			if (dep.Type == FScheduledOp::EType::Full)
 			{
-				m_pSystem->WorkingMemoryManager.CurrentInstanceCache->IncreaseHitCount(at);
+				System->WorkingMemoryManager.CurrentInstanceCache->IncreaseHitCount(at);
 			}
 		}
 
@@ -743,7 +743,7 @@ namespace  mu
 			else
 			{
 				// We could wait. Let's see if we have enough memory to issue tasks anyway.
-				bool bHaveEnoughMemory = !m_pSystem->WorkingMemoryManager.IsMemoryBudgetFull();
+				bool bHaveEnoughMemory = !System->WorkingMemoryManager.IsMemoryBudgetFull();
 				if (bHaveEnoughMemory)
 				{
 					return  true;
@@ -770,20 +770,20 @@ namespace  mu
 			},
 
 			// Release
-			[Runner](Ptr<Image>& i)
+			[Runner](TSharedPtr<FImage>& i)
 			{
 				Runner->Release(i);
 			},
 
 			// Clone
-			[Runner](const Image* i)
+			[Runner](const FImage* i)
 			{
-				Ptr<Image> New = Runner->CreateImage(i->GetSizeX(), i->GetSizeY(), i->GetLODCount(), i->GetFormat(), EInitializationType::NotInitialized);
+				TSharedPtr<FImage> New = Runner->CreateImage(i->GetSizeX(), i->GetSizeY(), i->GetLODCount(), i->GetFormat(), EInitializationType::NotInitialized);
 				New->Copy(i);
 				return New;
 			},
 
-			Runner->m_pSystem->ImagePixelFormatOverride
+			Runner->System->ImagePixelFormatOverride
 		);
 	}
 
diff --git a/Engine/Plugins/Mutable/Source/MutableRuntime/Private/MuR/CodeRunnerTasked.cpp b/Engine/Plugins/Mutable/Source/MutableRuntime/Private/MuR/CodeRunnerTasked.cpp
--- a/Engine/Plugins/Mutable/Source/MutableRuntime/Private/MuR/CodeRunnerTasked.cpp
+++ b/Engine/Plugins/Mutable/Source/MutableRuntime/Private/MuR/CodeRunnerTasked.cpp
@@ -75,7 +75,7 @@ namespace mu
 	* to be in this table.
 	*/
 	static constexpr int32 sMemoryWeightsStageCount = 4;
-	static const int8 sMemoryWeights[int32(OP_TYPE::COUNT)][sMemoryWeightsStageCount] =
+	static const int8 sMemoryWeights[int32(EOpType::COUNT)][sMemoryWeightsStageCount] =
 	{ 
 		{   0,   0,   0,   0 },	// NONE
 		    
@@ -201,16 +201,16 @@ namespace mu
 		{   0,   0,   0,   0 },	// LA_FROMMESH
 	};
 
-	static_assert(sizeof(sMemoryWeights)/sMemoryWeightsStageCount == int32(OP_TYPE::COUNT));
+	static_assert(sizeof(sMemoryWeights)/sMemoryWeightsStageCount == int32(EOpType::COUNT));
 
 
 	int32 CodeRunner::GetOpEstimatedMemoryDelta( const FScheduledOp& Candidate, const FProgram& Program )
 	{
 		int32 OpDelta = 0;
 		int32 Stage = FMath::Min(int32(Candidate.Stage), sMemoryWeightsStageCount - 1);
-		OP_TYPE OpType = Program.GetOpType(Candidate.At);
+		EOpType OpType = Program.GetOpType(Candidate.At);
 
-		if (OpType == OP_TYPE::IM_PIXELFORMAT && Stage == 1)
+		if (OpType == EOpType::IM_PIXELFORMAT && Stage == 1)
 		{
 			// TODO: We should get the actual format from the arguments for a more precise calculation.
 			OP::ImagePixelFormatArgs args = Program.GetOpArgs<OP::ImagePixelFormatArgs>(Candidate.At);
@@ -271,8 +271,8 @@ namespace mu
 
 		bUnrecoverableError = false;
 
-		m_heapData.SetNum(0, EAllowShrinking::No);
-		m_heapImageDesc.SetNum(0, EAllowShrinking::No);
+		HeapData.SetNum(0, EAllowShrinking::No);
+		HeapImageDesc.SetNum(0, EAllowShrinking::No);
 		
 		RunnerCompletionEvent = UE::Tasks::FTaskEvent(TEXT("CodeRunnerCompletionEvent"));
 		
@@ -368,7 +368,7 @@ namespace mu
 				//	int32 BestOp = 0;
 				//	int8 BestDelta = TNumericLimits<int8>::Max();
 				//	int32 OpenOpCount = OpenTasks.Num();
-				//	const FProgram& Program = m_pModel->GetPrivate()->m_program;
+				//	const FProgram& Program = Model->GetPrivate()->Program;
 				//	for (int32 OpIndex = 0; OpIndex < OpenOpCount; ++OpIndex)
 				//	{
 				//		const FScheduledOp& Candidate = OpenTasks[OpIndex];
@@ -403,7 +403,7 @@ namespace mu
 				// Special processing in case it is an ImageDesc operation
 				if (item.Type == FScheduledOp::EType::ImageDesc)
 				{
-					RunCodeImageDesc(item, m_pParams, m_pModel.Get(), m_lodMask);
+					RunCodeImageDesc(item, Params, Model.Get(), LODMask);
 					continue;
 				}
 
@@ -435,7 +435,7 @@ namespace mu
 				else
 				{
 					// Run immediately
-					RunCode(item, m_pParams, m_pModel, m_lodMask);
+					RunCode(item, Params, Model, LODMask);
 
 					if (ScheduledStagePerOp[item] == item.Stage + 1)
 					{
@@ -448,7 +448,7 @@ namespace mu
 				if (ProfileContext)
 				{
 					++ProfileContext->NumRunOps;
-					++ProfileContext->RunOpsPerType[int32(m_pModel->GetPrivate()->m_program.GetOpType(item.At))];
+					++ProfileContext->RunOpsPerType[int32(Model->GetPrivate()->Program.GetOpType(item.At))];
 				}
 			}
 
@@ -540,12 +540,12 @@ namespace mu
 						}
 					}
 
-					m_pSystem->WorkingMemoryManager.InvalidateRunnerThread();
+					System->WorkingMemoryManager.InvalidateRunnerThread();
 
 					UE::Tasks::Launch(TEXT("CodeRunnerFromIssuedTasksTask"),
 						[Runner = AsShared(), ProfileContext = MoveTemp(ProfileContext)]() mutable
 						{
-							Runner->m_pSystem->WorkingMemoryManager.ResetRunnerThread();
+							Runner->System->WorkingMemoryManager.ResetRunnerThread();
 
 							constexpr bool bForceInlineExecution = false;
 							Runner->Run(MoveTemp(ProfileContext), bForceInlineExecution);
@@ -574,12 +574,12 @@ namespace mu
 			{
 				if (FTimespan::FromSeconds(FPlatformTime::Seconds()) > TimeOut)
 				{
-					m_pSystem->WorkingMemoryManager.InvalidateRunnerThread();
+					System->WorkingMemoryManager.InvalidateRunnerThread();
 
 					UE::Tasks::Launch(TEXT("CodeRunnerFromTimeoutTask"),
 						[Runner = AsShared(), ProfileContext = MoveTemp(ProfileContext)]() mutable
 						{
-							Runner->m_pSystem->WorkingMemoryManager.ResetRunnerThread();
+							Runner->System->WorkingMemoryManager.ResetRunnerThread();
 
 							constexpr bool bForceInlineExecution = false;
 							Runner->Run(MoveTemp(ProfileContext), bForceInlineExecution);
@@ -593,12 +593,12 @@ namespace mu
 
 		if (ProfileContext)
 		{
-			UE_LOG(LogMutableCore, Log, TEXT("Mutable Heap Bytes: %d"), m_heapData.Num()* m_heapData.GetTypeSize());
+			UE_LOG(LogMutableCore, Log, TEXT("Mutable Heap Bytes: %d"), HeapData.Num()* HeapData.GetTypeSize());
 			UE_LOG(LogMutableCore, Log, TEXT("Ran ops : %5d "), ProfileContext->NumRunOps);
 
 			constexpr int32 HistogramSize = 8;
 			int32 MostCommonOps[HistogramSize] = {};
-			for (int32 OpIndex = 0; OpIndex < int32(OP_TYPE::COUNT); ++OpIndex)
+			for (int32 OpIndex = 0; OpIndex < int32(EOpType::COUNT); ++OpIndex)
 			{
 				for (int32 HistIndex = 0; HistIndex < HistogramSize; ++HistIndex)
 				{
@@ -630,8 +630,8 @@ namespace mu
 	//---------------------------------------------------------------------------------------------
 	void CodeRunner::GetImageDescResult(FImageDesc& OutDesc)
 	{
-		check( m_heapImageDesc.Num()>0 );
-		OutDesc = m_heapImageDesc[0];
+		check( HeapImageDesc.Num()>0 );
+		OutDesc = HeapImageDesc[0];
 	}
 
 	//---------------------------------------------------------------------------------------------
@@ -652,9 +652,9 @@ namespace mu
 	private:
 		int32 ImageCompressionQuality = 0;
 		OP::ImageLayerArgs Args;
-		Ptr<const Image> Blended;
-		Ptr<const Image> Mask;
-		Ptr<Image> Result;
+		TSharedPtr<const FImage> Blended;
+		TSharedPtr<const FImage> Mask;
+		TSharedPtr<FImage> Result;
 		EImageFormat InitialFormat = EImageFormat::None;
 	};
 
@@ -665,9 +665,9 @@ namespace mu
 		MUTABLE_CPUPROFILER_SCOPE(FImageLayerTask_Prepare);
 		bOutFailed = false;
 
-		ImageCompressionQuality = Runner->m_pSettings->ImageCompressionQuality;
+		ImageCompressionQuality = Runner->Settings.ImageCompressionQuality;
 
-		Ptr<const Image> Base = Runner->LoadImage({ Args.base, Op.ExecutionIndex, Op.ExecutionOptions });
+		TSharedPtr<const FImage> Base = Runner->LoadImage({ Args.base, Op.ExecutionIndex, Op.ExecutionOptions });
 		check(!Base || Base->GetFormat() < EImageFormat::Count);
 
 		Blended = Runner->LoadImage({ Args.blended, Op.ExecutionIndex, Op.ExecutionOptions });
@@ -703,9 +703,9 @@ namespace mu
 		if (IsCompressedFormat(InitialFormat))
 		{
 			EImageFormat UncompressedFormat = GetUncompressedFormat(InitialFormat);
-			Ptr<Image> Formatted = Runner->CreateImage( Base->GetSizeX(), Base->GetSizeY(), Base->GetLODCount(), UncompressedFormat, EInitializationType::NotInitialized );
+			TSharedPtr<FImage> Formatted = Runner->CreateImage( Base->GetSizeX(), Base->GetSizeY(), Base->GetLODCount(), UncompressedFormat, EInitializationType::NotInitialized );
 			bool bSuccess = false;
-			ImOp.ImagePixelFormat(bSuccess, ImageCompressionQuality, Formatted.get(), Base.get());
+			ImOp.ImagePixelFormat(bSuccess, ImageCompressionQuality, Formatted.Get(), Base.Get());
 			check(bSuccess); // Decompression cannot fail
 			Runner->Release(Base);
 			Base = Formatted;
@@ -714,9 +714,9 @@ namespace mu
 		bool bMustHaveSameFormat = !(Args.flags & (OP::ImageLayerArgs::F_BASE_RGB_FROM_ALPHA | OP::ImageLayerArgs::F_BLENDED_RGB_FROM_ALPHA));
 		if (Blended && InitialFormat != Blended->GetFormat() && bMustHaveSameFormat)
 		{
-			Ptr<Image> Formatted = Runner->CreateImage(Blended->GetSizeX(), Blended->GetSizeY(), Blended->GetLODCount(), Base->GetFormat(), EInitializationType::NotInitialized);
+			TSharedPtr<FImage> Formatted = Runner->CreateImage(Blended->GetSizeX(), Blended->GetSizeY(), Blended->GetLODCount(), Base->GetFormat(), EInitializationType::NotInitialized);
 			bool bSuccess = false;
-			ImOp.ImagePixelFormat(bSuccess, ImageCompressionQuality, Formatted.get(), Blended.get());
+			ImOp.ImagePixelFormat(bSuccess, ImageCompressionQuality, Formatted.Get(), Blended.Get());
 			check(bSuccess); // Decompression cannot fail
 			Runner->Release(Blended);
 			Blended = Formatted;
@@ -725,8 +725,8 @@ namespace mu
 		if (Base->GetSize() != Blended->GetSize())
 		{
 			MUTABLE_CPUPROFILER_SCOPE(ImageResize_EmergencyFix);
-			Ptr<Image> Resized = Runner->CreateImage(Base->GetSizeX(), Base->GetSizeY(), 1, Blended->GetFormat(), EInitializationType::NotInitialized);
-			ImOp.ImageResizeLinear(Resized.get(), ImageCompressionQuality, Blended.get());
+			TSharedPtr<FImage> Resized = Runner->CreateImage(Base->GetSizeX(), Base->GetSizeY(), 1, Blended->GetFormat(), EInitializationType::NotInitialized);
+			ImOp.ImageResizeLinear(Resized.Get(), ImageCompressionQuality, Blended.Get());
 			Runner->Release(Blended);
 			Blended = Resized;
 
@@ -737,8 +737,8 @@ namespace mu
 			if (Base->GetSize() != Mask->GetSize())
 			{
 				MUTABLE_CPUPROFILER_SCOPE(ImageResize_EmergencyFix);
-				Ptr<Image> Resized = Runner->CreateImage(Base->GetSizeX(), Base->GetSizeY(), 1, Mask->GetFormat(), EInitializationType::NotInitialized);
-				ImOp.ImageResizeLinear(Resized.get(), ImageCompressionQuality, Mask.get());
+				TSharedPtr<FImage> Resized = Runner->CreateImage(Base->GetSizeX(), Base->GetSizeY(), 1, Mask->GetFormat(), EInitializationType::NotInitialized);
+				ImOp.ImageResizeLinear(Resized.Get(), ImageCompressionQuality, Mask.Get());
 				Runner->Release(Mask);
 				Mask = Resized;
 
@@ -753,15 +753,15 @@ namespace mu
 				int32 LevelCount = Base->GetLODCount();
 
 				// Uncompress mask to avoid excessive RLE-compression and decompression in the following ops.
-				Ptr<Image> UncompressedMask = Runner->CreateImage(Mask->GetSizeX(), Mask->GetSizeY(), Mask->GetLODCount(), GetUncompressedFormat(Mask->GetFormat()), EInitializationType::NotInitialized);
+				TSharedPtr<FImage> UncompressedMask = Runner->CreateImage(Mask->GetSizeX(), Mask->GetSizeY(), Mask->GetLODCount(), GetUncompressedFormat(Mask->GetFormat()), EInitializationType::NotInitialized);
 				bool bSuccess = false;
-				ImOp.ImagePixelFormat(bSuccess, ImageCompressionQuality, UncompressedMask.get(), Mask.get());
+				ImOp.ImagePixelFormat(bSuccess, ImageCompressionQuality, UncompressedMask.Get(), Mask.Get());
 
-				Ptr<Image> MaskFix = UncompressedMask;
+				TSharedPtr<FImage> MaskFix = UncompressedMask;
 				MaskFix->DataStorage.SetNumLODs(LevelCount);
 
 				FMipmapGenerationSettings Settings{};
-				ImOp.ImageMipmap(ImageCompressionQuality, MaskFix.get(), MaskFix.get(), StartLevel, LevelCount, Settings);
+				ImOp.ImageMipmap(ImageCompressionQuality, MaskFix.Get(), MaskFix.Get(), StartLevel, LevelCount, Settings);
 
 				Runner->Release(Mask);
 				Mask = MaskFix;
@@ -805,11 +805,11 @@ namespace mu
 			constexpr bool bUseVectorImplementation = false;	
 			if constexpr (bUseVectorImplementation)
 			{
-				BufferLayerCompositeVector<VectorBlendChannelMasked, VectorLightenChannel, false>(Result.get(), Blended.get(), bOnlyOneMip, Args.BlendAlphaSourceChannel);
+				BufferLayerCompositeVector<VectorBlendChannelMasked, VectorLightenChannel, false>(Result.Get(), Blended.Get(), bOnlyOneMip, Args.BlendAlphaSourceChannel);
 			}
 			else
 			{
-				BufferLayerComposite<BlendChannelMasked, LightenChannel, false>(Result.get(), Blended.get(), bOnlyOneMip, Args.BlendAlphaSourceChannel);
+				BufferLayerComposite<BlendChannelMasked, LightenChannel, false>(Result.Get(), Blended.Get(), bOnlyOneMip, Args.BlendAlphaSourceChannel);
 			}
 		}
 
@@ -826,16 +826,16 @@ namespace mu
 
 			switch (EBlendType(Args.blendType))
 			{
-			case EBlendType::BT_NORMAL_COMBINE: ImageNormalCombine(Result.get(), Result.get(), Mask.get(), Blended.get(), bOnlyOneMip); break;
-			case EBlendType::BT_SOFTLIGHT: BufferLayer<SoftLightChannelMasked, SoftLightChannel, true>(Result.get(), Result.get(), Mask.get(), Blended.get(), bApplyColorBlendToAlpha, bOnlyOneMip); break;
-			case EBlendType::BT_HARDLIGHT: BufferLayer<HardLightChannelMasked, HardLightChannel, true>(Result.get(), Result.get(), Mask.get(), Blended.get(), bApplyColorBlendToAlpha, bOnlyOneMip); break;
-			case EBlendType::BT_BURN: BufferLayer<BurnChannelMasked, BurnChannel, true>(Result.get(), Result.get(), Mask.get(), Blended.get(), bApplyColorBlendToAlpha, bOnlyOneMip); break;
-			case EBlendType::BT_DODGE: BufferLayer<DodgeChannelMasked, DodgeChannel, true>(Result.get(), Result.get(), Mask.get(), Blended.get(), bApplyColorBlendToAlpha, bOnlyOneMip); break;
-			case EBlendType::BT_SCREEN: BufferLayer<ScreenChannelMasked, ScreenChannel, true>(Result.get(), Result.get(), Mask.get(), Blended.get(), bApplyColorBlendToAlpha, bOnlyOneMip); break;
-			case EBlendType::BT_OVERLAY: BufferLayer<OverlayChannelMasked, OverlayChannel, true>(Result.get(), Result.get(), Mask.get(), Blended.get(), bApplyColorBlendToAlpha, bOnlyOneMip); break;
-			case EBlendType::BT_LIGHTEN: BufferLayer<LightenChannelMasked, LightenChannel, true>(Result.get(), Result.get(), Mask.get(), Blended.get(), bApplyColorBlendToAlpha, bOnlyOneMip); break;
-			case EBlendType::BT_MULTIPLY: BufferLayer<MultiplyChannelMasked, MultiplyChannel, true>(Result.get(), Result.get(), Mask.get(), Blended.get(), bApplyColorBlendToAlpha, bOnlyOneMip); break;
-			case EBlendType::BT_BLEND: BufferLayer<BlendChannelMasked, BlendChannel, true>(Result.get(), Result.get(), Mask.get(), Blended.get(), bApplyColorBlendToAlpha, bOnlyOneMip); break;
+			case EBlendType::BT_NORMAL_COMBINE: ImageNormalCombine(Result.Get(), Result.Get(), Mask.Get(), Blended.Get(), bOnlyOneMip); break;
+			case EBlendType::BT_SOFTLIGHT: BufferLayer<SoftLightChannelMasked, SoftLightChannel, true>(Result.Get(), Result.Get(), Mask.Get(), Blended.Get(), bApplyColorBlendToAlpha, bOnlyOneMip); break;
+			case EBlendType::BT_HARDLIGHT: BufferLayer<HardLightChannelMasked, HardLightChannel, true>(Result.Get(), Result.Get(), Mask.Get(), Blended.Get(), bApplyColorBlendToAlpha, bOnlyOneMip); break;
+			case EBlendType::BT_BURN: BufferLayer<BurnChannelMasked, BurnChannel, true>(Result.Get(), Result.Get(), Mask.Get(), Blended.Get(), bApplyColorBlendToAlpha, bOnlyOneMip); break;
+			case EBlendType::BT_DODGE: BufferLayer<DodgeChannelMasked, DodgeChannel, true>(Result.Get(), Result.Get(), Mask.Get(), Blended.Get(), bApplyColorBlendToAlpha, bOnlyOneMip); break;
+			case EBlendType::BT_SCREEN: BufferLayer<ScreenChannelMasked, ScreenChannel, true>(Result.Get(), Result.Get(), Mask.Get(), Blended.Get(), bApplyColorBlendToAlpha, bOnlyOneMip); break;
+			case EBlendType::BT_OVERLAY: BufferLayer<OverlayChannelMasked, OverlayChannel, true>(Result.Get(), Result.Get(), Mask.Get(), Blended.Get(), bApplyColorBlendToAlpha, bOnlyOneMip); break;
+			case EBlendType::BT_LIGHTEN: BufferLayer<LightenChannelMasked, LightenChannel, true>(Result.Get(), Result.Get(), Mask.Get(), Blended.Get(), bApplyColorBlendToAlpha, bOnlyOneMip); break;
+			case EBlendType::BT_MULTIPLY: BufferLayer<MultiplyChannelMasked, MultiplyChannel, true>(Result.Get(), Result.Get(), Mask.Get(), Blended.Get(), bApplyColorBlendToAlpha, bOnlyOneMip); break;
+			case EBlendType::BT_BLEND: BufferLayer<BlendChannelMasked, BlendChannel, true>(Result.Get(), Result.Get(), Mask.Get(), Blended.Get(), bApplyColorBlendToAlpha, bOnlyOneMip); break;
 			case EBlendType::BT_NONE: break;
 			default: check(false);
 			}
@@ -850,15 +850,15 @@ namespace mu
 			switch (EBlendType(Args.blendType))
 			{
 			case EBlendType::BT_NORMAL_COMBINE: check(false); break;
-			case EBlendType::BT_SOFTLIGHT: BufferLayerEmbeddedMask<SoftLightChannelMasked, SoftLightChannel, false>(Result.get(), Result.get(), Blended.get(), bApplyColorBlendToAlpha, bOnlyOneMip); break;
-			case EBlendType::BT_HARDLIGHT: BufferLayerEmbeddedMask<HardLightChannelMasked, HardLightChannel, false>(Result.get(), Result.get(), Blended.get(), bApplyColorBlendToAlpha, bOnlyOneMip); break;
-			case EBlendType::BT_BURN: BufferLayerEmbeddedMask<BurnChannelMasked, BurnChannel, false>(Result.get(), Result.get(), Blended.get(), bApplyColorBlendToAlpha, bOnlyOneMip); break;
-			case EBlendType::BT_DODGE: BufferLayerEmbeddedMask<DodgeChannelMasked, DodgeChannel, false>(Result.get(), Result.get(), Blended.get(), bApplyColorBlendToAlpha, bOnlyOneMip); break;
-			case EBlendType::BT_SCREEN: BufferLayerEmbeddedMask<ScreenChannelMasked, ScreenChannel, false>(Result.get(), Result.get(), Blended.get(), bApplyColorBlendToAlpha, bOnlyOneMip); break;
-			case EBlendType::BT_OVERLAY: BufferLayerEmbeddedMask<OverlayChannelMasked, OverlayChannel, false>(Result.get(), Result.get(), Blended.get(), bApplyColorBlendToAlpha, bOnlyOneMip); break;
-			case EBlendType::BT_LIGHTEN: BufferLayerEmbeddedMask<LightenChannelMasked, LightenChannel, false>(Result.get(), Result.get(), Blended.get(), bApplyColorBlendToAlpha, bOnlyOneMip); break;
-			case EBlendType::BT_MULTIPLY: BufferLayerEmbeddedMask<MultiplyChannelMasked, MultiplyChannel, false>(Result.get(), Result.get(), Blended.get(), bApplyColorBlendToAlpha, bOnlyOneMip); break;
-			case EBlendType::BT_BLEND: BufferLayerEmbeddedMask<BlendChannelMasked, BlendChannel, false>(Result.get(), Result.get(), Blended.get(), bApplyColorBlendToAlpha, bOnlyOneMip); break;
+			case EBlendType::BT_SOFTLIGHT: BufferLayerEmbeddedMask<SoftLightChannelMasked, SoftLightChannel, false>(Result.Get(), Result.Get(), Blended.Get(), bApplyColorBlendToAlpha, bOnlyOneMip); break;
+			case EBlendType::BT_HARDLIGHT: BufferLayerEmbeddedMask<HardLightChannelMasked, HardLightChannel, false>(Result.Get(), Result.Get(), Blended.Get(), bApplyColorBlendToAlpha, bOnlyOneMip); break;
+			case EBlendType::BT_BURN: BufferLayerEmbeddedMask<BurnChannelMasked, BurnChannel, false>(Result.Get(), Result.Get(), Blended.Get(), bApplyColorBlendToAlpha, bOnlyOneMip); break;
+			case EBlendType::BT_DODGE: BufferLayerEmbeddedMask<DodgeChannelMasked, DodgeChannel, false>(Result.Get(), Result.Get(), Blended.Get(), bApplyColorBlendToAlpha, bOnlyOneMip); break;
+			case EBlendType::BT_SCREEN: BufferLayerEmbeddedMask<ScreenChannelMasked, ScreenChannel, false>(Result.Get(), Result.Get(), Blended.Get(), bApplyColorBlendToAlpha, bOnlyOneMip); break;
+			case EBlendType::BT_OVERLAY: BufferLayerEmbeddedMask<OverlayChannelMasked, OverlayChannel, false>(Result.Get(), Result.Get(), Blended.Get(), bApplyColorBlendToAlpha, bOnlyOneMip); break;
+			case EBlendType::BT_LIGHTEN: BufferLayerEmbeddedMask<LightenChannelMasked, LightenChannel, false>(Result.Get(), Result.Get(), Blended.Get(), bApplyColorBlendToAlpha, bOnlyOneMip); break;
+			case EBlendType::BT_MULTIPLY: BufferLayerEmbeddedMask<MultiplyChannelMasked, MultiplyChannel, false>(Result.Get(), Result.Get(), Blended.Get(), bApplyColorBlendToAlpha, bOnlyOneMip); break;
+			case EBlendType::BT_BLEND: BufferLayerEmbeddedMask<BlendChannelMasked, BlendChannel, false>(Result.Get(), Result.Get(), Blended.Get(), bApplyColorBlendToAlpha, bOnlyOneMip); break;
 			case EBlendType::BT_NONE: break;
 			default: check(false);
 			}
@@ -868,16 +868,16 @@ namespace mu
 			// Apply blend without mask to RGB
 			switch (EBlendType(Args.blendType))
 			{
-			case EBlendType::BT_NORMAL_COMBINE: ImageNormalCombine(Result.get(), Result.get(), Blended.get(), bOnlyOneMip); check(!bUseBlendSourceFromBlendAlpha);  break;
-			case EBlendType::BT_SOFTLIGHT: BufferLayer<SoftLightChannel, true>(Result.get(), Result.get(), Blended.get(), bApplyColorBlendToAlpha, bOnlyOneMip, bUseBlendSourceFromBlendAlpha); break;
-			case EBlendType::BT_HARDLIGHT: BufferLayer<HardLightChannel, true>(Result.get(), Result.get(), Blended.get(), bApplyColorBlendToAlpha, bOnlyOneMip, bUseBlendSourceFromBlendAlpha); break;
-			case EBlendType::BT_BURN: BufferLayer<BurnChannel, true>(Result.get(), Result.get(), Blended.get(), bApplyColorBlendToAlpha, bOnlyOneMip, bUseBlendSourceFromBlendAlpha); break;
-			case EBlendType::BT_DODGE: BufferLayer<DodgeChannel, true>(Result.get(), Result.get(), Blended.get(), bApplyColorBlendToAlpha, bOnlyOneMip, bUseBlendSourceFromBlendAlpha); break;
-			case EBlendType::BT_SCREEN: BufferLayer<ScreenChannel, true>(Result.get(), Result.get(), Blended.get(), bApplyColorBlendToAlpha, bOnlyOneMip, bUseBlendSourceFromBlendAlpha); break;
-			case EBlendType::BT_OVERLAY: BufferLayer<OverlayChannel, true>(Result.get(), Result.get(), Blended.get(), bApplyColorBlendToAlpha, bOnlyOneMip, bUseBlendSourceFromBlendAlpha); break;
-			case EBlendType::BT_LIGHTEN: BufferLayer<LightenChannel, true>(Result.get(), Result.get(), Blended.get(), bApplyColorBlendToAlpha, bOnlyOneMip, bUseBlendSourceFromBlendAlpha); break;
-			case EBlendType::BT_MULTIPLY: BufferLayer<MultiplyChannel, true>(Result.get(), Result.get(), Blended.get(), bApplyColorBlendToAlpha, bOnlyOneMip, bUseBlendSourceFromBlendAlpha); break;
-			case EBlendType::BT_BLEND: BufferLayer<BlendChannel, true>(Result.get(), Result.get(), Blended.get(), bApplyColorBlendToAlpha, bOnlyOneMip, bUseBlendSourceFromBlendAlpha); break;
+			case EBlendType::BT_NORMAL_COMBINE: ImageNormalCombine(Result.Get(), Result.Get(), Blended.Get(), bOnlyOneMip); check(!bUseBlendSourceFromBlendAlpha);  break;
+			case EBlendType::BT_SOFTLIGHT: BufferLayer<SoftLightChannel, true>(Result.Get(), Result.Get(), Blended.Get(), bApplyColorBlendToAlpha, bOnlyOneMip, bUseBlendSourceFromBlendAlpha); break;
+			case EBlendType::BT_HARDLIGHT: BufferLayer<HardLightChannel, true>(Result.Get(), Result.Get(), Blended.Get(), bApplyColorBlendToAlpha, bOnlyOneMip, bUseBlendSourceFromBlendAlpha); break;
+			case EBlendType::BT_BURN: BufferLayer<BurnChannel, true>(Result.Get(), Result.Get(), Blended.Get(), bApplyColorBlendToAlpha, bOnlyOneMip, bUseBlendSourceFromBlendAlpha); break;
+			case EBlendType::BT_DODGE: BufferLayer<DodgeChannel, true>(Result.Get(), Result.Get(), Blended.Get(), bApplyColorBlendToAlpha, bOnlyOneMip, bUseBlendSourceFromBlendAlpha); break;
+			case EBlendType::BT_SCREEN: BufferLayer<ScreenChannel, true>(Result.Get(), Result.Get(), Blended.Get(), bApplyColorBlendToAlpha, bOnlyOneMip, bUseBlendSourceFromBlendAlpha); break;
+			case EBlendType::BT_OVERLAY: BufferLayer<OverlayChannel, true>(Result.Get(), Result.Get(), Blended.Get(), bApplyColorBlendToAlpha, bOnlyOneMip, bUseBlendSourceFromBlendAlpha); break;
+			case EBlendType::BT_LIGHTEN: BufferLayer<LightenChannel, true>(Result.Get(), Result.Get(), Blended.Get(), bApplyColorBlendToAlpha, bOnlyOneMip, bUseBlendSourceFromBlendAlpha); break;
+			case EBlendType::BT_MULTIPLY: BufferLayer<MultiplyChannel, true>(Result.Get(), Result.Get(), Blended.Get(), bApplyColorBlendToAlpha, bOnlyOneMip, bUseBlendSourceFromBlendAlpha); break;
+			case EBlendType::BT_BLEND: BufferLayer<BlendChannel, true>(Result.Get(), Result.Get(), Blended.Get(), bApplyColorBlendToAlpha, bOnlyOneMip, bUseBlendSourceFromBlendAlpha); break;
 			case EBlendType::BT_NONE: break;
 			default: check(false);
 			}
@@ -889,15 +889,15 @@ namespace mu
 			// Separate alpha operation ignores the mask.
 			switch (EBlendType(Args.blendTypeAlpha))
 			{
-			case EBlendType::BT_SOFTLIGHT: BufferLayerInPlace<SoftLightChannel, false, 1>(Result.get(), Blended.get(), bOnlyOneMip, 3, Args.BlendAlphaSourceChannel); break;
-			case EBlendType::BT_HARDLIGHT: BufferLayerInPlace<HardLightChannel, false, 1>(Result.get(), Blended.get(), bOnlyOneMip, 3, Args.BlendAlphaSourceChannel); break;
-			case EBlendType::BT_BURN: BufferLayerInPlace<BurnChannel, false, 1>(Result.get(), Blended.get(), bOnlyOneMip, 3, Args.BlendAlphaSourceChannel); break;
-			case EBlendType::BT_DODGE: BufferLayerInPlace<DodgeChannel, false, 1>(Result.get(), Blended.get(), bOnlyOneMip, 3, Args.BlendAlphaSourceChannel); break;
-			case EBlendType::BT_SCREEN: BufferLayerInPlace<ScreenChannel, false, 1>(Result.get(), Blended.get(), bOnlyOneMip, 3, Args.BlendAlphaSourceChannel); break;
-			case EBlendType::BT_OVERLAY: BufferLayerInPlace<OverlayChannel, false, 1>(Result.get(), Blended.get(), bOnlyOneMip, 3, Args.BlendAlphaSourceChannel); break;
-			case EBlendType::BT_LIGHTEN: BufferLayerInPlace<LightenChannel, false, 1>(Result.get(), Blended.get(), bOnlyOneMip, 3, Args.BlendAlphaSourceChannel); break;
-			case EBlendType::BT_MULTIPLY: BufferLayerInPlace<MultiplyChannel, false, 1>(Result.get(), Blended.get(), bOnlyOneMip, 3, Args.BlendAlphaSourceChannel); break;
-			case EBlendType::BT_BLEND: BufferLayerInPlace<BlendChannel, false, 1>(Result.get(), Blended.get(), bOnlyOneMip, 3, Args.BlendAlphaSourceChannel); break;
+			case EBlendType::BT_SOFTLIGHT: BufferLayerInPlace<SoftLightChannel, false, 1>(Result.Get(), Blended.Get(), bOnlyOneMip, 3, Args.BlendAlphaSourceChannel); break;
+			case EBlendType::BT_HARDLIGHT: BufferLayerInPlace<HardLightChannel, false, 1>(Result.Get(), Blended.Get(), bOnlyOneMip, 3, Args.BlendAlphaSourceChannel); break;
+			case EBlendType::BT_BURN: BufferLayerInPlace<BurnChannel, false, 1>(Result.Get(), Blended.Get(), bOnlyOneMip, 3, Args.BlendAlphaSourceChannel); break;
+			case EBlendType::BT_DODGE: BufferLayerInPlace<DodgeChannel, false, 1>(Result.Get(), Blended.Get(), bOnlyOneMip, 3, Args.BlendAlphaSourceChannel); break;
+			case EBlendType::BT_SCREEN: BufferLayerInPlace<ScreenChannel, false, 1>(Result.Get(), Blended.Get(), bOnlyOneMip, 3, Args.BlendAlphaSourceChannel); break;
+			case EBlendType::BT_OVERLAY: BufferLayerInPlace<OverlayChannel, false, 1>(Result.Get(), Blended.Get(), bOnlyOneMip, 3, Args.BlendAlphaSourceChannel); break;
+			case EBlendType::BT_LIGHTEN: BufferLayerInPlace<LightenChannel, false, 1>(Result.Get(), Blended.Get(), bOnlyOneMip, 3, Args.BlendAlphaSourceChannel); break;
+			case EBlendType::BT_MULTIPLY: BufferLayerInPlace<MultiplyChannel, false, 1>(Result.Get(), Blended.Get(), bOnlyOneMip, 3, Args.BlendAlphaSourceChannel); break;
+			case EBlendType::BT_BLEND: BufferLayerInPlace<BlendChannel, false, 1>(Result.Get(), Blended.Get(), bOnlyOneMip, 3, Args.BlendAlphaSourceChannel); break;
 			case EBlendType::BT_NONE: break;
 			default: check(false);
 			}
@@ -907,11 +907,11 @@ namespace mu
 		{
 			MUTABLE_CPUPROFILER_SCOPE(ImageLayer_MipFix);
 			FMipmapGenerationSettings DummyMipSettings{};
-			ImageMipmapInPlace(ImageCompressionQuality, Result.get(), DummyMipSettings);
+			ImageMipmapInPlace(ImageCompressionQuality, Result.Get(), DummyMipSettings);
 		}
 
 		// Reset relevancy map.
-		Result->m_flags &= ~Image::EImageFlags::IF_HAS_RELEVANCY_MAP;
+		Result->Flags &= ~FImage::EImageFlags::IF_HAS_RELEVANCY_MAP;
 	}
 
 
@@ -926,11 +926,11 @@ namespace mu
 		{
 			if (InitialFormat != Result->GetFormat())
 			{
-				Ptr<Image> Formatted = Runner->CreateImage(Result->GetSizeX(), Result->GetSizeY(), Result->GetLODCount(), InitialFormat, EInitializationType::NotInitialized);
+				TSharedPtr<FImage> Formatted = Runner->CreateImage(Result->GetSizeX(), Result->GetSizeY(), Result->GetLODCount(), InitialFormat, EInitializationType::NotInitialized);
 				bool bSuccess = false;
 
 				FImageOperator ImOp = MakeImageOperator(Runner);
-				ImOp.ImagePixelFormat(bSuccess, ImageCompressionQuality, Formatted.get(), Result.get());
+				ImOp.ImagePixelFormat(bSuccess, ImageCompressionQuality, Formatted.Get(), Result.Get());
 				check(bSuccess);
 
 				Runner->Release(Result);
@@ -964,8 +964,8 @@ namespace mu
 		int32 ImageCompressionQuality = 0;
 		OP::ImageLayerColourArgs Args;
 		FVector4f Color;
-		Ptr<const Image> Mask;
-		Ptr<Image> Result;
+		TSharedPtr<const FImage> Mask;
+		TSharedPtr<FImage> Result;
 		EImageFormat InitialFormat = EImageFormat::None;
 	};
 
@@ -976,9 +976,9 @@ namespace mu
 		MUTABLE_CPUPROFILER_SCOPE(FImageLayerColourTask_Prepare);
 		bOutFailed = false;
 
-		ImageCompressionQuality = Runner->m_pSettings->ImageCompressionQuality;
+		ImageCompressionQuality = Runner->Settings.ImageCompressionQuality;
 
-		Ptr<const Image> Base = Runner->LoadImage({ Args.base, Op.ExecutionIndex, Op.ExecutionOptions });
+		TSharedPtr<const FImage> Base = Runner->LoadImage({ Args.base, Op.ExecutionIndex, Op.ExecutionOptions });
 		check(!Base || Base->GetFormat() < EImageFormat::Count);
 
 		Color = Runner->LoadColor({ Args.colour, Op.ExecutionIndex, 0 });
@@ -1015,8 +1015,8 @@ namespace mu
 			if (Base->GetSize() != Mask->GetSize())
 			{
 				MUTABLE_CPUPROFILER_SCOPE(ImageResize_EmergencyFixSize);
-				Ptr<Image> Resized = Runner->CreateImage(Base->GetSizeX(), Base->GetSizeY(), 1, Mask->GetFormat(), EInitializationType::NotInitialized);
-				ImOp.ImageResizeLinear(Resized.get(), ImageCompressionQuality, Mask.get() );
+				TSharedPtr<FImage> Resized = Runner->CreateImage(Base->GetSizeX(), Base->GetSizeY(), 1, Mask->GetFormat(), EInitializationType::NotInitialized);
+				ImOp.ImageResizeLinear(Resized.Get(), ImageCompressionQuality, Mask.Get() );
 				Runner->Release(Mask);
 				Mask = Resized;
 			}
@@ -1028,11 +1028,11 @@ namespace mu
 				int32 StartLevel = Mask->GetLODCount() - 1;
 				int32 LevelCount = Base->GetLODCount();
 
-				Ptr<Image> MaskFix = Runner->CloneOrTakeOver(Mask);
+				TSharedPtr<FImage> MaskFix = Runner->CloneOrTakeOver(Mask);
 				MaskFix->DataStorage.SetNumLODs(LevelCount);
 
 				FMipmapGenerationSettings Settings{};
-				ImOp.ImageMipmap(ImageCompressionQuality, MaskFix.get(), MaskFix.get(), StartLevel, LevelCount, Settings);
+				ImOp.ImageMipmap(ImageCompressionQuality, MaskFix.Get(), MaskFix.Get(), StartLevel, LevelCount, Settings);
 
 				Mask = MaskFix;
 			}
@@ -1063,16 +1063,16 @@ namespace mu
 				// \todo: precalculated tables for softlight
 				switch (EBlendType(Args.blendType))
 				{
-				case EBlendType::BT_NORMAL_COMBINE: ImageNormalCombine(Result.get(), Result.get(), Mask.get(), Color); break;
-				case EBlendType::BT_SOFTLIGHT: BufferLayerColour<SoftLightChannelMasked, SoftLightChannel>(Result.get(), Result.get(), Mask.get(), Color); break;
-				case EBlendType::BT_HARDLIGHT: BufferLayerColour<HardLightChannelMasked, HardLightChannel>(Result.get(), Result.get(), Mask.get(), Color); break;
-				case EBlendType::BT_BURN: BufferLayerColour<BurnChannelMasked, BurnChannel>(Result.get(), Result.get(), Mask.get(), Color); break;
-				case EBlendType::BT_DODGE: BufferLayerColour<DodgeChannelMasked, DodgeChannel>(Result.get(), Result.get(), Mask.get(), Color); break;
-				case EBlendType::BT_SCREEN: BufferLayerColour<ScreenChannelMasked, ScreenChannel>(Result.get(), Result.get(), Mask.get(), Color); break;
-				case EBlendType::BT_OVERLAY: BufferLayerColour<OverlayChannelMasked, OverlayChannel>(Result.get(), Result.get(), Mask.get(), Color); break;
-				case EBlendType::BT_LIGHTEN: BufferLayerColour<LightenChannelMasked, LightenChannel>(Result.get(), Result.get(), Mask.get(), Color); break;
-				case EBlendType::BT_MULTIPLY: BufferLayerColour<MultiplyChannelMasked, MultiplyChannel>(Result.get(), Result.get(), Mask.get(), Color); break;
-				case EBlendType::BT_BLEND: BufferLayerColour<BlendChannelMasked, BlendChannel>(Result.get(), Result.get(), Mask.get(), Color); break;
+				case EBlendType::BT_NORMAL_COMBINE: ImageNormalCombine(Result.Get(), Result.Get(), Mask.Get(), Color); break;
+				case EBlendType::BT_SOFTLIGHT: BufferLayerColour<SoftLightChannelMasked, SoftLightChannel>(Result.Get(), Result.Get(), Mask.Get(), Color); break;
+				case EBlendType::BT_HARDLIGHT: BufferLayerColour<HardLightChannelMasked, HardLightChannel>(Result.Get(), Result.Get(), Mask.Get(), Color); break;
+				case EBlendType::BT_BURN: BufferLayerColour<BurnChannelMasked, BurnChannel>(Result.Get(), Result.Get(), Mask.Get(), Color); break;
+				case EBlendType::BT_DODGE: BufferLayerColour<DodgeChannelMasked, DodgeChannel>(Result.Get(), Result.Get(), Mask.Get(), Color); break;
+				case EBlendType::BT_SCREEN: BufferLayerColour<ScreenChannelMasked, ScreenChannel>(Result.Get(), Result.Get(), Mask.Get(), Color); break;
+				case EBlendType::BT_OVERLAY: BufferLayerColour<OverlayChannelMasked, OverlayChannel>(Result.Get(), Result.Get(), Mask.Get(), Color); break;
+				case EBlendType::BT_LIGHTEN: BufferLayerColour<LightenChannelMasked, LightenChannel>(Result.Get(), Result.Get(), Mask.Get(), Color); break;
+				case EBlendType::BT_MULTIPLY: BufferLayerColour<MultiplyChannelMasked, MultiplyChannel>(Result.Get(), Result.Get(), Mask.Get(), Color); break;
+				case EBlendType::BT_BLEND: BufferLayerColour<BlendChannelMasked, BlendChannel>(Result.Get(), Result.Get(), Mask.Get(), Color); break;
 				default: check(false);
 				}
 
@@ -1084,14 +1084,14 @@ namespace mu
 					switch (EBlendType(Args.blendType))
 					{
 					case EBlendType::BT_NORMAL_COMBINE: check(false); break;
-					case EBlendType::BT_SOFTLIGHT: BufferLayerColourFromAlpha<SoftLightChannel>(Result.get(), Result.get(), Color); break;
-					case EBlendType::BT_HARDLIGHT: BufferLayerColourFromAlpha<HardLightChannel>(Result.get(), Result.get(), Color); break;
-					case EBlendType::BT_BURN: BufferLayerColourFromAlpha<BurnChannel>(Result.get(), Result.get(), Color); break;
-					case EBlendType::BT_DODGE: BufferLayerColourFromAlpha<DodgeChannel>(Result.get(), Result.get(), Color); break;
-					case EBlendType::BT_SCREEN: BufferLayerColourFromAlpha<ScreenChannel>(Result.get(), Result.get(), Color); break;
-					case EBlendType::BT_OVERLAY: BufferLayerColourFromAlpha<OverlayChannel>(Result.get(), Result.get(), Color); break;
-					case EBlendType::BT_LIGHTEN: BufferLayerColourFromAlpha<LightenChannel>(Result.get(), Result.get(), Color); break;
-					case EBlendType::BT_MULTIPLY: BufferLayerColourFromAlpha<MultiplyChannel>(Result.get(), Result.get(), Color); break;
+					case EBlendType::BT_SOFTLIGHT: BufferLayerColourFromAlpha<SoftLightChannel>(Result.Get(), Result.Get(), Color); break;
+					case EBlendType::BT_HARDLIGHT: BufferLayerColourFromAlpha<HardLightChannel>(Result.Get(), Result.Get(), Color); break;
+					case EBlendType::BT_BURN: BufferLayerColourFromAlpha<BurnChannel>(Result.Get(), Result.Get(), Color); break;
+					case EBlendType::BT_DODGE: BufferLayerColourFromAlpha<DodgeChannel>(Result.Get(), Result.Get(), Color); break;
+					case EBlendType::BT_SCREEN: BufferLayerColourFromAlpha<ScreenChannel>(Result.Get(), Result.Get(), Color); break;
+					case EBlendType::BT_OVERLAY: BufferLayerColourFromAlpha<OverlayChannel>(Result.Get(), Result.Get(), Color); break;
+					case EBlendType::BT_LIGHTEN: BufferLayerColourFromAlpha<LightenChannel>(Result.Get(), Result.Get(), Color); break;
+					case EBlendType::BT_MULTIPLY: BufferLayerColourFromAlpha<MultiplyChannel>(Result.Get(), Result.Get(), Color); break;
 					case EBlendType::BT_BLEND: check(false); break;
 					default: check(false);
 					}
@@ -1100,20 +1100,20 @@ namespace mu
 				{
 					switch (EBlendType(Args.blendType))
 					{
-					case EBlendType::BT_NORMAL_COMBINE: ImageNormalCombine(Result.get(), Result.get(), Color); break;
-					case EBlendType::BT_SOFTLIGHT: BufferLayerColour<SoftLightChannel>(Result.get(), Result.get(), Color); break;
-					case EBlendType::BT_HARDLIGHT: BufferLayerColour<HardLightChannel>(Result.get(), Result.get(), Color); break;
-					case EBlendType::BT_BURN: BufferLayerColour<BurnChannel>(Result.get(), Result.get(), Color); break;
-					case EBlendType::BT_DODGE: BufferLayerColour<DodgeChannel>(Result.get(), Result.get(), Color); break;
-					case EBlendType::BT_SCREEN: BufferLayerColour<ScreenChannel>(Result.get(), Result.get(), Color); break;
-					case EBlendType::BT_OVERLAY: BufferLayerColour<OverlayChannel>(Result.get(), Result.get(), Color); break;
-					case EBlendType::BT_LIGHTEN: BufferLayerColour<LightenChannel>(Result.get(), Result.get(), Color); break;
-					case EBlendType::BT_MULTIPLY: BufferLayerColour<MultiplyChannel>(Result.get(), Result.get(), Color); break;
+					case EBlendType::BT_NORMAL_COMBINE: ImageNormalCombine(Result.Get(), Result.Get(), Color); break;
+					case EBlendType::BT_SOFTLIGHT: BufferLayerColour<SoftLightChannel>(Result.Get(), Result.Get(), Color); break;
+					case EBlendType::BT_HARDLIGHT: BufferLayerColour<HardLightChannel>(Result.Get(), Result.Get(), Color); break;
+					case EBlendType::BT_BURN: BufferLayerColour<BurnChannel>(Result.Get(), Result.Get(), Color); break;
+					case EBlendType::BT_DODGE: BufferLayerColour<DodgeChannel>(Result.Get(), Result.Get(), Color); break;
+					case EBlendType::BT_SCREEN: BufferLayerColour<ScreenChannel>(Result.Get(), Result.Get(), Color); break;
+					case EBlendType::BT_OVERLAY: BufferLayerColour<OverlayChannel>(Result.Get(), Result.Get(), Color); break;
+					case EBlendType::BT_LIGHTEN: BufferLayerColour<LightenChannel>(Result.Get(), Result.Get(), Color); break;
+					case EBlendType::BT_MULTIPLY: BufferLayerColour<MultiplyChannel>(Result.Get(), Result.Get(), Color); break;
 					case EBlendType::BT_BLEND: 
 					{
 						// In this case we know it is already an uncompressed image, and we won't need additional allocations;
 						FImageOperator ImOp = FImageOperator::GetDefault(FImageOperator::FImagePixelFormatFunc());
-						ImOp.FillColor( Result.get(), Color); 
+						ImOp.FillColor( Result.Get(), Color); 
 						break;
 					}
 					default: check(false);
@@ -1131,15 +1131,15 @@ namespace mu
 				switch (EBlendType(Args.blendTypeAlpha))
 				{
 				case EBlendType::BT_NORMAL_COMBINE: check(false); break;
-				case EBlendType::BT_SOFTLIGHT: BufferLayerColourInPlace<SoftLightChannelMasked, SoftLightChannel, 1>(Result.get(), Mask.get(), Color, bOnlyOneMip, 3, Args.BlendAlphaSourceChannel); break;
-				case EBlendType::BT_HARDLIGHT: BufferLayerColourInPlace<HardLightChannelMasked, HardLightChannel, 1>(Result.get(), Mask.get(), Color, bOnlyOneMip, 3, Args.BlendAlphaSourceChannel); break;
-				case EBlendType::BT_BURN: BufferLayerColourInPlace<BurnChannelMasked, BurnChannel, 1>(Result.get(), Mask.get(), Color, bOnlyOneMip, 3, Args.BlendAlphaSourceChannel); break;
-				case EBlendType::BT_DODGE: BufferLayerColourInPlace<DodgeChannelMasked, DodgeChannel, 1>(Result.get(), Mask.get(), Color, bOnlyOneMip, 3, Args.BlendAlphaSourceChannel); break;
-				case EBlendType::BT_SCREEN: BufferLayerColourInPlace<ScreenChannelMasked, ScreenChannel, 1>(Result.get(), Mask.get(), Color, bOnlyOneMip, 3, Args.BlendAlphaSourceChannel); break;
-				case EBlendType::BT_OVERLAY: BufferLayerColourInPlace<OverlayChannelMasked, OverlayChannel, 1>(Result.get(), Mask.get(), Color, bOnlyOneMip, 3, Args.BlendAlphaSourceChannel); break;
-				case EBlendType::BT_LIGHTEN: BufferLayerColourInPlace<LightenChannelMasked, LightenChannel, 1>(Result.get(), Mask.get(), Color, bOnlyOneMip, 3, Args.BlendAlphaSourceChannel); break;
-				case EBlendType::BT_MULTIPLY: BufferLayerColourInPlace<MultiplyChannelMasked, MultiplyChannel, 1>(Result.get(), Mask.get(), Color, bOnlyOneMip, 3, Args.BlendAlphaSourceChannel); break;
-				case EBlendType::BT_BLEND: BufferLayerColourInPlace<BlendChannelMasked, BlendChannel, 1>(Result.get(), Mask.get(), Color, bOnlyOneMip, 3, Args.BlendAlphaSourceChannel); break;
+				case EBlendType::BT_SOFTLIGHT: BufferLayerColourInPlace<SoftLightChannelMasked, SoftLightChannel, 1>(Result.Get(), Mask.Get(), Color, bOnlyOneMip, 3, Args.BlendAlphaSourceChannel); break;
+				case EBlendType::BT_HARDLIGHT: BufferLayerColourInPlace<HardLightChannelMasked, HardLightChannel, 1>(Result.Get(), Mask.Get(), Color, bOnlyOneMip, 3, Args.BlendAlphaSourceChannel); break;
+				case EBlendType::BT_BURN: BufferLayerColourInPlace<BurnChannelMasked, BurnChannel, 1>(Result.Get(), Mask.Get(), Color, bOnlyOneMip, 3, Args.BlendAlphaSourceChannel); break;
+				case EBlendType::BT_DODGE: BufferLayerColourInPlace<DodgeChannelMasked, DodgeChannel, 1>(Result.Get(), Mask.Get(), Color, bOnlyOneMip, 3, Args.BlendAlphaSourceChannel); break;
+				case EBlendType::BT_SCREEN: BufferLayerColourInPlace<ScreenChannelMasked, ScreenChannel, 1>(Result.Get(), Mask.Get(), Color, bOnlyOneMip, 3, Args.BlendAlphaSourceChannel); break;
+				case EBlendType::BT_OVERLAY: BufferLayerColourInPlace<OverlayChannelMasked, OverlayChannel, 1>(Result.Get(), Mask.Get(), Color, bOnlyOneMip, 3, Args.BlendAlphaSourceChannel); break;
+				case EBlendType::BT_LIGHTEN: BufferLayerColourInPlace<LightenChannelMasked, LightenChannel, 1>(Result.Get(), Mask.Get(), Color, bOnlyOneMip, 3, Args.BlendAlphaSourceChannel); break;
+				case EBlendType::BT_MULTIPLY: BufferLayerColourInPlace<MultiplyChannelMasked, MultiplyChannel, 1>(Result.Get(), Mask.Get(), Color, bOnlyOneMip, 3, Args.BlendAlphaSourceChannel); break;
+				case EBlendType::BT_BLEND: BufferLayerColourInPlace<BlendChannelMasked, BlendChannel, 1>(Result.Get(), Mask.Get(), Color, bOnlyOneMip, 3, Args.BlendAlphaSourceChannel); break;
 				default: check(false);
 				}
 
@@ -1149,22 +1149,22 @@ namespace mu
 				switch (EBlendType(Args.blendTypeAlpha))
 				{
 				case EBlendType::BT_NORMAL_COMBINE: check(false); break;
-				case EBlendType::BT_SOFTLIGHT: BufferLayerColourInPlace<SoftLightChannel, 1>(Result.get(), Color, bOnlyOneMip, 3, Args.BlendAlphaSourceChannel); break;
-				case EBlendType::BT_HARDLIGHT: BufferLayerColourInPlace<HardLightChannel, 1>(Result.get(), Color, bOnlyOneMip, 3, Args.BlendAlphaSourceChannel); break;
-				case EBlendType::BT_BURN: BufferLayerColourInPlace<BurnChannel, 1>(Result.get(), Color, bOnlyOneMip, 3, Args.BlendAlphaSourceChannel); break;
-				case EBlendType::BT_DODGE: BufferLayerColourInPlace<DodgeChannel, 1>(Result.get(), Color, bOnlyOneMip, 3, Args.BlendAlphaSourceChannel); break;
-				case EBlendType::BT_SCREEN: BufferLayerColourInPlace<ScreenChannel, 1>(Result.get(), Color, bOnlyOneMip, 3, Args.BlendAlphaSourceChannel); break;
-				case EBlendType::BT_OVERLAY: BufferLayerColourInPlace<OverlayChannel, 1>(Result.get(), Color, bOnlyOneMip, 3, Args.BlendAlphaSourceChannel); break;
-				case EBlendType::BT_LIGHTEN: BufferLayerColourInPlace<LightenChannel, 1>(Result.get(), Color, bOnlyOneMip, 3, Args.BlendAlphaSourceChannel); break;
-				case EBlendType::BT_MULTIPLY: BufferLayerColourInPlace<MultiplyChannel, 1>(Result.get(), Color, bOnlyOneMip, 3, Args.BlendAlphaSourceChannel); break;
-				case EBlendType::BT_BLEND: BufferLayerColourInPlace<BlendChannel, 1>(Result.get(), Color, bOnlyOneMip, 3, Args.BlendAlphaSourceChannel); break;
+				case EBlendType::BT_SOFTLIGHT: BufferLayerColourInPlace<SoftLightChannel, 1>(Result.Get(), Color, bOnlyOneMip, 3, Args.BlendAlphaSourceChannel); break;
+				case EBlendType::BT_HARDLIGHT: BufferLayerColourInPlace<HardLightChannel, 1>(Result.Get(), Color, bOnlyOneMip, 3, Args.BlendAlphaSourceChannel); break;
+				case EBlendType::BT_BURN: BufferLayerColourInPlace<BurnChannel, 1>(Result.Get(), Color, bOnlyOneMip, 3, Args.BlendAlphaSourceChannel); break;
+				case EBlendType::BT_DODGE: BufferLayerColourInPlace<DodgeChannel, 1>(Result.Get(), Color, bOnlyOneMip, 3, Args.BlendAlphaSourceChannel); break;
+				case EBlendType::BT_SCREEN: BufferLayerColourInPlace<ScreenChannel, 1>(Result.Get(), Color, bOnlyOneMip, 3, Args.BlendAlphaSourceChannel); break;
+				case EBlendType::BT_OVERLAY: BufferLayerColourInPlace<OverlayChannel, 1>(Result.Get(), Color, bOnlyOneMip, 3, Args.BlendAlphaSourceChannel); break;
+				case EBlendType::BT_LIGHTEN: BufferLayerColourInPlace<LightenChannel, 1>(Result.Get(), Color, bOnlyOneMip, 3, Args.BlendAlphaSourceChannel); break;
+				case EBlendType::BT_MULTIPLY: BufferLayerColourInPlace<MultiplyChannel, 1>(Result.Get(), Color, bOnlyOneMip, 3, Args.BlendAlphaSourceChannel); break;
+				case EBlendType::BT_BLEND: BufferLayerColourInPlace<BlendChannel, 1>(Result.Get(), Color, bOnlyOneMip, 3, Args.BlendAlphaSourceChannel); break;
 				default: check(false);
 				}
 			}
 		}
 
 		// Reset relevancy map.
-		Result->m_flags &= ~Image::EImageFlags::IF_HAS_RELEVANCY_MAP;
+		Result->Flags &= ~FImage::EImageFlags::IF_HAS_RELEVANCY_MAP;
 	}
 
 
@@ -1203,8 +1203,8 @@ namespace mu
 		int ImageCompressionQuality = 0;
 		OP::ImagePixelFormatArgs Args;
 		EImageFormat TargetFormat = EImageFormat::None;
-		Ptr<const Image> Base;
-		Ptr<Image> Result;
+		TSharedPtr<const FImage> Base;
+		TSharedPtr<FImage> Result;
 		FImageOperator::FImagePixelFormatFunc ImagePixelFormatFunc;
 	};
 
@@ -1215,7 +1215,7 @@ namespace mu
 		MUTABLE_CPUPROFILER_SCOPE(FImageLayerPixelFormatTask_Prepare);
 		bOutFailed = false;
 
-		ImageCompressionQuality = Runner->m_pSettings->ImageCompressionQuality;
+		ImageCompressionQuality = Runner->Settings.ImageCompressionQuality;
 
 		Base = Runner->LoadImage({ Args.source, Op.ExecutionIndex, Op.ExecutionOptions });
 		check(!Base || Base->GetFormat() < EImageFormat::Count);		
@@ -1248,7 +1248,7 @@ namespace mu
 			return false;
 		}
 
-		ImagePixelFormatFunc = Runner->m_pSystem->ImagePixelFormatOverride;
+		ImagePixelFormatFunc = Runner->System->ImagePixelFormatOverride;
 
 		// Create destination data
 		Result = Runner->CreateImage(Base->GetSizeX(), Base->GetSizeY(), Base->GetLODCount(), TargetFormat, EInitializationType::NotInitialized);
@@ -1263,7 +1263,7 @@ namespace mu
 		
 		bool bSuccess = false;
 		FImageOperator ImOp = FImageOperator::GetDefault(ImagePixelFormatFunc);
-		ImOp.ImagePixelFormat(bSuccess, ImageCompressionQuality, Result.get(), Base.get(), -1);
+		ImOp.ImagePixelFormat(bSuccess, ImageCompressionQuality, Result.Get(), Base.Get(), -1);
 		check(bSuccess);
 	}
 
@@ -1299,8 +1299,8 @@ namespace mu
 		int32 ImageCompressionQuality = 0;
 		int32 StartLevel = -1;
 		OP::ImageMipmapArgs Args;
-		Ptr<const Image> Base;
-		Ptr<Image> Result;
+		TSharedPtr<const FImage> Base;
+		TSharedPtr<FImage> Result;
 		FImageOperator::FScratchImageMipmap Scratch;
 		FImageOperator::FImagePixelFormatFunc ImagePixelFormatFunc;
 	};
@@ -1312,7 +1312,7 @@ namespace mu
 		MUTABLE_CPUPROFILER_SCOPE(FImageMipmapTask_Prepare);
 		bOutFailed = false;
 
-		ImageCompressionQuality = Runner->m_pSettings->ImageCompressionQuality;
+		ImageCompressionQuality = Runner->Settings.ImageCompressionQuality;
 
 		Base = Runner->LoadImage({ Args.source, Op.ExecutionIndex, Op.ExecutionOptions });
 		check(!Base || Base->GetFormat() < EImageFormat::Count);
@@ -1332,7 +1332,7 @@ namespace mu
 		}
 
 		int32 LevelCount = Args.levels;
-		int32 MaxLevelCount = Image::GetMipmapCount(Base->GetSizeX(), Base->GetSizeY());
+		int32 MaxLevelCount = FImage::GetMipmapCount(Base->GetSizeX(), Base->GetSizeY());
 		if (LevelCount == 0)
 		{
 			LevelCount = MaxLevelCount;
@@ -1361,9 +1361,9 @@ namespace mu
 		Result->DataStorage.SetNumLODs(LevelCount);
 
 		FImageOperator ImOp = MakeImageOperator(Runner);
-		ImOp.ImageMipmap_PrepareScratch(Result.get(), StartLevel, LevelCount, Scratch);
+		ImOp.ImageMipmap_PrepareScratch(Result.Get(), StartLevel, LevelCount, Scratch);
 
-		ImagePixelFormatFunc = Runner->m_pSystem->ImagePixelFormatOverride;
+		ImagePixelFormatFunc = Runner->System->ImagePixelFormatOverride;
 
 		return true;
 	}
@@ -1378,7 +1378,7 @@ namespace mu
 
 		FMipmapGenerationSettings Settings{Args.FilterType, Args.AddressMode};
 		FImageOperator ImOp = FImageOperator::GetDefault(ImagePixelFormatFunc);
-		ImOp.ImageMipmap(Scratch, ImageCompressionQuality, Result.get(), Result.get(), StartLevel, Result->GetLODCount(), Settings);
+		ImOp.ImageMipmap(Scratch, ImageCompressionQuality, Result.Get(), Result.Get(), StartLevel, Result->GetLODCount(), Settings);
 	}
 
 
@@ -1421,8 +1421,8 @@ namespace mu
 
 	private:
 		OP::ImageSwizzleArgs Args;
-		Ptr<const Image> Sources[MUTABLE_OP_MAX_SWIZZLE_CHANNELS];
-		Ptr<Image> Result;
+		TSharedPtr<const FImage> Sources[MUTABLE_OP_MAX_SWIZZLE_CHANNELS];
+		TSharedPtr<FImage> Result;
 	};
 
 
@@ -1477,10 +1477,10 @@ namespace mu
 				MUTABLE_CPUPROFILER_SCOPE(ImageFormat_ForSwizzle);
 
 				EImageFormat UncompressedFormat = GetUncompressedFormat(Sources[SourceIndex]->GetFormat());
-				Ptr<Image> Formatted = Runner->CreateImage(Sources[SourceIndex]->GetSizeX(), Sources[SourceIndex]->GetSizeY(), 1, UncompressedFormat, EInitializationType::NotInitialized);
+				TSharedPtr<FImage> Formatted = Runner->CreateImage(Sources[SourceIndex]->GetSizeX(), Sources[SourceIndex]->GetSizeY(), 1, UncompressedFormat, EInitializationType::NotInitialized);
 				bool bSuccess = false;
 				int32 ImageCompressionQuality = 4; // TODO
-				ImOp.ImagePixelFormat(bSuccess, ImageCompressionQuality, Formatted.get(), Sources[SourceIndex].get());
+				ImOp.ImagePixelFormat(bSuccess, ImageCompressionQuality, Formatted.Get(), Sources[SourceIndex].Get());
 				check(bSuccess); // Decompression cannot fail
 				Runner->Release(Sources[SourceIndex]);
 				Sources[SourceIndex] = Formatted;
@@ -1496,8 +1496,8 @@ namespace mu
 			if (Sources[SourceIndex] && ResultSize != Sources[SourceIndex]->GetSize())
 			{
 				MUTABLE_CPUPROFILER_SCOPE(ImageResize_ForSwizzle);
-				Ptr<Image> Resized = Runner->CreateImage(ResultSize.X, ResultSize.Y, 1, Sources[SourceIndex]->GetFormat(), EInitializationType::NotInitialized);
-				ImOp.ImageResizeLinear(Resized.get(), 0, Sources[SourceIndex].get());
+				TSharedPtr<FImage> Resized = Runner->CreateImage(ResultSize.X, ResultSize.Y, 1, Sources[SourceIndex]->GetFormat(), EInitializationType::NotInitialized);
+				ImOp.ImageResizeLinear(Resized.Get(), 0, Sources[SourceIndex].Get());
 				Runner->Release(Sources[SourceIndex]);
 				Sources[SourceIndex] = Resized;
 				ResultLODs = 1;
@@ -1523,7 +1523,7 @@ namespace mu
 		// This runs in a worker thread
 		MUTABLE_CPUPROFILER_SCOPE(FImageSwizzleTask);
 
-		ImageSwizzle(Result.get(), Sources, Args.sourceChannels);
+		ImageSwizzle(Result.Get(), Sources, Args.sourceChannels);
 	}
 
 
@@ -1563,7 +1563,7 @@ namespace mu
 
 	private:
 		const OP::ImageSaturateArgs Args;
-		Ptr<Image> Result;
+		TSharedPtr<FImage> Result;
 		float Factor;
 	};
 
@@ -1582,7 +1582,7 @@ namespace mu
 		MUTABLE_CPUPROFILER_SCOPE(FImageSaturateTask_Prepare);
 		bOutFailed = false;
 
-		Ptr<const Image> Source = Runner->LoadImage(FCacheAddress(Args.base, Op));
+		TSharedPtr<const FImage> Source = Runner->LoadImage(FCacheAddress(Args.base, Op));
 		Factor = Runner->LoadScalar(FScheduledOp::FromOpAndOptions(Args.factor, Op, 0));
 
 		if (!Source)
@@ -1609,7 +1609,7 @@ namespace mu
 		MUTABLE_CPUPROFILER_SCOPE(FImageSaturateTask);
 
 		constexpr bool bUseVectorIntrinsics = true;
-		ImageSaturate<bUseVectorIntrinsics>(Result.get(), Factor);
+		ImageSaturate<bUseVectorIntrinsics>(Result.Get(), Factor);
 	}
 
 	//---------------------------------------------------------------------------------------------
@@ -1646,8 +1646,8 @@ namespace mu
 	private:
 		int32 ImageCompressionQuality = 0;
 		OP::ImageResizeArgs Args;
-		Ptr<const Image> Base;
-		Ptr<Image> Result;
+		TSharedPtr<const FImage> Base;
+		TSharedPtr<FImage> Result;
 		FImageOperator::FImagePixelFormatFunc ImagePixelFormatFunc;
 	};
 
@@ -1658,7 +1658,7 @@ namespace mu
 		MUTABLE_CPUPROFILER_SCOPE(FImageResizeTask_Prepare);
 		bOutFailed = false;
 
-		ImageCompressionQuality = Runner->m_pSettings->ImageCompressionQuality;
+		ImageCompressionQuality = Runner->Settings.ImageCompressionQuality;
 
 		FImageSize destSize = FImageSize(Args.size[0], Args.size[1]);
 
@@ -1686,7 +1686,7 @@ namespace mu
 		bool bSourceHasMips = Base->GetLODCount() > 1;
 		if (bSourceHasMips)
 		{
-			Lods = Image::GetMipmapCount(destSize[0], destSize[1]);
+			Lods = FImage::GetMipmapCount(destSize[0], destSize[1]);
 		}
 
 		if (Base->IsReference())
@@ -1698,7 +1698,7 @@ namespace mu
 
 		Result = Runner->CreateImage( destSize[0], destSize[1], Lods, Base->GetFormat(), EInitializationType::NotInitialized );
 
-		ImagePixelFormatFunc = Runner->m_pSystem->ImagePixelFormatOverride;
+		ImagePixelFormatFunc = Runner->System->ImagePixelFormatOverride;
 
 		return true;
 	}
@@ -1720,13 +1720,13 @@ namespace mu
 		// Warning: This will actually allocate temp memory that may exceed the budget.
 		// \TODO: Fix it.
 		FImageOperator ImOp = FImageOperator::GetDefault(ImagePixelFormatFunc);
-		ImOp.ImageResizeLinear( Result.get(), ImageCompressionQuality, Base.get());
+		ImOp.ImageResizeLinear( Result.Get(), ImageCompressionQuality, Base.Get());
 
 		int32 LodCount = Result->GetLODCount();
 		if (LodCount>1)
 		{
 			FMipmapGenerationSettings mipSettings{};
-			ImageMipmapInPlace( ImageCompressionQuality, Result.get(), mipSettings );
+			ImageMipmapInPlace( ImageCompressionQuality, Result.Get(), mipSettings );
 		}
 	}
 
@@ -1766,8 +1766,8 @@ namespace mu
 	private:
 		OP::ImageResizeRelArgs Args;
 		int32 ImageCompressionQuality=0;
-		Ptr<const Image> Base;
-		Ptr<Image> Result;
+		TSharedPtr<const FImage> Base;
+		TSharedPtr<FImage> Result;
 		FImageSize DestSize;
 		FImageOperator::FImagePixelFormatFunc ImagePixelFormatFunc;
 	};
@@ -1779,7 +1779,7 @@ namespace mu
 		MUTABLE_CPUPROFILER_SCOPE(FImageResizeRelTask_Prepare);
 		bOutFailed = false;
 
-		ImageCompressionQuality = Runner->m_pSettings->ImageCompressionQuality;
+		ImageCompressionQuality = Runner->Settings.ImageCompressionQuality;
 
 		Base = Runner->LoadImage(FCacheAddress(Args.source, Op));
 		if (!Base)
@@ -1807,7 +1807,7 @@ namespace mu
 		bool bSourceHasMips = Base->GetLODCount() > 1;
 		if (bSourceHasMips)
 		{
-			Lods = Image::GetMipmapCount(DestSize[0], DestSize[1]);
+			Lods = FImage::GetMipmapCount(DestSize[0], DestSize[1]);
 		}
 
 		if (Base->IsReference())
@@ -1819,7 +1819,7 @@ namespace mu
 
 		Result = Runner->CreateImage(DestSize[0], DestSize[1], Lods, Base->GetFormat(), EInitializationType::NotInitialized);
 
-		ImagePixelFormatFunc = Runner->m_pSystem->ImagePixelFormatOverride;
+		ImagePixelFormatFunc = Runner->System->ImagePixelFormatOverride;
 
 		return true;
 	}
@@ -1832,13 +1832,13 @@ namespace mu
 
 		// \TODO: Track allocs
 		FImageOperator ImOp = FImageOperator::GetDefault(ImagePixelFormatFunc);
-		ImOp.ImageResizeLinear(Result.get(), ImageCompressionQuality, Base.get());
+		ImOp.ImageResizeLinear(Result.Get(), ImageCompressionQuality, Base.Get());
 
 		int32 LodCount = Result->GetLODCount();
 		if (LodCount > 1)
 		{
 			FMipmapGenerationSettings mipSettings{};
-			ImageMipmapInPlace(ImageCompressionQuality, Result.get(), mipSettings);
+			ImageMipmapInPlace(ImageCompressionQuality, Result.Get(), mipSettings);
 		}
 	}
 
@@ -1875,7 +1875,7 @@ namespace mu
 		virtual bool Complete(CodeRunner* Runner) override;
 
 	private:
-		Ptr<Image> Result;
+		TSharedPtr<FImage> Result;
 		OP::ImageInvertArgs Args;
 	};
 
@@ -1886,7 +1886,7 @@ namespace mu
 		MUTABLE_CPUPROFILER_SCOPE(FImageInvertTask_Prepare);
 		bOutFailed = false;
 
-		Ptr<const Image> Source = Runner->LoadImage({ Args.base, Op.ExecutionIndex, Op.ExecutionOptions });
+		TSharedPtr<const FImage> Source = Runner->LoadImage({ Args.base, Op.ExecutionIndex, Op.ExecutionOptions });
 
 		// Create destination data
 		Result = Runner->CloneOrTakeOver(Source);
@@ -1899,7 +1899,7 @@ namespace mu
 		// This runs on a random worker thread
 		MUTABLE_CPUPROFILER_SCOPE(FImageInvertTask);
 
-		ImageInvert(Result.get());
+		ImageInvert(Result.Get());
 	}
 
 
@@ -1922,7 +1922,7 @@ namespace mu
 	class FImageComposeTask : public CodeRunner::FIssuedTask
 	{
 	public:
-		FImageComposeTask(const FScheduledOp& InOp, const OP::ImageComposeArgs& InArgs, const Ptr<const Layout>& InLayout)
+		FImageComposeTask(const FScheduledOp& InOp, const OP::ImageComposeArgs& InArgs, const TSharedPtr<const FLayout>& InLayout)
 			: FIssuedTask(InOp), Args(InArgs), Layout(InLayout)
 		{}
 
@@ -1934,10 +1934,10 @@ namespace mu
 	private:
 		int32 ImageCompressionQuality = 0;
 		OP::ImageComposeArgs Args;
-		Ptr<const Layout> Layout;
-		Ptr<const Image> Block;
-		Ptr<const Image> Mask;
-		Ptr<Image> Result;
+		TSharedPtr<const FLayout> Layout;
+		TSharedPtr<const FImage> Block;
+		TSharedPtr<const FImage> Mask;
+		TSharedPtr<FImage> Result;
 		box<FIntVector2> Rect;
 		FImageOperator::FImagePixelFormatFunc ImagePixelFormatFunc;
 	};
@@ -1949,9 +1949,9 @@ namespace mu
 		MUTABLE_CPUPROFILER_SCOPE(FImageComposeTask_Prepare);
 		bOutFailed = false;
 
-		ImageCompressionQuality = Runner->m_pSettings->ImageCompressionQuality;
+		ImageCompressionQuality = Runner->Settings.ImageCompressionQuality;
 
-		Ptr<const Image> Base = Runner->LoadImage({ Args.base, Op.ExecutionIndex, Op.ExecutionOptions });
+		TSharedPtr<const FImage> Base = Runner->LoadImage({ Args.base, Op.ExecutionIndex, Op.ExecutionOptions });
 
 		int32 RelBlockIndex = Layout->FindBlock(Args.BlockId);
 
@@ -1993,7 +1993,7 @@ namespace mu
 
 		// Create destination data
 		Result = Runner->CloneOrTakeOver(Base);
-		Result->m_flags = 0;
+		Result->Flags = 0;
 
 		bool useMask = Args.mask != 0;
 		if (!useMask)
@@ -2011,8 +2011,8 @@ namespace mu
 
 				// This now happens more often since the generation of specific mips on request. For this reason
 				// this warning is usually acceptable.
-				Ptr<Image> Resized = Runner->CreateImage(Rect.size[0], Rect.size[1], 1, Block->GetFormat(), EInitializationType::NotInitialized );
-				ImOp.ImageResizeLinear(Resized.get(), ImageCompressionQuality, Block.get());
+				TSharedPtr<FImage> Resized = Runner->CreateImage(Rect.size[0], Rect.size[1], 1, Block->GetFormat(), EInitializationType::NotInitialized );
+				ImOp.ImageResizeLinear(Resized.Get(), ImageCompressionQuality, Block.Get());
 				Runner->Release(Block);
 				Block = Resized;
 			}
@@ -2025,18 +2025,18 @@ namespace mu
 
 				if (Result->GetFormat() != Format)
 				{
-					Ptr<Image> Formatted = Runner->CreateImage(Result->GetSizeX(), Result->GetSizeY(), Result->GetLODCount(), Format, EInitializationType::NotInitialized);
+					TSharedPtr<FImage> Formatted = Runner->CreateImage(Result->GetSizeX(), Result->GetSizeY(), Result->GetLODCount(), Format, EInitializationType::NotInitialized);
 					bool bSuccess = false;
-					ImOp.ImagePixelFormat(bSuccess, ImageCompressionQuality, Formatted.get(), Result.get());
+					ImOp.ImagePixelFormat(bSuccess, ImageCompressionQuality, Formatted.Get(), Result.Get());
 					check(bSuccess); // Decompression cannot fail
 					Runner->Release(Result);
 					Result = Formatted;
 				}
 				if (Block->GetFormat() != Format)
 				{
-					Ptr<Image> Formatted = Runner->CreateImage(Block->GetSizeX(), Block->GetSizeY(), Block->GetLODCount(), Format, EInitializationType::NotInitialized);
+					TSharedPtr<FImage> Formatted = Runner->CreateImage(Block->GetSizeX(), Block->GetSizeY(), Block->GetLODCount(), Format, EInitializationType::NotInitialized);
 					bool bSuccess = false;
-					ImOp.ImagePixelFormat(bSuccess, ImageCompressionQuality, Formatted.get(), Block.get());
+					ImOp.ImagePixelFormat(bSuccess, ImageCompressionQuality, Formatted.Get(), Block.Get());
 					check(bSuccess); // Decompression cannot fail
 					Runner->Release(Block);
 					Block = Formatted;
@@ -2044,7 +2044,7 @@ namespace mu
 			}
 		}
 
-		ImagePixelFormatFunc = Runner->m_pSystem->ImagePixelFormatOverride;
+		ImagePixelFormatFunc = Runner->System->ImagePixelFormatOverride;
 
 		return true;
 	}
@@ -2063,14 +2063,14 @@ namespace mu
 			// Compose without a mask
 			// \TODO: track allocs
 			FImageOperator ImOp = FImageOperator::GetDefault(ImagePixelFormatFunc);
-			ImOp.ImageCompose(Result.get(), Block.get(), Rect);
+			ImOp.ImageCompose(Result.Get(), Block.Get(), Rect);
 		}
 		else
 		{
 			MUTABLE_CPUPROFILER_SCOPE(ImageComposeWithMask);
 
 			// Compose with a mask
-			ImageBlendOnBaseNoAlpha(Result.get(), Mask.get(), Block.get(), Rect);
+			ImageBlendOnBaseNoAlpha(Result.Get(), Mask.Get(), Block.Get(), Rect);
 		}
 
 		Layout = nullptr;
@@ -2102,18 +2102,18 @@ namespace mu
 		// This runs in the mutable Runner thread
 		MUTABLE_CPUPROFILER_SCOPE(FLoadMeshRomTask_Prepare);
 
-		if (!Runner || !Runner->m_pSystem)
+		if (!Runner || !Runner->System)
 		{
 			return false;
 		}
 
 		bOutFailed = false;
 
-		const FProgram& Program = Runner->m_pModel->GetPrivate()->m_program;
+		const FProgram& Program = Runner->Model->GetPrivate()->Program;
 
-		check(RomIndex < Program.m_roms.Num());
+		check(RomIndex < Program.Roms.Num());
 		
-		FWorkingMemoryManager::FModelCacheEntry* ModelCache = Runner->m_pSystem->WorkingMemoryManager.FindModelCache(Runner->m_pModel.Get());
+		FWorkingMemoryManager::FModelCacheEntry* ModelCache = Runner->System->WorkingMemoryManager.FindModelCache(Runner->Model.Get());
 		++ModelCache->PendingOpsPerRom[RomIndex];
 
 		if (Program.IsRomLoaded(RomIndex))
@@ -2129,17 +2129,17 @@ namespace mu
 
 		FRomLoadOp& RomLoadOp = Runner->RomLoadOps.Create(RomIndex);
 		
-		check(Runner->m_pSystem->StreamInterface);
+		check(Runner->System->StreamInterface);
 
-		const uint32 RomSize = Program.m_roms[RomIndex].Size;
+		const uint32 RomSize = Program.Roms[RomIndex].Size;
 		check(RomSize > 0);
 
 		// Free roms if necessary
 		{
 			MUTABLE_CPUPROFILER_SCOPE(FreeingRoms);
 
-			Runner->m_pSystem->WorkingMemoryManager.MarkRomUsed(RomIndex, Runner->m_pModel);
-			Runner->m_pSystem->WorkingMemoryManager.EnsureBudgetBelow(RomSize);
+			Runner->System->WorkingMemoryManager.MarkRomUsed(RomIndex, Runner->Model);
+			Runner->System->WorkingMemoryManager.EnsureBudgetBelow(RomSize);
 		}
 
 		const int32 SizeBefore = RomLoadOp.m_streamBuffer.GetAllocatedSize();
@@ -2154,8 +2154,8 @@ namespace mu
 			ReadCompletionEvent.Trigger();
 		};
 		
-		const uint32 RomId = Program.m_roms[RomIndex].Id;
-		RomLoadOp.m_streamID = Runner->m_pSystem->StreamInterface->BeginReadBlock(Runner->m_pModel.Get(), RomId, RomLoadOp.m_streamBuffer.GetData(), RomSize, &Callback);
+		const uint32 RomId = Program.Roms[RomIndex].Id;
+		RomLoadOp.m_streamID = Runner->System->StreamInterface->BeginReadBlock(Runner->Model.Get(), RomId, RomLoadOp.m_streamBuffer.GetData(), RomSize, &Callback);
 		if (RomLoadOp.m_streamID < 0)
 		{
 			bOutFailed = true;
@@ -2173,27 +2173,27 @@ namespace mu
 		// This runs in the Runner thread
 		MUTABLE_CPUPROFILER_SCOPE(FLoadMeshRomTask_Complete);
 
-		if (!Runner || !Runner->m_pSystem)
+		if (!Runner || !Runner->System)
 		{
 			return false;
 		}
 
-		FProgram& Program = Runner->m_pModel->GetPrivate()->m_program;
+		FProgram& Program = Runner->Model->GetPrivate()->Program;
 
 		// Since task could be reordered, we need to make sure we end the rom read before continuing
 		if (FRomLoadOp* RomLoadOp = Runner->RomLoadOps.Find(RomIndex))
 		{
-			Runner->m_pSystem->StreamInterface->EndRead(RomLoadOp->m_streamID);									
+			Runner->System->StreamInterface->EndRead(RomLoadOp->m_streamID);									
 
-			const int32 ResIndex = Program.m_roms[RomIndex].ResourceIndex;
+			const int32 ResIndex = Program.Roms[RomIndex].ResourceIndex;
 			check(!Program.ConstantMeshes[ResIndex].Value)
 			MUTABLE_CPUPROFILER_SCOPE(Unserialise);
 
-			InputMemoryStream Stream(RomLoadOp->m_streamBuffer.GetData(), RomLoadOp->m_streamBuffer.Num());
-			InputArchive Arch(&Stream);
+			FInputMemoryStream Stream(RomLoadOp->m_streamBuffer.GetData(), RomLoadOp->m_streamBuffer.Num());
+			FInputArchive Arch(&Stream);
 
 			check(!Program.ConstantMeshes[ResIndex].Value);
-			Ptr<Mesh> Value = Mesh::StaticUnserialise(Arch);
+			TSharedPtr<FMesh> Value = FMesh::StaticUnserialise(Arch);
 			Program.SetMeshRomValue(RomIndex, Value);
 
 			check(Program.ConstantMeshes[ResIndex].Value);
@@ -2202,11 +2202,11 @@ namespace mu
 		}
 
 		// Process the constant op normally, now that the rom is loaded.
-		Runner->RunCode(Op, Runner->m_pParams, Runner->m_pModel, Runner->m_lodMask);
+		Runner->RunCode(Op, Runner->Params, Runner->Model, Runner->LODMask);
 
-		FWorkingMemoryManager::FModelCacheEntry* ModelCache = Runner->m_pSystem->WorkingMemoryManager.FindModelCache(Runner->m_pModel.Get());
+		FWorkingMemoryManager::FModelCacheEntry* ModelCache = Runner->System->WorkingMemoryManager.FindModelCache(Runner->Model.Get());
 
-		Runner->m_pSystem->WorkingMemoryManager.MarkRomUsed(RomIndex, Runner->m_pModel);
+		Runner->System->WorkingMemoryManager.MarkRomUsed(RomIndex, Runner->Model);
 		--ModelCache->PendingOpsPerRom[RomIndex];
 
 		bool bSuccess = true;
@@ -2223,10 +2223,10 @@ namespace mu
 		MUTABLE_CPUPROFILER_SCOPE(FLoadExtensionDataTask_Prepare);
 		bOutFailed = false;
 
-		FProgram& Program = Runner->m_pModel->GetPrivate()->m_program;
-		check(Program.m_constantExtensionData.IsValidIndex(ModelConstantIndex));
+		FProgram& Program = Runner->Model->GetPrivate()->Program;
+		check(Program.ConstantExtensionData.IsValidIndex(ModelConstantIndex));
 
-		const FExtensionDataConstant& Constant = Program.m_constantExtensionData[ModelConstantIndex];
+		const FExtensionDataConstant& Constant = Program.ConstantExtensionData[ModelConstantIndex];
 		check(Constant.Data == Data);
 
 		if (Constant.LoadState == FExtensionDataConstant::ELoadState::FailedToLoad)
@@ -2238,16 +2238,16 @@ namespace mu
 			return false;
 		}
 
-		TArray<ExtensionDataPtrConst> UnloadedConstants;
-		LoadHandle = Runner->m_pSystem->GetExtensionDataStreamer()->StartLoad(Data, UnloadedConstants);
+		TArray<TSharedPtr<const FExtensionData>> UnloadedConstants;
+		LoadHandle = Runner->System->GetExtensionDataStreamer()->StartLoad(Data, UnloadedConstants);
 		check(LoadHandle.IsValid());
 
 		// Update the load state of any constants that were unloaded
-		for (const ExtensionDataPtrConst& Unloaded : UnloadedConstants)
+		for (const TSharedPtr<const FExtensionData>& Unloaded : UnloadedConstants)
 		{
-			check(Unloaded->Origin == ExtensionData::EOrigin::ConstantStreamed);
+			check(Unloaded->Origin == FExtensionData::EOrigin::ConstantStreamed);
 
-			FExtensionDataConstant* FoundConstant = Algo::FindBy(Program.m_constantExtensionData, Unloaded, &FExtensionDataConstant::Data);
+			FExtensionDataConstant* FoundConstant = Algo::FindBy(Program.ConstantExtensionData, Unloaded, &FExtensionDataConstant::Data);
 			check(FoundConstant);
 
 			FoundConstant->LoadState = FExtensionDataConstant::ELoadState::Unloaded;
@@ -2268,10 +2268,10 @@ namespace mu
 		// Complete should only be called if the load is finished
 		check(LoadHandle->LoadState != FExtensionDataLoadHandle::ELoadState::Pending);
 
-		FProgram& Program = Runner->m_pModel->GetPrivate()->m_program;
-		check(Program.m_constantExtensionData.IsValidIndex(ModelConstantIndex));
+		FProgram& Program = Runner->Model->GetPrivate()->Program;
+		check(Program.ConstantExtensionData.IsValidIndex(ModelConstantIndex));
 
-		FExtensionDataConstant& Constant = Program.m_constantExtensionData[ModelConstantIndex];
+		FExtensionDataConstant& Constant = Program.ConstantExtensionData[ModelConstantIndex];
 		check(Constant.Data == Data);
 
 		// Update the load state in the Program
@@ -2287,7 +2287,7 @@ namespace mu
 		}
 
 		// Process the constant op normally, now that the data is loaded.
-		Runner->RunCode(Op, Runner->m_pParams, Runner->m_pModel, Runner->m_lodMask);
+		Runner->RunCode(Op, Runner->Params, Runner->Model, Runner->LODMask);
 
 		return bSuccess;
 	}
@@ -2306,7 +2306,7 @@ namespace mu
 	//---------------------------------------------------------------------------------------------
 	bool CodeRunner::FLoadImageRomsTask::Prepare(CodeRunner* Runner, bool& bOutFailed )
 	{
-		if (!Runner || !Runner->m_pSystem)
+		if (!Runner || !Runner->System)
 		{
 			return false;
 		}
@@ -2315,9 +2315,9 @@ namespace mu
 		MUTABLE_CPUPROFILER_SCOPE(FLoadImageRomsTask_Prepare);
 		bOutFailed = false;
 
-		FProgram& program = Runner->m_pModel->GetPrivate()->m_program;
+		FProgram& program = Runner->Model->GetPrivate()->Program;
 
-		FWorkingMemoryManager::FModelCacheEntry* ModelCache = Runner->m_pSystem->WorkingMemoryManager.FindModelCache(Runner->m_pModel.Get());
+		FWorkingMemoryManager::FModelCacheEntry* ModelCache = Runner->System->WorkingMemoryManager.FindModelCache(Runner->Model.Get());
 
 		TArray<UE::Tasks::FTask> ReadCompleteEvents;
 
@@ -2335,7 +2335,7 @@ namespace mu
 			}
 
 			int32 RomIndex = program.ConstantImageLODs[CurrentIndex].Key;
-			check(RomIndex < program.m_roms.Num());
+			check(RomIndex < program.Roms.Num());
 
 			++ModelCache->PendingOpsPerRom[RomIndex];
 
@@ -2357,15 +2357,15 @@ namespace mu
 
 			FRomLoadOp& RomLoadOp = Runner->RomLoadOps.Create(RomIndex);
 			
-			check(Runner->m_pSystem->StreamInterface);
+			check(Runner->System->StreamInterface);
 
-			const uint32 RomSize = program.m_roms[RomIndex].Size;
+			const uint32 RomSize = program.Roms[RomIndex].Size;
 			check(RomSize > 0);
 
 			// Free roms if necessary
 			{
-				Runner->m_pSystem->WorkingMemoryManager.MarkRomUsed(RomIndex, Runner->m_pModel);
-				Runner->m_pSystem->WorkingMemoryManager.EnsureBudgetBelow(RomSize);
+				Runner->System->WorkingMemoryManager.MarkRomUsed(RomIndex, Runner->Model);
+				Runner->System->WorkingMemoryManager.EnsureBudgetBelow(RomSize);
 			}
 
 			const int32 SizeBefore = RomLoadOp.m_streamBuffer.GetAllocatedSize();
@@ -2381,8 +2381,8 @@ namespace mu
 				ReadCompletionEvent.Trigger();
 			};
 			
-			const uint32 RomId = program.m_roms[RomIndex].Id;
-			RomLoadOp.m_streamID = Runner->m_pSystem->StreamInterface->BeginReadBlock(Runner->m_pModel.Get(), RomId, RomLoadOp.m_streamBuffer.GetData(), RomSize, &Callback);
+			const uint32 RomId = program.Roms[RomIndex].Id;
+			RomLoadOp.m_streamID = Runner->System->StreamInterface->BeginReadBlock(Runner->Model.Get(), RomId, RomLoadOp.m_streamBuffer.GetData(), RomSize, &Callback);
 			if (RomLoadOp.m_streamID < 0)
 			{
 				bOutFailed = true;
@@ -2407,14 +2407,14 @@ namespace mu
 		// This runs in the Runner thread
 		MUTABLE_CPUPROFILER_SCOPE(FLoadImageRomsTask_Complete);
 
-		if (!Runner || !Runner->m_pSystem)
+		if (!Runner || !Runner->System)
 		{
 			return false;
 		}
 
-		FProgram& Program = Runner->m_pModel->GetPrivate()->m_program;
+		FProgram& Program = Runner->Model->GetPrivate()->Program;
 		
-		FWorkingMemoryManager::FModelCacheEntry* ModelCache = Runner->m_pSystem->WorkingMemoryManager.FindModelCache(Runner->m_pModel.Get());
+		FWorkingMemoryManager::FModelCacheEntry* ModelCache = Runner->System->WorkingMemoryManager.FindModelCache(Runner->Model.Get());
 
 		bool bSomeMissingData = false;
 
@@ -2423,20 +2423,20 @@ namespace mu
 			// Since task could be reordered, we need to make sure we end the rom read before continuing
 			if (FRomLoadOp* RomLoadOp = Runner->RomLoadOps.Find(RomIndex))
 			{				
-				bool bSuccess = Runner->m_pSystem->StreamInterface->EndRead(RomLoadOp->m_streamID);
+				bool bSuccess = Runner->System->StreamInterface->EndRead(RomLoadOp->m_streamID);
 
 				MUTABLE_CPUPROFILER_SCOPE(Unserialise);
 
 				if (bSuccess)
 				{
-					InputMemoryStream Stream(RomLoadOp->m_streamBuffer.GetData(), RomLoadOp->m_streamBuffer.Num());
-					InputArchive Arch(&Stream);
+					FInputMemoryStream Stream(RomLoadOp->m_streamBuffer.GetData(), RomLoadOp->m_streamBuffer.Num());
+					FInputArchive Arch(&Stream);
 
-					const int32 ResIndex = Program.m_roms[RomIndex].ResourceIndex;
+					const int32 ResIndex = Program.Roms[RomIndex].ResourceIndex;
 
 					// TODO: Try to reuse buffer from PooledImages.
 					check(!Program.ConstantImageLODs[ResIndex].Value);
-					Ptr<Image> Value = Image::StaticUnserialise(Arch);
+					TSharedPtr<FImage> Value = FImage::StaticUnserialise(Arch);
 
 					Program.SetImageRomValue(RomIndex, Value);
 					check(Program.ConstantImageLODs[ResIndex].Value);
@@ -2457,7 +2457,7 @@ namespace mu
 		}
 
 		// Process the constant op normally, now that the rom is loaded.	
-		bool bSuccess = Runner->RunCode_ConstantResource(Op, Runner->m_pModel.Get());
+		bool bSuccess = Runner->RunCode_ConstantResource(Op, Runner->Model.Get());
 
 		for (int32 LODIndex = 0; bSuccess && (LODIndex < LODIndexCount); ++LODIndex)
 		{
@@ -2471,9 +2471,9 @@ namespace mu
 			}
 
 			int32 RomIndex = Program.ConstantImageLODs[CurrentIndex].Key;
-			check(RomIndex < Program.m_roms.Num());
+			check(RomIndex < Program.Roms.Num());
 
-			Runner->m_pSystem->WorkingMemoryManager.MarkRomUsed(RomIndex, Runner->m_pModel);
+			Runner->System->WorkingMemoryManager.MarkRomUsed(RomIndex, Runner->Model);
 			--ModelCache->PendingOpsPerRom[RomIndex];
 
 			if (DebugRom && (DebugRomAll || RomIndex == DebugRomIndex))
@@ -2502,7 +2502,7 @@ namespace mu
 		uint8 MipmapsToSkip;
 		CodeRunner::FExternalResourceId Id;
 
-		Ptr<Image> Result;
+		TSharedPtr<FImage> Result;
 		
 		TFunction<void()> ExternalCleanUpFunc;
 	};
@@ -2528,7 +2528,7 @@ namespace mu
 		// the task and the later will always outlive the former. 
 		
 		// Probably we could simply pass a reference to the result image. 
-		TFunction<void (Ptr<Image>)> ResultCallback = [this](Ptr<Image> InResult)
+		TFunction<void (TSharedPtr<FImage>)> ResultCallback = [this](TSharedPtr<FImage> InResult)
 		{
 			Result = InResult;
 		};
@@ -2570,7 +2570,7 @@ namespace mu
 		uint8 MipmapsToSkip;
 		CodeRunner::FExternalResourceId Id;
 
-		Ptr<Mesh> Result;
+		TSharedPtr<FMesh> Result;
 
 		TFunction<void()> ExternalCleanUpFunc;
 	};
@@ -2595,7 +2595,7 @@ namespace mu
 		// the task and the later will always outlive the former. 
 
 		// Probably we could simply pass a reference to the result mesh. 
-		TFunction<void(Ptr<Mesh>)> ResultCallback = [this](Ptr<Mesh> InResult)
+		TFunction<void(TSharedPtr<FMesh>)> ResultCallback = [this](TSharedPtr<FMesh> InResult)
 			{
 				Result = InResult;
 			};
@@ -2628,13 +2628,13 @@ namespace mu
 	{
 		TSharedPtr<FIssuedTask> Issued;
 
-		FProgram& program = m_pModel->GetPrivate()->m_program;
+		FProgram& program = Model->GetPrivate()->Program;
 
-		OP_TYPE type = program.GetOpType(item.At);
+		EOpType type = program.GetOpType(item.At);
 
 		switch (type)
 		{
-		case OP_TYPE::ME_CONSTANT:
+		case EOpType::ME_CONSTANT:
 		{
 			OP::MeshConstantArgs args = program.GetOpArgs<OP::MeshConstantArgs>(item.At);
 			int32 RomIndex = program.ConstantMeshes[args.value].Key;
@@ -2649,11 +2649,11 @@ namespace mu
 			break;
 		}
 		
-		case OP_TYPE::ED_CONSTANT:
+		case EOpType::ED_CONSTANT:
 		{
 			OP::ResourceConstantArgs Args = program.GetOpArgs<OP::ResourceConstantArgs>(item.At);
 
-			const FExtensionDataConstant::ELoadState LoadState = program.m_constantExtensionData[Args.value].LoadState;
+			const FExtensionDataConstant::ELoadState LoadState = program.ConstantExtensionData[Args.value].LoadState;
 
 			check(LoadState != FExtensionDataConstant::ELoadState::Invalid);
 			
@@ -2664,12 +2664,12 @@ namespace mu
 			}
 			else
 			{
-				Issued = MakeShared<FLoadExtensionDataTask>(item, program.m_constantExtensionData[Args.value].Data, Args.value);
+				Issued = MakeShared<FLoadExtensionDataTask>(item, program.ConstantExtensionData[Args.value].Data, Args.value);
 			}
 			break;
 		}
 
-		case OP_TYPE::IM_CONSTANT:
+		case EOpType::IM_CONSTANT:
 		{
 			OP::ResourceConstantArgs args = program.GetOpArgs<OP::ResourceConstantArgs>(item.At);
 			int32 MipsToSkip = item.ExecutionOptions;
@@ -2709,12 +2709,12 @@ namespace mu
 			break;
 		}		
 
-		case OP_TYPE::IM_PARAMETER:
+		case EOpType::IM_PARAMETER:
 		{
 			OP::ParameterArgs args = program.GetOpArgs<OP::ParameterArgs>(item.At);
-			Ptr<RangeIndex> Index = BuildCurrentOpRangeIndex(item, m_pParams, m_pModel.Get(), args.variable);
+			TSharedPtr<FRangeIndex> Index = BuildCurrentOpRangeIndex(item, Params, Model.Get(), args.variable);
 
-			const FName Id = m_pParams->GetImageValue(args.variable, Index);
+			const FName Id = Params->GetImageValue(args.variable, Index.Get());
 
 			check(ImageLOD < TNumericLimits<uint8>::Max() && ImageLOD >= 0);
 			check(ImageLOD + static_cast<int32>(item.ExecutionOptions) < TNumericLimits<uint8>::Max());
@@ -2728,9 +2728,9 @@ namespace mu
 			break;
 		}
 
-		case OP_TYPE::IM_REFERENCE:
+		case EOpType::IM_REFERENCE:
 		{
-			OP::ResourceReferenceArgs Args = m_pModel->GetPrivate()->m_program.GetOpArgs<OP::ResourceReferenceArgs>(item.At);
+			OP::ResourceReferenceArgs Args = Model->GetPrivate()->Program.GetOpArgs<OP::ResourceReferenceArgs>(item.At);
 
 			// We only convert references to images if indicated in the operation.
 			if (Args.ForceLoad)
@@ -2747,9 +2747,9 @@ namespace mu
 			break;
 		}
 
-		case OP_TYPE::ME_REFERENCE:
+		case EOpType::ME_REFERENCE:
 		{
-			OP::ResourceReferenceArgs Args = m_pModel->GetPrivate()->m_program.GetOpArgs<OP::ResourceReferenceArgs>(item.At);
+			OP::ResourceReferenceArgs Args = Model->GetPrivate()->Program.GetOpArgs<OP::ResourceReferenceArgs>(item.At);
 
 			// We only convert references to meshes if indicated in the operation.
 			if (Args.ForceLoad)
@@ -2764,7 +2764,7 @@ namespace mu
 			break;
 		}
 
-		case OP_TYPE::IM_PIXELFORMAT:
+		case EOpType::IM_PIXELFORMAT:
 		{
 			if (item.Stage == 1)
 			{
@@ -2774,7 +2774,7 @@ namespace mu
 			break;
 		}
 
-		case OP_TYPE::IM_LAYERCOLOUR:
+		case EOpType::IM_LAYERCOLOUR:
 		{
 			if (item.Stage == 1)
 			{
@@ -2784,7 +2784,7 @@ namespace mu
 			break;
 		}
 
-		case OP_TYPE::IM_LAYER:
+		case EOpType::IM_LAYER:
 		{
 			if ((ExecutionStrategy == EExecutionStrategy::MinimizeMemory && item.Stage == 2)
 				||
@@ -2797,7 +2797,7 @@ namespace mu
 			break;
 		}
 
-		case OP_TYPE::IM_MIPMAP:
+		case EOpType::IM_MIPMAP:
 		{
 			if (item.Stage == 1)
 			{
@@ -2807,7 +2807,7 @@ namespace mu
 			break;
 		}
 
-		case OP_TYPE::IM_SWIZZLE:
+		case EOpType::IM_SWIZZLE:
 		{
 			if (item.Stage == 1)
 			{
@@ -2817,7 +2817,7 @@ namespace mu
 			break;
 		}
 
-		case OP_TYPE::IM_SATURATE:
+		case EOpType::IM_SATURATE:
 		{
 			if (item.Stage == 1)
 			{
@@ -2827,7 +2827,7 @@ namespace mu
 			break;
 		}
 
-		case OP_TYPE::IM_INVERT:
+		case EOpType::IM_INVERT:
 		{
 			if (item.Stage == 1)
 			{
@@ -2837,7 +2837,7 @@ namespace mu
 			break;
 		}
 
-		case OP_TYPE::IM_RESIZE:
+		case EOpType::IM_RESIZE:
 		{
 			if (item.Stage == 1)
 			{
@@ -2847,7 +2847,7 @@ namespace mu
 			break;
 		}
 
-		case OP_TYPE::IM_RESIZEREL:
+		case EOpType::IM_RESIZEREL:
 		{
 			if (item.Stage == 1)
 			{
@@ -2857,13 +2857,13 @@ namespace mu
 			break;
 		}
 
-		case OP_TYPE::IM_COMPOSE:
+		case EOpType::IM_COMPOSE:
 		{
 			if ((ExecutionStrategy == EExecutionStrategy::MinimizeMemory && item.Stage == 3) ||
 				(ExecutionStrategy != EExecutionStrategy::MinimizeMemory && item.Stage == 2))
 			{
 				OP::ImageComposeArgs Args = program.GetOpArgs<OP::ImageComposeArgs>(item.At);
-				Ptr<const Layout> ComposeLayout = static_cast<const Layout*>( m_heapData[item.CustomState].Resource.get());
+				TSharedPtr<const FLayout> ComposeLayout = StaticCastSharedPtr<const FLayout>( HeapData[item.CustomState].Resource);
 				Issued = MakeShared<FImageComposeTask>(item, Args, ComposeLayout);
 			}
 			break;
