[mutable] More enforcement of coding standards.

diff --git a/Engine/Plugins/Mutable/Source/MutableRuntime/Private/MuR/CodeVisitor.cpp b/Engine/Plugins/Mutable/Source/MutableRuntime/Private/MuR/CodeVisitor.cpp
--- a/Engine/Plugins/Mutable/Source/MutableRuntime/Private/MuR/CodeVisitor.cpp
+++ b/Engine/Plugins/Mutable/Source/MutableRuntime/Private/MuR/CodeVisitor.cpp
@@ -24,19 +24,19 @@ namespace mu
 		// Not cached
 		MUTABLE_CPUPROFILER_SCOPE(SubtreeParametersVisitor);
 
-		Visited.SetNum(Program.m_opAddress.Num());
-		if (Program.m_opAddress.Num())
+		Visited.SetNum(Program.OpAddress.Num());
+		if (Program.OpAddress.Num())
 		{
 			FMemory::Memzero(&Visited[0], Visited.Num());
 		}
 
-		CurrentParams.SetNum(Program.m_parameters.Num());
+		CurrentParams.SetNum(Program.Parameters.Num());
 		if (CurrentParams.Num())
 		{
 			FMemory::Memzero(CurrentParams.GetData(), CurrentParams.GetAllocatedSize());
 		}
 
-		Pending.Reserve(Program.m_opAddress.Num() / 4);
+		Pending.Reserve(Program.OpAddress.Num() / 4);
 		Pending.Add(Root);
 
 		while (Pending.Num())
@@ -49,13 +49,13 @@ namespace mu
 
 				switch (Program.GetOpType(at))
 				{
-				case OP_TYPE::NU_PARAMETER:
-				case OP_TYPE::SC_PARAMETER:
-				case OP_TYPE::BO_PARAMETER:
-				case OP_TYPE::CO_PARAMETER:
-				case OP_TYPE::PR_PARAMETER:
-				case OP_TYPE::IM_PARAMETER:
-				case OP_TYPE::ST_PARAMETER:
+				case EOpType::NU_PARAMETER:
+				case EOpType::SC_PARAMETER:
+				case EOpType::BO_PARAMETER:
+				case EOpType::CO_PARAMETER:
+				case EOpType::PR_PARAMETER:
+				case EOpType::IM_PARAMETER:
+				case EOpType::ST_PARAMETER:
 					CurrentParams[Program.GetOpArgs<OP::ParameterArgs>(at).variable]++;
 					break;
 
diff --git a/Engine/Plugins/Mutable/Source/MutableRuntime/Private/MuR/ExtensionData.cpp b/Engine/Plugins/Mutable/Source/MutableRuntime/Private/MuR/ExtensionData.cpp
--- a/Engine/Plugins/Mutable/Source/MutableRuntime/Private/MuR/ExtensionData.cpp
+++ b/Engine/Plugins/Mutable/Source/MutableRuntime/Private/MuR/ExtensionData.cpp
@@ -10,19 +10,19 @@ namespace mu
 {
 
 //---------------------------------------------------------------------------------------------
-void ExtensionData::Serialise(const ExtensionData* Data, OutputArchive& Archive)
+void FExtensionData::Serialise(const FExtensionData* Data, FOutputArchive& Archive)
 {
 	Archive << *Data;
 }
 
 
 //---------------------------------------------------------------------------------------------
-ExtensionDataPtr ExtensionData::StaticUnserialise(InputArchive& Archive)
+TSharedPtr<FExtensionData> FExtensionData::StaticUnserialise(FInputArchive& Archive)
 {
 	MUTABLE_CPUPROFILER_SCOPE(ExtensionDataUnserialise);
 	LLM_SCOPE_BYNAME(TEXT("MutableRuntime"));
 
-	ExtensionDataPtr Result = new ExtensionData();
+	TSharedPtr<FExtensionData> Result = MakeShared<FExtensionData>();
 	Archive >> *Result;
 
 	return Result;
@@ -30,7 +30,7 @@ ExtensionDataPtr ExtensionData::StaticUnserialise(InputArchive& Archive)
 
 
 //---------------------------------------------------------------------------------------------
-uint32 ExtensionData::Hash() const
+uint32 FExtensionData::Hash() const
 {
 	uint32 Result = ::GetTypeHash(Index);
 	Result = HashCombine(Result, ::GetTypeHash((uint8)Origin));
@@ -40,7 +40,7 @@ uint32 ExtensionData::Hash() const
 
 
 //---------------------------------------------------------------------------------------------
-void ExtensionData::Serialise(OutputArchive& Archive) const
+void FExtensionData::Serialise(FOutputArchive& Archive) const
 {
 	Archive << Index;
 
@@ -50,7 +50,7 @@ void ExtensionData::Serialise(OutputArchive& Archive) const
 
 
 //---------------------------------------------------------------------------------------------
-void ExtensionData::Unserialise(InputArchive& Archive)
+void FExtensionData::Unserialise(FInputArchive& Archive)
 {
 	Archive >> Index;
 
@@ -61,9 +61,9 @@ void ExtensionData::Unserialise(InputArchive& Archive)
 
 
 //---------------------------------------------------------------------------------------------
-int32 ExtensionData::GetDataSize() const
+int32 FExtensionData::GetDataSize() const
 {
-	return sizeof(ExtensionData);
+	return sizeof(FExtensionData);
 }
 
 
diff --git a/Engine/Plugins/Mutable/Source/MutableRuntime/Private/MuR/Image.cpp b/Engine/Plugins/Mutable/Source/MutableRuntime/Private/MuR/Image.cpp
--- a/Engine/Plugins/Mutable/Source/MutableRuntime/Private/MuR/Image.cpp
+++ b/Engine/Plugins/Mutable/Source/MutableRuntime/Private/MuR/Image.cpp
@@ -12,16 +12,16 @@
 namespace mu
 {
 
-	Image::Image()
+	FImage::FImage()
 	{
 	}
 
-	Image::Image(uint32 SizeX, uint32 SizeY, uint32 Lods, EImageFormat Format, EInitializationType InitType)
+	FImage::FImage(uint32 SizeX, uint32 SizeY, uint32 Lods, EImageFormat Format, EInitializationType InitType)
 	{
 		Init(SizeX, SizeY, Lods, Format, InitType);
 	}
 	
-	void Image::Init(uint32 SizeX, uint32 SizeY, uint32 Lods, EImageFormat Format, EInitializationType InitType)
+	void FImage::Init(uint32 SizeX, uint32 SizeY, uint32 Lods, EImageFormat Format, EInitializationType InitType)
 	{
 		MUTABLE_CPUPROFILER_SCOPE(NewImage)
 		LLM_SCOPE_BYNAME(TEXT("MutableRuntime"));
@@ -31,113 +31,71 @@ namespace mu
 	}
 
 
-	void Image::InitToBlack()
+	void FImage::InitToBlack()
 	{
 		const FImageDesc Desc = DataStorage.MakeImageDesc();
 		DataStorage.Init(Desc, EInitializationType::Black);
 
-		m_flags = 0;
+		Flags = 0;
 		RelevancyMinY = 0;
 		RelevancyMaxY = 0;
 	}
 
 
-	Ptr<Image> Image::CreateAsReference(uint32 ID, const FImageDesc& Desc, bool bForceLoad)
+	TSharedPtr<FImage> FImage::CreateAsReference(uint32 ID, const FImageDesc& Desc, bool bForceLoad)
 	{
-		Ptr<Image> Result = new Image;
+		TSharedPtr<FImage> Result = MakeShared<FImage>();
 		Result->ReferenceID = ID;
 	
 		Result->DataStorage.InitVoid(Desc);
 
-		Result->m_flags = EImageFlags::IF_IS_REFERENCE;
+		Result->Flags = EImageFlags::IF_IS_REFERENCE;
 		if (bForceLoad)
 		{
-			Result->m_flags = Result->m_flags | EImageFlags::IF_IS_FORCELOAD;
+			Result->Flags = Result->Flags | EImageFlags::IF_IS_FORCELOAD;
 		}
 
 		return Result;
 	}
 
 
-	void Image::Serialise(const Image* ValuePtr, OutputArchive& Arch)
+	void FImage::Serialise(const FImage* ValuePtr, FOutputArchive& Arch)
 	{
 		Arch << *ValuePtr;
 	}
 
-	void Image::Serialise(OutputArchive& Arch) const
+	void FImage::Serialise(FOutputArchive& Arch) const
 	{
-		uint32 Version = 4;
-		Arch << Version;
-
 		Arch << DataStorage;
 
 		// Remove non-persistent flags.
-		uint8 flags = m_flags & ~IF_HAS_RELEVANCY_MAP;
+		uint8 flags = Flags & ~IF_HAS_RELEVANCY_MAP;
 		Arch << flags;
 	}
 
-	void Image::Unserialise(InputArchive& Arch)
+	void FImage::Unserialise(FInputArchive& Arch)
 	{
-		uint32 Version = 0;
-		Arch >> Version;
-		check(Version <= 4);
-
-		if (Version <= 3)
-		{
-			FImageSize Size;
-			Arch >> Size;
-			
-			uint8 LODs;
-			Arch >> LODs;
-
-			uint8 Format;
-			Arch >> Format;
-
-			FImageArray OldData;
-			Arch >> OldData;
-
-
-			// This does not work for RLE compressed formats.
-			if (!OldData.IsEmpty())
-			{
-				DataStorage.Init(FImageDesc(Size, static_cast<EImageFormat>(Format), LODs), EInitializationType::NotInitialized);
-				check(!DataStorage.IsEmpty());
-				
-				int32 OldDataConsumedOffset = 0;
-				for (FImageArray& Buffer : DataStorage.Buffers)
-				{
-					check(OldData.Num() >= OldDataConsumedOffset + Buffer.Num());
-					FMemory::Memcpy(Buffer.GetData(), OldData.GetData() + OldDataConsumedOffset, Buffer.Num());
-
-					OldDataConsumedOffset += Buffer.Num();
-				}
-			}
-		}
-		else
-		{
-			Arch >> DataStorage;
-		}
-
-		Arch >> m_flags;
+		Arch >> DataStorage;
+		Arch >> Flags;
 	}
 
 
-	Ptr<Image> Image::StaticUnserialise(InputArchive& Arch)
+	TSharedPtr<FImage> FImage::StaticUnserialise(FInputArchive& Arch)
 	{
 		LLM_SCOPE_BYNAME(TEXT("MutableRuntime"));
 
-		Ptr<Image> ResultPtr = new Image();
+		TSharedPtr<FImage> ResultPtr = MakeShared<FImage>();
 		Arch >> *ResultPtr;
 		
 		return ResultPtr;
 	}
 
-	int32 Image::GetDataSize() const
+	int32 FImage::GetDataSize() const
 	{
 		return DataStorage.GetAllocatedSize();
 	}
 
-	Ptr<Image> FImageOperator::ExtractMip(const Image* This, int32 Mip)
+	TSharedPtr<FImage> FImageOperator::ExtractMip(const FImage* This, int32 Mip)
 	{
 		if (Mip == 0 && This->GetLODCount() == 1)
 		{
@@ -152,8 +110,8 @@ namespace mu
 		{
 			TArrayView<const uint8> SrcView = This->DataStorage.GetLOD(Mip); 
 
-			Ptr<Image> ResultImage = CreateImage(MipSize.X, MipSize.Y, 1, This->GetFormat(), EInitializationType::NotInitialized);
-			ResultImage->m_flags = This->m_flags;
+			TSharedPtr<FImage> ResultImage = CreateImage(MipSize.X, MipSize.Y, 1, This->GetFormat(), EInitializationType::NotInitialized);
+			ResultImage->Flags = This->Flags;
 			
 			// Probably an RLE texture, resize.
 			if (ResultImage->DataStorage.IsEmpty())
@@ -172,74 +130,74 @@ namespace mu
 		// We need to generate the mip
 		// \TODO: optimize, quality
 
-		Ptr<Image> Resized = CreateImage(MipSize.X, MipSize.Y, 1, This->GetFormat(), EInitializationType::NotInitialized);
-		ImageResizeLinear(Resized.get(), Quality, This);
+		TSharedPtr<FImage> Resized = CreateImage(MipSize.X, MipSize.Y, 1, This->GetFormat(), EInitializationType::NotInitialized);
+		ImageResizeLinear(Resized.Get(), Quality, This);
 		
 		//TODO: Do we realy need to extract the mip again? The resized Image looks already to be what we want.
-		Ptr<Image> Result = ExtractMip(Resized.get(), 0);
+		TSharedPtr<FImage> Result = ExtractMip(Resized.Get(), 0);
 		ReleaseImage(Resized);
 
 		return Result;
 	}
 	
-	uint16 Image::GetSizeX() const
+	uint16 FImage::GetSizeX() const
 	{
 		return GetSize().X;
 	}
 
-	uint16 Image::GetSizeY() const
+	uint16 FImage::GetSizeY() const
 	{
 		return GetSize().Y;
 	}
 
-	const FImageSize& Image::GetSize() const
+	const FImageSize& FImage::GetSize() const
 	{
 		return DataStorage.ImageSize;
 	}
 
-	EImageFormat Image::GetFormat() const
+	EImageFormat FImage::GetFormat() const
 	{
 		return DataStorage.ImageFormat;
 	}
 
-	int32 Image::GetLODCount() const
+	int32 FImage::GetLODCount() const
 	{
 		return FMath::Max(1, (int32)DataStorage.NumLODs);
 	}
 
-	const uint8* Image::GetLODData(int32 LOD) const
+	const uint8* FImage::GetLODData(int32 LOD) const
 	{
 		return DataStorage.GetLOD(LOD).GetData();
 	}
 
-	uint8* Image::GetLODData(int32 LOD)
+	uint8* FImage::GetLODData(int32 LOD)
 	{
 		return DataStorage.GetLOD(LOD).GetData();
 	}
 
-	int32 Image::GetLODDataSize(int32 LOD) const
+	int32 FImage::GetLODDataSize(int32 LOD) const
 	{
 		return DataStorage.GetLOD(LOD).Num();
 	}
 
-	bool Image::IsReference() const
+	bool FImage::IsReference() const
 	{
-		return m_flags & EImageFlags::IF_IS_REFERENCE;
+		return Flags & EImageFlags::IF_IS_REFERENCE;
 	}
 
-	bool Image::IsForceLoad() const
+	bool FImage::IsForceLoad() const
 	{
-		return m_flags & EImageFlags::IF_IS_FORCELOAD;
+		return Flags & EImageFlags::IF_IS_FORCELOAD;
 	}
 
 	//---------------------------------------------------------------------------------------------
-	uint32 Image::GetReferencedTexture() const
+	uint32 FImage::GetReferencedTexture() const
 	{
 		ensure(IsReference());
 		return ReferenceID;
 	}
 
-	int32 Image::CalculateDataSize(int32 SizeX, int32 SizeY, int32 LODCount, EImageFormat Format)
+	int32 FImage::CalculateDataSize(int32 SizeX, int32 SizeY, int32 LODCount, EImageFormat Format)
 	{
 		int32 Result = 0;
 
@@ -263,7 +221,7 @@ namespace mu
 
 
 	//---------------------------------------------------------------------------------------------
-	FIntVector2 Image::CalculateMipSize(int32 Mip) const
+	FIntVector2 FImage::CalculateMipSize(int32 Mip) const
 	{
 		FIntVector2 Result = FIntVector2(GetSizeX(), GetSizeY());
 
@@ -283,7 +241,7 @@ namespace mu
 
 	
 	//---------------------------------------------------------------------------------------------
-	int32 Image::GetMipmapCount(int32 SizeX, int32 SizeY)
+	int32 FImage::GetMipmapCount(int32 SizeX, int32 SizeY)
 	{
 		if (SizeX <= 0 || SizeY <= 0)
 		{
@@ -295,23 +253,23 @@ namespace mu
 	}
 
 
-	const uint8* Image::GetMipData(int32 Mip) const
+	const uint8* FImage::GetMipData(int32 Mip) const
 	{
 		return DataStorage.GetLOD(Mip).GetData();
 	}
 
 
-	uint8* Image::GetMipData(int32 Mip)
+	uint8* FImage::GetMipData(int32 Mip)
 	{
 		return DataStorage.GetLOD(Mip).GetData();
 	}
 
-	int32 Image::GetMipsDataSize() const
+	int32 FImage::GetMipsDataSize() const
 	{
 		return DataStorage.GetDataSize();
 	}
 
-	FVector4f Image::Sample(FVector2f Coords) const
+	FVector4f FImage::Sample(FVector2f Coords) const
 	{
 		FVector4f Result = FVector4f(0.0f);
 
@@ -392,7 +350,7 @@ namespace mu
 	}
 
 	//---------------------------------------------------------------------------------------------
-	bool Image::IsPlainColour(FVector4f& OutColor) const
+	bool FImage::IsPlainColour(FVector4f& OutColor) const
 	{
 		bool Result = true;
 
@@ -488,7 +446,7 @@ namespace mu
 	}
 
 
-	bool Image::IsFullAlpha() const
+	bool FImage::IsFullAlpha() const
 	{
 		const TArrayView<const uint8> DataView = DataStorage.GetLOD(0);
 		
@@ -542,7 +500,7 @@ namespace mu
 	}
 
 
-	void Image::GetNonBlackRect_Reference(FImageRect& Rect) const
+	void FImage::GetNonBlackRect_Reference(FImageRect& Rect) const
 	{
 		const FImageSize Size = GetSize();
 		Rect.min[0] = Rect.min[1] = 0;
@@ -630,7 +588,7 @@ namespace mu
 
 
     //---------------------------------------------------------------------------------------------
-    void Image::GetNonBlackRect(FImageRect& rect) const
+    void FImage::GetNonBlackRect(FImageRect& rect) const
     {            
 		// TODO: There is a bug here with cyborg-windows.
 		// Meanwhile do this.
@@ -789,17 +747,17 @@ namespace mu
   //       }
     }
 
-	void Image::ReduceLODsTo(int32 NewLODCount)
+	void FImage::ReduceLODsTo(int32 NewLODCount)
 	{
 		DataStorage.SetNumLODs(NewLODCount);
 	}
 
-	void Image::ReduceLODs(int32 LODsToSkip)
+	void FImage::ReduceLODs(int32 LODsToSkip)
 	{
 		DataStorage.DropLODs(LODsToSkip);
 	}
 
-	void FImageOperator::FillColor(Image* Target, FVector4f Color)
+	void FImageOperator::FillColor(FImage* Target, FVector4f Color)
 	{
 		const EImageFormat Format = Target->GetFormat();
 
@@ -932,7 +890,7 @@ namespace mu
 		{
 			// Generic case that supports compressed formats.
 			const FImageFormatData& FormatData = GetImageFormatData(Format);
-			Ptr<Image> BlockImage = CreateImage(
+			TSharedPtr<FImage> BlockImage = CreateImage(
 						FormatData.PixelsPerBlockX, FormatData.PixelsPerBlockY, 1, EImageFormat::RGBA_UByte, EInitializationType::NotInitialized);
 
 			const uint32 PixelData = 
@@ -948,7 +906,7 @@ namespace mu
 				FMemory::Memcpy(UncompressedBlockData + I*4, &PixelData, 4);
 			}
 
-			Ptr<Image> Converted = ImagePixelFormat(0, BlockImage.get(), Format);
+			TSharedPtr<FImage> Converted = ImagePixelFormat(0, BlockImage.Get(), Format);
 			ReleaseImage(BlockImage);
 
 			const TArrayView<const uint8> BlockView = Converted->DataStorage.GetLOD(0); 
diff --git a/Engine/Plugins/Mutable/Source/MutableRuntime/Private/MuR/ImageDataStorage.cpp b/Engine/Plugins/Mutable/Source/MutableRuntime/Private/MuR/ImageDataStorage.cpp
--- a/Engine/Plugins/Mutable/Source/MutableRuntime/Private/MuR/ImageDataStorage.cpp
+++ b/Engine/Plugins/Mutable/Source/MutableRuntime/Private/MuR/ImageDataStorage.cpp
@@ -370,11 +370,8 @@ namespace mu
 		}
 	}
 
-	void FImageDataStorage::Serialise(OutputArchive& Arch) const
+	void FImageDataStorage::Serialise(FOutputArchive& Arch) const
 	{
-		int32 Version = 0;
-		Arch << Version;
-		
 		Arch << ImageSize.X;
 		Arch << ImageSize.Y;
 		Arch << ImageFormat;
@@ -391,13 +388,8 @@ namespace mu
 		Arch << CompactedTailOffsets;
 	}
 
-	void FImageDataStorage::Unserialise(InputArchive& Arch)
+	void FImageDataStorage::Unserialise(FInputArchive& Arch)
 	{
-		int32 Version;
-		Arch >> Version;
-
-		check(Version == 0);
-	
 		Arch >> ImageSize.X;
 		Arch >> ImageSize.Y;
 		Arch >> ImageFormat;
diff --git a/Engine/Plugins/Mutable/Source/MutableRuntime/Private/MuR/ImageRLE.h b/Engine/Plugins/Mutable/Source/MutableRuntime/Private/MuR/ImageRLE.h
--- a/Engine/Plugins/Mutable/Source/MutableRuntime/Private/MuR/ImageRLE.h
+++ b/Engine/Plugins/Mutable/Source/MutableRuntime/Private/MuR/ImageRLE.h
@@ -92,7 +92,7 @@ namespace mu
                                    uint8* pDestDataB);
 
 
-	inline void UncompressRLE_L(const Image* BaseImage, Image* DestImage)
+	inline void UncompressRLE_L(const FImage* BaseImage, FImage* DestImage)
 	{
 		check(DestImage->GetLODCount() <= BaseImage->GetLODCount());
 		check(DestImage->GetSize() == BaseImage->GetSize());
@@ -110,7 +110,7 @@ namespace mu
 		}
 	}
 
-    inline void CompressRLE_L(const Image* BaseImage, Image* DestImage)
+    inline void CompressRLE_L(const FImage* BaseImage, FImage* DestImage)
     {
 		check(DestImage->GetLODCount() <= BaseImage->GetLODCount());
 		check(DestImage->GetSize() == BaseImage->GetSize());
@@ -144,7 +144,7 @@ namespace mu
         }
     }
 
-    inline void CompressRLE_L1(const Image* BaseImage, Image* DestImage)
+    inline void CompressRLE_L1(const FImage* BaseImage, FImage* DestImage)
     {
 		check(DestImage->GetLODCount() <= BaseImage->GetLODCount());
 		check(DestImage->GetSize() == BaseImage->GetSize());
@@ -178,7 +178,7 @@ namespace mu
         }
     }
 
-    inline void UncompressRLE_L1(const Image* BaseImage, Image* DestImage)
+    inline void UncompressRLE_L1(const FImage* BaseImage, FImage* DestImage)
     {
 		check(DestImage->GetLODCount() <= BaseImage->GetLODCount());
 		check(DestImage->GetSize() == BaseImage->GetSize());
@@ -196,7 +196,7 @@ namespace mu
 		}
     }
 
-    inline void CompressRLE_RGB(const Image* BaseImage, Image* DestImage)
+    inline void CompressRLE_RGB(const FImage* BaseImage, FImage* DestImage)
     {
 		check(DestImage->GetLODCount() <= BaseImage->GetLODCount());
 		check(DestImage->GetSize() == BaseImage->GetSize());
@@ -230,7 +230,7 @@ namespace mu
         }
     }
 
-    inline void UncompressRLE_RGB(const Image* BaseImage, Image* DestImage)
+    inline void UncompressRLE_RGB(const FImage* BaseImage, FImage* DestImage)
     {
 		check(DestImage->GetLODCount() <= BaseImage->GetLODCount());
 		check(DestImage->GetSize() == BaseImage->GetSize());
@@ -248,7 +248,7 @@ namespace mu
         }
     }
 
-    inline void CompressRLE_RGBA(const Image* BaseImage, Image* DestImage)
+    inline void CompressRLE_RGBA(const FImage* BaseImage, FImage* DestImage)
     {
 		check(DestImage->GetLODCount() <= BaseImage->GetLODCount());
 		check(DestImage->GetSize() == BaseImage->GetSize());
@@ -282,7 +282,7 @@ namespace mu
         }
     }
 
-    inline void UncompressRLE_RGBA(const Image* BaseImage, Image* DestImage)
+    inline void UncompressRLE_RGBA(const FImage* BaseImage, FImage* DestImage)
     {
 		check(DestImage->GetLODCount() <= BaseImage->GetLODCount());
 		check(DestImage->GetSize() == BaseImage->GetSize());
diff --git a/Engine/Plugins/Mutable/Source/MutableRuntime/Private/MuR/ImageTypes.cpp b/Engine/Plugins/Mutable/Source/MutableRuntime/Private/MuR/ImageTypes.cpp
--- a/Engine/Plugins/Mutable/Source/MutableRuntime/Private/MuR/ImageTypes.cpp
+++ b/Engine/Plugins/Mutable/Source/MutableRuntime/Private/MuR/ImageTypes.cpp
@@ -64,37 +64,16 @@ namespace mu
         return s_imageFormatData[ uint8(format) ];
     }
 
-	void FMipmapGenerationSettings::Serialise(OutputArchive& Arch) const
+	void FMipmapGenerationSettings::Serialise(FOutputArchive& Arch) const
 	{
-		uint32 Version = 1;
-		Arch << Version;
-
 		Arch << FilterType;
 		Arch << AddressMode;
 	}
 
-	void FMipmapGenerationSettings::Unserialise(InputArchive& Arch)
+	void FMipmapGenerationSettings::Unserialise(FInputArchive& Arch)
 	{
-		uint32 Version = 0;
-		Arch >> Version;
-
-		check(Version <= 1);
-
-		if (Version < 1)
-		{
-			float SharpenFactor = 0.0f;
-			Arch >> SharpenFactor;
-
-			Arch >> FilterType;
-
-			bool bDitherMipmapAlpha = false;
-			Arch >> bDitherMipmapAlpha;
-		}
-		else
-		{
-			Arch >> FilterType;
-			Arch >> AddressMode;
-		}
+		Arch >> FilterType;
+		Arch >> AddressMode;
 	}
 }
 
diff --git a/Engine/Plugins/Mutable/Source/MutableRuntime/Private/MuR/Instance.cpp b/Engine/Plugins/Mutable/Source/MutableRuntime/Private/MuR/Instance.cpp
--- a/Engine/Plugins/Mutable/Source/MutableRuntime/Private/MuR/Instance.cpp
+++ b/Engine/Plugins/Mutable/Source/MutableRuntime/Private/MuR/Instance.cpp
@@ -14,7 +14,7 @@ namespace mu
 	//---------------------------------------------------------------------------------------------
 	//---------------------------------------------------------------------------------------------
 	//---------------------------------------------------------------------------------------------
-	Instance::Instance()
+	FInstance::FInstance()
 	{
 		LLM_SCOPE_BYNAME(TEXT("MutableRuntime"));
 		m_pD = new Private();
@@ -22,57 +22,57 @@ namespace mu
 
 
 	//---------------------------------------------------------------------------------------------
-	Instance::~Instance()
+	FInstance::~FInstance()
 	{
 		LLM_SCOPE_BYNAME(TEXT("MutableRuntime"));
         check( m_pD );
 		delete m_pD;
-		m_pD = 0;
+		m_pD = nullptr;
 	}
 
 
 	//---------------------------------------------------------------------------------------------
-	Instance::Private* Instance::GetPrivate() const
+	FInstance::Private* FInstance::GetPrivate() const
 	{
 		return m_pD;
 	}
 
 
     //---------------------------------------------------------------------------------------------
-    InstancePtr Instance::Clone() const
+    TSharedPtr<FInstance> FInstance::Clone() const
     {
 		LLM_SCOPE_BYNAME(TEXT("MutableRuntime"));
-        InstancePtr pResult = new Instance();
+		TSharedPtr<FInstance> Result = MakeShared<FInstance>();
 
-        *pResult->GetPrivate() = *m_pD;
+        *Result->GetPrivate() = *m_pD;
 
-        return pResult;
+        return Result;
     }
 
 
 	//---------------------------------------------------------------------------------------------
-	int32 Instance::GetDataSize() const
+	int32 FInstance::GetDataSize() const
 	{
 		return 16 + sizeof(Private) + m_pD->Components.GetAllocatedSize() + m_pD->ExtensionData.GetAllocatedSize();
 	}
 
 
     //---------------------------------------------------------------------------------------------
-    Instance::ID Instance::GetId() const
+    FInstance::FID FInstance::GetId() const
     {
         return m_pD->Id;
     }
 
 
     //---------------------------------------------------------------------------------------------
-    int32 Instance::GetComponentCount() const
+    int32 FInstance::GetComponentCount() const
     {
 		return m_pD->Components.Num();
     }
 
 
 	//---------------------------------------------------------------------------------------------
-	int32 Instance::GetLODCount( int32 ComponentIndex ) const
+	int32 FInstance::GetLODCount( int32 ComponentIndex ) const
 	{
 		if (m_pD->Components.IsValidIndex(ComponentIndex))
 		{
@@ -84,7 +84,7 @@ namespace mu
 
 	
 	//---------------------------------------------------------------------------------------------
-	uint16 Instance::GetComponentId( int32 ComponentIndex ) const
+	uint16 FInstance::GetComponentId( int32 ComponentIndex ) const
 	{
 		if (m_pD->Components.IsValidIndex(ComponentIndex))
 		{
@@ -100,7 +100,7 @@ namespace mu
 
 
     //---------------------------------------------------------------------------------------------
-    int32 Instance::GetSurfaceCount( int32 ComponentIndex, int32 LODIndex ) const
+    int32 FInstance::GetSurfaceCount( int32 ComponentIndex, int32 LODIndex ) const
     {
 		if (m_pD->Components.IsValidIndex(ComponentIndex) &&
 			m_pD->Components[ComponentIndex].LODs.IsValidIndex(LODIndex))
@@ -117,7 +117,7 @@ namespace mu
 
 
     //---------------------------------------------------------------------------------------------
-    uint32 Instance::GetSurfaceId(int32 ComponentIndex, int32 LODIndex, int32 SurfaceIndex ) const
+    uint32 FInstance::GetSurfaceId(int32 ComponentIndex, int32 LODIndex, int32 SurfaceIndex ) const
     {
         if (m_pD->Components.IsValidIndex(ComponentIndex) &&
 			m_pD->Components[ComponentIndex].LODs.IsValidIndex(LODIndex) &&
@@ -135,7 +135,7 @@ namespace mu
 
 
     //---------------------------------------------------------------------------------------------
-    int32 Instance::FindSurfaceById(int32 ComponentIndex, int32 LODIndex, uint32 id ) const
+    int32 FInstance::FindSurfaceById(int32 ComponentIndex, int32 LODIndex, uint32 id ) const
     {
 		if (m_pD->Components.IsValidIndex(ComponentIndex) &&
 			m_pD->Components[ComponentIndex].LODs.IsValidIndex(LODIndex))
@@ -158,7 +158,7 @@ namespace mu
 
 	
 	//---------------------------------------------------------------------------------------------
-	void Instance::FindBaseSurfaceBySharedId(int32 CompIndex, int32 SharedId, int32& OutSurfaceIndex, int32& OutLODIndex) const
+	void FInstance::FindBaseSurfaceBySharedId(int32 CompIndex, int32 SharedId, int32& OutSurfaceIndex, int32& OutLODIndex) const
 	{
 		if (m_pD->Components.IsValidIndex(CompIndex))
 		{
@@ -184,7 +184,7 @@ namespace mu
 
 
 	//---------------------------------------------------------------------------------------------
-	int32 Instance::GetSharedSurfaceId(int32 ComponentIndex, int32 LODIndex, int32 SurfaceIndex) const
+	int32 FInstance::GetSharedSurfaceId(int32 ComponentIndex, int32 LODIndex, int32 SurfaceIndex) const
 	{
 		if (m_pD->Components.IsValidIndex(ComponentIndex) &&
 			m_pD->Components[ComponentIndex].LODs.IsValidIndex(LODIndex) &&
@@ -202,7 +202,7 @@ namespace mu
 	
 
     //---------------------------------------------------------------------------------------------
-    uint32 Instance::GetSurfaceCustomId(int32 ComponentIndex, int32 LODIndex, int32 SurfaceIndex ) const
+    uint32 FInstance::GetSurfaceCustomId(int32 ComponentIndex, int32 LODIndex, int32 SurfaceIndex ) const
     {
         if (m_pD->Components.IsValidIndex(ComponentIndex) &&
 			m_pD->Components[ComponentIndex].LODs.IsValidIndex(LODIndex) &&
@@ -220,7 +220,7 @@ namespace mu
 
 
 	//---------------------------------------------------------------------------------------------
-    int32 Instance::GetImageCount(int32 ComponentIndex, int32 LODIndex, int32 SurfaceIndex ) const
+    int32 FInstance::GetImageCount(int32 ComponentIndex, int32 LODIndex, int32 SurfaceIndex ) const
 	{
 		check(m_pD->Components.IsValidIndex(ComponentIndex));
 		check(m_pD->Components[ComponentIndex].LODs.IsValidIndex(LODIndex));
@@ -231,7 +231,7 @@ namespace mu
 
 
 	//---------------------------------------------------------------------------------------------
-    int32 Instance::GetVectorCount(int32 ComponentIndex, int32 LODIndex, int32 SurfaceIndex ) const
+    int32 FInstance::GetVectorCount(int32 ComponentIndex, int32 LODIndex, int32 SurfaceIndex ) const
 	{
 		check(m_pD->Components.IsValidIndex(ComponentIndex));
 		check(m_pD->Components[ComponentIndex].LODs.IsValidIndex(LODIndex));
@@ -242,7 +242,7 @@ namespace mu
 
 
 	//---------------------------------------------------------------------------------------------
-    int32 Instance::GetScalarCount(int32 ComponentIndex, int32 LODIndex, int32 SurfaceIndex ) const
+    int32 FInstance::GetScalarCount(int32 ComponentIndex, int32 LODIndex, int32 SurfaceIndex ) const
 	{
 		check(m_pD->Components.IsValidIndex(ComponentIndex));
 		check(m_pD->Components[ComponentIndex].LODs.IsValidIndex(LODIndex));
@@ -253,7 +253,7 @@ namespace mu
 
 
     //---------------------------------------------------------------------------------------------
-    int32 Instance::GetStringCount(int32 ComponentIndex, int32 LODIndex, int32 SurfaceIndex ) const
+    int32 FInstance::GetStringCount(int32 ComponentIndex, int32 LODIndex, int32 SurfaceIndex ) const
     {
 		check(m_pD->Components.IsValidIndex(ComponentIndex));
 		check(m_pD->Components[ComponentIndex].LODs.IsValidIndex(LODIndex));
@@ -264,7 +264,7 @@ namespace mu
 
 
     //---------------------------------------------------------------------------------------------
-	FResourceID Instance::GetMeshId(int32 ComponentIndex, int32 LODIndex ) const
+	FResourceID FInstance::GetMeshId(int32 ComponentIndex, int32 LODIndex ) const
     {
         check(m_pD->Components.IsValidIndex(ComponentIndex));
         check(m_pD->Components[ComponentIndex].LODs.IsValidIndex(LODIndex));
@@ -274,7 +274,7 @@ namespace mu
 
 
 	//---------------------------------------------------------------------------------------------
-	FResourceID Instance::GetImageId(int32 ComponentIndex, int32 LODIndex, int32 SurfaceIndex, int32 ImageIndex) const
+	FResourceID FInstance::GetImageId(int32 ComponentIndex, int32 LODIndex, int32 SurfaceIndex, int32 ImageIndex) const
 	{
 		check(m_pD->Components.IsValidIndex(ComponentIndex));
 		check(m_pD->Components[ComponentIndex].LODs.IsValidIndex(LODIndex));
@@ -286,7 +286,7 @@ namespace mu
 
 
 	//---------------------------------------------------------------------------------------------
-    FName Instance::GetImageName(int32 ComponentIndex, int32 LODIndex, int32 SurfaceIndex, int32 ImageIndex ) const
+    FName FInstance::GetImageName(int32 ComponentIndex, int32 LODIndex, int32 SurfaceIndex, int32 ImageIndex ) const
 	{
 		check(m_pD->Components.IsValidIndex(ComponentIndex));
 		check(m_pD->Components[ComponentIndex].LODs.IsValidIndex(LODIndex));
@@ -298,7 +298,7 @@ namespace mu
 
 
 	//---------------------------------------------------------------------------------------------
-	FVector4f Instance::GetVector(int32 ComponentIndex, int32 LODIndex, int32 SurfaceIndex, int32 VectorIndex) const
+	FVector4f FInstance::GetVector(int32 ComponentIndex, int32 LODIndex, int32 SurfaceIndex, int32 VectorIndex) const
 	{
 		check(m_pD->Components.IsValidIndex(ComponentIndex));
 		check(m_pD->Components[ComponentIndex].LODs.IsValidIndex(LODIndex));
@@ -310,7 +310,7 @@ namespace mu
 
 
 	//---------------------------------------------------------------------------------------------
-	FName Instance::GetVectorName(int32 ComponentIndex, int32 LODIndex, int32 SurfaceIndex, int32 VectorIndex) const
+	FName FInstance::GetVectorName(int32 ComponentIndex, int32 LODIndex, int32 SurfaceIndex, int32 VectorIndex) const
 	{
 		check(m_pD->Components.IsValidIndex(ComponentIndex));
 		check(m_pD->Components[ComponentIndex].LODs.IsValidIndex(LODIndex));
@@ -322,7 +322,7 @@ namespace mu
 
 
 	//---------------------------------------------------------------------------------------------
-    float Instance::GetScalar(int32 ComponentIndex, int32 LODIndex, int32 SurfaceIndex, int32 ScalarIndex ) const
+    float FInstance::GetScalar(int32 ComponentIndex, int32 LODIndex, int32 SurfaceIndex, int32 ScalarIndex ) const
 	{
 		check(m_pD->Components.IsValidIndex(ComponentIndex));
 		check(m_pD->Components[ComponentIndex].LODs.IsValidIndex(LODIndex));
@@ -334,7 +334,7 @@ namespace mu
 
 
 	//---------------------------------------------------------------------------------------------
-	FName Instance::GetScalarName(int32 ComponentIndex, int32 LODIndex, int32 SurfaceIndex, int32 ScalarIndex) const
+	FName FInstance::GetScalarName(int32 ComponentIndex, int32 LODIndex, int32 SurfaceIndex, int32 ScalarIndex) const
 	{
 		check(m_pD->Components.IsValidIndex(ComponentIndex));
 		check(m_pD->Components[ComponentIndex].LODs.IsValidIndex(LODIndex));
@@ -346,7 +346,7 @@ namespace mu
 
 
     //---------------------------------------------------------------------------------------------
-    FString Instance::GetString(int32 ComponentIndex, int32 LODIndex, int32 SurfaceIndex, int32 StringIndex ) const
+    FString FInstance::GetString(int32 ComponentIndex, int32 LODIndex, int32 SurfaceIndex, int32 StringIndex ) const
     {
         check(m_pD->Components.IsValidIndex(ComponentIndex));
         check(m_pD->Components[ComponentIndex].LODs.IsValidIndex(LODIndex));
@@ -364,7 +364,7 @@ namespace mu
 
 
     //---------------------------------------------------------------------------------------------
-	FName Instance::GetStringName(int32 ComponentIndex, int32 LODIndex, int32 SurfaceIndex, int32 StringIndex) const
+	FName FInstance::GetStringName(int32 ComponentIndex, int32 LODIndex, int32 SurfaceIndex, int32 StringIndex) const
     {
         check(m_pD->Components.IsValidIndex(ComponentIndex));
         check(m_pD->Components[ComponentIndex].LODs.IsValidIndex(LODIndex));
@@ -381,14 +381,14 @@ namespace mu
 
 
     //---------------------------------------------------------------------------------------------
-	int32 Instance::GetExtensionDataCount() const
+	int32 FInstance::GetExtensionDataCount() const
 	{
 		return m_pD->ExtensionData.Num();
 	}
 
 
     //---------------------------------------------------------------------------------------------
-	void Instance::GetExtensionData(int32 Index, Ptr<const class ExtensionData>& OutExtensionData, FName& OutName) const
+	void FInstance::GetExtensionData(int32 Index, TSharedPtr<const FExtensionData>& OutExtensionData, FName& OutName) const
 	{
 		check(m_pD->ExtensionData.IsValidIndex(Index));
 
@@ -398,7 +398,7 @@ namespace mu
 
 
     //---------------------------------------------------------------------------------------------
-	int32 Instance::Private::AddComponent()
+	int32 FInstance::Private::AddComponent()
 	{
 		int32 result = Components.Emplace();
 		return result;
@@ -406,7 +406,7 @@ namespace mu
 
 
     //---------------------------------------------------------------------------------------------
-    int32 Instance::Private::AddLOD( int32 ComponentIndex )
+    int32 FInstance::Private::AddLOD( int32 ComponentIndex )
     {
         // Automatically create the necessary lods and components
         while (ComponentIndex >= Components.Num())
@@ -419,7 +419,7 @@ namespace mu
 
 
     //---------------------------------------------------------------------------------------------
-    int32 Instance::Private::AddSurface( int32 ComponentIndex, int32 LODIndex )
+    int32 FInstance::Private::AddSurface( int32 ComponentIndex, int32 LODIndex )
     {
         // Automatically create the necessary lods and components
         while (ComponentIndex >= Components.Num())
@@ -436,7 +436,7 @@ namespace mu
 
 
     //---------------------------------------------------------------------------------------------
-    void Instance::Private::SetSurfaceName( int32 ComponentIndex, int32 LODIndex, int32 SurfaceIndex, FName Name)
+    void FInstance::Private::SetSurfaceName( int32 ComponentIndex, int32 LODIndex, int32 SurfaceIndex, FName Name)
     {
         // Automatically create the necessary lods and components
 		while (ComponentIndex >= Components.Num())
@@ -458,7 +458,7 @@ namespace mu
 
 
 	//---------------------------------------------------------------------------------------------
-	void Instance::Private::SetMesh(int32 ComponentIndex, int32 LODIndex, FResourceID meshId, FName Name)
+	void FInstance::Private::SetMesh(int32 ComponentIndex, int32 LODIndex, FResourceID meshId, FName Name)
 	{
 		// Automatically create the necessary lods and components
 		while (ComponentIndex >= Components.Num())
@@ -477,7 +477,7 @@ namespace mu
 
 
 	//---------------------------------------------------------------------------------------------
-    int32 Instance::Private::AddImage( int32 ComponentIndex, int32 LODIndex, int32 SurfaceIndex, FResourceID imageId, FName Name)
+    int32 FInstance::Private::AddImage( int32 ComponentIndex, int32 LODIndex, int32 SurfaceIndex, FResourceID imageId, FName Name)
 	{
 		// Automatically create the necessary lods and components
 		while (ComponentIndex >= Components.Num())
@@ -499,7 +499,7 @@ namespace mu
 
 
 	//---------------------------------------------------------------------------------------------
-    int32 Instance::Private::AddVector( int32 ComponentIndex, int32 LODIndex, int32 SurfaceIndex, const FVector4f& vec, FName Name)
+    int32 FInstance::Private::AddVector( int32 ComponentIndex, int32 LODIndex, int32 SurfaceIndex, const FVector4f& vec, FName Name)
 	{
 		// Automatically create the necessary lods and components
 		while (ComponentIndex >= Components.Num())
@@ -521,7 +521,7 @@ namespace mu
 
 
     //---------------------------------------------------------------------------------------------
-    int32 Instance::Private::AddScalar( int32 ComponentIndex, int32 LODIndex, int32 SurfaceIndex, float sca, FName Name)
+    int32 FInstance::Private::AddScalar( int32 ComponentIndex, int32 LODIndex, int32 SurfaceIndex, float sca, FName Name)
     {
         // Automatically create the necessary lods and components
 		while (ComponentIndex >= Components.Num())
@@ -543,7 +543,7 @@ namespace mu
 
 
     //---------------------------------------------------------------------------------------------
-    int32 Instance::Private::AddString( int32 ComponentIndex, int32 LODIndex, int32 SurfaceIndex, const FString& Value, FName Name)
+    int32 FInstance::Private::AddString( int32 ComponentIndex, int32 LODIndex, int32 SurfaceIndex, const FString& Value, FName Name)
     {
         // Automatically create the necessary lods and components
 		while (ComponentIndex >= Components.Num())
@@ -565,7 +565,7 @@ namespace mu
 
 
     //---------------------------------------------------------------------------------------------
-	void Instance::Private::AddExtensionData(const Ptr<const class ExtensionData>& Data, FName Name)
+	void FInstance::Private::AddExtensionData(const TSharedPtr<const FExtensionData>& Data, FName Name)
 	{
 		check(Data);
 
diff --git a/Engine/Plugins/Mutable/Source/MutableRuntime/Private/MuR/InstancePrivate.h b/Engine/Plugins/Mutable/Source/MutableRuntime/Private/MuR/InstancePrivate.h
--- a/Engine/Plugins/Mutable/Source/MutableRuntime/Private/MuR/InstancePrivate.h
+++ b/Engine/Plugins/Mutable/Source/MutableRuntime/Private/MuR/InstancePrivate.h
@@ -78,16 +78,16 @@ namespace mu
 
 	struct NamedExtensionData
 	{
-		Ptr<const ExtensionData> Data;
+		TSharedPtr<const FExtensionData> Data;
 		FName Name;
 	};
 
-	class Instance::Private
+	class FInstance::Private
 	{
 	public:
 
         //!
-        Instance::ID Id = 0;
+        FInstance::FID Id = 0;
 
 		//!
 		TArray<FInstanceComponent,TInlineAllocator<4>> Components;
@@ -106,6 +106,6 @@ namespace mu
         int32 AddString( int32 ComponentIndex, int32 LODIndex, int32 SurfaceIndex, const FString& Value, FName Name);
 		
 		// Data must be non-null
-		void AddExtensionData(const Ptr<const class ExtensionData>& Data, FName Name);
+		void AddExtensionData(const TSharedPtr<const FExtensionData>& Data, FName Name);
     };
 }
diff --git a/Engine/Plugins/Mutable/Source/MutableRuntime/Private/MuR/Layout.cpp b/Engine/Plugins/Mutable/Source/MutableRuntime/Private/MuR/Layout.cpp
--- a/Engine/Plugins/Mutable/Source/MutableRuntime/Private/MuR/Layout.cpp
+++ b/Engine/Plugins/Mutable/Source/MutableRuntime/Private/MuR/Layout.cpp
@@ -14,80 +14,67 @@ namespace mu {
 	MUTABLE_IMPLEMENT_POD_SERIALISABLE(FLayoutBlock);
 	MUTABLE_IMPLEMENT_POD_VECTOR_SERIALISABLE(FLayoutBlock);
 
-	//---------------------------------------------------------------------------------------------
-	Layout::Layout()
-	{
-	}
 
-
-	//---------------------------------------------------------------------------------------------
-	void Layout::Serialise( const Layout* p, OutputArchive& arch )
+	void FLayout::Serialise( const FLayout* In, FOutputArchive& Arch )
 	{
-		arch << *p;
+		Arch << *In;
 	}
 
 
-	//---------------------------------------------------------------------------------------------
-	Ptr<Layout> Layout::StaticUnserialise( InputArchive& arch )
+	TSharedPtr<FLayout> FLayout::StaticUnserialise( FInputArchive& Arch )
 	{
 		LLM_SCOPE_BYNAME(TEXT("MutableRuntime"));
-		Ptr<Layout> pResult = new Layout();
-		arch >> *pResult;
-		return pResult;
+		TSharedPtr<FLayout> Result = MakeShared<FLayout>();
+		Arch >> *Result;
+		return Result;
 	}
 
 
-	//---------------------------------------------------------------------------------------------
-	Ptr<Layout> Layout::Clone() const
+	TSharedPtr<FLayout> FLayout::Clone() const
 	{
 		LLM_SCOPE_BYNAME(TEXT("MutableRuntime"));
-		Ptr<Layout> pResult = new Layout();
-		pResult->Size = Size;
-		pResult->MaxSize = MaxSize;
-		pResult->Blocks = Blocks;
-		pResult->Strategy = Strategy;
-		pResult->ReductionMethod = ReductionMethod;
-		return pResult;
+		TSharedPtr<FLayout> Result = MakeShared<FLayout>();
+		Result->Size = Size;
+		Result->MaxSize = MaxSize;
+		Result->Blocks = Blocks;
+		Result->Strategy = Strategy;
+		Result->ReductionMethod = ReductionMethod;
+		return Result;
 	}
 
 
-	//---------------------------------------------------------------------------------------------
-	bool Layout::operator==( const Layout& o ) const
+	bool FLayout::operator==( const FLayout& Other ) const
 	{
-		return (Size == o.Size) &&
-			(MaxSize == o.MaxSize) &&
-			(Blocks == o.Blocks) &&
-			(Strategy == o.Strategy) &&
+		return (Size == Other.Size) &&
+			(MaxSize == Other.MaxSize) &&
+			(Blocks == Other.Blocks) &&
+			(Strategy == Other.Strategy) &&
 			// maybe this is not needed
-			(ReductionMethod==o.ReductionMethod);
+			(ReductionMethod == Other.ReductionMethod);
 	}
 
 
-	//---------------------------------------------------------------------------------------------
-	int32 Layout::GetDataSize() const
+	int32 FLayout::GetDataSize() const
 	{
-		return sizeof(Layout) + Blocks.GetAllocatedSize();
+		return sizeof(FLayout) + Blocks.GetAllocatedSize();
 	}
 
 
-	//---------------------------------------------------------------------------------------------
-	FIntPoint Layout::GetGridSize() const
+	FIntPoint FLayout::GetGridSize() const
 	{
 		return FIntPoint(Size[0], Size[1]);
 	}
 
 
-	//---------------------------------------------------------------------------------------------
-	void Layout::SetGridSize( int32 sizeX, int32 sizeY )
+	void FLayout::SetGridSize( int32 SizeX, int32 SizeY )
 	{
-		check( sizeX>=0 && sizeY>=0 );
-		Size[0] = (uint16)sizeX;
-		Size[1] = (uint16)sizeY;
+		check( SizeX>=0 && SizeY>=0 );
+		Size[0] = (uint16)SizeX;
+		Size[1] = (uint16)SizeY;
 	}
 
 
-	//---------------------------------------------------------------------------------------------
-	void Layout::GetMaxGridSize(int32* SizeX, int32* SizeY) const
+	void FLayout::GetMaxGridSize(int32* SizeX, int32* SizeY) const
 	{
 		check(SizeX && SizeY);
 
@@ -99,47 +86,41 @@ namespace mu {
 	}
 
 
-	//---------------------------------------------------------------------------------------------
-	void Layout::SetMaxGridSize(int32 sizeX, int32 sizeY)
+	void FLayout::SetMaxGridSize(int32 SizeX, int32 SizeY)
 	{
-		check(sizeX >= 0 && sizeY >= 0);
-		MaxSize[0] = (uint16)sizeX;
-		MaxSize[1] = (uint16)sizeY;
+		check(SizeX >= 0 && SizeY >= 0);
+		MaxSize[0] = (uint16)SizeX;
+		MaxSize[1] = (uint16)SizeY;
 	}
 
 
-	//---------------------------------------------------------------------------------------------
-	int32 Layout::GetBlockCount() const
+	int32 FLayout::GetBlockCount() const
 	{
 		return Blocks.Num();
 	}
 
 
-	//---------------------------------------------------------------------------------------------
-	void Layout::SetBlockCount( int32 n )
+	void FLayout::SetBlockCount( int32 Count )
 	{
-		check( n>=0 );
+		check(Count >=0 );
 		LLM_SCOPE_BYNAME(TEXT("MutableRuntime"));
-		Blocks.SetNum( n );
+		Blocks.SetNum(Count);
 	}
 
 
-	//---------------------------------------------------------------------------------------------
-	void Layout::SetLayoutPackingStrategy(EPackStrategy InStrategy)
+	void FLayout::SetLayoutPackingStrategy(EPackStrategy InStrategy)
 	{
 		Strategy = InStrategy;
 	}
 
 
-	//---------------------------------------------------------------------------------------------
-	EPackStrategy Layout::GetLayoutPackingStrategy() const
+	EPackStrategy FLayout::GetLayoutPackingStrategy() const
 	{
 		return Strategy;
 	}
 
 
-	//---------------------------------------------------------------------------------------------
-	void Layout::Serialise(OutputArchive& arch) const
+	void FLayout::Serialise(FOutputArchive& arch) const
 	{
 		arch << Size;
 		arch << Blocks;
@@ -150,8 +131,7 @@ namespace mu {
 	}
 
 	
-	//---------------------------------------------------------------------------------------------
-	void Layout::Unserialise(InputArchive& arch)
+	void FLayout::Unserialise(FInputArchive& arch)
 	{
 		arch >> Size;
 		arch >> Blocks;
@@ -166,16 +146,15 @@ namespace mu {
 	}
 
 
-	//---------------------------------------------------------------------------------------------
-	bool Layout::IsSimilar(const Layout& o) const
+	bool FLayout::IsSimilar(const FLayout& Other) const
 	{
-		if (Size != o.Size || MaxSize != o.MaxSize ||
-			Blocks.Num() != o.Blocks.Num() || Strategy != o.Strategy)
+		if (Size != Other.Size || MaxSize != Other.MaxSize ||
+			Blocks.Num() != Other.Blocks.Num() || Strategy != Other.Strategy)
 			return false;
 
 		for (int32 i = 0; i < Blocks.Num(); ++i)
 		{
-			if (!Blocks[i].IsSimilar(o.Blocks[i])) return false;
+			if (!Blocks[i].IsSimilar(Other.Blocks[i])) return false;
 		}
 
 		return true;
@@ -183,8 +162,7 @@ namespace mu {
 	}
 
 
-	//---------------------------------------------------------------------------------------------
-	int32 Layout::FindBlock(uint64 Id) const
+	int32 FLayout::FindBlock(uint64 Id) const
 	{
 		for (int32 i = 0; i < Blocks.Num(); ++i)
 		{
@@ -198,8 +176,7 @@ namespace mu {
 	}
 
 
-	//---------------------------------------------------------------------------------------------
-	bool Layout::IsSingleBlockAndFull() const
+	bool FLayout::IsSingleBlockAndFull() const
 	{
 		if (Blocks.Num() == 1
 			&& Blocks[0].Min == FIntVector2(0, 0)
@@ -210,6 +187,5 @@ namespace mu {
 		return false;
 	}
 
-
 }
 
diff --git a/Engine/Plugins/Mutable/Source/MutableRuntime/Private/MuR/Mesh.cpp b/Engine/Plugins/Mutable/Source/MutableRuntime/Private/MuR/Mesh.cpp
--- a/Engine/Plugins/Mutable/Source/MutableRuntime/Private/MuR/Mesh.cpp
+++ b/Engine/Plugins/Mutable/Source/MutableRuntime/Private/MuR/Mesh.cpp
@@ -19,19 +19,19 @@ MUTABLE_IMPLEMENT_ENUM_SERIALISABLE(EShapeBindingMethod);
 MUTABLE_IMPLEMENT_ENUM_SERIALISABLE(EVertexColorUsage);
 
 
-void Mesh::Serialise(const Mesh* MeshPtr, OutputArchive& Arch)
+void FMesh::Serialise(const FMesh* MeshPtr, FOutputArchive& Arch)
 {
     //MeshPtr->m_pD->CheckIntegrity();
     Arch << *MeshPtr;
 }
 
 
-Ptr<Mesh> Mesh::StaticUnserialise(InputArchive& Arch)
+TSharedPtr<FMesh> FMesh::StaticUnserialise(FInputArchive& Arch)
 {
     MUTABLE_CPUPROFILER_SCOPE(MeshUnserialise)
 	LLM_SCOPE_BYNAME(TEXT("MutableRuntime"));
 
-    Ptr<Mesh> Result = new Mesh();
+    TSharedPtr<FMesh> Result = MakeShared<FMesh>();
     Arch >> *Result;
 
     //Result->m_pD->CheckIntegrity();
@@ -40,9 +40,9 @@ Ptr<Mesh> Mesh::StaticUnserialise(InputArchive& Arch)
 }
 
 
-Ptr<Mesh> Mesh::CreateAsReference(uint32 ID,  bool bForceLoad)
+TSharedPtr<FMesh> FMesh::CreateAsReference(uint32 ID,  bool bForceLoad)
 {
-	Ptr<Mesh> Result = new Mesh;
+	TSharedPtr<FMesh> Result = MakeShared<FMesh>();
 	Result->ReferenceID = ID;
 
 	EnumAddFlags(Result->Flags, EMeshFlags::IsResourceReference);
@@ -55,31 +55,31 @@ Ptr<Mesh> Mesh::CreateAsReference(uint32 ID,  bool bForceLoad)
 }
 
 
-bool Mesh::IsReference() const
+bool FMesh::IsReference() const
 {
 	return EnumHasAnyFlags(Flags, EMeshFlags::IsResourceReference);
 }
 
 
-bool Mesh::IsForceLoad() const
+bool FMesh::IsForceLoad() const
 {
 	return EnumHasAnyFlags(Flags, EMeshFlags::IsResourceForceLoad);
 }
 
 
-uint32 Mesh::GetReferencedMesh() const
+uint32 FMesh::GetReferencedMesh() const
 {
 	ensure(IsReference());
 	return ReferenceID;
 }
 
 
-Ptr<Mesh> Mesh::Clone() const
+TSharedPtr<FMesh> FMesh::Clone() const
 {
     //MUTABLE_CPUPROFILER_SCOPE(MeshClone);
 	LLM_SCOPE_BYNAME(TEXT("MutableRuntime"));
 
-    Ptr<Mesh> Result = new Mesh();
+    TSharedPtr<FMesh> Result = MakeShared<FMesh>();
 
     Result->InternalId = InternalId;
 	Result->Flags = Flags;
@@ -118,12 +118,12 @@ Ptr<Mesh> Mesh::Clone() const
 }
 
 
-Ptr<Mesh> Mesh::Clone(EMeshCopyFlags InFlags) const
+TSharedPtr<FMesh> FMesh::Clone(EMeshCopyFlags InFlags) const
 {
     //MUTABLE_CPUPROFILER_SCOPE(MeshClone);
 	LLM_SCOPE_BYNAME(TEXT("MutableRuntime"));
 
-	Ptr<Mesh> Result = new Mesh();
+	TSharedPtr<FMesh> Result = MakeShared<FMesh>();
 
     Result->InternalId = InternalId;
 	Result->MeshIDPrefix = MeshIDPrefix;
@@ -209,7 +209,7 @@ Ptr<Mesh> Mesh::Clone(EMeshCopyFlags InFlags) const
 }
 
 
-void Mesh::CopyFrom(const Mesh& From, EMeshCopyFlags InFlags)
+void FMesh::CopyFrom(const FMesh& From, EMeshCopyFlags InFlags)
 {
     //MUTABLE_CPUPROFILER_SCOPE(CopyFrom);
 
@@ -297,31 +297,31 @@ void Mesh::CopyFrom(const Mesh& From, EMeshCopyFlags InFlags)
 }
 
 
-uint32 Mesh::GetId() const
+uint32 FMesh::GetId() const
 {
     return InternalId;
 }
 
 
-int Mesh::GetVertexCount() const
+int FMesh::GetVertexCount() const
 {
     return GetVertexBuffers().GetElementCount();
 }
 
 
-FMeshBufferSet& Mesh::GetVertexBuffers()
+FMeshBufferSet& FMesh::GetVertexBuffers()
 {
     return VertexBuffers;
 }
 
 
-const FMeshBufferSet& Mesh::GetVertexBuffers() const
+const FMeshBufferSet& FMesh::GetVertexBuffers() const
 {
     return VertexBuffers;
 }
 
 
-bool Mesh::AreVertexIdsImplicit() const
+bool FMesh::AreVertexIdsImplicit() const
 {
 	// Is there a buffer for vertex ids?
 	int32 BufferIndex = -1;
@@ -332,7 +332,7 @@ bool Mesh::AreVertexIdsImplicit() const
 }
 
 
-bool Mesh::AreVertexIdsExplicit() const
+bool FMesh::AreVertexIdsExplicit() const
 {
 	// Is there a buffer for vertex ids?
 	int32 BufferIndex = -1;
@@ -352,7 +352,7 @@ bool Mesh::AreVertexIdsExplicit() const
 }
 
 
-void Mesh::MakeVertexIdsRelative()
+void FMesh::MakeVertexIdsRelative()
 {
 	check(AreVertexIdsImplicit());
 
@@ -374,7 +374,7 @@ void Mesh::MakeVertexIdsRelative()
 }
 
 
-void Mesh::MakeIdsExplicit()
+void FMesh::MakeIdsExplicit()
 {
 	MUTABLE_CPUPROFILER_SCOPE(Mesh_MakeIdsExplicit);
 
@@ -438,37 +438,37 @@ void Mesh::MakeIdsExplicit()
 }
 
 
-mu::Ptr<const Skeleton> Mesh::GetSkeleton() const
+TSharedPtr<const FSkeleton> FMesh::GetSkeleton() const
 {
     return Skeleton;
 }
 
 
-void Mesh::SetSkeleton(Ptr<const mu::Skeleton> InSkeleton)
+void FMesh::SetSkeleton(TSharedPtr<const FSkeleton> InSkeleton)
 {
     Skeleton = InSkeleton;
 }
 
 
-mu::Ptr<const PhysicsBody> Mesh::GetPhysicsBody() const
+TSharedPtr<const FPhysicsBody> FMesh::GetPhysicsBody() const
 {
     return PhysicsBody;
 }
 
 
-void Mesh::SetPhysicsBody(Ptr<const mu::PhysicsBody> InPhysicsBody)
+void FMesh::SetPhysicsBody(TSharedPtr<const FPhysicsBody> InPhysicsBody)
 {
     PhysicsBody = InPhysicsBody;
 }
 
 
-int32 Mesh::AddAdditionalPhysicsBody(Ptr<const mu::PhysicsBody> Body)
+int32 FMesh::AddAdditionalPhysicsBody(TSharedPtr<const FPhysicsBody> Body)
 {
 	return AdditionalPhysicsBodies.Add(Body);
 }
 
 
-Ptr<const PhysicsBody> Mesh::GetAdditionalPhysicsBody(int32 Index) const
+TSharedPtr<const FPhysicsBody> FMesh::GetAdditionalPhysicsBody(int32 Index) const
 {
 	check(AdditionalPhysicsBodies.IsValidIndex(Index));
 
@@ -476,37 +476,37 @@ Ptr<const PhysicsBody> Mesh::GetAdditionalPhysicsBody(int32 Index) const
 }
 
 
-int32 Mesh::GetFaceCount() const
+int32 FMesh::GetFaceCount() const
 {
     return GetIndexBuffers().GetElementCount() / 3;
 }
 
 
-int32 Mesh::GetIndexCount() const
+int32 FMesh::GetIndexCount() const
 {
     return GetIndexBuffers().GetElementCount();
 }
 
 
-FMeshBufferSet& Mesh::GetIndexBuffers()
+FMeshBufferSet& FMesh::GetIndexBuffers()
 {
     return IndexBuffers;
 }
 
 
-const FMeshBufferSet& Mesh::GetIndexBuffers() const
+const FMeshBufferSet& FMesh::GetIndexBuffers() const
 {
     return IndexBuffers;
 }
 
 
-int32 Mesh::GetSurfaceCount() const
+int32 FMesh::GetSurfaceCount() const
 {
     return Surfaces.Num();
 }
 
 
-void Mesh::GetSurface(
+void FMesh::GetSurface(
 		int32 SurfaceIndex,
         int32& OutFirstVertex, int32& OutVertexCount,
     	int32& OutFirstIndex, int32& OutIndexCount,
@@ -558,7 +558,7 @@ void Mesh::GetSurface(
 }
 
 
-uint32 Mesh::GetSurfaceId(int32 SurfaceIndex) const
+uint32 FMesh::GetSurfaceId(int32 SurfaceIndex) const
 {
     if (SurfaceIndex >= 0 && SurfaceIndex < Surfaces.Num())
     {
@@ -570,28 +570,28 @@ uint32 Mesh::GetSurfaceId(int32 SurfaceIndex) const
 }
 
 
-void Mesh::AddLayout(Ptr<const Layout> InLayout)
+void FMesh::AddLayout(TSharedPtr<const FLayout> InLayout)
 {
 	LLM_SCOPE_BYNAME(TEXT("MutableRuntime"));
     Layouts.Add(InLayout);
 }
 
 
-int32 Mesh::GetLayoutCount() const
+int32 FMesh::GetLayoutCount() const
 {
     return Layouts.Num();
 }
 
 
-const Layout* Mesh::GetLayout(int32 LayoutIndex) const
+TSharedPtr<const FLayout> FMesh::GetLayout(int32 LayoutIndex) const
 {
     check(Layouts.IsValidIndex(LayoutIndex));
 
-    return Layouts[LayoutIndex].get();
+    return Layouts[LayoutIndex];
 }
 
 
-void Mesh::SetLayout(int32 LayoutIndex, Ptr<const Layout> InLayout)
+void FMesh::SetLayout(int32 LayoutIndex, TSharedPtr<const FLayout> InLayout)
 {
     check(Layouts.IsValidIndex(LayoutIndex));
 
@@ -599,20 +599,20 @@ void Mesh::SetLayout(int32 LayoutIndex, Ptr<const Layout> InLayout)
 }
 
 
-int32 Mesh::GetTagCount() const
+int32 FMesh::GetTagCount() const
 {
     return Tags.Num();
 }
 
 
-void Mesh::SetTagCount(int32 Count)
+void FMesh::SetTagCount(int32 Count)
 {
 	LLM_SCOPE_BYNAME(TEXT("MutableRuntime"));
     Tags.SetNum(Count);
 }
 
 
-const FString& Mesh::GetTag(int32 TagIndex) const
+const FString& FMesh::GetTag(int32 TagIndex) const
 {
     check(Tags.IsValidIndex(TagIndex));
 
@@ -628,7 +628,7 @@ const FString& Mesh::GetTag(int32 TagIndex) const
 }
 
 
-void Mesh::SetTag(int32 TagIndex, const FString& Name)
+void FMesh::SetTag(int32 TagIndex, const FString& Name)
 {
     check(Tags.IsValidIndex(TagIndex));
 	LLM_SCOPE_BYNAME(TEXT("MutableRuntime"));
@@ -640,38 +640,38 @@ void Mesh::SetTag(int32 TagIndex, const FString& Name)
 }
 
 
-void Mesh::AddStreamedResource(uint64 ResourceId)
+void FMesh::AddStreamedResource(uint64 ResourceId)
 {
 	StreamedResources.AddUnique(ResourceId);
 }
 
 
-const TArray<uint64>& Mesh::GetStreamedResources() const
+const TArray<uint64>& FMesh::GetStreamedResources() const
 {
 	return StreamedResources;
 }
 
 
-int32 Mesh::FindBonePose(const FBoneName& BoneId) const
+int32 FMesh::FindBonePose(const FBoneName& BoneId) const
 {
 	return BonePoses.IndexOfByPredicate([BoneId](const FBonePose& Pose) { return Pose.BoneId == BoneId; });
 }
 
 
-void mu::Mesh::SetBonePoseCount(int32 count)
+void mu::FMesh::SetBonePoseCount(int32 count)
 {
 	LLM_SCOPE_BYNAME(TEXT("MutableRuntime"));
 	BonePoses.SetNum(count);
 }
 
 
-int32 mu::Mesh::GetBonePoseCount() const
+int32 mu::FMesh::GetBonePoseCount() const
 {
 	return BonePoses.Num();
 }
 
 
-void mu::Mesh::SetBonePose(int32 Index, const FBoneName& BoneId, FTransform3f Transform, EBoneUsageFlags BoneUsageFlags)
+void mu::FMesh::SetBonePose(int32 Index, const FBoneName& BoneId, FTransform3f Transform, EBoneUsageFlags BoneUsageFlags)
 {
 	check(BonePoses.IsValidIndex(Index));
 	if (BonePoses.IsValidIndex(Index))
@@ -681,59 +681,59 @@ void mu::Mesh::SetBonePose(int32 Index, const FBoneName& BoneId, FTransform3f Tr
 }
 
 
-const FBoneName& Mesh::GetBonePoseId(int32 Index) const
+const FBoneName& FMesh::GetBonePoseId(int32 Index) const
 {
 	check(BonePoses.IsValidIndex(Index));
 	return BonePoses[Index].BoneId;
 }
 
 
-void mu::Mesh::GetBonePoseTransform(int32 BoneIndex, FTransform3f& Transform) const
+void mu::FMesh::GetBonePoseTransform(int32 BoneIndex, FTransform3f& Transform) const
 {
 	check(BoneIndex >= 0 && BoneIndex < BonePoses.Num());
 	Transform = BoneIndex > INDEX_NONE ? BonePoses[BoneIndex].BoneTransform : FTransform3f::Identity;
 }
 
 
-EBoneUsageFlags Mesh::GetBoneUsageFlags(int32 BoneIndex) const
+EBoneUsageFlags FMesh::GetBoneUsageFlags(int32 BoneIndex) const
 {
 	check(BoneIndex >= 0 && BoneIndex < BonePoses.Num());
 	return BoneIndex > INDEX_NONE ? BonePoses[BoneIndex].BoneUsageFlags : EBoneUsageFlags::None;
 }
 
 
-void Mesh::SetBoneMap(const TArray<FBoneName>& InBoneMap)
+void FMesh::SetBoneMap(const TArray<FBoneName>& InBoneMap)
 {
 	BoneMap = InBoneMap;
 }
 
 
-const TArray<FBoneName>& Mesh::GetBoneMap() const
+const TArray<FBoneName>& FMesh::GetBoneMap() const
 {
 	return BoneMap;
 }
 
 
-int32 Mesh::GetSkeletonIDsCount() const
+int32 FMesh::GetSkeletonIDsCount() const
 {
     return SkeletonIDs.Num();
 }
 
 
-int32 Mesh::GetSkeletonID(int32 SkeletonIndex) const
+int32 FMesh::GetSkeletonID(int32 SkeletonIndex) const
 {
 	return SkeletonIDs.IsValidIndex(SkeletonIndex) ? SkeletonIDs[SkeletonIndex] : INDEX_NONE;
 }
 
 
-void Mesh::AddSkeletonID(int32 SkeletonID)
+void FMesh::AddSkeletonID(int32 SkeletonID)
 {
 	check(SkeletonID != INDEX_NONE);
 	SkeletonIDs.AddUnique(SkeletonID);
 }
 
 
-int32 Mesh::GetDataSize() const
+int32 FMesh::GetDataSize() const
 {
 	// TODO: review if other mesh fields like additional physics assets
 	// are relevant and add them to the count.
@@ -745,7 +745,7 @@ int32 Mesh::GetDataSize() const
 		AdditionalBuffersSize += AdditionalBuffer.Value.GetDataSize();
 	}
 
-	return sizeof(Mesh)
+	return sizeof(FMesh)
 		+ IndexBuffers.GetDataSize()
 		+ VertexBuffers.GetDataSize()
 		+ BonePoses.Num() * sizeof(FBonePose)
@@ -753,7 +753,7 @@ int32 Mesh::GetDataSize() const
 }
 
 
-bool Mesh::HasCompatibleFormat(const Mesh* Other) const
+bool FMesh::HasCompatibleFormat(const FMesh* Other) const
 {
     bool bCompatible = true;
 
@@ -798,7 +798,7 @@ bool Mesh::HasCompatibleFormat(const Mesh* Other) const
 }
 
 
-UE::Math::TIntVector3<uint32> Mesh::GetFaceVertexIndices(int32 FaceIndex) const
+UE::Math::TIntVector3<uint32> FMesh::GetFaceVertexIndices(int32 FaceIndex) const
 {
 	UE::Math::TIntVector3<uint32> Result;
 
@@ -818,7 +818,7 @@ UE::Math::TIntVector3<uint32> Mesh::GetFaceVertexIndices(int32 FaceIndex) const
 }
 
 
-bool Mesh::FVertexMatchMap::DoMatch(int32 Vertex, int32 OtherVertex) const
+bool FMesh::FVertexMatchMap::DoMatch(int32 Vertex, int32 OtherVertex) const
 {
 	if (Vertex >= 0 && Vertex < FirstMatch.Num())
 	{
@@ -843,7 +843,7 @@ bool Mesh::FVertexMatchMap::DoMatch(int32 Vertex, int32 OtherVertex) const
 }
 
 
-void Mesh::GetVertexMap(const Mesh& Other, FVertexMatchMap& VertexMap, float Tolerance) const
+void FMesh::GetVertexMap(const FMesh& Other, FVertexMatchMap& VertexMap, float Tolerance) const
 {
     int32 VertexCount = VertexBuffers.GetElementCount();
     VertexMap.FirstMatch.SetNum(VertexCount);
@@ -944,7 +944,7 @@ void Mesh::GetVertexMap(const Mesh& Other, FVertexMatchMap& VertexMap, float Tol
 }
 
 
-void Mesh::EnsureSurfaceData()
+void FMesh::EnsureSurfaceData()
 {
 	if (!Surfaces.Num() && VertexBuffers.GetElementCount())
 	{
@@ -961,7 +961,7 @@ void Mesh::EnsureSurfaceData()
 }
 
 
-void Mesh::ResetBufferIndices()
+void FMesh::ResetBufferIndices()
 {
 	VertexBuffers.ResetBufferIndices();
 	IndexBuffers.ResetBufferIndices();
@@ -970,7 +970,7 @@ void Mesh::ResetBufferIndices()
 MUTABLE_IMPLEMENT_POD_SERIALISABLE(FSurfaceSubMesh);
 MUTABLE_IMPLEMENT_POD_VECTOR_SERIALISABLE(FSurfaceSubMesh);
 
-void FMeshSurface::Serialise(OutputArchive& Arch) const
+void FMeshSurface::Serialise(FOutputArchive& Arch) const
 {
 	Arch << SubMeshes;
 
@@ -980,7 +980,7 @@ void FMeshSurface::Serialise(OutputArchive& Arch) const
 }
 
 
-void FMeshSurface::Unserialise(InputArchive& Arch)
+void FMeshSurface::Unserialise(FInputArchive& Arch)
 {
 	Arch >> SubMeshes;
 
@@ -990,7 +990,7 @@ void FMeshSurface::Unserialise(InputArchive& Arch)
 }
 
 
-void Mesh::FBonePose::Serialise(OutputArchive& Arch) const
+void FMesh::FBonePose::Serialise(FOutputArchive& Arch) const
 {
 	Arch << BoneId;
 	Arch << BoneUsageFlags;
@@ -998,7 +998,7 @@ void Mesh::FBonePose::Serialise(OutputArchive& Arch) const
 }
 
 
-void Mesh::FBonePose::Unserialise(InputArchive& Arch)
+void FMesh::FBonePose::Unserialise(FInputArchive& Arch)
 {
 	Arch >> BoneId;
 	Arch >> BoneUsageFlags;
@@ -1006,11 +1006,8 @@ void Mesh::FBonePose::Unserialise(InputArchive& Arch)
 }
 
 
-void Mesh::Serialise(OutputArchive& Arch) const
+void FMesh::Serialise(FOutputArchive& Arch) const
 {
-	uint32 Version = 23;
-	Arch << Version;
-
 	Arch << IndexBuffers;
 	Arch << VertexBuffers;
 	Arch << AdditionalBuffers;
@@ -1037,12 +1034,8 @@ void Mesh::Serialise(OutputArchive& Arch) const
 }
 
 
-void Mesh::Unserialise(InputArchive& Arch)
+void FMesh::Unserialise(FInputArchive& Arch)
 {
-	uint32 Version;
-	Arch >> Version;
-	check(Version == 23);
-
 	Arch >> IndexBuffers;
 	Arch >> VertexBuffers;
 	Arch >> AdditionalBuffers;
@@ -1072,7 +1065,7 @@ void Mesh::Unserialise(InputArchive& Arch)
 }
 
 
-bool Mesh::IsSimilar(const Mesh& Other, bool bCompareLayouts) const
+bool FMesh::IsSimilar(const FMesh& Other, bool bCompareLayouts) const
 {
 	// Some meshes are just vertex indices (masks) we don't consider them for similarity,
 	// because the kind of vertex channel data they store is the kind that is ignored.
@@ -1137,7 +1130,7 @@ bool Mesh::IsSimilar(const Mesh& Other, bool bCompareLayouts) const
 }
 
 
-void Mesh::CheckIntegrity() const
+void FMesh::CheckIntegrity() const
 {
 #ifdef MUTABLE_DEBUG
 
@@ -1301,7 +1294,7 @@ void Mesh::CheckIntegrity() const
 }
 
 
-static bool StaticMeshFormatIdentify_Project(const Mesh* InMesh)
+static bool StaticMeshFormatIdentify_Project(const FMesh* InMesh)
 {
     // This format is used internally for the mesh project
     bool bResult = true;
@@ -1370,7 +1363,7 @@ static bool StaticMeshFormatIdentify_Project(const Mesh* InMesh)
 }
 
 
-static bool StaticMeshFormatIdentify_ProjectWrapping(const Mesh* InMesh)
+static bool StaticMeshFormatIdentify_ProjectWrapping(const FMesh* InMesh)
 {
     // This format is used internally for the mesh project
     bool bResult = true;
@@ -1453,7 +1446,7 @@ static bool StaticMeshFormatIdentify_ProjectWrapping(const Mesh* InMesh)
 }
 
 
-void Mesh::ResetStaticFormatFlags() const
+void FMesh::ResetStaticFormatFlags() const
 {
 	EnumRemoveFlags(Flags, EMeshFlags::ProjectFormat);
 	EnumRemoveFlags(Flags, EMeshFlags::ProjectWrappingFormat);
@@ -1545,7 +1538,7 @@ namespace
 }
 
 
-void Mesh::Log( FString& out, int32 BufferElementLimit)
+void FMesh::Log( FString& out, int32 BufferElementLimit)
 {
     out += "Mesh:\n";
 
diff --git a/Engine/Plugins/Mutable/Source/MutableRuntime/Private/MuR/MeshBufferSet.cpp b/Engine/Plugins/Mutable/Source/MutableRuntime/Private/MuR/MeshBufferSet.cpp
--- a/Engine/Plugins/Mutable/Source/MutableRuntime/Private/MuR/MeshBufferSet.cpp
+++ b/Engine/Plugins/Mutable/Source/MutableRuntime/Private/MuR/MeshBufferSet.cpp
@@ -61,13 +61,13 @@ namespace mu
 		return s_meshBufferFormatData[uint32(Format)];
 	}
 
-	void FMeshBufferSet::Serialise(OutputArchive& Arch) const
+	void FMeshBufferSet::Serialise(FOutputArchive& Arch) const
 	{
 		Arch << ElementCount;
 		Arch << Buffers;
 	}
 
-	void FMeshBufferSet::Unserialise(InputArchive& Arch)
+	void FMeshBufferSet::Unserialise(FInputArchive& Arch)
 	{
 		Arch >> ElementCount;
 		Arch >> Buffers;
@@ -656,14 +656,14 @@ namespace mu
 		return false;
 	}
 
-	void FMeshBuffer::Serialise(OutputArchive& Arch) const
+	void FMeshBuffer::Serialise(FOutputArchive& Arch) const
 	{
 		Arch << Channels;
 		Arch << Data;
 		Arch << ElementSize;
 	}
 
-	void FMeshBuffer::Unserialise(InputArchive& Arch)
+	void FMeshBuffer::Unserialise(FInputArchive& Arch)
 	{
 		Arch >> Channels;
 		Arch >> Data;
