[mutable] More enforcement of coding standards.

diff --git a/Engine/Plugins/Mutable/Source/MutableRuntime/Private/MuR/Model.cpp b/Engine/Plugins/Mutable/Source/MutableRuntime/Private/MuR/Model.cpp
--- a/Engine/Plugins/Mutable/Source/MutableRuntime/Private/MuR/Model.cpp
+++ b/Engine/Plugins/Mutable/Source/MutableRuntime/Private/MuR/Model.cpp
@@ -39,31 +39,31 @@ namespace mu
     void FProgram::LogHistogram() const
     {
 #if 0
-        uint64 countPerType[(int32)OP_TYPE::COUNT];
+        uint64 countPerType[(int32)EOpType::COUNT];
         mutable_memset(countPerType,0,sizeof(countPerType));
 
-        for ( const uint32& o: m_opAddress )
+        for ( const uint32& o: OpAddress )
         {
-            OP_TYPE type = GetOpType(o);
+            EOpType type = GetOpType(o);
             countPerType[(int32)type]++;
         }
 
-		TArray< TPair<uint64,OP_TYPE> > sorted((int32)OP_TYPE::COUNT);
-        for (int32 i=0; i<(int32)OP_TYPE::COUNT; ++i)
+		TArray< TPair<uint64,EOpType> > sorted((int32)EOpType::COUNT);
+        for (int32 i=0; i<(int32)EOpType::COUNT; ++i)
         {
-            sorted[i].second = (OP_TYPE)i;
+            sorted[i].second = (EOpType)i;
             sorted[i].first = countPerType[i];
         }
 
-        std::sort(sorted.begin(),sorted.end(), []( const pair<uint64,OP_TYPE>& a, const pair<uint64,OP_TYPE>& b )
+        std::sort(sorted.begin(),sorted.end(), []( const pair<uint64,EOpType>& a, const pair<uint64,EOpType>& b )
         {
             return a.first>b.first;
         });
 
-        UE_LOG(LogMutableCore,Log, TEXT("Op histogram (%llu ops):"), m_opAddress.Num());
+        UE_LOG(LogMutableCore,Log, TEXT("Op histogram (%llu ops):"), OpAddress.Num());
         for(int32 i=0; i<8; ++i)
         {
-            float p = sorted[i].first/float(m_opAddress.Num())*100.0f;
+            float p = sorted[i].first/float(OpAddress.Num())*100.0f;
             UE_LOG(LogMutableCore,Log, TEXT("  %3.2f%% : %d"), p, (int32)sorted[i].second );
         }
 #endif
@@ -71,11 +71,11 @@ namespace mu
 
 	
 	//---------------------------------------------------------------------------------------------
-    void Model::Private::UnloadRoms()
+    void FModel::Private::UnloadRoms()
     {
-	    for (int32 RomIndex = 0; RomIndex < m_program.m_roms.Num(); ++RomIndex)
+	    for (int32 RomIndex = 0; RomIndex < Program.Roms.Num(); ++RomIndex)
 	    {
-		    m_program.UnloadRom(RomIndex);
+		    Program.UnloadRom(RomIndex);
 	    }
     }
 	
@@ -83,14 +83,14 @@ namespace mu
     //---------------------------------------------------------------------------------------------
     //---------------------------------------------------------------------------------------------
     //---------------------------------------------------------------------------------------------
-    Model::Model()
+    FModel::FModel()
     {
         m_pD = new Private();
     }
 
 
     //---------------------------------------------------------------------------------------------
-    Model::~Model()
+    FModel::~FModel()
     {
 		MUTABLE_CPUPROFILER_SCOPE(ModelDestructor);
 
@@ -101,7 +101,7 @@ namespace mu
 
 
     //---------------------------------------------------------------------------------------------
-    void Model::Serialise( const Model* p, OutputArchive& arch )
+    void FModel::Serialise( const FModel* p, FOutputArchive& arch )
     {
 		LLM_SCOPE_BYNAME(TEXT("MutableRuntime"));
 
@@ -111,54 +111,49 @@ namespace mu
     //---------------------------------------------------------------------------------------------
     //---------------------------------------------------------------------------------------------
     //---------------------------------------------------------------------------------------------
-    class OutputModelStream : public OutputStream
+    class FOutputModelStream : public FOutputStream
     {
     public:
 
-        //-----------------------------------------------------------------------------------------
         // Life cycle
-        //-----------------------------------------------------------------------------------------
-        OutputModelStream(ModelWriter* pStreamer )
-            : m_pStreamer( pStreamer )
+        FOutputModelStream(FModelWriter* InStreamer )
+            : Streamer( InStreamer )
         {
         }
 
-        //-----------------------------------------------------------------------------------------
-        // InputStream interface
-        //-----------------------------------------------------------------------------------------
-        void Write( const void* pData, uint64 size ) override
+        // FInputStream interface
+        void Write( const void* Data, uint64 Size ) override
         {
-            m_pStreamer->Write( pData, size );
+            Streamer->Write( Data, Size );
         }
 
-
     private:
 
-        ModelWriter* m_pStreamer;
+        FModelWriter* Streamer;
     };
 
 
     //---------------------------------------------------------------------------------------------
-    void Model::Serialise( Model* p, ModelWriter& streamer, bool bDropData )
+    void FModel::Serialise( FModel* p, FModelWriter& streamer, bool bDropData )
     {
 		LLM_SCOPE_BYNAME(TEXT("MutableRuntime"));
 
-		mu::FProgram& program = p->m_pD->m_program;
+		mu::FProgram& program = p->m_pD->Program;
 
-		TArray<TPair<int32, mu::ImagePtrConst>> InitialImages;
-		TArray<TPair<int32, mu::MeshPtrConst>> InitialMeshes;
+		TArray<TPair<int32, TSharedPtr<const FImage>>> InitialImages;
+		TArray<TPair<int32, TSharedPtr<const FMesh>>> InitialMeshes;
 		if (!bDropData)
 		{
 			InitialImages = program.ConstantImageLODs;
 			InitialMeshes = program.ConstantMeshes;
 		}
 
-		OutputMemoryStream MemStream(16 * 1024 * 1024);
+		FOutputMemoryStream MemStream(16 * 1024 * 1024);
 
 		// Save images and unload from memory
 		for (int32 ResourceIndex = 0; ResourceIndex < program.ConstantImageLODs.Num(); ++ResourceIndex)
 		{
-			TPair<int32, mu::ImagePtrConst>& ResData = program.ConstantImageLODs[ResourceIndex];
+			TPair<int32, TSharedPtr<const mu::FImage>>& ResData = program.ConstantImageLODs[ResourceIndex];
 
 			// This shouldn't have been serialised with rom support before.
 			if (ResData.Key < 0)
@@ -167,14 +162,14 @@ namespace mu
 			}
 
 			int32 RomIndex = ResData.Key;
-			const FRomData& RomData = program.m_roms[RomIndex];
+			const FRomData& RomData = program.Roms[RomIndex];
 			check(RomData.ResourceType == DT_IMAGE);
 			check(RomData.ResourceIndex == ResourceIndex);
 
 			// Serialize to memory, to find out final size of this rom
 			MemStream.Reset();
-			OutputArchive MemoryArch(&MemStream);
-			Image::Serialise(ResData.Value.get(), MemoryArch);
+			FOutputArchive MemoryArch(&MemStream);
+			FImage::Serialise(ResData.Value.Get(), MemoryArch);
 			check(RomData.Size == MemStream.GetBufferSize());
 
 			streamer.OpenWriteFile(RomData.Id);
@@ -188,7 +183,7 @@ namespace mu
 		// Save meshes and unload from memory
 		for (int32 ResourceIndex = 0; ResourceIndex < program.ConstantMeshes.Num(); ++ResourceIndex)
 		{
-			TPair<int32, mu::MeshPtrConst>& ResData = program.ConstantMeshes[ResourceIndex];
+			TPair<int32, TSharedPtr<const FMesh>>& ResData = program.ConstantMeshes[ResourceIndex];
 
 			// This shouldn't have been serialised with rom support before.
 			if (ResData.Key < 0)
@@ -197,14 +192,14 @@ namespace mu
 			}
 
 			int32 RomIndex = ResData.Key;
-			const FRomData& RomData = program.m_roms[RomIndex];
+			const FRomData& RomData = program.Roms[RomIndex];
 			check(RomData.ResourceType == DT_MESH);
 			check(RomData.ResourceIndex == ResourceIndex);
 
 			// Serialize to memory, to find out final size of this rom
 			MemStream.Reset();
-			OutputArchive MemoryArch(&MemStream);
-			Mesh::Serialise(ResData.Value.get(), MemoryArch);
+			FOutputArchive MemoryArch(&MemStream);
+			FMesh::Serialise(ResData.Value.Get(), MemoryArch);
 			check(RomData.Size == MemStream.GetBufferSize());
 
 			streamer.OpenWriteFile(RomData.Id);
@@ -218,8 +213,8 @@ namespace mu
 		// Store the main data of the model
 		{
 			streamer.OpenWriteFile(0);
-			OutputModelStream stream(&streamer);
-			OutputArchive arch(&stream);
+			FOutputModelStream stream(&streamer);
+			FOutputArchive arch(&stream);
 
 			arch << *p->m_pD;
 
@@ -235,159 +230,159 @@ namespace mu
 	}
 
 
-	bool Model::HasExternalData() const
+	bool FModel::HasExternalData() const
     {
-		return m_pD->m_program.m_roms.Num() > 0;
+		return m_pD->Program.Roms.Num() > 0;
     }
 
 #if WITH_EDITOR
-	bool Model::IsValid() const
+	bool FModel::IsValid() const
 	{
-		return m_pD->m_program.bIsValid;
+		return m_pD->Program.bIsValid;
 	}
 
 
-	void Model::Invalidate()
+	void FModel::Invalidate()
     {
-		m_pD->m_program.bIsValid = false;
+		m_pD->Program.bIsValid = false;
     }
 #endif
 
 
-	void Model::UnloadExternalData()
+	void FModel::UnloadExternalData()
     {
 		LLM_SCOPE_BYNAME(TEXT("MutableRuntime"));
 
-		for (int32 ResIndex = 0; ResIndex < m_pD->m_program.ConstantImageLODs.Num(); ++ResIndex)
+		for (int32 ResIndex = 0; ResIndex < m_pD->Program.ConstantImageLODs.Num(); ++ResIndex)
 		{
-			if (m_pD->m_program.ConstantImageLODs[ResIndex].Key >= 0)
+			if (m_pD->Program.ConstantImageLODs[ResIndex].Key >= 0)
 			{
-				m_pD->m_program.ConstantImageLODs[ResIndex].Value = nullptr;
+				m_pD->Program.ConstantImageLODs[ResIndex].Value = nullptr;
 			}
 		}
 
-		for (int32 ResIndex = 0; ResIndex < m_pD->m_program.ConstantMeshes.Num(); ++ResIndex)
+		for (int32 ResIndex = 0; ResIndex < m_pD->Program.ConstantMeshes.Num(); ++ResIndex)
 		{
-			if (m_pD->m_program.ConstantMeshes[ResIndex].Key >= 0)
+			if (m_pD->Program.ConstantMeshes[ResIndex].Key >= 0)
 			{
-				m_pD->m_program.ConstantMeshes[ResIndex].Value = nullptr;
+				m_pD->Program.ConstantMeshes[ResIndex].Value = nullptr;
 			}
 		}
 	}
 
 
-	TSharedPtr<Model> Model::StaticUnserialise( InputArchive& arch )
+	TSharedPtr<FModel> FModel::StaticUnserialise( FInputArchive& arch )
     {
 		LLM_SCOPE_BYNAME(TEXT("MutableRuntime"));
 
-		TSharedPtr<Model> pResult = MakeShared<Model>();
+		TSharedPtr<FModel> pResult = MakeShared<FModel>();
         arch >> *pResult->m_pD;
         return pResult;
     }
 
 
-	Model::Private* Model::GetPrivate() const
+	FModel::Private* FModel::GetPrivate() const
     {
         return m_pD;
     }
 
 
-    bool Model::GetBoolDefaultValue(int32 Index) const
+    bool FModel::GetBoolDefaultValue(int32 Index) const
     {
-    	check(m_pD->m_program.m_parameters.IsValidIndex(Index));
-		check(m_pD->m_program.m_parameters[Index].m_type == EParameterType::Bool);
+    	check(m_pD->Program.Parameters.IsValidIndex(Index));
+		check(m_pD->Program.Parameters[Index].Type == EParameterType::Bool);
 
         // Early out in case of invalid parameters
-        if (!m_pD->m_program.m_parameters.IsValidIndex(Index) ||
-            m_pD->m_program.m_parameters[Index].m_type != EParameterType::Bool)
+        if (!m_pD->Program.Parameters.IsValidIndex(Index) ||
+            m_pD->Program.Parameters[Index].Type != EParameterType::Bool)
         {
             return false;
         }
 		
-        return m_pD->m_program.m_parameters[Index].m_defaultValue.Get<FParamBoolType>();
+        return m_pD->Program.Parameters[Index].DefaultValue.Get<FParamBoolType>();
     }
 
 
-	int32 Model::GetIntDefaultValue(int32 Index) const
+	int32 FModel::GetIntDefaultValue(int32 Index) const
 	{
-		check(m_pD->m_program.m_parameters.IsValidIndex(Index));
-		check(m_pD->m_program.m_parameters[Index].m_type == EParameterType::Int);
+		check(m_pD->Program.Parameters.IsValidIndex(Index));
+		check(m_pD->Program.Parameters[Index].Type == EParameterType::Int);
 
         // Early out in case of invalid parameters
-        if (!m_pD->m_program.m_parameters.IsValidIndex(Index) ||
-            m_pD->m_program.m_parameters[Index].m_type != EParameterType::Int)
+        if (!m_pD->Program.Parameters.IsValidIndex(Index) ||
+            m_pD->Program.Parameters[Index].Type != EParameterType::Int)
         {
             return 0;
         }
 		
-        return m_pD->m_program.m_parameters[Index].m_defaultValue.Get<FParamIntType>();
+        return m_pD->Program.Parameters[Index].DefaultValue.Get<FParamIntType>();
 	}
 
 
-	float Model::GetFloatDefaultValue(int32 Index) const
+	float FModel::GetFloatDefaultValue(int32 Index) const
 	{
-    	check(m_pD->m_program.m_parameters.IsValidIndex(Index));
-		check(m_pD->m_program.m_parameters[Index].m_type == EParameterType::Float);
+    	check(m_pD->Program.Parameters.IsValidIndex(Index));
+		check(m_pD->Program.Parameters[Index].Type == EParameterType::Float);
 
         // Early out in case of invalid parameters
-        if (!m_pD->m_program.m_parameters.IsValidIndex(Index) ||
-            m_pD->m_program.m_parameters[Index].m_type != EParameterType::Float)
+        if (!m_pD->Program.Parameters.IsValidIndex(Index) ||
+            m_pD->Program.Parameters[Index].Type != EParameterType::Float)
         {
             return 0.0f;
         }
 		
-        return m_pD->m_program.m_parameters[Index].m_defaultValue.Get<FParamFloatType>();
+        return m_pD->Program.Parameters[Index].DefaultValue.Get<FParamFloatType>();
 	}
 
 
-	void Model::GetColourDefaultValue(int32 Index, FVector4f& OutValue) const
+	void FModel::GetColourDefaultValue(int32 Index, FVector4f& OutValue) const
     {
-    	check(m_pD->m_program.m_parameters.IsValidIndex(Index));
-		check(m_pD->m_program.m_parameters[Index].m_type == EParameterType::Color);
+    	check(m_pD->Program.Parameters.IsValidIndex(Index));
+		check(m_pD->Program.Parameters[Index].Type == EParameterType::Color);
 
         // Early out in case of invalid parameters
-        if (!m_pD->m_program.m_parameters.IsValidIndex(Index) ||
-            m_pD->m_program.m_parameters[Index].m_type != EParameterType::Color)
+        if (!m_pD->Program.Parameters.IsValidIndex(Index) ||
+            m_pD->Program.Parameters[Index].Type != EParameterType::Color)
         {
             return;
         }
 
-        FParamColorType& Color = m_pD->m_program.m_parameters[Index].m_defaultValue.Get<FParamColorType>();
+        FParamColorType& Color = m_pD->Program.Parameters[Index].DefaultValue.Get<FParamColorType>();
 		
 		OutValue = Color;
     }
 
 
-	FMatrix44f Model::GetMatrixDefaultValue(int32 Index) const
+	FMatrix44f FModel::GetMatrixDefaultValue(int32 Index) const
 	{
-    	check(m_pD->m_program.m_parameters.IsValidIndex(Index));
-    	check(m_pD->m_program.m_parameters[Index].m_type == EParameterType::Matrix);
+    	check(m_pD->Program.Parameters.IsValidIndex(Index));
+    	check(m_pD->Program.Parameters[Index].Type == EParameterType::Matrix);
 
     	// Early out in case of invalid parameters
-    	if (!m_pD->m_program.m_parameters.IsValidIndex(Index) ||
-			m_pD->m_program.m_parameters[Index].m_type != EParameterType::Matrix)
+    	if (!m_pD->Program.Parameters.IsValidIndex(Index) ||
+			m_pD->Program.Parameters[Index].Type != EParameterType::Matrix)
     	{
     		return FMatrix44f::Identity;
     	}
 
-    	return m_pD->m_program.m_parameters[Index].m_defaultValue.Get<FParamMatrixType>();
+    	return m_pD->Program.Parameters[Index].DefaultValue.Get<FParamMatrixType>();
 	}
 
 
-	void Model::GetProjectorDefaultValue(int32 Index, EProjectorType* OutProjectionType, FVector3f* OutPos,
+	void FModel::GetProjectorDefaultValue(int32 Index, EProjectorType* OutProjectionType, FVector3f* OutPos,
 	                                     FVector3f* OutDir, FVector3f* OutUp, FVector3f* OutScale, float* OutProjectionAngle) const
 	{
-    	check(m_pD->m_program.m_parameters.IsValidIndex(Index));
-		check(m_pD->m_program.m_parameters[Index].m_type == EParameterType::Projector);
+    	check(m_pD->Program.Parameters.IsValidIndex(Index));
+		check(m_pD->Program.Parameters[Index].Type == EParameterType::Projector);
 
         // Early out in case of invalid parameters
-        if (!m_pD->m_program.m_parameters.IsValidIndex(Index) ||
-            m_pD->m_program.m_parameters[Index].m_type != EParameterType::Projector)
+        if (!m_pD->Program.Parameters.IsValidIndex(Index) ||
+            m_pD->Program.Parameters[Index].Type != EParameterType::Projector)
         {
             return;
         }
 
-        const FProjector& Projector = m_pD->m_program.m_parameters[Index].m_defaultValue.Get<FParamProjectorType>();
+        const FProjector& Projector = m_pD->Program.Parameters[Index].DefaultValue.Get<FParamProjectorType>();
         if (OutProjectionType) *OutProjectionType = Projector.type;
     	if (OutPos) *OutPos = Projector.position;
 		if (OutDir) *OutDir = Projector.direction;
@@ -397,99 +392,99 @@ namespace mu
 	}
 
 
-	FName Model::GetImageDefaultValue(int32 Index) const
+	FName FModel::GetImageDefaultValue(int32 Index) const
     {
-	    check(m_pD->m_program.m_parameters.IsValidIndex(Index));
-		check(m_pD->m_program.m_parameters[Index].m_type == EParameterType::Image);
+	    check(m_pD->Program.Parameters.IsValidIndex(Index));
+		check(m_pD->Program.Parameters[Index].Type == EParameterType::Image);
 
         // Early out in case of invalid parameters
-        if (!m_pD->m_program.m_parameters.IsValidIndex(Index) ||
-            m_pD->m_program.m_parameters[Index].m_type != EParameterType::Image)
+        if (!m_pD->Program.Parameters.IsValidIndex(Index) ||
+            m_pD->Program.Parameters[Index].Type != EParameterType::Image)
         {
             return {};
         }
 		
-        return m_pD->m_program.m_parameters[Index].m_defaultValue.Get<FParamImageType>();
+        return m_pD->Program.Parameters[Index].DefaultValue.Get<FParamImageType>();
     }
 
 
-    int32 Model::GetRomCount() const
+    int32 FModel::GetRomCount() const
     {
-    	return m_pD->m_program.m_roms.Num();
+    	return m_pD->Program.Roms.Num();
     }
 
 
-    uint32 Model::GetRomId(int32 Index) const
+    uint32 FModel::GetRomId(int32 Index) const
     {
-    	return m_pD->m_program.m_roms[Index].Id;
+    	return m_pD->Program.Roms[Index].Id;
     }
 
-	uint32 Model::GetRomSourceId(int32 Index) const
+	uint32 FModel::GetRomSourceId(int32 Index) const
 	{
-		return m_pD->m_program.m_roms[Index].SourceId;
+		return m_pD->Program.Roms[Index].SourceId;
 	}
 
 
-    uint32 Model::GetRomSize(int32 Index) const
+    uint32 FModel::GetRomSize(int32 Index) const
     {
-    	return m_pD->m_program.m_roms[Index].Size;
+    	return m_pD->Program.Roms[Index].Size;
     }
 
-	uint16 Model::GetRomType(int32 Index) const
+	uint16 FModel::GetRomType(int32 Index) const
 	{
-		return m_pD->m_program.m_roms[Index].ResourceType;
+		return m_pD->Program.Roms[Index].ResourceType;
 	}
 
-	ERomFlags Model::GetRomFlags(int32 Index) const
+	ERomFlags FModel::GetRomFlags(int32 Index) const
 	{
-		return m_pD->m_program.m_roms[Index].Flags;
+		return m_pD->Program.Roms[Index].Flags;
 	}
 
 
-    ParametersPtr Model::NewParameters(TSharedPtr<const Model> Model, const Parameters* pOld )
+    TSharedPtr<FParameters> FModel::NewParameters(TSharedPtr<const FModel> Model, const FParameters* OldParameters )
     {
 		LLM_SCOPE_BYNAME(TEXT("MutableRuntime"));
 
-        ParametersPtr pRes = new Parameters();
+        TSharedPtr<FParameters> pRes = MakeShared<FParameters>();
 
-        pRes->GetPrivate()->m_pModel = Model;
+        pRes->GetPrivate()->Model = Model;
 
-		const FProgram& Program = Model->GetPrivate()->m_program;
-        pRes->GetPrivate()->m_values.SetNum(Program.m_parameters.Num());
-        for ( int32 p=0; p< Program.m_parameters.Num(); ++p )
+		const FProgram& Program = Model->GetPrivate()->Program;
+        pRes->GetPrivate()->Values.SetNum(Program.Parameters.Num());
+        for ( int32 p=0; p< Program.Parameters.Num(); ++p )
         {
-            pRes->GetPrivate()->m_values[p] = Program.m_parameters[p].m_defaultValue;
+            pRes->GetPrivate()->Values[p] = Program.Parameters[p].DefaultValue;
         }
 
         // Copy old values
-        if ( pOld )
+        if ( OldParameters )
         {
-            for ( int32 p=0; p<pOld->GetCount(); ++p )
+            for ( int32 p=0; p<OldParameters->GetCount(); ++p )
             {
-                int32 thisP = pRes->GetPrivate()->Find( pOld->GetName(p) );
+                int32 thisP = pRes->GetPrivate()->Find( OldParameters->GetName(p) );
 
                 if ( thisP>=0 )
                 {
-                    if ( pOld->GetType(p)==pRes->GetType(thisP) )
+                    if ( OldParameters->GetType(p)==pRes->GetType(thisP) )
                     {
                         switch ( pRes->GetType(thisP) )
                         {
 						case EParameterType::Bool:
-                            pRes->SetBoolValue( thisP, pOld->GetBoolValue(p) );
+                            pRes->SetBoolValue( thisP, OldParameters->GetBoolValue(p) );
                             break;
 
                         case EParameterType::Int:
-                            pRes->SetIntValue( thisP, pOld->GetIntValue(p) );
+                            pRes->SetIntValue( thisP, OldParameters->GetIntValue(p) );
                             break;
 
                         case EParameterType::Float:
-                            pRes->SetFloatValue( thisP, pOld->GetFloatValue(p) );
+                            pRes->SetFloatValue( thisP, OldParameters->GetFloatValue(p) );
                             break;
 
                         case EParameterType::Color:
                         {
                             FVector4f V;
-                            pOld->GetColourValue( p, V );
+                            OldParameters->GetColourValue( p, V );
                             pRes->SetColourValue( thisP, V );
                             break;
                         }
@@ -497,13 +492,13 @@ namespace mu
                         case EParameterType::Projector:
                         {
 //							float m[16];
-//							pOld->GetProjectorValue( p, m );
-                            pRes->GetPrivate()->m_values[thisP].Set<FParamProjectorType>(pOld->GetPrivate()->m_values[p].Get<FParamProjectorType>());
+//							OldParameters->GetProjectorValue( p, m );
+                            pRes->GetPrivate()->Values[thisP].Set<FParamProjectorType>(OldParameters->GetPrivate()->Values[p].Get<FParamProjectorType>());
                             break;
                         }
 
                         case EParameterType::Image:
-                            pRes->SetImageValue( thisP, pOld->GetImageValue(p) );
+                            pRes->SetImageValue( thisP, OldParameters->GetImageValue(p) );
                             break;
 
                         default:
@@ -520,11 +515,11 @@ namespace mu
 
 
     //---------------------------------------------------------------------------------------------
-    bool Model::IsParameterMultidimensional(const int32 ParamIndex) const
+    bool FModel::IsParameterMultidimensional(const int32 ParamIndex) const
     {
-		if (m_pD->m_program.m_parameters.IsValidIndex(ParamIndex))
+		if (m_pD->Program.Parameters.IsValidIndex(ParamIndex))
 		{
-			return m_pD->m_program.m_parameters[ParamIndex].m_ranges.Num() > 0;
+			return m_pD->Program.Parameters[ParamIndex].Ranges.Num() > 0;
 		}
 
 		return false;
@@ -532,20 +527,20 @@ namespace mu
 	
 
     //---------------------------------------------------------------------------------------------
-    int32 Model::GetStateCount() const
+    int32 FModel::GetStateCount() const
     {
-        return (int32)m_pD->m_program.m_states.Num();
+        return (int32)m_pD->Program.States.Num();
     }
 
 
     //---------------------------------------------------------------------------------------------
-    const FString& Model::GetStateName( int32 index ) const
+    const FString& FModel::GetStateName( int32 index ) const
     {
         const char* strRes = 0;
 
-        if ( index>=0 && index<(int32)m_pD->m_program.m_states.Num() )
+        if ( index>=0 && index<(int32)m_pD->Program.States.Num() )
         {
-            return m_pD->m_program.m_states[index].Name;
+            return m_pD->Program.States[index].Name;
         }
 
 		static FString None;
@@ -554,13 +549,13 @@ namespace mu
 
 
     //---------------------------------------------------------------------------------------------
-    int32 Model::FindState( const FString& Name ) const
+    int32 FModel::FindState( const FString& Name ) const
     {
         int32 res = -1;
 
-        for ( int32 i=0; res<0 && i<(int32)m_pD->m_program.m_states.Num(); ++i )
+        for ( int32 i=0; res<0 && i<(int32)m_pD->Program.States.Num(); ++i )
         {
-            if ( m_pD->m_program.m_states[i].Name == Name )
+            if ( m_pD->Program.States[i].Name == Name )
             {
                 res = i;
             }
@@ -571,13 +566,13 @@ namespace mu
 
 
     //---------------------------------------------------------------------------------------------
-    int32 Model::GetStateParameterCount( int32 stateIndex ) const
+    int32 FModel::GetStateParameterCount( int32 stateIndex ) const
     {
         int32 res = -1;
 
-        if ( stateIndex>=0 && stateIndex<(int32)m_pD->m_program.m_states.Num() )
+        if ( stateIndex>=0 && stateIndex<(int32)m_pD->Program.States.Num() )
         {
-            res = (int32)m_pD->m_program.m_states[stateIndex].m_runtimeParameters.Num();
+            res = (int32)m_pD->Program.States[stateIndex].m_runtimeParameters.Num();
         }
 
         return res;
@@ -585,13 +580,13 @@ namespace mu
 
 
     //---------------------------------------------------------------------------------------------
-    int32 Model::GetStateParameterIndex( int32 stateIndex, int32 paramIndex ) const
+    int32 FModel::GetStateParameterIndex( int32 stateIndex, int32 paramIndex ) const
     {
         int32 res = -1;
 
-        if ( stateIndex>=0 && stateIndex<(int32)m_pD->m_program.m_states.Num() )
+        if ( stateIndex>=0 && stateIndex<(int32)m_pD->Program.States.Num() )
         {
-            const FProgram::FState& state = m_pD->m_program.m_states[stateIndex];
+            const FProgram::FState& state = m_pD->Program.States[stateIndex];
             if ( paramIndex>=0 && paramIndex<(int32)state.m_runtimeParameters.Num() )
             {
                 res = (int32)state.m_runtimeParameters[paramIndex];
@@ -601,393 +596,5 @@ namespace mu
         return res;
     }
 
-
-    //---------------------------------------------------------------------------------------------
-    //---------------------------------------------------------------------------------------------
-    //---------------------------------------------------------------------------------------------
-    ModelParametersGenerator::ModelParametersGenerator
-        (
-			TSharedPtr<const Model> pModel,
-            System* pSystem
-        )
-    {
-		LLM_SCOPE_BYNAME(TEXT("MutableRuntime"));
-
-        m_pD = new Private();
-        m_pD->m_pModel = pModel;
-        m_pD->m_pSystem = pSystem;
-
-        size_t paramCount = pModel->GetPrivate()->m_program.m_parameters.Num();
-
-		m_pD->m_instanceCount = 1;
-        for (size_t i=0;i<paramCount; ++i)
-        {
-            const FParameterDesc& param =
-                    m_pD->m_pModel->GetPrivate()->m_program.m_parameters[ i ];
-            switch ( param.m_type )
-            {
-
-            case EParameterType::Int:
-            {
-                m_pD->m_instanceCount *= param.m_possibleValues.Num();
-                break;
-            }
-
-            case EParameterType::Bool:
-            {
-                m_pD->m_instanceCount *= 2;
-                break;
-            }
-
-            default:
-                // Parameter not discrete, so ignored for combinations.
-                break;
-            }
-        }
-    }
-
-
-    //---------------------------------------------------------------------------------------------
-    ModelParametersGenerator::~ModelParametersGenerator()
-    {
-		LLM_SCOPE_BYNAME(TEXT("MutableRuntime"));
-
-        check( m_pD );
-        delete m_pD;
-        m_pD = 0;
-    }
-
-
-    //---------------------------------------------------------------------------------------------
-    ModelParametersGenerator::Private* ModelParametersGenerator::GetPrivate() const
-    {
-        return m_pD;
-    }
-
-
-    //---------------------------------------------------------------------------------------------
-    int64 ModelParametersGenerator::GetInstanceCount()
-    {
-        return m_pD->m_instanceCount;
-    }
-
-
-    //---------------------------------------------------------------------------------------------
-    Ptr<Parameters> ModelParametersGenerator::GetInstance( int64 index, TFunction<float()> RandomGenerator )
-    {
-		LLM_SCOPE_BYNAME(TEXT("MutableRuntime"));
-
-		Ptr<Parameters> res = Model::NewParameters(m_pD->m_pModel);
-
-		const FProgram& Program = m_pD->m_pModel->GetPrivate()->m_program;
-        int32 paramCount = Program.m_parameters.Num();
-        int64 currentInstance = index;
-        for (int32 i=0;i<paramCount;++i)
-        {
-            const FParameterDesc& param = Program.m_parameters[ i ];
-			Ptr<RangeIndex> RangeIndex = res->NewRangeIndex(i);
-
-            switch ( param.m_type )
-            {
-
-            case EParameterType::Int:
-            {
-				bool bIsRangeSize = Program.m_ranges.ContainsByPredicate([i](const FRangeDesc& r) { return r.m_dimensionParameter == i; });
-				if (bIsRangeSize)
-				{
-					res->SetIntValue(i, m_pD->DefaultRangeDimension);
-				}
-				else
-				{
-					int32 numOptions = res->GetIntPossibleValueCount(i);
-					int32 value = res->GetIntPossibleValue(i, (int32)(currentInstance % numOptions));
-					res->SetIntValue(i, value);
-					currentInstance /= numOptions;
-				}
-                break;
-            }
-
-            case EParameterType::Bool:
-            {
-                res->SetBoolValue( i, currentInstance%2!=0 );
-                currentInstance /= 2;
-                break;
-            }
-
-            case EParameterType::Float:
-			{
-				bool bIsRangeSize = Program.m_ranges.ContainsByPredicate([i](const FRangeDesc& r) { return r.m_dimensionParameter == i; });
-				if (bIsRangeSize)
-				{
-					if (RangeIndex)
-					{
-						for (int32 Dimensions = 0; Dimensions < RangeIndex->GetRangeCount(); ++Dimensions)
-						{
-							for (int32 RangePosition = 0; RangePosition < m_pD->DefaultRangeDimension; ++RangePosition)
-							{
-								RangeIndex->SetPosition(Dimensions, RangePosition);
-								res->SetFloatValue(i, float(m_pD->DefaultRangeDimension), RangeIndex);
-							}
-						}
-					}
-					else
-					{
-						res->SetFloatValue(i, float(m_pD->DefaultRangeDimension));
-					}
-				}
-				else if (RandomGenerator)
-				{
-					if (RangeIndex)
-					{
-						for (int32 Dimensions = 0; Dimensions < RangeIndex->GetRangeCount(); ++Dimensions)
-						{
-							for (int32 RangePosition = 0; RangePosition < m_pD->DefaultRangeDimension; ++RangePosition)
-							{
-								RangeIndex->SetPosition(Dimensions, RangePosition);
-								res->SetFloatValue(i, RandomGenerator(), RangeIndex);
-							}
-						}
-					}
-					else
-					{
-						res->SetFloatValue(i, RandomGenerator());
-					}
-				}
-				break;
-			}
-
-            case EParameterType::Color:
-                if (RandomGenerator)
-                {
-					if (RangeIndex)
-					{
-						for (int32 Dimensions = 0; Dimensions < RangeIndex->GetRangeCount(); ++Dimensions)
-						{
-							for (int32 RangePosition = 0; RangePosition < m_pD->DefaultRangeDimension; ++RangePosition)
-							{
-								RangeIndex->SetPosition(Dimensions, RangePosition);
-								FVector4f Random(RandomGenerator(), RandomGenerator(), RandomGenerator(), RandomGenerator());
-								res->SetColourValue(i, Random, RangeIndex);
-							}
-						}
-					}
-					else
-					{
-						FVector4f Random(RandomGenerator(), RandomGenerator(), RandomGenerator(), RandomGenerator());
-						res->SetColourValue(i, Random);
-					}
-                }
-                break;
-
-			case EParameterType::Projector:
-				if (RandomGenerator)
-				{
-					// For projectors we just warp the position a little bit just to get something different
-					FVector3f Position, Direction, Up, Scale;
-					float Angle = 1.0f;
-					res->GetProjectorValue(i, nullptr, &Position, &Direction, &Up, &Scale, &Angle);
-
-					if (RangeIndex)
-					{
-						for (int32 Dimensions = 0; Dimensions < RangeIndex->GetRangeCount(); ++Dimensions)
-						{
-							for (int32 RangePosition = 0; RangePosition < m_pD->DefaultRangeDimension; ++RangePosition)
-							{
-								RangeIndex->SetPosition(Dimensions, RangePosition);
-								Position.X *= 0.9 + RandomGenerator() * 0.2f;
-								Position.Y *= 0.9 + RandomGenerator() * 0.2f;
-								Position.Z *= 0.9 + RandomGenerator() * 0.2f;
-								res->SetProjectorValue(i, Position, Direction, Up, Scale, Angle, RangeIndex);
-
-							}
-						}
-					}
-					else
-					{
-						Position.X *= 0.9 + RandomGenerator() * 0.2f;
-						Position.Y *= 0.9 + RandomGenerator() * 0.2f;
-						Position.Z *= 0.9 + RandomGenerator() * 0.2f;
-						res->SetProjectorValue(i, Position, Direction, Up, Scale, Angle);
-					}
-				}
-				break;
-
-            default:
-                break;
-            }
-        }
-
-        return res;
-    }
-
-
-    //---------------------------------------------------------------------------------------------
-    ParametersPtr ModelParametersGenerator::GetRandomInstance( TFunctionRef<float()> RandomGenerator)
-    {
-		LLM_SCOPE_BYNAME(TEXT("MutableRuntime"));
-
-        ParametersPtr res = Model::NewParameters(m_pD->m_pModel);
-
-		const FProgram& Program = m_pD->m_pModel->GetPrivate()->m_program;
-        int32 ParamCount = Program.m_parameters.Num();
-
-        for (int32 i=0;i<ParamCount;++i)
-        {
-            const FParameterDesc& param = Program.m_parameters[ i ];
-			Ptr<RangeIndex> RangeIndex = res->NewRangeIndex(i);
-
-            switch ( param.m_type )
-            {
-
-            case EParameterType::Int:
-            {
-				bool bIsRangeSize = Program.m_ranges.ContainsByPredicate([i](const FRangeDesc& r) { return r.m_dimensionParameter == i; });
-				if (bIsRangeSize)
-				{
-					if (RangeIndex)
-					{
-						for (int32 Dimensions = 0; Dimensions < RangeIndex->GetRangeCount(); ++Dimensions)
-						{
-							for (int32 RangePosition = 0; RangePosition < m_pD->DefaultRangeDimension; ++RangePosition)
-							{
-								RangeIndex->SetPosition(Dimensions, RangePosition);
-								res->SetIntValue(i, m_pD->DefaultRangeDimension, RangeIndex);
-							}
-						}
-					}
-					else
-					{
-						res->SetIntValue(i, m_pD->DefaultRangeDimension);
-					}
-				}
-				else
-				{
-					if (RangeIndex)
-					{
-						for (int32 Dimensions = 0; Dimensions < RangeIndex->GetRangeCount(); ++Dimensions)
-						{
-							for (int32 RangePosition = 0; RangePosition < m_pD->DefaultRangeDimension; ++RangePosition)
-							{
-								int32 numOptions = res->GetIntPossibleValueCount(i);
-								int32 valueIndex = (int32)(FMath::Min(numOptions - 1, int32(RandomGenerator() * numOptions)));
-								int32 value = res->GetIntPossibleValue(i, valueIndex);
-								res->SetIntValue(i, value, RangeIndex);
-							}
-						}
-					}
-					else
-					{
-						int32 numOptions = res->GetIntPossibleValueCount(i);
-						int32 valueIndex = (int32)(FMath::Min(numOptions - 1, int32(RandomGenerator() * numOptions)));
-						int32 value = res->GetIntPossibleValue(i, valueIndex);
-						res->SetIntValue(i, value);
-					}
-				}
-                break;
-            }
-
-            case EParameterType::Bool:
-            {
-                res->SetBoolValue( i, RandomGenerator()>0.5f );
-                break;
-            }
-
-            case EParameterType::Float:
-			{
-				bool bIsRangeSize = Program.m_ranges.ContainsByPredicate([i](const FRangeDesc& r) { return r.m_dimensionParameter == i; });
-				if (bIsRangeSize)
-				{
-					if (RangeIndex)
-					{
-						for (int32 Dimensions = 0; Dimensions < RangeIndex->GetRangeCount(); ++Dimensions)
-						{
-							for (int32 RangePosition = 0; RangePosition < m_pD->DefaultRangeDimension; ++RangePosition)
-							{
-								res->SetFloatValue(i, float(m_pD->DefaultRangeDimension), RangeIndex);
-							}
-						}
-					}
-					else
-					{
-						res->SetFloatValue(i, float(m_pD->DefaultRangeDimension));
-					}
-				}
-				else
-				{
-					if (RangeIndex)
-					{
-						for (int32 Dimensions = 0; Dimensions < RangeIndex->GetRangeCount(); ++Dimensions)
-						{
-							for (int32 RangePosition = 0; RangePosition < m_pD->DefaultRangeDimension; ++RangePosition)
-							{
-								res->SetFloatValue(i, RandomGenerator(), RangeIndex);
-							}
-						}
-					}
-					else
-					{
-						res->SetFloatValue(i, RandomGenerator());
-					}
-				}
-				break;
-			}
-
-            case EParameterType::Color:
-				if (RangeIndex)
-				{
-					for (int32 Dimensions = 0; Dimensions < RangeIndex->GetRangeCount(); ++Dimensions)
-					{
-						for (int32 RangePosition = 0; RangePosition < m_pD->DefaultRangeDimension; ++RangePosition)
-						{
-							FVector4f Random(RandomGenerator(), RandomGenerator(), RandomGenerator(), RandomGenerator());
-							res->SetColourValue(i, Random, RangeIndex);
-						}
-					}
-				}
-				else
-				{
-					FVector4f Random(RandomGenerator(), RandomGenerator(), RandomGenerator(), RandomGenerator());
-					res->SetColourValue(i, Random);
-				}
-                break;
-
-			case EParameterType::Projector:
-			{
-				// For projectors we just warp the position a little bit just to get something different
-				FVector3f Position, Direction, Up, Scale;
-				float Angle = 1.0f;
-				res->GetProjectorValue(i, nullptr, &Position, &Direction, &Up, &Scale, &Angle);
-				if (RangeIndex)
-				{
-					for (int32 Dimensions = 0; Dimensions < RangeIndex->GetRangeCount(); ++Dimensions)
-					{
-						for (int32 RangePosition = 0; RangePosition < m_pD->DefaultRangeDimension; ++RangePosition)
-						{
-							Position.X *= 0.9 + RandomGenerator() * 0.2f;
-							Position.Y *= 0.9 + RandomGenerator() * 0.2f;
-							Position.Z *= 0.9 + RandomGenerator() * 0.2f;
-							res->SetProjectorValue(i, Position, Direction, Up, Scale, Angle, RangeIndex);
-						}
-					}
-				}
-				else
-				{
-					Position.X *= 0.9 + RandomGenerator() * 0.2f;
-					Position.Y *= 0.9 + RandomGenerator() * 0.2f;
-					Position.Z *= 0.9 + RandomGenerator() * 0.2f;
-					res->SetProjectorValue(i, Position, Direction, Up, Scale, Angle);
-				}
-				break;
-			}
-
-            default:
-                break;
-            }
-        }
-
-        return res;
-    }
-
-
 }
 
diff --git a/Engine/Plugins/Mutable/Source/MutableRuntime/Private/MuR/MutableRuntimeModule.cpp b/Engine/Plugins/Mutable/Source/MutableRuntime/Private/MuR/MutableRuntimeModule.cpp
--- a/Engine/Plugins/Mutable/Source/MutableRuntime/Private/MuR/MutableRuntimeModule.cpp
+++ b/Engine/Plugins/Mutable/Source/MutableRuntime/Private/MuR/MutableRuntimeModule.cpp
@@ -2,27 +2,41 @@
 
 #include "MuR/MutableRuntimeModule.h"
 
-#include "HAL/UnrealMemory.h"
-#include "Logging/LogMacros.h"
-#include "Misc/AssertionMacros.h"
 #include "Modules/ModuleManager.h"
-#include "MuR/MutableMemory.h"
 #include "MuR/Platform.h"
+#include "MuR/BlockCompression/Miro/Miro.h"
 
 IMPLEMENT_MODULE(FMutableRuntimeModule, MutableRuntime);
 
 DEFINE_LOG_CATEGORY(LogMutableCore);
 
+namespace
+{
+	static int32 SInitialized = 0;
+	static int32 SFinalized = 0;
+}
+
 
 void FMutableRuntimeModule::StartupModule()
 {
-	mu::Initialize();
+	if (!SInitialized)
+	{
+		SInitialized = 1;
+		SFinalized = 0;
+
+		miro::initialize();
+	}
 }
 
 
 void FMutableRuntimeModule::ShutdownModule()
 {
-	// Finalize the mutable runtime in this module
-	mu::Finalize();
+	if (SInitialized && !SFinalized)
+	{
+		miro::finalize();
+
+		SFinalized = 1;
+		SInitialized = 0;
+	}
 }
 
diff --git a/Engine/Plugins/Mutable/Source/MutableRuntime/Private/MuR/MutableString.cpp b/Engine/Plugins/Mutable/Source/MutableRuntime/Private/MuR/MutableString.cpp
--- a/Engine/Plugins/Mutable/Source/MutableRuntime/Private/MuR/MutableString.cpp
+++ b/Engine/Plugins/Mutable/Source/MutableRuntime/Private/MuR/MutableString.cpp
@@ -4,47 +4,28 @@
 #include "MuR/MutableString.h"
 
 
-namespace mu {
+namespace mu 
+{
 
-	//---------------------------------------------------------------------------------------------
 	String::String( const FString& InValue)
 	{
         Value = InValue;
 	}
 
 
-	////---------------------------------------------------------------------------------------------
-	//void String::Serialise( const String* p, OutputArchive& arch )
-	//{
-	//	arch << *p;
-	//}
-
-
-	////---------------------------------------------------------------------------------------------
-	//Ptr<String> String::StaticUnserialise( InputArchive& arch )
-	//{
-	//	Ptr<String> pResult = new String();
-	//	arch >> *pResult;
-	//	return pResult;
-	//}
-
-
-	//---------------------------------------------------------------------------------------------
-	Ptr<String> String::Clone() const
+	TSharedPtr<String> String::Clone() const
 	{
-		Ptr<String> pResult = new String(Value);
-		return pResult;
+		TSharedPtr<String> Result = MakeShared<String>(Value);
+		return Result;
 	}
 
 
-	//---------------------------------------------------------------------------------------------
 	int32 String::GetDataSize() const
 	{
 		return sizeof(String) + Value.GetAllocatedSize();
 	}
 
 
-	//---------------------------------------------------------------------------------------------
 	const FString& String::GetValue() const
 	{
         return Value;
diff --git a/Engine/Plugins/Mutable/Source/MutableRuntime/Private/MuR/NullExtensionDataStreamer.cpp b/Engine/Plugins/Mutable/Source/MutableRuntime/Private/MuR/NullExtensionDataStreamer.cpp
--- a/Engine/Plugins/Mutable/Source/MutableRuntime/Private/MuR/NullExtensionDataStreamer.cpp
+++ b/Engine/Plugins/Mutable/Source/MutableRuntime/Private/MuR/NullExtensionDataStreamer.cpp
@@ -8,20 +8,20 @@
 namespace mu
 {
 
-ExtensionDataPtr NullExtensionDataStreamer::CloneExtensionData(const ExtensionDataPtrConst& Source)
+TSharedPtr<FExtensionData> NullExtensionDataStreamer::CloneExtensionData(const TSharedPtr<const FExtensionData>& Source)
 {
 	// Note that these appear to be valid and should be good enough for testing, but they won't be
 	// usable since we're not copying the actual data pointed to by Source.
 
-	ExtensionDataPtr Result = new ExtensionData();
+	TSharedPtr<FExtensionData> Result = MakeShared<FExtensionData>();
 	Result->Index = NextIndex;
 	NextIndex++;
-	Result->Origin = ExtensionData::EOrigin::Runtime;
+	Result->Origin = FExtensionData::EOrigin::Runtime;
 
 	return Result;
 }
 
-TSharedRef<const FExtensionDataLoadHandle> NullExtensionDataStreamer::StartLoad(const ExtensionDataPtrConst& Data, TArray<ExtensionDataPtrConst>& OutUnloadedConstants)
+TSharedRef<const FExtensionDataLoadHandle> NullExtensionDataStreamer::StartLoad(const TSharedPtr<const FExtensionData>& Data, TArray<TSharedPtr<const FExtensionData>>& OutUnloadedConstants)
 {
 	OutUnloadedConstants.Reset();
 
diff --git a/Engine/Plugins/Mutable/Source/MutableRuntime/Private/MuR/NullExtensionDataStreamer.h b/Engine/Plugins/Mutable/Source/MutableRuntime/Private/MuR/NullExtensionDataStreamer.h
--- a/Engine/Plugins/Mutable/Source/MutableRuntime/Private/MuR/NullExtensionDataStreamer.h
+++ b/Engine/Plugins/Mutable/Source/MutableRuntime/Private/MuR/NullExtensionDataStreamer.h
@@ -7,12 +7,12 @@
 namespace mu
 {
 
-/** A stub implementation that can be used if there's no valid ExtensionDataStreamer */
-class NullExtensionDataStreamer : public ExtensionDataStreamer
+/** A stub implementation that can be used if there's no valid FExtensionDataStreamer */
+class NullExtensionDataStreamer : public FExtensionDataStreamer
 {
 public:
-	virtual ExtensionDataPtr CloneExtensionData(const ExtensionDataPtrConst& Source) override;
-	virtual TSharedRef<const FExtensionDataLoadHandle> StartLoad(const ExtensionDataPtrConst& Data, TArray<ExtensionDataPtrConst>& OutUnloadedConstants) override;
+	virtual TSharedPtr<FExtensionData> CloneExtensionData(const TSharedPtr<const FExtensionData>& Source) override;
+	virtual TSharedRef<const FExtensionDataLoadHandle> StartLoad(const TSharedPtr<const FExtensionData>& Data, TArray<TSharedPtr<const FExtensionData>>& OutUnloadedConstants) override;
 
 private:
 	int32 NextIndex = 0;
diff --git a/Engine/Plugins/Mutable/Source/MutableRuntime/Private/MuR/OpImageApplyComposite.h b/Engine/Plugins/Mutable/Source/MutableRuntime/Private/MuR/OpImageApplyComposite.h
--- a/Engine/Plugins/Mutable/Source/MutableRuntime/Private/MuR/OpImageApplyComposite.h
+++ b/Engine/Plugins/Mutable/Source/MutableRuntime/Private/MuR/OpImageApplyComposite.h
@@ -49,7 +49,7 @@ namespace mu
 		return  (Value << ChannelShift) | (Base & ResMask); 
 	}
 
-	inline void ImageNormalComposite(Image* pResult, const Image* pBase, const Image* pNormal, ECompositeImageMode mode, float power)
+	inline void ImageNormalComposite(FImage* pResult, const FImage* pBase, const FImage* pNormal, ECompositeImageMode mode, float power)
 	{
 		if (mode == ECompositeImageMode::CIM_Disabled)
 		{
diff --git a/Engine/Plugins/Mutable/Source/MutableRuntime/Private/MuR/OpImageBinarise.h b/Engine/Plugins/Mutable/Source/MutableRuntime/Private/MuR/OpImageBinarise.h
--- a/Engine/Plugins/Mutable/Source/MutableRuntime/Private/MuR/OpImageBinarise.h
+++ b/Engine/Plugins/Mutable/Source/MutableRuntime/Private/MuR/OpImageBinarise.h
@@ -14,7 +14,7 @@ namespace mu
 	//---------------------------------------------------------------------------------------------
 	//!
 	//---------------------------------------------------------------------------------------------
-	inline void ImageBinarise(Image* DestImage, const Image* AImage, float Threshold)
+	inline void ImageBinarise(FImage* DestImage, const FImage* AImage, float Threshold)
 	{
 		if (!AImage || !DestImage)
 		{
diff --git a/Engine/Plugins/Mutable/Source/MutableRuntime/Private/MuR/OpImageColourMap.h b/Engine/Plugins/Mutable/Source/MutableRuntime/Private/MuR/OpImageColourMap.h
--- a/Engine/Plugins/Mutable/Source/MutableRuntime/Private/MuR/OpImageColourMap.h
+++ b/Engine/Plugins/Mutable/Source/MutableRuntime/Private/MuR/OpImageColourMap.h
@@ -10,7 +10,7 @@
 namespace mu
 {
 
-	inline void ImageColourMap(Image* DestImage, const Image* SourceImage, const Image* MaskImage, const Image* MapImage, bool bOnlyOneMip)
+	inline void ImageColourMap(FImage* DestImage, const FImage* SourceImage, const FImage* MaskImage, const FImage* MapImage, bool bOnlyOneMip)
 	{
 		check(SourceImage->GetSizeX() == MaskImage->GetSizeX());
 		check(SourceImage->GetSizeY() == MaskImage->GetSizeY());
diff --git a/Engine/Plugins/Mutable/Source/MutableRuntime/Private/MuR/OpImageCompose.cpp b/Engine/Plugins/Mutable/Source/MutableRuntime/Private/MuR/OpImageCompose.cpp
--- a/Engine/Plugins/Mutable/Source/MutableRuntime/Private/MuR/OpImageCompose.cpp
+++ b/Engine/Plugins/Mutable/Source/MutableRuntime/Private/MuR/OpImageCompose.cpp
@@ -7,7 +7,7 @@
 namespace mu
 {
 
-	void FImageOperator::ImageCompose(Image* Base, const Image* Block, const box< FIntVector2 >& Rect)
+	void FImageOperator::ImageCompose(FImage* Base, const FImage* Block, const box< FIntVector2 >& Rect)
 	{
 		check(Base && Block);
 		check(Base != Block);
@@ -40,14 +40,14 @@ namespace mu
 			MUTABLE_CPUPROFILER_SCOPE(ImageCompose_SubBlock_Fix);
 			EImageFormat UncompressedFormat = GetUncompressedFormat(OriginalFormat);
 			constexpr int32 Quality = 0;
-			Ptr<Image> TempBase = ImagePixelFormat(Quality, Base, UncompressedFormat);
-			Ptr<Image> TempBlock = ImagePixelFormat(Quality, Block, UncompressedFormat);
-			ImageCompose(TempBase.get(), TempBlock.get(), Rect);
+			TSharedPtr<FImage> TempBase = ImagePixelFormat(Quality, Base, UncompressedFormat);
+			TSharedPtr<FImage> TempBlock = ImagePixelFormat(Quality, Block, UncompressedFormat);
+			ImageCompose(TempBase.Get(), TempBlock.Get(), Rect);
 			ReleaseImage(TempBlock);
-			Ptr<Image> NewBase = ImagePixelFormat(Quality, TempBase.get(), OriginalFormat);
+			TSharedPtr<FImage> NewBase = ImagePixelFormat(Quality, TempBase.Get(), OriginalFormat);
 			ReleaseImage(TempBase);
 			// ImageCompose above may have reduced the LODs because of the block lods.
-			Base->CopyMove(NewBase.get());
+			Base->CopyMove(NewBase.Get());
 			ReleaseImage(NewBase);
 			return;
 		}
diff --git a/Engine/Plugins/Mutable/Source/MutableRuntime/Private/MuR/OpImageCrop.cpp b/Engine/Plugins/Mutable/Source/MutableRuntime/Private/MuR/OpImageCrop.cpp
--- a/Engine/Plugins/Mutable/Source/MutableRuntime/Private/MuR/OpImageCrop.cpp
+++ b/Engine/Plugins/Mutable/Source/MutableRuntime/Private/MuR/OpImageCrop.cpp
@@ -8,22 +8,23 @@
 
 namespace mu
 {
-	bool FImageOperator::ImageCrop( Image* InCropped, int32 CompressionQuality, const Image* InBase, const box< FIntVector2 >& Rect)
+	bool FImageOperator::ImageCrop(TSharedPtr<FImage>& OutCropped, int32 CompressionQuality, const FImage* InBase, const box< FIntVector2 >& Rect)
 	{
-		Ptr<const Image> Base = InBase;
-		Ptr<Image> BaseReformat;
-		Ptr<Image> Cropped = InCropped;
+		check(InBase);
 
-		EImageFormat BaseFormat = Base->GetFormat();
+		TSharedPtr<FImage> BaseReformat;
+		TSharedPtr<FImage> Cropped = OutCropped;
+
+		EImageFormat BaseFormat = InBase->GetFormat();
 		EImageFormat UncompressedFormat = GetUncompressedFormat(BaseFormat);
 
 		if (BaseFormat != UncompressedFormat)
 		{
 			// Compressed formats need decompression + compression after crop			
 			// \TODO: This may use some additional untracked memory locally in this function.
-			BaseReformat = ImagePixelFormat(CompressionQuality, Base.get(), UncompressedFormat);
-			Base = BaseReformat;
-			Cropped = CreateImage(InCropped->GetSizeX(), InCropped->GetSizeY(), InCropped->GetLODCount(), UncompressedFormat, EInitializationType::NotInitialized);
+			BaseReformat = ImagePixelFormat(CompressionQuality, InBase, UncompressedFormat);
+			InBase = BaseReformat.Get();
+			Cropped = CreateImage(OutCropped->GetSizeX(), OutCropped->GetSizeY(), OutCropped->GetLODCount(), UncompressedFormat, EInitializationType::NotInitialized);
         }
 
 		// In case the base is smaller than the crop extension, make it as if the base extends indefinitely
@@ -35,7 +36,7 @@ namespace mu
 
 		const FImageFormatData& finfo = GetImageFormatData(UncompressedFormat);
 
-		check(Base && Cropped);
+		check(Cropped);
 		check(Cropped->GetSizeX() == Rect.size[0]);
 		check(Cropped->GetSizeY() == Rect.size[1]);
 
@@ -43,8 +44,8 @@ namespace mu
 		bool bCorrect =
 				(Rect.min[0] >= 0 && Rect.min[1] >= 0) &&
 				(Rect.size[0] >= 0 && Rect.size[1] >= 0) &&
-				(Rect.min[0] + Rect.size[0] <= Base->GetSizeX()) &&
-				(Rect.min[1] + Rect.size[1] <= Base->GetSizeY());
+				(Rect.min[0] + Rect.size[0] <= InBase->GetSizeX()) &&
+				(Rect.min[1] + Rect.size[1] <= InBase->GetSizeY());
 
 		if (!bCorrect)
 		{
@@ -62,10 +63,10 @@ namespace mu
 	
 		const uint32 BytesPerPixel = finfo.BytesPerBlock;
 
-        const uint8* BaseBuf = Base->GetLODData(0);
+        const uint8* BaseBuf = InBase->GetLODData(0);
         uint8* CropBuf = Cropped->GetLODData(0);
 
-		FImageSize BaseSize = Base->GetSize(); 
+		FImageSize BaseSize = InBase->GetSize();
 		FImageSize CroppedSize = Cropped->GetSize();
 
 		FIntRect BaseRect = FIntRect(
@@ -88,7 +89,7 @@ namespace mu
 			ReleaseImage(BaseReformat);
 
 			bool bSuccess = false;
-			ImagePixelFormat(bSuccess, CompressionQuality, InCropped, Cropped.get());
+			ImagePixelFormat(bSuccess, CompressionQuality, OutCropped.Get(), Cropped.Get());
 			check(bSuccess);
 
 			ReleaseImage(Cropped);
diff --git a/Engine/Plugins/Mutable/Source/MutableRuntime/Private/MuR/OpImageDisplace.h b/Engine/Plugins/Mutable/Source/MutableRuntime/Private/MuR/OpImageDisplace.h
--- a/Engine/Plugins/Mutable/Source/MutableRuntime/Private/MuR/OpImageDisplace.h
+++ b/Engine/Plugins/Mutable/Source/MutableRuntime/Private/MuR/OpImageDisplace.h
@@ -32,7 +32,7 @@ namespace mu
 	//---------------------------------------------------------------------------------------------
 	//!
 	//---------------------------------------------------------------------------------------------
-	inline void ImageMakeGrowMap(Image* ResultImage, const Image* MaskImage, int32 InBorder)
+	inline void ImageMakeGrowMap(FImage* ResultImage, const FImage* MaskImage, int32 InBorder)
 	{
 		MUTABLE_CPUPROFILER_SCOPE(ImageMakeGrowMap);
 
@@ -57,10 +57,11 @@ namespace mu
 			FMemory::Memset(Buffer.GetData(), ZeroOffset, Buffer.Num());
 		}
 
-		Ptr<const Image> ThisMask = MaskImage;
+		TSharedPtr<const FImage> ThisMaskCopy;
+		const FImage* ThisMask = MaskImage;
 		for (int32 B = 0; B < InBorder; ++B)
 		{
-			Ptr<Image> NextMask = ThisMask->Clone();
+			TSharedPtr<FImage> NextMask = ThisMask->Clone();
 
 			for(int32 CurrentMip = 0; CurrentMip < MipCount; ++CurrentMip)
 			{ 
@@ -127,13 +128,14 @@ namespace mu
 				ParallelFor(MipSize.Y, ProcessRow);
 			}
 
-			ThisMask = NextMask;
+			ThisMaskCopy = NextMask;
+			ThisMask = ThisMaskCopy.Get();
 		}
 
 	}
 
 
-	inline void ImageDisplace(Image* ResultImage, const Image* SourceImage, const Image* MapImage)
+	inline void ImageDisplace(FImage* ResultImage, const FImage* SourceImage, const FImage* MapImage)
 	{
 		MUTABLE_CPUPROFILER_SCOPE(ImageDisplace)
 
@@ -357,7 +359,7 @@ namespace mu
 
 
 		// Update the relevancy data of the image for the worst case. 
-		if (ResultImage->m_flags & Image::EImageFlags::IF_HAS_RELEVANCY_MAP)
+		if (ResultImage->Flags & FImage::EImageFlags::IF_HAS_RELEVANCY_MAP)
 		{
 			// Displace can encode at max MUTABLE_GROW_BORDER_VALUE pixels away according to "border" in MakeGrowMap.
 			ResultImage->RelevancyMinY = FMath::Max(int32(ResultImage->RelevancyMinY) - MUTABLE_GROW_BORDER_VALUE, 0);
diff --git a/Engine/Plugins/Mutable/Source/MutableRuntime/Private/MuR/OpImageGradient.h b/Engine/Plugins/Mutable/Source/MutableRuntime/Private/MuR/OpImageGradient.h
--- a/Engine/Plugins/Mutable/Source/MutableRuntime/Private/MuR/OpImageGradient.h
+++ b/Engine/Plugins/Mutable/Source/MutableRuntime/Private/MuR/OpImageGradient.h
@@ -9,7 +9,7 @@
 namespace mu
 {
 
-	inline void ImageGradient( Image* pDest, const FVector4f& c0, const FVector4f& c1 )
+	inline void ImageGradient( FImage* pDest, const FVector4f& c0, const FVector4f& c1 )
 	{
 		check(pDest && pDest->GetFormat()==EImageFormat::RGBA_UByte);
 
diff --git a/Engine/Plugins/Mutable/Source/MutableRuntime/Private/MuR/OpImageInterpolate.h b/Engine/Plugins/Mutable/Source/MutableRuntime/Private/MuR/OpImageInterpolate.h
--- a/Engine/Plugins/Mutable/Source/MutableRuntime/Private/MuR/OpImageInterpolate.h
+++ b/Engine/Plugins/Mutable/Source/MutableRuntime/Private/MuR/OpImageInterpolate.h
@@ -10,7 +10,7 @@
 namespace mu
 {
 
-    inline void ImageInterpolate(Image* DestImage, const Image* BImage, float Factor)
+    inline void ImageInterpolate(FImage* DestImage, const FImage* BImage, float Factor)
 	{
 		MUTABLE_CPUPROFILER_SCOPE(ImageInterpolate)
 
diff --git a/Engine/Plugins/Mutable/Source/MutableRuntime/Private/MuR/OpImageInvert.h b/Engine/Plugins/Mutable/Source/MutableRuntime/Private/MuR/OpImageInvert.h
--- a/Engine/Plugins/Mutable/Source/MutableRuntime/Private/MuR/OpImageInvert.h
+++ b/Engine/Plugins/Mutable/Source/MutableRuntime/Private/MuR/OpImageInvert.h
@@ -12,7 +12,7 @@ namespace mu
 
 	/** Invert the colour (RGB or L) components of an image. Leave alpha untouched. */
 	template<bool bUseVectorImpl = false>
-	inline void ImageInvert(Image* AImage)
+	inline void ImageInvert(FImage* AImage)
 	{
 		MUTABLE_CPUPROFILER_SCOPE(ImageInvertInPlace);
 			
diff --git a/Engine/Plugins/Mutable/Source/MutableRuntime/Private/MuR/OpImageLayer.h b/Engine/Plugins/Mutable/Source/MutableRuntime/Private/MuR/OpImageLayer.h
--- a/Engine/Plugins/Mutable/Source/MutableRuntime/Private/MuR/OpImageLayer.h
+++ b/Engine/Plugins/Mutable/Source/MutableRuntime/Private/MuR/OpImageLayer.h
@@ -47,7 +47,7 @@ namespace mu
 	}
 
 	template< unsigned (*BLEND_FUNC)(unsigned,unsigned), bool CLAMP >
-    inline void BufferLayerColour(Image* ResultImage, const Image* BaseImage, FVector4f Color)
+    inline void BufferLayerColour(FImage* ResultImage, const FImage* BaseImage, FVector4f Color)
 	{
         check(ResultImage->GetFormat() == BaseImage->GetFormat());
         check(ResultImage->GetSizeX() == BaseImage->GetSizeX());
@@ -118,7 +118,7 @@ namespace mu
 
 
 	template< uint32 (*BLEND_FUNC)(uint32, uint32) >
-	inline void BufferLayerColour(Image* ResultImage, const Image* BaseImage, FVector4f Color)
+	inline void BufferLayerColour(FImage* ResultImage, const FImage* BaseImage, FVector4f Color)
 	{
 		bool bIsClampNeeded = IsAnyComponentLargerThan1(Color);
 		if (bIsClampNeeded)
@@ -171,7 +171,7 @@ namespace mu
 	template<
 		uint32 (*BLEND_FUNC)(uint32, uint32), 
 		bool CLAMP>
-    inline void BufferLayerColourFromAlpha(Image* ResultImage, const Image* BaseImage, FVector4f Color)
+    inline void BufferLayerColourFromAlpha(FImage* ResultImage, const FImage* BaseImage, FVector4f Color)
 	{
         check(ResultImage->GetFormat() == BaseImage->GetFormat());
         check(ResultImage->GetSizeX() == BaseImage->GetSizeX());
@@ -242,7 +242,7 @@ namespace mu
 
 
 	template<uint32 (*BLEND_FUNC)(uint32, uint32)>
-	inline void BufferLayerColourFromAlpha(Image* ResultImage, const Image* BaseImage, FVector4f Color)
+	inline void BufferLayerColourFromAlpha(FImage* ResultImage, const FImage* BaseImage, FVector4f Color)
 	{
 		bool bIsClampNeeded = IsAnyComponentLargerThan1(Color);
 		if (bIsClampNeeded)
@@ -264,7 +264,7 @@ namespace mu
 		uint32 CHANNELS_TO_BLEND,
 		// Number of total channels in the base image
 		int32 BASE_CHANNEL_STRIDE>
-    inline void BufferLayerColourFormat(Image* DestImage, const Image* BaseImage, const Image* MaskImage, FVector4f Col, uint32 BaseChannelOffset, uint8 ColorChannelOffset, bool bOnlyFirstLOD)
+    inline void BufferLayerColourFormat(FImage* DestImage, const FImage* BaseImage, const FImage* MaskImage, FVector4f Col, uint32 BaseChannelOffset, uint8 ColorChannelOffset, bool bOnlyFirstLOD)
 	{
 		check(CHANNELS_TO_BLEND + BaseChannelOffset <= BASE_CHANNEL_STRIDE);
 		check(DestImage->GetLODCount() <= BaseImage->GetLODCount());
@@ -497,7 +497,7 @@ namespace mu
 		uint32 (*BLEND_FUNC_MASKED)(uint32,uint32,uint32),
 		uint32 (*BLEND_FUNC)(uint32,uint32),
 		bool CLAMP>
-    inline void BufferLayerColour(Image* ResultImage, const Image* BaseImage, const Image* MaskImage, FVector4f Col)
+    inline void BufferLayerColour(FImage* ResultImage, const FImage* BaseImage, const FImage* MaskImage, FVector4f Col)
 	{
 		check(BaseImage->GetSizeX() == MaskImage->GetSizeX());
 		check(BaseImage->GetSizeY() == MaskImage->GetSizeY());
@@ -542,7 +542,7 @@ namespace mu
 	template< 
 		uint32 (*BLEND_FUNC_MASKED)(uint32, uint32, uint32),
 		uint32 (*BLEND_FUNC)(uint32, uint32)>
-	inline void BufferLayerColour(Image* DestImage, const Image* BaseImage, const Image* MaskImage, FVector4f Col)
+	inline void BufferLayerColour(FImage* DestImage, const FImage* BaseImage, const FImage* MaskImage, FVector4f Col)
 	{
 		bool bIsClampNeeded = IsAnyComponentLargerThan1(Col);
 		if (bIsClampNeeded)
@@ -566,7 +566,7 @@ namespace mu
 		uint32 CHANNELS_TO_BLEND,
 		// Number of total channels in the base image
 		int32 BASE_CHANNEL_STRIDE>
-	inline void BufferLayerColourFormatInPlace(Image* BaseImage, FVector4f Color, uint32 BaseChannelOffset, uint8 ColorChannelOffset, bool bOnlyFirstLOD)
+	inline void BufferLayerColourFormatInPlace(FImage* BaseImage, FVector4f Color, uint32 BaseChannelOffset, uint8 ColorChannelOffset, bool bOnlyFirstLOD)
 	{
 		FUintVector4 TopColor = FUintVector4(Color.X * 255.0f, Color.Y * 255.0f, Color.Z * 255.0f, Color.W * 255.0f);
 
@@ -623,7 +623,7 @@ namespace mu
 		uint32 (*BLEND_FUNC)(uint32, uint32), 
 		bool CLAMP, 
 		uint32 CHANNEL_COUNT>
-	inline void BufferLayerColourInPlace(Image* BaseImage, FVector4f Col, bool bOnlyOneMip, uint32 BaseOffset, uint8 ColOffset)
+	inline void BufferLayerColourInPlace(FImage* BaseImage, FVector4f Col, bool bOnlyOneMip, uint32 BaseOffset, uint8 ColOffset)
 	{
 		EImageFormat BaseFormat = BaseImage->GetFormat();
 
@@ -664,7 +664,7 @@ namespace mu
 	template< 
 		uint32 (*BLEND_FUNC)(uint32, uint32), 
 		uint32 CHANNEL_COUNT>
-	inline void BufferLayerColourInPlace(Image* BaseImage, FVector4f Color, bool bOnlyOneMip, uint32 BaseOffset, uint8 ColorOffset)
+	inline void BufferLayerColourInPlace(FImage* BaseImage, FVector4f Color, bool bOnlyOneMip, uint32 BaseOffset, uint8 ColorOffset)
 	{
 		bool bIsClampNeeded = IsAnyComponentLargerThan1(Color);
 		if (bIsClampNeeded)
@@ -683,7 +683,7 @@ namespace mu
 		uint32 (*BLEND_FUNC)(uint32, uint32), 
 		bool CLAMP, 
 		int32 CHANNEL_COUNT>
-	inline void BufferLayerColourInPlace(Image* BaseImage, const Image* MaskImage, FVector4f Color, bool bOnlyOneMip, uint32 BaseOffset, uint8 ColorOffset)
+	inline void BufferLayerColourInPlace(FImage* BaseImage, const FImage* MaskImage, FVector4f Color, bool bOnlyOneMip, uint32 BaseOffset, uint8 ColorOffset)
 	{
 		check(BaseImage->GetSizeX() == MaskImage->GetSizeX());
 		check(BaseImage->GetSizeY() == MaskImage->GetSizeY());
@@ -727,7 +727,7 @@ namespace mu
 	template< 
 		uint32 (*BLEND_FUNC_MASKED)(uint32, uint32, uint32), 
 		uint32 (*BLEND_FUNC)(uint32, uint32), uint32 CHANNEL_COUNT>
-	inline void BufferLayerColourInPlace(Image* BaseImage, const Image* MaskImage, FVector4f Color, bool bOnlyOneMip, uint32 BaseOffset, uint8 ColorOffset)
+	inline void BufferLayerColourInPlace(FImage* BaseImage, const FImage* MaskImage, FVector4f Color, bool bOnlyOneMip, uint32 BaseOffset, uint8 ColorOffset)
 	{
 		bool bIsClampNeeded = IsAnyComponentLargerThan1(Color);
 		if (bIsClampNeeded)
@@ -754,8 +754,8 @@ namespace mu
 		// Number of total channels in the blend image
 		int32 BLENDED_CHANNEL_STRIDE>
     inline void BufferLayerFormatInPlace(
-		Image* BaseImage, 
-		const Image* BlendedImage,
+		FImage* BaseImage, 
+		const FImage* BlendedImage,
 		uint32 BaseChannelOffset,
 		uint32 BlendedChannelOffset,
 		bool bOnlyFirstLOD)
@@ -842,7 +842,7 @@ namespace mu
 		// Number of total channels in the blend image
 		int32 BLENDED_CHANNEL_STRIDE,
 		int32 BLENDED_CHANNEL_OFFSET>
-	inline void BufferLayerFormat(Image* DestImage, const Image* BaseImage, const Image* BlendedImage, bool bOnlyFirstLOD)
+	inline void BufferLayerFormat(FImage* DestImage, const FImage* BaseImage, const FImage* BlendedImage, bool bOnlyFirstLOD)
 	{
 		check(BaseImage->GetSizeX() == BlendedImage->GetSizeX());
 		check(BaseImage->GetSizeY() == BlendedImage->GetSizeY());
@@ -934,7 +934,7 @@ namespace mu
 	* Apply a blending function to an image with another image as blending layer
 	*/
 	template< uint32 (*BLEND_FUNC)(uint32, uint32), bool CLAMP >
-	inline void BufferLayer(Image* ResultImage, const Image* BaseImage, const Image* BlendedImage, bool bApplyToAlpha, bool bOnlyOneMip, bool bUseBlendSourceFromBlendAlpha)
+	inline void BufferLayer(FImage* ResultImage, const FImage* BaseImage, const FImage* BlendedImage, bool bApplyToAlpha, bool bOnlyOneMip, bool bUseBlendSourceFromBlendAlpha)
 	{
 		check(ResultImage->GetFormat() == BaseImage->GetFormat());
 		check(ResultImage->GetSizeX() == BaseImage->GetSizeX());
@@ -998,7 +998,7 @@ namespace mu
 
 
 	template< uint32 (*BLEND_FUNC)(uint32, uint32), bool CLAMP, int32 CHANNEL_COUNT >
-	inline void BufferLayerInPlace(Image* BaseImage, const Image* BlendedImage, bool bOnlyOneMip, uint32 BaseOffset, uint32 BlendedOffset)
+	inline void BufferLayerInPlace(FImage* BaseImage, const FImage* BlendedImage, bool bOnlyOneMip, uint32 BaseOffset, uint32 BlendedOffset)
 	{
 		check(BaseImage->GetSizeX() == BlendedImage->GetSizeX());
 		check(BaseImage->GetSizeY() == BlendedImage->GetSizeY());
@@ -1053,10 +1053,10 @@ namespace mu
 		// Number of total channels in the blend image
 		int32 BLENDED_CHANNEL_STRIDE>
 	inline void BufferLayerFormat(
-		Image* DestImage,
-		const Image* BaseImage,
-		const Image* MaskImage,
-		const Image* BlendImage,
+		FImage* DestImage,
+		const FImage* BaseImage,
+		const FImage* MaskImage,
+		const FImage* BlendImage,
 		uint32 DestOffset,
 		uint32 BaseChannelOffset,
 		uint32 BlendedChannelOffset,
@@ -1321,9 +1321,9 @@ namespace mu
 		// Number of total channels in the blend image
 		int32 BLENDED_CHANNEL_STRIDE>
 	inline void BufferLayerFormatEmbeddedMask(
-		Image* DestImage,
-		const Image* BaseImage,
-		const Image* BlendImage,
+		FImage* DestImage,
+		const FImage* BaseImage,
+		const FImage* BlendImage,
 		uint32 DestOffset,
 		uint32 BaseChannelOffset,
 		bool bOnlyFirstLOD)
@@ -1446,11 +1446,11 @@ namespace mu
 		bool CLAMP,
 		uint32 NC>
     inline void BufferLayerFormatStrideNoAlpha(
-		Image* DestImage,
+		FImage* DestImage,
 		int32 DestOffset,
 		int32 Stride,
-		const Image* MaskImage,
-		const Image* BlendImage/*, int32 LODCount*/)
+		const FImage* MaskImage,
+		const FImage* BlendImage/*, int32 LODCount*/)
 	{
         const uint8* MaskBuf = MaskImage->GetLODData(0);
         const uint8* BlendedBuf = BlendImage->GetLODData(0);
@@ -1547,10 +1547,10 @@ namespace mu
 	template<uint32(*BLEND_FUNC_MASKED)(uint32, uint32, uint32),
 			 uint32(*BLEND_FUNC)(uint32, uint32),
 			 bool CLAMP>
-    inline void BufferLayer(Image* DestImage,
-							const Image* BaseImage,
-							const Image* MaskImage,
-                            const Image* BlendImage,
+    inline void BufferLayer(FImage* DestImage,
+							const FImage* BaseImage,
+							const FImage* MaskImage,
+                            const FImage* BlendImage,
                             bool bApplyToAlpha,
 							bool bOnlyFirstLOD)
 	{
@@ -1589,9 +1589,9 @@ namespace mu
 		uint32 (*BLEND_FUNC)(uint32, uint32),
 		bool CLAMP >
 	inline void BufferLayerEmbeddedMask(
-		Image* DestImage,
-		const Image* BaseImage,
-		const Image* BlendImage,
+		FImage* DestImage,
+		const FImage* BaseImage,
+		const FImage* BlendImage,
 		bool bApplyToAlpha,
 		bool bOnlyFirstLOD)
 	{
@@ -1626,8 +1626,8 @@ namespace mu
 		uint32 (*A_FUNC)(uint32, uint32),
 		bool CLAMP >
 	inline void BufferLayerComposite(
-		Image* BaseImage,
-		const Image* BlendImage,
+		FImage* BaseImage,
+		const FImage* BlendImage,
 		bool bOnlyFirstLOD,
 		uint8 BlendAlphaSourceChannel)
 	{
@@ -1641,7 +1641,7 @@ namespace mu
 		int32 FirstLODDataOffset = 0;	
 		int32 NumRelevantElems = -1;
 
-		if (BlendImage->m_flags & Image::IF_HAS_RELEVANCY_MAP && bOnlyFirstLOD)
+		if (BlendImage->Flags & FImage::IF_HAS_RELEVANCY_MAP && bOnlyFirstLOD)
 		{
 			check(BlendImage->RelevancyMaxY < BaseImage->GetSizeY());
 			check(BlendImage->RelevancyMaxY >= BlendImage->RelevancyMinY);
@@ -1743,8 +1743,8 @@ namespace mu
 	template<>
 	inline void BufferLayerComposite<BlendChannelMasked, LightenChannel, false>
 	(
-		Image* BaseImage,
-		const Image* BlendImage,
+		FImage* BaseImage,
+		const FImage* BlendImage,
 		bool bOnlyFirstLOD,
 		uint8 BlendAlphaSourceChannel)
 	{
@@ -1758,7 +1758,7 @@ namespace mu
 
 		int32 FirstLODDataOffset = 0;
 		int32 NumRelevantElems = -1;
-		if (BlendImage->m_flags & Image::IF_HAS_RELEVANCY_MAP && bOnlyFirstLOD)
+		if (BlendImage->Flags & FImage::IF_HAS_RELEVANCY_MAP && bOnlyFirstLOD)
 		{
 			check(BlendImage->RelevancyMaxY < BaseImage->GetSizeY());
 			check(BlendImage->RelevancyMaxY >= BlendImage->RelevancyMinY);
@@ -1842,8 +1842,8 @@ namespace mu
 		int32 (*A_FUNC)(int32, int32),
 		bool CLAMP >
 	inline void BufferLayerCompositeVector(
-		Image* BaseImage,
-		const Image* BlendImage,
+		FImage* BaseImage,
+		const FImage* BlendImage,
 		bool bOnlyFirstLOD,
 		uint8 BlendAlphaSourceChannel)
 	{
@@ -1940,7 +1940,7 @@ namespace mu
 
 	template<uint32 (*BLEND_FUNC)(uint32,uint32),
 			 bool CLAMP>
-    inline void BufferLayerStrideNoAlpha(Image* DestImage, int32 DestOffset, int32 Stride, const Image* MaskImage, const Image* BlendImage/*, int32 LODCount*/)
+    inline void BufferLayerStrideNoAlpha(FImage* DestImage, int32 DestOffset, int32 Stride, const FImage* MaskImage, const FImage* BlendImage/*, int32 LODCount*/)
 	{
 		if (BlendImage->GetFormat() == EImageFormat::RGB_UByte)
 		{
@@ -1973,9 +1973,9 @@ namespace mu
 	*/
 	template<uint32 (*BLEND_FUNC)(uint32, uint32), bool CLAMP>
 	inline void ImageLayerOnBaseNoAlpha(
-			Image* BaseImage,
-			const Image* MaskImage,
-			const Image* BlendedImage,
+			FImage* BaseImage,
+			const FImage* MaskImage,
+			const FImage* BlendedImage,
 			const box<FIntVector2>& Rect)
 	{
 		check(BaseImage->GetSizeX() >= Rect.min[0] + Rect.size[0]);
@@ -2025,7 +2025,7 @@ namespace mu
 	template< 
 		uint32 (*BLEND_FUNC)(uint32, uint32), 
 		uint32 NC>
-	inline void BufferLayerCombineColour(Image* ResultImage, const Image* BaseImage, FVector4f Color, bool bOnlyFirstLOD = false)
+	inline void BufferLayerCombineColour(FImage* ResultImage, const FImage* BaseImage, FVector4f Color, bool bOnlyFirstLOD = false)
 	{
 		static_assert(NC > 0 && NC <= 4);
 
@@ -2081,7 +2081,7 @@ namespace mu
 	template< 
 		uint32 (*BLEND_FUNC)(uint32, uint32), 
 		uint32 NC>
-	inline void BufferLayerCombine(Image* ResultImage, const Image* BaseImage, const Image* BlendImage, bool bOnlyFirstLOD)
+	inline void BufferLayerCombine(FImage* ResultImage, const FImage* BaseImage, const FImage* BlendImage, bool bOnlyFirstLOD)
 	{
 		static_assert(NC > 0 && NC <= 4);
 
@@ -2142,7 +2142,7 @@ namespace mu
 	template< 
 		uint32 (*BLEND_FUNC_MASKED)(uint32, uint32, uint32), 
 		uint32 NC>
-	inline void BufferLayerCombine(Image* ResultImage, const Image* BaseImage, const Image* MaskImage, const Image* BlendImage, bool bOnlyFirstLOD)
+	inline void BufferLayerCombine(FImage* ResultImage, const FImage* BaseImage, const FImage* MaskImage, const FImage* BlendImage, bool bOnlyFirstLOD)
 	{
 		static_assert(NC > 0 && NC <= 4);
 
@@ -2210,7 +2210,7 @@ namespace mu
 	template< 
 		uint32 (*BLEND_FUNC_MASKED)(uint32, uint32, uint32), 
 		uint32 NC>
-	inline void BufferLayerCombineColour(Image* ResultImage, const Image* BaseImage, const Image* MaskImage, FVector4f Color)
+	inline void BufferLayerCombineColour(FImage* ResultImage, const FImage* BaseImage, const FImage* MaskImage, FVector4f Color)
 	{
 		static_assert(NC > 0 && NC <= 4);
 
@@ -2264,7 +2264,7 @@ namespace mu
 	}
 
 	template<uint32 (*BLEND_FUNC)(uint32, uint32)>
-	inline void ImageLayerCombine(Image* ResultImage, const Image* BaseImage, const Image* BlendedImage, bool bOnlyFirstLOD)
+	inline void ImageLayerCombine(FImage* ResultImage, const FImage* BaseImage, const FImage* BlendedImage, bool bOnlyFirstLOD)
 	{
 		check(ResultImage->GetFormat() == BaseImage->GetFormat());
 		check(ResultImage->GetSizeX() == BaseImage->GetSizeX());
@@ -2299,7 +2299,7 @@ namespace mu
 	template< 
 		uint32 (*BLEND_FUNC)(uint32, uint32), 
 		uint32 (*BLEND_FUNC_MASKED)(uint32, uint32, uint32)>
-	inline void ImageLayerCombine(Image* ResultImage, const Image* BaseImage, const Image* MaskImage, const Image* BlendedImage, bool bOnlyFirstLOD)
+	inline void ImageLayerCombine(FImage* ResultImage, const FImage* BaseImage, const FImage* MaskImage, const FImage* BlendedImage, bool bOnlyFirstLOD)
 	{
 		check(ResultImage->GetFormat() == BaseImage->GetFormat());
 		check(ResultImage->GetSizeX() == BaseImage->GetSizeX());
@@ -2312,7 +2312,7 @@ namespace mu
 
 		const EImageFormat BaseFormat = BaseImage->GetFormat();
 
-		Ptr<Image> TempMaskImage;
+		TSharedPtr<FImage> TempMaskImage;
 		if (MaskImage->GetFormat() != EImageFormat::L_UByte)
 		{
 			UE_LOG(LogMutableCore, Log, TEXT("Image layer format not supported. A generic one will be used. "));
@@ -2320,7 +2320,7 @@ namespace mu
 			FImageOperator ImOp = FImageOperator::GetDefault(nullptr);
 			constexpr int32 Quality = 4;
 			TempMaskImage = ImOp.ImagePixelFormat( Quality, MaskImage, EImageFormat::L_UByte );
-			MaskImage = TempMaskImage.get();
+			MaskImage = TempMaskImage.Get();
 		}
 
 		if (BaseFormat == EImageFormat::L_UByte)
@@ -2342,14 +2342,14 @@ namespace mu
 
 			FImageOperator ImOp = FImageOperator::GetDefault(nullptr);
 			constexpr int32 Quality = 4;
-			Ptr<Image> TempBaseImage = ImOp.ImagePixelFormat(Quality, BaseImage, EImageFormat::RGBA_UByte);
-			Ptr<Image> TempBlededImage = ImOp.ImagePixelFormat(Quality, BlendedImage, EImageFormat::RGBA_UByte);
-			BufferLayerCombine<BLEND_FUNC_MASKED, 4>(ResultImage, TempBaseImage.get(), MaskImage, TempBlededImage.get(), bOnlyFirstLOD);
+			TSharedPtr<FImage> TempBaseImage = ImOp.ImagePixelFormat(Quality, BaseImage, EImageFormat::RGBA_UByte);
+			TSharedPtr<FImage> TempBlededImage = ImOp.ImagePixelFormat(Quality, BlendedImage, EImageFormat::RGBA_UByte);
+			BufferLayerCombine<BLEND_FUNC_MASKED, 4>(ResultImage, TempBaseImage.Get(), MaskImage, TempBlededImage.Get(), bOnlyFirstLOD);
 		}
 	}
 
 	template<uint32 (*BLEND_FUNC)(uint32, uint32)>
-	inline void ImageLayerCombineColour(Image* ResultImage, const Image* BaseImage, FVector4f Color)
+	inline void ImageLayerCombineColour(FImage* ResultImage, const FImage* BaseImage, FVector4f Color)
 	{
 		check(ResultImage->GetFormat() == BaseImage->GetFormat());
 		check(ResultImage->GetSizeX() == BaseImage->GetSizeX());
@@ -2380,7 +2380,7 @@ namespace mu
 	template< 
 		uint32 (*BLEND_FUNC)(uint32, uint32),
 		uint32 (*BLEND_FUNC_MASKED)(uint32, uint32, uint32)>
-	inline void ImageLayerCombineColour(Image* ResultImage, const Image* BaseImage, const Image* MaskImage, FVector4f Color)
+	inline void ImageLayerCombineColour(FImage* ResultImage, const FImage* BaseImage, const FImage* MaskImage, FVector4f Color)
 	{
 		check(ResultImage->GetFormat() == BaseImage->GetFormat());
 		check(ResultImage->GetSizeX() == BaseImage->GetSizeX());
@@ -2416,7 +2416,7 @@ namespace mu
 	}
 
 	template<uint32 NCBase, uint32 NCBlend, class ImageCombineFn>
-	inline void BufferLayerCombineFunctor(Image* DestImage, const Image* BaseImage, const Image* BlendImage, ImageCombineFn&& ImageCombine)
+	inline void BufferLayerCombineFunctor(FImage* DestImage, const FImage* BaseImage, const FImage* BlendImage, ImageCombineFn&& ImageCombine)
 	{
 		static_assert(NCBase > 0 && NCBase <= 4);
 		static_assert(NCBlend > 0 && NCBlend <= 4);
@@ -2472,7 +2472,7 @@ namespace mu
 
 	// Same functionality as above, in this case we use a functor which allows to pass user data. 
 	template<class ImageCombineFn>
-	inline void ImageLayerCombineFunctor(Image* ResultImage, const Image* BaseImage, const Image* BlendedImage, ImageCombineFn&& ImageCombine)
+	inline void ImageLayerCombineFunctor(FImage* ResultImage, const FImage* BaseImage, const FImage* BlendedImage, ImageCombineFn&& ImageCombine)
 	{
 		check(ResultImage->GetFormat() == BaseImage->GetFormat());
 		check(ResultImage->GetSizeX() == BaseImage->GetSizeX());
@@ -2550,9 +2550,9 @@ namespace mu
 
 
 	//! Blend a subimage on the base using a mask.
-	inline void ImageBlendOnBaseNoAlpha(Image* BaseImage,
-		const Image* MaskImage,
-		const Image* BlendedImage,
+	inline void ImageBlendOnBaseNoAlpha(FImage* BaseImage,
+		const FImage* MaskImage,
+		const FImage* BlendedImage,
 		const box<FIntVector2>& Rect)
 	{
 		ImageLayerOnBaseNoAlpha<BlendChannel, false>(BaseImage, MaskImage, BlendedImage, Rect);
diff --git a/Engine/Plugins/Mutable/Source/MutableRuntime/Private/MuR/OpImageLuminance.h b/Engine/Plugins/Mutable/Source/MutableRuntime/Private/MuR/OpImageLuminance.h
--- a/Engine/Plugins/Mutable/Source/MutableRuntime/Private/MuR/OpImageLuminance.h
+++ b/Engine/Plugins/Mutable/Source/MutableRuntime/Private/MuR/OpImageLuminance.h
@@ -10,7 +10,7 @@
 namespace mu
 {
 
-	inline void ImageLuminance(Image* DestImage, const Image* AImage)
+	inline void ImageLuminance(FImage* DestImage, const FImage* AImage)
 	{
 		check(DestImage && AImage && DestImage->GetFormat() == EImageFormat::L_UByte);
 
diff --git a/Engine/Plugins/Mutable/Source/MutableRuntime/Private/MuR/OpImageMipmap.cpp b/Engine/Plugins/Mutable/Source/MutableRuntime/Private/MuR/OpImageMipmap.cpp
--- a/Engine/Plugins/Mutable/Source/MutableRuntime/Private/MuR/OpImageMipmap.cpp
+++ b/Engine/Plugins/Mutable/Source/MutableRuntime/Private/MuR/OpImageMipmap.cpp
@@ -371,7 +371,7 @@ namespace OpImageMipmapInternal
 	template<int32 PIXEL_SIZE>
 	inline void GenerateMipmapUint8LODRange(
 		int32 SrcLOD, int32 DestLODBegin, int32 DestLODEnd,
-		const Image* SourceImage, Image* DestImage,
+		const FImage* SourceImage, FImage* DestImage,
 		const FMipmapGenerationSettings& Settings)
 	{
 		using namespace OpImageMipmapInternal;
@@ -439,7 +439,7 @@ namespace OpImageMipmapInternal
 	template<int32 PixelSize>
 	inline void GenerateMipmapsBlockCompressedLODRange(
 			int32 SrcLOD, int32 DestLODBegin, int32 DestLODEnd,
-			const Image* SourceImage, Image* DestImage,
+			const FImage* SourceImage, FImage* DestImage,
 			const FMipmapGenerationSettings& Settings)
 	{
 		using namespace OpImageMipmapInternal;
@@ -531,7 +531,7 @@ namespace OpImageMipmapInternal
 	}
 
 
-    void FImageOperator::ImageMipmap_PrepareScratch(const Image* DestImage, int32 StartLevel, int32 LevelCount, FScratchImageMipmap& Scratch )
+    void FImageOperator::ImageMipmap_PrepareScratch(const FImage* DestImage, int32 StartLevel, int32 LevelCount, FScratchImageMipmap& Scratch )
     {
         check(DestImage->GetLODCount() == LevelCount);
 
@@ -609,10 +609,10 @@ namespace OpImageMipmapInternal
 		ReleaseImage(Scratch.CompressedMips);
 	}
 
-	void FImageOperator::ImageMipmap(FScratchImageMipmap& Scratch, int32 CompressionQuality, Image* DestImage, const Image* BaseImage,
+	void FImageOperator::ImageMipmap(FScratchImageMipmap& Scratch, int32 CompressionQuality, FImage* DestImage, const FImage* BaseImage,
 		int32 StartLevel, int32 NumLODs, const FMipmapGenerationSettings& Settings, bool bGenerateOnlyTail)
 	{
-		check(!(BaseImage->m_flags & Image::IF_CANNOT_BE_SCALED));
+		check(!(BaseImage->Flags & FImage::IF_CANNOT_BE_SCALED));
 		check(DestImage->GetFormat() == BaseImage->GetFormat());
 		
 		if (!bGenerateOnlyTail)
@@ -653,7 +653,7 @@ namespace OpImageMipmapInternal
 		if (bIsBlockCompressedFormat && bEnableCompressedMipGenerationMemoryOptimizations)
 		{
 			const EImageFormat DestFormat = Scratch.UncompressedMips->GetFormat();
-			Image* UncompressedImage = Scratch.UncompressedMips.get();
+			FImage* UncompressedImage = Scratch.UncompressedMips.Get();
 	
 			const int32 UncompressedNumLODs = UncompressedImage->GetLODCount();
 
@@ -698,20 +698,20 @@ namespace OpImageMipmapInternal
 			// Bad case.
 			// Uncompress the last mip that we already have
 			bool bSuccess = false;
-			ImagePixelFormat(bSuccess, CompressionQuality, Scratch.Uncompressed.get(), BaseImage, StartLevel);
+			ImagePixelFormat(bSuccess, CompressionQuality, Scratch.Uncompressed.Get(), BaseImage, StartLevel);
 			check(bSuccess);
 
 			// Generate the mipmaps from there on
 			constexpr bool bGenerateOnlyTailForCompressed = true;
 
 			const int32 NumScratchMips = Scratch.UncompressedMips->GetLODCount() + Scratch.Uncompressed->GetLODCount();
-			ImageMipmap(Scratch, CompressionQuality, Scratch.UncompressedMips.get(),
-					Scratch.Uncompressed.get(), 0, NumScratchMips, Settings, bGenerateOnlyTailForCompressed);
+			ImageMipmap(Scratch, CompressionQuality, Scratch.UncompressedMips.Get(),
+					Scratch.Uncompressed.Get(), 0, NumScratchMips, Settings, bGenerateOnlyTailForCompressed);
 
 			// Compress the mipmapped image
 			bSuccess = false;
 			ImagePixelFormat(
-					bSuccess, CompressionQuality, DestImage, Scratch.UncompressedMips.get(), 
+					bSuccess, CompressionQuality, DestImage, Scratch.UncompressedMips.Get(), 
 					StartLevel + 1, 0, Scratch.UncompressedMips->GetLODCount()); 
 			check(bSuccess);
 		}
@@ -744,7 +744,7 @@ namespace OpImageMipmapInternal
 		}
 	}
 
-	void FImageOperator::ImageMipmap(int32 CompressionQuality, Image* Dest, const Image* Base,
+	void FImageOperator::ImageMipmap(int32 CompressionQuality, FImage* Dest, const FImage* Base,
 		int32 StartLevel, int32 LevelCount, const FMipmapGenerationSettings& Settings, bool bGenerateOnlyTail)
 	{
 		FScratchImageMipmap Scratch;
@@ -758,9 +758,9 @@ namespace OpImageMipmapInternal
 	 * Update all the mipmaps in the image from the data in the base one. 
 	 * Only the mipmaps already existing in the image are updated.
 	 */
-	void ImageMipmapInPlace(int32 InImageCompressionQuality, Image* InBase, const FMipmapGenerationSettings& InSettings)
+	void ImageMipmapInPlace(int32 InImageCompressionQuality, FImage* InBase, const FMipmapGenerationSettings& InSettings)
 	{
-		check(!(InBase->m_flags & Image::IF_CANNOT_BE_SCALED));
+		check(!(InBase->Flags & FImage::IF_CANNOT_BE_SCALED));
 
 		int32 LevelCount = InBase->GetLODCount();
 		
diff --git a/Engine/Plugins/Mutable/Source/MutableRuntime/Private/MuR/OpImageNormalCombine.h b/Engine/Plugins/Mutable/Source/MutableRuntime/Private/MuR/OpImageNormalCombine.h
--- a/Engine/Plugins/Mutable/Source/MutableRuntime/Private/MuR/OpImageNormalCombine.h
+++ b/Engine/Plugins/Mutable/Source/MutableRuntime/Private/MuR/OpImageNormalCombine.h
@@ -62,22 +62,22 @@ namespace mu
 			(FMath::Clamp(static_cast<unsigned>((n.Z + 1.0f) * 0.5f * 255.0f), 0u, 255u) << 16);
 	}
 
-	inline void ImageNormalCombine(Image* pResult, const Image* pBase, FVector4f col)
+	inline void ImageNormalCombine(FImage* pResult, const FImage* pBase, FVector4f col)
 	{
 		ImageLayerCombineColour<CombineNormal>(pResult, pBase, col);
 	}
 
-	inline void ImageNormalCombine(Image* pResult, const Image* pBase, const Image* pBlended, bool bOnlyFirstLOD)
+	inline void ImageNormalCombine(FImage* pResult, const FImage* pBase, const FImage* pBlended, bool bOnlyFirstLOD)
 	{
 		ImageLayerCombine<CombineNormal>(pResult, pBase, pBlended, bOnlyFirstLOD);
 	}
 
-	inline void ImageNormalCombine(Image* pResult, const Image* pBase, const Image* pMask, FVector4f col)
+	inline void ImageNormalCombine(FImage* pResult, const FImage* pBase, const FImage* pMask, FVector4f col)
 	{
 		ImageLayerCombineColour<CombineNormal, CombineNormalMasked>(pResult, pBase, pMask, col);
 	}
 
-	inline void ImageNormalCombine(Image* pResult, const Image* pBase, const Image* pMask, const Image* pBlended, bool bOnlyFirstLOD)
+	inline void ImageNormalCombine(FImage* pResult, const FImage* pBase, const FImage* pMask, const FImage* pBlended, bool bOnlyFirstLOD)
 	{
 		ImageLayerCombine<CombineNormal, CombineNormalMasked>(pResult, pBase, pMask, pBlended, bOnlyFirstLOD);
 	}
