[mutable] More enforcement of coding standards.

diff --git a/Engine/Plugins/Mutable/Source/MutableRuntime/Private/MuR/OpImagePixelFormat.cpp b/Engine/Plugins/Mutable/Source/MutableRuntime/Private/MuR/OpImagePixelFormat.cpp
--- a/Engine/Plugins/Mutable/Source/MutableRuntime/Private/MuR/OpImagePixelFormat.cpp
+++ b/Engine/Plugins/Mutable/Source/MutableRuntime/Private/MuR/OpImagePixelFormat.cpp
@@ -19,7 +19,7 @@ namespace mu
 {
 
 
-Ptr<Image> FImageOperator::ImagePixelFormat(int32 CompressionQuality, const Image* Base, EImageFormat TargetFormat, int32 OnlyLOD)
+TSharedPtr<FImage> FImageOperator::ImagePixelFormat(int32 CompressionQuality, const FImage* Base, EImageFormat TargetFormat, int32 OnlyLOD)
 {
 	MUTABLE_CPUPROFILER_SCOPE(ImagePixelFormat);
 	TRACE_CPUPROFILER_EVENT_SCOPE_TEXT(*FString::Printf(TEXT("From %d to %d"), int32(Base->GetFormat()), int32(TargetFormat)));
@@ -39,8 +39,8 @@ Ptr<Image> FImageOperator::ImagePixelFormat(int32 CompressionQuality, const Imag
 		ResultLODCount = 1;
 	}
 
-	Ptr<Image> Result = CreateImage(ResultSize[0], ResultSize[1], ResultLODCount, TargetFormat, EInitializationType::NotInitialized);
-	Result->m_flags = Base->m_flags;
+	TSharedPtr<FImage> Result = CreateImage(ResultSize[0], ResultSize[1], ResultLODCount, TargetFormat, EInitializationType::NotInitialized);
+	Result->Flags = Base->Flags;
 
 	if (Base->GetSizeX() <= 0 || Base->GetSizeY() <= 0)
 	{
@@ -48,7 +48,7 @@ Ptr<Image> FImageOperator::ImagePixelFormat(int32 CompressionQuality, const Imag
 	}
 
 	bool bSuccess = false;
-	ImagePixelFormat(bSuccess, CompressionQuality, Result.get(), Base, OnlyLOD);
+	ImagePixelFormat(bSuccess, CompressionQuality, Result.Get(), Base, OnlyLOD);
 	check(bSuccess);
 
 	return Result;
@@ -58,7 +58,7 @@ namespace ImagePixelFormatInternal
 {
 	template<class PixelFormatFuncType>
 	void ProcessBatchedHelper(
-			Image* Dest, const Image* Src, 
+			FImage* Dest, const FImage* Src, 
 			int32 DestElemSizeInBytes, int32 SrcElemSizeInBytes, 
 			int32 DestLODBegin, int32 SrcLODBegin, int32 NumLODs, 
 			PixelFormatFuncType&& BatchFunc)
@@ -89,7 +89,7 @@ namespace ImagePixelFormatInternal
 
 	template<class PixelFormatFuncType>
 	void ProcessBatchedHelper(
-			Image* Dest, int32 DestElemSizeInBytes, int32 DestLODBegin, int32 NumLODs, 
+			FImage* Dest, int32 DestElemSizeInBytes, int32 DestLODBegin, int32 NumLODs, 
 			PixelFormatFuncType&& BatchFunc)
 	{
 		check(Dest->GetLODCount() >= DestLODBegin + NumLODs);
@@ -113,7 +113,7 @@ namespace ImagePixelFormatInternal
 
 	template<class BlockDecompressionFuncType> 
 	void BlockDecompressionLODRangeHelper(
-		Image* Result, const Image* Base, 
+		FImage* Result, const FImage* Base, 
 		int32 ResultLODBegin, int32 BaseLODBegin, int32 NumLODs, 
 		BlockDecompressionFuncType&& DecFunc)
 	{
@@ -133,7 +133,7 @@ namespace ImagePixelFormatInternal
 
 	template<class BlockCompressionFuncType> 
 	void BlockCompressionLODRangeHelper(
-		Image* Result, const Image* Base, 
+		FImage* Result, const FImage* Base, 
 		int32 ResultLODBegin, int32 BaseLODBegin, int32 NumLODs, int32 Quality,
 		BlockCompressionFuncType&& CompFunc)
 	{
@@ -153,7 +153,7 @@ namespace ImagePixelFormatInternal
 
 }
 
-void FImageOperator::ImagePixelFormat(bool& bOutSuccess, int32 CompressionQuality, Image* Result, const Image* Base, int32 OnlyLOD)
+void FImageOperator::ImagePixelFormat(bool& bOutSuccess, int32 CompressionQuality, FImage* Result, const FImage* Base, int32 OnlyLOD)
 {
 	if (FormatImageOverride)
 	{
@@ -175,7 +175,7 @@ void FImageOperator::ImagePixelFormat(bool& bOutSuccess, int32 CompressionQualit
 	check(bOutSuccess);
 }
 
-void FImageOperator::ImagePixelFormat(bool& bOutSuccess, int32 CompressionQuality, Image* Result, const Image* Base, 
+void FImageOperator::ImagePixelFormat(bool& bOutSuccess, int32 CompressionQuality, FImage* Result, const FImage* Base, 
 		int32 ResultLODBegin, int32 BaseLODBegin, int32 NumLODs)
 {
 	MUTABLE_CPUPROFILER_SCOPE(ImagePixelFormatInPlace);
@@ -263,10 +263,10 @@ void FImageOperator::ImagePixelFormat(bool& bOutSuccess, int32 CompressionQualit
 			}
 			case EImageFormat::RGB_UByteRLE:
 			{
-				Ptr<Image> TempBase =
+				TSharedPtr<FImage> TempBase =
 						ImagePixelFormat(CompressionQuality, Base, EImageFormat::RGB_UByte);
 
-				ImagePixelFormatInternal::ProcessBatchedHelper(Result, TempBase.get(), 1, 3, ResultLODBegin, BaseLODBegin, NumLODs,
+				ImagePixelFormatInternal::ProcessBatchedHelper(Result, TempBase.Get(), 1, 3, ResultLODBegin, BaseLODBegin, NumLODs,
 				[](uint8* DestBuf, const uint8* BaseBuf, int32 NumElems)
 				{
 					for (int32 I = 0; I < NumElems; ++I)
@@ -296,9 +296,9 @@ void FImageOperator::ImagePixelFormat(bool& bOutSuccess, int32 CompressionQualit
 			}
 			case EImageFormat::RGBA_UByteRLE:
 			{
-				Ptr<Image> TempBase = ImagePixelFormat(CompressionQuality, Base, EImageFormat::RGBA_UByte);
+				TSharedPtr<FImage> TempBase = ImagePixelFormat(CompressionQuality, Base, EImageFormat::RGBA_UByte);
 				
-				ImagePixelFormatInternal::ProcessBatchedHelper(Result, TempBase.get(), 1, 4, ResultLODBegin, BaseLODBegin, NumLODs,
+				ImagePixelFormatInternal::ProcessBatchedHelper(Result, TempBase.Get(), 1, 4, ResultLODBegin, BaseLODBegin, NumLODs,
 				[](uint8* DestBuf, const uint8* BaseBuf, int32 NumElems)
 				{
 					for (int32 I = 0; I < NumElems; ++I)
@@ -330,9 +330,9 @@ void FImageOperator::ImagePixelFormat(bool& bOutSuccess, int32 CompressionQualit
 			case EImageFormat::BC2:
 			case EImageFormat::BC3:
 			{
-				Ptr<Image> TempBase = ImagePixelFormat(CompressionQuality, Base, EImageFormat::RGB_UByte);
+				TSharedPtr<FImage> TempBase = ImagePixelFormat(CompressionQuality, Base, EImageFormat::RGB_UByte);
 				
-				ImagePixelFormatInternal::ProcessBatchedHelper(Result, TempBase.get(), 1, 3, ResultLODBegin, BaseLODBegin, NumLODs,
+				ImagePixelFormatInternal::ProcessBatchedHelper(Result, TempBase.Get(), 1, 3, ResultLODBegin, BaseLODBegin, NumLODs,
 				[](uint8* DestBuf, const uint8* BaseBuf, int32 NumElems)
 				{
 					for (int32 I = 0; I < NumElems; ++I)
@@ -388,7 +388,7 @@ void FImageOperator::ImagePixelFormat(bool& bOutSuccess, int32 CompressionQualit
 //				if (bOutSuccess)
 //				{
 //					// Verify
-//					Ptr<Image> Test = ImagePixelFormat(CompressionQuality, Result, EImageFormat::L_UByte);
+//					TSharedPtr<FImage> Test = ImagePixelFormat(CompressionQuality, Result, EImageFormat::L_UByte);
 //					check(Test->DataStorage == Base->DataStorage);
 //				}
 //#endif
@@ -400,7 +400,7 @@ void FImageOperator::ImagePixelFormat(bool& bOutSuccess, int32 CompressionQualit
 			case EImageFormat::BGRA_UByte:
 			{
 				check(bFormatAllLODs);
-				Ptr<Image> TempBase = ImagePixelFormat(CompressionQuality, Base, EImageFormat::L_UByte, -1);
+				TSharedPtr<FImage> TempBase = ImagePixelFormat(CompressionQuality, Base, EImageFormat::L_UByte, -1);
  
 				// Try to compress
 				if (Result->DataStorage.IsEmpty())
@@ -416,13 +416,13 @@ void FImageOperator::ImagePixelFormat(bool& bOutSuccess, int32 CompressionQualit
 					Result->DataStorage.CompactedTailOffsets = Base->DataStorage.CompactedTailOffsets;
 				}
 
-				CompressRLE_L(TempBase.get(), Result);
+				CompressRLE_L(TempBase.Get(), Result);
 
 //#ifdef MUTABLE_DEBUG_RLE
 //				if (bOutSuccess)
 //				{
 //					// Verify
-//					Ptr<Image> Test = ImagePixelFormat(CompressionQuality, Result, EImageFormat::L_UByte);
+//					TSharedPtr<FImage> Test = ImagePixelFormat(CompressionQuality, Result, EImageFormat::L_UByte);
 //					check(Test->DataStorage == TempBase->DataStorage);
 //				}
 //#endif
@@ -469,7 +469,7 @@ void FImageOperator::ImagePixelFormat(bool& bOutSuccess, int32 CompressionQualit
 //				if (bOutSuccess)
 //				{
 //					// Verify
-//					Ptr<Image> Test = ImagePixelFormat(CompressionQuality, Result, EImageFormat::L_UByte);
+//					TSharedPtr<FImage> Test = ImagePixelFormat(CompressionQuality, Result, EImageFormat::L_UByte);
 //					check(Test->DataStorage == Base->DataStorage);
 //				}
 //#endif
@@ -482,14 +482,14 @@ void FImageOperator::ImagePixelFormat(bool& bOutSuccess, int32 CompressionQualit
 			case EImageFormat::BGRA_UByte:
 			{
 				check(bFormatAllLODs);
-				Ptr<Image> TempBase = ImagePixelFormat(CompressionQuality, Base, EImageFormat::L_UByte, -1);
-				CompressRLE_L1(TempBase.get(), Result);
+				TSharedPtr<FImage> TempBase = ImagePixelFormat(CompressionQuality, Base, EImageFormat::L_UByte, -1);
+				CompressRLE_L1(TempBase.Get(), Result);
 
 //#ifdef MUTABLE_DEBUG_RLE
 //				if (bOutSuccess)
 //				{
 //					// Verify
-//					Ptr<Image> Test = ImagePixelFormat(CompressionQuality, Result, EImageFormat::L_UByte);
+//					TSharedPtr<FImage> Test = ImagePixelFormat(CompressionQuality, Result, EImageFormat::L_UByte);
 //					check(Test->DataStorage == TempBase->DataStorage);
 //				}
 //#endif
@@ -575,9 +575,9 @@ void FImageOperator::ImagePixelFormat(bool& bOutSuccess, int32 CompressionQualit
 			{
 				check(bFormatAllLODs);
 
-				Ptr<Image> TempBase = ImagePixelFormat(CompressionQuality, Base, EImageFormat::RGBA_UByte, -1);
+				TSharedPtr<FImage> TempBase = ImagePixelFormat(CompressionQuality, Base, EImageFormat::RGBA_UByte, -1);
 
-				ImagePixelFormatInternal::ProcessBatchedHelper(Result, TempBase.get(), 3, 4, ResultLODBegin, BaseLODBegin, NumLODs,
+				ImagePixelFormatInternal::ProcessBatchedHelper(Result, TempBase.Get(), 3, 4, ResultLODBegin, BaseLODBegin, NumLODs,
 				[](uint8* DestBuf, const uint8* BaseBuf, int32 NumElems)
 				{
 					for (int32 I = 0; I < NumElems; ++I)
@@ -934,9 +934,9 @@ void FImageOperator::ImagePixelFormat(bool& bOutSuccess, int32 CompressionQualit
 			case EImageFormat::RGB_UByteRLE:
 			{
 				check(bFormatAllLODs);
-				Ptr<Image> TempBase = ImagePixelFormat(CompressionQuality, Base, EImageFormat::RGB_UByte);
+				TSharedPtr<FImage> TempBase = ImagePixelFormat(CompressionQuality, Base, EImageFormat::RGB_UByte);
 
-				ImagePixelFormatInternal::ProcessBatchedHelper(Result, TempBase.get(), 4, 3, ResultLODBegin, BaseLODBegin, NumLODs,
+				ImagePixelFormatInternal::ProcessBatchedHelper(Result, TempBase.Get(), 4, 3, ResultLODBegin, BaseLODBegin, NumLODs,
 				[](uint8* DestBuf, const uint8* BaseBuf, int32 NumElems)
 				{
 					for (int32 I = 0; I < NumElems; ++I)
@@ -954,9 +954,9 @@ void FImageOperator::ImagePixelFormat(bool& bOutSuccess, int32 CompressionQualit
 			case EImageFormat::L_UByteRLE:
 			{
 				check(bFormatAllLODs);
-				Ptr<Image> TempBase = ImagePixelFormat(CompressionQuality, Base, EImageFormat::L_UByte);
+				TSharedPtr<FImage> TempBase = ImagePixelFormat(CompressionQuality, Base, EImageFormat::L_UByte);
 
-				ImagePixelFormatInternal::ProcessBatchedHelper(Result, TempBase.get(), 4, 1, ResultLODBegin, BaseLODBegin, NumLODs,
+				ImagePixelFormatInternal::ProcessBatchedHelper(Result, TempBase.Get(), 4, 1, ResultLODBegin, BaseLODBegin, NumLODs,
 				[](uint8* DestBuf, const uint8* BaseBuf, int32 NumElems)
 				{
 					for (int32 I = 0; I < NumElems; ++I)
@@ -997,11 +997,11 @@ void FImageOperator::ImagePixelFormat(bool& bOutSuccess, int32 CompressionQualit
 				check(bFormatAllLODs);
 
 				// \todo: optimise
-				Ptr<Image> TempBase = ImagePixelFormat(CompressionQuality, Base, EImageFormat::RGBA_UByte);
-				CompressRLE_RGBA(TempBase.get(), Result);
+				TSharedPtr<FImage> TempBase = ImagePixelFormat(CompressionQuality, Base, EImageFormat::RGBA_UByte);
+				CompressRLE_RGBA(TempBase.Get(), Result);
 
 				// Test
-				//Ptr<Image> Test = ImagePixelFormat(CompressionQuality, Result, EImageFormat::RGBA_UByte); 
+				//TSharedPtr<FImage> Test = ImagePixelFormat(CompressionQuality, Result, EImageFormat::RGBA_UByte); 
 				//check(Test->DataStorage == Base->DataStorage);
 
 				ReleaseImage(TempBase);
@@ -1013,10 +1013,10 @@ void FImageOperator::ImagePixelFormat(bool& bOutSuccess, int32 CompressionQualit
 				check(bFormatAllLODs);
 
 				// \todo: optimise
-				Ptr<Image> Temp1 = ImagePixelFormat(CompressionQuality, Base, EImageFormat::RGB_UByte);
-				Ptr<Image> Temp2 = ImagePixelFormat(CompressionQuality, Temp1.get(), EImageFormat::RGBA_UByte);
+				TSharedPtr<FImage> Temp1 = ImagePixelFormat(CompressionQuality, Base, EImageFormat::RGB_UByte);
+				TSharedPtr<FImage> Temp2 = ImagePixelFormat(CompressionQuality, Temp1.Get(), EImageFormat::RGBA_UByte);
 				ReleaseImage(Temp1);
-				CompressRLE_RGBA(Temp2.get(), Result);
+				CompressRLE_RGBA(Temp2.Get(), Result);
 				ReleaseImage(Temp2);
 
 				break;
@@ -1041,7 +1041,7 @@ void FImageOperator::ImagePixelFormat(bool& bOutSuccess, int32 CompressionQualit
 				CompressRLE_RGB(Base, Result);
 
 				// Test
-				// Ptr<Image> Test = ImagePixelFormat(CompressionQuality, Result, EImageFormat::RGB_UByte); 
+				// TSharedPtr<FImage> Test = ImagePixelFormat(CompressionQuality, Result, EImageFormat::RGB_UByte); 
 				// check(Test->DataStorage == Base->DataStorage);
 				break;
 			}
@@ -1051,10 +1051,10 @@ void FImageOperator::ImagePixelFormat(bool& bOutSuccess, int32 CompressionQualit
 				check(bFormatAllLODs);
 
 				// \todo: optimise
-				Ptr<Image> Temp1 = ImagePixelFormat(CompressionQuality, Base, EImageFormat::RGBA_UByte);
-				Ptr<Image> Temp2 = ImagePixelFormat(CompressionQuality, Temp1.get(), EImageFormat::RGB_UByte);
+				TSharedPtr<FImage> Temp1 = ImagePixelFormat(CompressionQuality, Base, EImageFormat::RGBA_UByte);
+				TSharedPtr<FImage> Temp2 = ImagePixelFormat(CompressionQuality, Temp1.Get(), EImageFormat::RGB_UByte);
 				ReleaseImage(Temp1);
-				CompressRLE_RGB(Temp2.get(), Result);
+				CompressRLE_RGB(Temp2.Get(), Result);
 				ReleaseImage(Temp2);
 
 				break;
@@ -1086,18 +1086,18 @@ void FImageOperator::ImagePixelFormat(bool& bOutSuccess, int32 CompressionQualit
 			}
 			case EImageFormat::L_UByte:
 			{
-				Ptr<Image> TempBase = ImagePixelFormat(CompressionQuality, Base, EImageFormat::RGB_UByte);
+				TSharedPtr<FImage> TempBase = ImagePixelFormat(CompressionQuality, Base, EImageFormat::RGB_UByte);
 				ImagePixelFormatInternal::BlockCompressionLODRangeHelper(
-						Result, TempBase.get(), ResultLODBegin, BaseLODBegin, NumLODs, CompressionQuality, miro::RGB_to_BC1);
+						Result, TempBase.Get(), ResultLODBegin, BaseLODBegin, NumLODs, CompressionQuality, miro::RGB_to_BC1);
 				ReleaseImage(TempBase);
 
 				break;
 			}
 			case EImageFormat::BC3:
 			{
-				Ptr<Image> TempBase = ImagePixelFormat(CompressionQuality, Base, EImageFormat::RGBA_UByte);
+				TSharedPtr<FImage> TempBase = ImagePixelFormat(CompressionQuality, Base, EImageFormat::RGBA_UByte);
 				ImagePixelFormatInternal::BlockCompressionLODRangeHelper(
-						Result, TempBase.get(), ResultLODBegin, BaseLODBegin, NumLODs, CompressionQuality, miro::RGBA_to_BC1);
+						Result, TempBase.Get(), ResultLODBegin, BaseLODBegin, NumLODs, CompressionQuality, miro::RGBA_to_BC1);
 				ReleaseImage(TempBase);
 
 				break;
@@ -1107,9 +1107,9 @@ void FImageOperator::ImagePixelFormat(bool& bOutSuccess, int32 CompressionQualit
 				UE_LOG(LogMutableCore, Log, TEXT("Image format conversion not implemented. Expensive generic one used. "));
 
 				// Case not implemented, try a generic more expensive operation.
-				Ptr<Image> TempBase = ImagePixelFormat(CompressionQuality, Base, EImageFormat::RGBA_UByte);
+				TSharedPtr<FImage> TempBase = ImagePixelFormat(CompressionQuality, Base, EImageFormat::RGBA_UByte);
 				ImagePixelFormatInternal::BlockCompressionLODRangeHelper(
-						Result, TempBase.get(), ResultLODBegin, BaseLODBegin, NumLODs, CompressionQuality, miro::RGBA_to_BC1);
+						Result, TempBase.Get(), ResultLODBegin, BaseLODBegin, NumLODs, CompressionQuality, miro::RGBA_to_BC1);
 				ReleaseImage(TempBase);
 			}
 
@@ -1133,9 +1133,9 @@ void FImageOperator::ImagePixelFormat(bool& bOutSuccess, int32 CompressionQualit
 				UE_LOG(LogMutableCore, Log, TEXT("Image format conversion not implemented. Expensive generic one used."));
 
 				// Case not implemented, try a generic more expensive operatino.
-				Ptr<Image> TempBase = ImagePixelFormat(CompressionQuality, Base, EImageFormat::RGBA_UByte);
+				TSharedPtr<FImage> TempBase = ImagePixelFormat(CompressionQuality, Base, EImageFormat::RGBA_UByte);
 				ImagePixelFormatInternal::BlockCompressionLODRangeHelper(
-						Result, TempBase.get(), ResultLODBegin, BaseLODBegin, NumLODs, CompressionQuality, miro::RGBA_to_BC2);
+						Result, TempBase.Get(), ResultLODBegin, BaseLODBegin, NumLODs, CompressionQuality, miro::RGBA_to_BC2);
 				ReleaseImage(TempBase);
 			}
 
@@ -1172,9 +1172,9 @@ void FImageOperator::ImagePixelFormat(bool& bOutSuccess, int32 CompressionQualit
 				UE_LOG(LogMutableCore, Log, TEXT("Image format conversion not implemented. Expensive generic one used."));
 
 				// Case not implemented, try a generic more expensive operation.
-				Ptr<Image> TempBase = ImagePixelFormat(CompressionQuality, Base, EImageFormat::RGBA_UByte);
+				TSharedPtr<FImage> TempBase = ImagePixelFormat(CompressionQuality, Base, EImageFormat::RGBA_UByte);
 				ImagePixelFormatInternal::BlockCompressionLODRangeHelper(
-						Result, TempBase.get(), ResultLODBegin, BaseLODBegin, NumLODs, CompressionQuality, miro::RGBA_to_BC3);
+						Result, TempBase.Get(), ResultLODBegin, BaseLODBegin, NumLODs, CompressionQuality, miro::RGBA_to_BC3);
 				ReleaseImage(TempBase);
 			}
 
@@ -1198,9 +1198,9 @@ void FImageOperator::ImagePixelFormat(bool& bOutSuccess, int32 CompressionQualit
 			case EImageFormat::BGRA_UByte:
 			default:
 			{	
-				Ptr<Image> TempBase = ImagePixelFormat(CompressionQuality, Base, EImageFormat::L_UByte, -1);
+				TSharedPtr<FImage> TempBase = ImagePixelFormat(CompressionQuality, Base, EImageFormat::L_UByte, -1);
 				ImagePixelFormatInternal::BlockCompressionLODRangeHelper(
-						Result, TempBase.get(), ResultLODBegin, BaseLODBegin, NumLODs, CompressionQuality, miro::L_to_BC4);
+						Result, TempBase.Get(), ResultLODBegin, BaseLODBegin, NumLODs, CompressionQuality, miro::L_to_BC4);
 				ReleaseImage(TempBase);
 				break;
 			}
@@ -1229,9 +1229,9 @@ void FImageOperator::ImagePixelFormat(bool& bOutSuccess, int32 CompressionQualit
 			{
 				UE_LOG(LogMutableCore, Log, TEXT("Image format conversion not implemented. Expensive generic one used. "));
 
-				Ptr<Image> TempBase = ImagePixelFormat(CompressionQuality, Base, EImageFormat::RGBA_UByte, -1);
+				TSharedPtr<FImage> TempBase = ImagePixelFormat(CompressionQuality, Base, EImageFormat::RGBA_UByte, -1);
 				ImagePixelFormatInternal::BlockCompressionLODRangeHelper(
-						Result, TempBase.get(), ResultLODBegin, BaseLODBegin, NumLODs, CompressionQuality, miro::RGBA_to_BC5);
+						Result, TempBase.Get(), ResultLODBegin, BaseLODBegin, NumLODs, CompressionQuality, miro::RGBA_to_BC5);
 				ReleaseImage(TempBase);
 				break;
 			}
@@ -1272,9 +1272,9 @@ void FImageOperator::ImagePixelFormat(bool& bOutSuccess, int32 CompressionQualit
 			{
 				UE_LOG(LogMutableCore, Log, TEXT("Image format conversion not implemented. Expensive generic one used. "));
 
-				Ptr<Image> TempBase = ImagePixelFormat(CompressionQuality, Base, EImageFormat::RGBA_UByte, -1);
+				TSharedPtr<FImage> TempBase = ImagePixelFormat(CompressionQuality, Base, EImageFormat::RGBA_UByte, -1);
 				ImagePixelFormatInternal::BlockCompressionLODRangeHelper(
-						Result, TempBase.get(), ResultLODBegin, BaseLODBegin, NumLODs, CompressionQuality, miro::RGBA_to_ASTC4x4RGBL);
+						Result, TempBase.Get(), ResultLODBegin, BaseLODBegin, NumLODs, CompressionQuality, miro::RGBA_to_ASTC4x4RGBL);
 				ReleaseImage(TempBase);
 				break;
 			}
@@ -1309,9 +1309,9 @@ void FImageOperator::ImagePixelFormat(bool& bOutSuccess, int32 CompressionQualit
 			{
 				UE_LOG(LogMutableCore, Log, TEXT("Image format conversion not implemented. Expensive generic one used. "));
 
-				Ptr<Image> TempBase = ImagePixelFormat(CompressionQuality, Base, EImageFormat::RGBA_UByte, -1);
+				TSharedPtr<FImage> TempBase = ImagePixelFormat(CompressionQuality, Base, EImageFormat::RGBA_UByte, -1);
 				ImagePixelFormatInternal::BlockCompressionLODRangeHelper(
-						Result, TempBase.get(), ResultLODBegin, BaseLODBegin, NumLODs, CompressionQuality, miro::RGBA_to_ASTC4x4RGBAL);
+						Result, TempBase.Get(), ResultLODBegin, BaseLODBegin, NumLODs, CompressionQuality, miro::RGBA_to_ASTC4x4RGBAL);
 				ReleaseImage(TempBase);
 				break;
 			}
@@ -1347,9 +1347,9 @@ void FImageOperator::ImagePixelFormat(bool& bOutSuccess, int32 CompressionQualit
 			{
 				UE_LOG(LogMutableCore, Log, TEXT("Image format conversion not implemented. Expensive generic one used. "));
 
-				Ptr<Image> TempBase = ImagePixelFormat(CompressionQuality, Base, EImageFormat::RGBA_UByte, -1);
+				TSharedPtr<FImage> TempBase = ImagePixelFormat(CompressionQuality, Base, EImageFormat::RGBA_UByte, -1);
 				ImagePixelFormatInternal::BlockCompressionLODRangeHelper(
-						Result, TempBase.get(), ResultLODBegin, BaseLODBegin, NumLODs, CompressionQuality, miro::RGBA_to_ASTC4x4RGL);
+						Result, TempBase.Get(), ResultLODBegin, BaseLODBegin, NumLODs, CompressionQuality, miro::RGBA_to_ASTC4x4RGL);
 				ReleaseImage(TempBase);
 				break;
 			}
@@ -1391,9 +1391,9 @@ void FImageOperator::ImagePixelFormat(bool& bOutSuccess, int32 CompressionQualit
 			{
 				UE_LOG(LogMutableCore, Log, TEXT("Image format conversion not implemented. Expensive generic one used. "));
 
-				Ptr<Image> TempBase = ImagePixelFormat(CompressionQuality, Base, EImageFormat::RGBA_UByte, -1);
+				TSharedPtr<FImage> TempBase = ImagePixelFormat(CompressionQuality, Base, EImageFormat::RGBA_UByte, -1);
 				ImagePixelFormatInternal::BlockCompressionLODRangeHelper(
-						Result, TempBase.get(), ResultLODBegin, BaseLODBegin, NumLODs, CompressionQuality, miro::RGBA_to_ASTC6x6RGBL);
+						Result, TempBase.Get(), ResultLODBegin, BaseLODBegin, NumLODs, CompressionQuality, miro::RGBA_to_ASTC6x6RGBL);
 				ReleaseImage(TempBase);
 				break;
 			}
@@ -1428,9 +1428,9 @@ void FImageOperator::ImagePixelFormat(bool& bOutSuccess, int32 CompressionQualit
 			{
 				UE_LOG(LogMutableCore, Log, TEXT("Image format conversion not implemented. Expensive generic one used. "));
 
-				Ptr<Image> TempBase = ImagePixelFormat(CompressionQuality, Base, EImageFormat::RGBA_UByte, -1);
+				TSharedPtr<FImage> TempBase = ImagePixelFormat(CompressionQuality, Base, EImageFormat::RGBA_UByte, -1);
 				ImagePixelFormatInternal::BlockCompressionLODRangeHelper(
-						Result, TempBase.get(), ResultLODBegin, BaseLODBegin, NumLODs, CompressionQuality, miro::RGBA_to_ASTC6x6RGBAL);
+						Result, TempBase.Get(), ResultLODBegin, BaseLODBegin, NumLODs, CompressionQuality, miro::RGBA_to_ASTC6x6RGBAL);
 				ReleaseImage(TempBase);
 				break;
 			}
@@ -1467,9 +1467,9 @@ void FImageOperator::ImagePixelFormat(bool& bOutSuccess, int32 CompressionQualit
 			{
 				UE_LOG(LogMutableCore, Log, TEXT("Image format conversion not implemented. Expensive generic one used. "));
 
-				Ptr<Image> TempBase = ImagePixelFormat(CompressionQuality, Base, EImageFormat::RGBA_UByte, -1);
+				TSharedPtr<FImage> TempBase = ImagePixelFormat(CompressionQuality, Base, EImageFormat::RGBA_UByte, -1);
 				ImagePixelFormatInternal::BlockCompressionLODRangeHelper(
-						Result, TempBase.get(), ResultLODBegin, BaseLODBegin, NumLODs, CompressionQuality, miro::RGBA_to_ASTC6x6RGL);
+						Result, TempBase.Get(), ResultLODBegin, BaseLODBegin, NumLODs, CompressionQuality, miro::RGBA_to_ASTC6x6RGL);
 				ReleaseImage(TempBase);
 				break;
 			}
@@ -1510,9 +1510,9 @@ void FImageOperator::ImagePixelFormat(bool& bOutSuccess, int32 CompressionQualit
 			{
 				UE_LOG(LogMutableCore, Log, TEXT("Image format conversion not implemented. Expensive generic one used. "));
 
-				Ptr<Image> TempBase = ImagePixelFormat(CompressionQuality, Base, EImageFormat::RGBA_UByte, -1);
+				TSharedPtr<FImage> TempBase = ImagePixelFormat(CompressionQuality, Base, EImageFormat::RGBA_UByte, -1);
 				ImagePixelFormatInternal::BlockCompressionLODRangeHelper(
-						Result, TempBase.get(), ResultLODBegin, BaseLODBegin, NumLODs, CompressionQuality, miro::RGBA_to_ASTC8x8RGBL);
+						Result, TempBase.Get(), ResultLODBegin, BaseLODBegin, NumLODs, CompressionQuality, miro::RGBA_to_ASTC8x8RGBL);
 				ReleaseImage(TempBase);
 				break;
 			}
@@ -1547,9 +1547,9 @@ void FImageOperator::ImagePixelFormat(bool& bOutSuccess, int32 CompressionQualit
 			{
 				UE_LOG(LogMutableCore, Log, TEXT("Image format conversion not implemented. Expensive generic one used. "));
 
-				Ptr<Image> TempBase = ImagePixelFormat(CompressionQuality, Base, EImageFormat::RGBA_UByte, -1);
+				TSharedPtr<FImage> TempBase = ImagePixelFormat(CompressionQuality, Base, EImageFormat::RGBA_UByte, -1);
 				ImagePixelFormatInternal::BlockCompressionLODRangeHelper(
-						Result, TempBase.get(), ResultLODBegin, BaseLODBegin, NumLODs, CompressionQuality, miro::RGBA_to_ASTC8x8RGBAL);
+						Result, TempBase.Get(), ResultLODBegin, BaseLODBegin, NumLODs, CompressionQuality, miro::RGBA_to_ASTC8x8RGBAL);
 				ReleaseImage(TempBase);
 				break;
 			}
@@ -1578,9 +1578,9 @@ void FImageOperator::ImagePixelFormat(bool& bOutSuccess, int32 CompressionQualit
 			{
 				UE_LOG(LogMutableCore, Log, TEXT("Image format conversion not implemented. Expensive generic one used. "));
 
-				Ptr<Image> TempBase = ImagePixelFormat(CompressionQuality, Base, EImageFormat::RGBA_UByte, -1);
+				TSharedPtr<FImage> TempBase = ImagePixelFormat(CompressionQuality, Base, EImageFormat::RGBA_UByte, -1);
 				ImagePixelFormatInternal::BlockCompressionLODRangeHelper(
-						Result, TempBase.get(), ResultLODBegin, BaseLODBegin, NumLODs, CompressionQuality, miro::RGBA_to_ASTC8x8RGL);
+						Result, TempBase.Get(), ResultLODBegin, BaseLODBegin, NumLODs, CompressionQuality, miro::RGBA_to_ASTC8x8RGL);
 				ReleaseImage(TempBase);
 				break;
 			}
diff --git a/Engine/Plugins/Mutable/Source/MutableRuntime/Private/MuR/OpImageProject.cpp b/Engine/Plugins/Mutable/Source/MutableRuntime/Private/MuR/OpImageProject.cpp
--- a/Engine/Plugins/Mutable/Source/MutableRuntime/Private/MuR/OpImageProject.cpp
+++ b/Engine/Plugins/Mutable/Source/MutableRuntime/Private/MuR/OpImageProject.cpp
@@ -51,10 +51,10 @@ namespace Private
 {
 	struct FImageRasterInvokeArgs
 	{
-		const Mesh* MeshPtr;
-		Image* ImagePtr;
-		const Image* SourcePtr;
-		const Image* MaskPtr;
+		const FMesh* MeshPtr;
+		FImage* ImagePtr;
+		const FImage* SourcePtr;
+		const FImage* MaskPtr;
 		bool bIsRGBFadingEnabled;
 		bool bIsAlphaFadingEnabled;
 		ESamplingMethod SamplingMethod;
@@ -174,7 +174,7 @@ namespace Private
 
 	public:
 		static FProjectedPixelProcessorContext MakeContext(
-			const Image* Source,
+			const FImage* Source,
 			const uint8* InTargetData,
 			const uint8* InMaskData,
 			bool bInIsRGBFadingEnabled, bool bInIsAlphaFadingEnabled,
@@ -626,7 +626,7 @@ namespace Private
 	};
 	
 	//-------------------------------------------------------------------------------------------------
-	FORCENOINLINE void ImageRasterProjected_Optimised(const Mesh* pMesh, Image* pImage,
+	FORCENOINLINE void ImageRasterProjected_Optimised(const FMesh* pMesh, FImage* pImage,
 		TTriangleRasterPixelProcRefType<4> PixelProc, 
 		float FadeEnd,
 		UE::Math::TIntVector2<uint16> CropMin, UE::Math::TIntVector2<uint16> UncroppedSize,
@@ -738,7 +738,7 @@ namespace Private
 				}
 			}
 
-			pImage->m_flags |= Image::IF_HAS_RELEVANCY_MAP;
+			pImage->Flags |= FImage::IF_HAS_RELEVANCY_MAP;
 			pImage->RelevancyMinY = uint16(FMath::FloorToFloat(MinY));
 			pImage->RelevancyMaxY = uint16(FMath::Min( int32(FMath::CeilToFloat(MaxY)), SizeY - 1));
 		}
@@ -746,7 +746,7 @@ namespace Private
 
 
 	//-------------------------------------------------------------------------------------------------
-	FORCENOINLINE void ImageRasterProjected_OptimisedWrapping(const Mesh* pMesh, Image* pImage, 
+	FORCENOINLINE void ImageRasterProjected_OptimisedWrapping(const FMesh* pMesh, FImage* pImage, 
 		TTriangleRasterPixelProcRefType<4> PixelProc, 
 		float FadeEnd,
 		uint64 BlockId,
@@ -928,9 +928,9 @@ namespace Private
 
 
 //-------------------------------------------------------------------------------------------------
-void ImageRasterProjectedPlanar(const Mesh* pMesh, Image* pImage,
-	const Image* pSource,
-	const Image* pMask,
+void ImageRasterProjectedPlanar(const FMesh* pMesh, FImage* pImage,
+	const FImage* pSource,
+	const FImage* pMask,
 	bool bIsRGBFadingEnabled, bool bIsAlphaFadingEnabled,
 	ESamplingMethod SamplingMethod,
 	float FadeStart, float FadeEnd, float MipInterpolationFactor,
@@ -1143,8 +1143,8 @@ void ImageRasterProjectedPlanar(const Mesh* pMesh, Image* pImage,
 
 
 //-------------------------------------------------------------------------------------------------
-void ImageRasterProjectedWrapping( const Mesh* pMesh, Image* pImage,
-	const Image* pSource, const Image* pMask,
+void ImageRasterProjectedWrapping( const FMesh* pMesh, FImage* pImage,
+	const FImage* pSource, const FImage* pMask,
 	bool bIsRGBFadingEnabled, bool bIsAlphaFadingEnabled,
 	ESamplingMethod SamplingMethod,
 	float FadeStart, float FadeEnd, float MipInterpolationFactor,
@@ -1356,8 +1356,8 @@ void ImageRasterProjectedWrapping( const Mesh* pMesh, Image* pImage,
 }
 
 //-------------------------------------------------------------------------------------------------
-void ImageRasterProjectedCylindrical( const Mesh* pMesh, Image* pImage,
-	const Image* pSource, const Image* pMask,
+void ImageRasterProjectedCylindrical( const FMesh* pMesh, FImage* pImage,
+	const FImage* pSource, const FImage* pMask,
 	bool bIsRGBFadingEnabled, bool bIsAlphaFadingEnabled,
 	ESamplingMethod SamplingMethod,
 	float FadeStart, float FadeEnd, float MipInterpolationFactor,
@@ -1569,7 +1569,7 @@ void ImageRasterProjectedCylindrical( const Mesh* pMesh, Image* pImage,
 }
 
 float ComputeProjectedFootprintBestMip(
-	const Mesh* pMesh, const FProjector& Projector, const FVector2f& TargetSize, const FVector2f& SourceSize)
+	const FMesh* pMesh, const FProjector& Projector, const FVector2f& TargetSize, const FVector2f& SourceSize)
 {
 	// Compute projected mesh footprint on the traget image and extract a optimal mip for source image.
 	const bool bIsPlanarProjection = Projector.type == EProjectorType::Planar;
@@ -1694,7 +1694,7 @@ float ComputeProjectedFootprintBestMip(
 UE_DISABLE_OPTIMIZATION
 
 #undef assert
-int* assert_aux = 0;
+int32* assert_aux = 0;
 #define assert(x) if((x) == 0) assert_aux[0] = 1;
 #endif
 
@@ -1704,7 +1704,7 @@ constexpr float vert_collapse_eps = 0.0001f;
 //---------------------------------------------------------------------------------------------
 //! Create a map from vertices into vertices, collapsing vertices that have the same position
 //---------------------------------------------------------------------------------------------
-inline void MeshCreateCollapsedVertexMap( const Mesh* pMesh,
+inline void MeshCreateCollapsedVertexMap( const FMesh* pMesh,
 	TArray<int32>& CollapsedVertices,
 	TMultiMap<int32, int32>& CollapsedVerticesMap,
 	TArray<FVector3f>& OutVertices)
@@ -1763,8 +1763,8 @@ inline void MeshCreateCollapsedVertexMap( const Mesh* pMesh,
 
 struct AdjacentFaces
 {
-	int faces[3];
-	int newVertices[3];
+	int32 faces[3];
+	int32 newVertices[3];
 	bool changesUVIsland[3];
 
 	AdjacentFaces()
@@ -1782,13 +1782,13 @@ struct AdjacentFaces
 		changesUVIsland[2] = false;
 	}
 
-	void addConnectedFace(int newConnectedFace, int newVertex, bool changesUVIslandParam)
+	void addConnectedFace(int32 newConnectedFace, int32 newVertex, bool changesUVIslandParam)
 	{
 #ifdef DEBUG_PROJECTION
 		assert( newConnectedFace >= 0 );
 #endif
 
-		for (int i = 0; i < 3; ++i)
+		for (int32 i = 0; i < 3; ++i)
 		{
 			if (faces[i] == newConnectedFace)
 			{
@@ -1796,7 +1796,7 @@ struct AdjacentFaces
 			}
 		}
 
-		for (int i = 0; i < 3; ++i)
+		for (int32 i = 0; i < 3; ++i)
 		{
 			if (faces[i] == -1)
 			{
@@ -1886,11 +1886,11 @@ float getTriangleRatio(const FVector2f& a, const FVector2f& b, const FVector2f&
 
 struct NeighborFace
 {
-	int neighborFace;
-	int newVertex;
-	int previousFace;
-	int numUVIslandChanges = 0;
-	int step;
+	int32 neighborFace;
+	int32 newVertex;
+	int32 previousFace;
+	int32 numUVIslandChanges = 0;
+	int32 step;
 	bool changesUVIsland;
 
 	friend bool operator <(const NeighborFace& a, const NeighborFace& b)
@@ -1905,12 +1905,12 @@ struct NeighborFace
 };
 
 
-void getEdgeHorizontalLength( int edgeVert0, int edgeVert1, int opposedVert,
+void getEdgeHorizontalLength( int32 edgeVert0, int32 edgeVert1, int32 opposedVert,
                               const FOptimizedVertex* pVertices,
                               float &out_uvSpaceLen, float &out_objSpaceLen,
                               float& out_midEdgePointFraction)
 {
-	const int oldVertices[2] = { edgeVert0, edgeVert1 };
+	const int32 oldVertices[2] = { edgeVert0, edgeVert1 };
 
 	const FVector2f& edge0_oldUVSpace = pVertices[oldVertices[0]].Uv;
 	const FVector2f& edge1_oldUVSpace = pVertices[oldVertices[1]].Uv;
@@ -1984,13 +1984,13 @@ static_assert( sizeof(PROJECTED_VERTEX)==16, "Unexpected struct size" );
 
 
 //-------------------------------------------------------------------------------------------------
-void MeshProject_Optimised_Planar(const FOptimizedVertex* pVertices, int vertexCount,
-                                  const uint32* pIndices, int faceCount,
+void MeshProject_Optimised_Planar(const FOptimizedVertex* pVertices, int32 vertexCount,
+                                  const uint32* pIndices, int32 faceCount,
                                   const FVector3f& projectorPosition, const FVector3f& projectorDirection,
                                   const FVector3f& projectorSide, const FVector3f& projectorUp,
                                   const FVector3f& projectorScale,
-                                  FOptimizedVertex* pResultVertices, int& currentVertex,
-                                  uint32* pResultIndices, int& currentIndex)
+                                  FOptimizedVertex* pResultVertices, int32& currentVertex,
+                                  uint32* pResultIndices, int32& currentIndex)
 {
 	MUTABLE_CPUPROFILER_SCOPE(MeshProject_Optimised_Planar)
 
@@ -2000,7 +2000,7 @@ void MeshProject_Optimised_Planar(const FOptimizedVertex* pVertices, int vertexC
 	TArray<PROJECTED_VERTEX> projectedPositions;
 	projectedPositions.SetNumZeroed(vertexCount);
 
-    for ( int v=0; v<vertexCount; ++v )
+    for ( int32 v=0; v<vertexCount; ++v )
     {
         float x = FVector3f::DotProduct(pVertices[v].Position - projectorPosition, projectorSide) / projectorScale[0] + 0.5f;
         float y = FVector3f::DotProduct(pVertices[v].Position - projectorPosition, projectorUp) / projectorScale[1] + 0.5f;
@@ -2022,11 +2022,11 @@ void MeshProject_Optimised_Planar(const FOptimizedVertex* pVertices, int vertexC
     }
 
     // Iterate the faces
-    for ( int f=0; f<faceCount; ++f )
+    for ( int32 f=0; f<faceCount; ++f )
     {
-        int i0 = pIndices[f*3+0];
-        int i1 = pIndices[f*3+1];
-        int i2 = pIndices[f*3+2];
+        int32 i0 = pIndices[f*3+0];
+        int32 i1 = pIndices[f*3+1];
+        int32 i2 = pIndices[f*3+2];
 
         // Approximate test: discard the triangle if any of the 6 planes entirely discards all the vertices.
         // This will let some triangles through that could be discarded with a precise test.
@@ -2040,7 +2040,7 @@ void MeshProject_Optimised_Planar(const FOptimizedVertex* pVertices, int vertexC
         if ( !discarded )
         {
             // This face is required.
-            for (int v=0;v<3;++v)
+            for (int32 v=0;v<3;++v)
             {
                 uint32 i = pIndices[f*3+v];
                 if (oldToNewVertex[i]<0)
@@ -2068,13 +2068,13 @@ void MeshProject_Optimised_Planar(const FOptimizedVertex* pVertices, int vertexC
 
 
 //-------------------------------------------------------------------------------------------------
-void MeshProject_Optimised_Cylindrical(const FOptimizedVertex* pVertices, int vertexCount,
-                                       const uint32* pIndices, int faceCount,
+void MeshProject_Optimised_Cylindrical(const FOptimizedVertex* pVertices, int32 vertexCount,
+                                       const uint32* pIndices, int32 faceCount,
                                        const FVector3f& projectorPosition, const FVector3f& projectorDirection,
                                        const FVector3f& projectorSide, const FVector3f& projectorUp,
                                        const FVector3f& projectorScale,
-                                       FOptimizedVertex* pResultVertices, int& currentVertex,
-                                       uint32* pResultIndices, int& currentIndex)
+                                       FOptimizedVertex* pResultVertices, int32& currentVertex,
+                                       uint32* pResultIndices, int32& currentIndex)
 {
 	MUTABLE_CPUPROFILER_SCOPE(MeshProject_Optimised_Cylindrical)
 
@@ -2089,7 +2089,7 @@ void MeshProject_Optimised_Cylindrical(const FOptimizedVertex* pVertices, int ve
     float height = projectorScale[0];
 	FMatrix44f WorldToCylinder(projectorDirection, projectorSide, projectorUp,FVector3f(0,0,0));
 
-    for ( int v=0; v<vertexCount; ++v )
+    for ( int32 v=0; v<vertexCount; ++v )
     {
         // Cylinder is along the X axis
 
@@ -2112,11 +2112,11 @@ void MeshProject_Optimised_Cylindrical(const FOptimizedVertex* pVertices, int ve
     }
 
     // Iterate the faces
-    for ( int f=0; f<faceCount; ++f )
+    for ( int32 f=0; f<faceCount; ++f )
     {
-        int i0 = pIndices[f*3+0];
-        int i1 = pIndices[f*3+1];
-        int i2 = pIndices[f*3+2];
+        int32 i0 = pIndices[f*3+0];
+        int32 i1 = pIndices[f*3+1];
+        int32 i2 = pIndices[f*3+2];
 
         // Approximate test: discard the triangle if any of the 3 "planes" (top, bottom and radius)
         // entirely discards all the vertices.
@@ -2133,7 +2133,7 @@ void MeshProject_Optimised_Cylindrical(const FOptimizedVertex* pVertices, int ve
         if ( !discarded )
         {
             // This face is required.
-            for (int v=0;v<3;++v)
+            for (int32 v=0;v<3;++v)
             {
                 uint32 i = pIndices[f*3+v];
                 if (oldToNewVertex[i]<0)
@@ -2161,17 +2161,17 @@ void MeshProject_Optimised_Cylindrical(const FOptimizedVertex* pVertices, int ve
 }
 
 //-------------------------------------------------------------------------------------------------
-void MeshProject_Optimised_Wrapping(const Mesh* pMesh,
+void MeshProject_Optimised_Wrapping(const FMesh* pMesh,
                                     const FVector3f& projectorPosition, const FVector3f& projectorDirection,
                                     const FVector3f& projectorSide, const FVector3f& projectorUp,
                                     const FVector3f& projectorScale,
                                     FOptimizedVertex* pResultVertices, uint64* pResultLayoutBlockIds, int32& currentVertex,
-                                    uint32* pResultIndices, int& currentIndex)
+                                    uint32* pResultIndices, int32& currentIndex)
 {
 	MUTABLE_CPUPROFILER_SCOPE(MeshProject_Optimised_Wrapping)
     
 		// Get the vertices
-    int vertexCount = pMesh->GetVertexCount();
+    int32 vertexCount = pMesh->GetVertexCount();
 	check(pMesh->GetVertexBuffers().GetElementSize(0) == sizeof(FOptimizedVertex));
 	check(pMesh->GetVertexBuffers().Buffers.Num() == 2);
 	const FOptimizedVertex* pVertices = reinterpret_cast<const FOptimizedVertex*>(pMesh->GetVertexBuffers().GetBufferData(0));
@@ -2181,7 +2181,7 @@ void MeshProject_Optimised_Wrapping(const Mesh* pMesh,
     // Get the indices
     check(pMesh->GetIndexBuffers().GetElementSize(0) == 4);
 	const uint32* pIndices = reinterpret_cast<const uint32*>(pMesh->GetIndexBuffers().GetBufferData(0));
-    int faceCount = pMesh->GetFaceCount();
+    int32 faceCount = pMesh->GetFaceCount();
 
 	TArray<PROJECTED_VERTEX> projectedPositions;
 	projectedPositions.SetNum(vertexCount);
@@ -2190,18 +2190,18 @@ void MeshProject_Optimised_Wrapping(const Mesh* pMesh,
     const float maxDist = 100000.f; // TODO: This should be proportional to the mesh bounding box size.
     float min_t = maxDist;
     float rayLength = maxDist;
-    int intersectedFace = -1;
-	TSet<int> processedVertices;
+    int32 intersectedFace = -1;
+	TSet<int32> processedVertices;
     FVector3f projectionPlaneNormal;
     FVector3f out_intersection;
 
 	TArray<AdjacentFaces> faceConnectivity;
 	faceConnectivity.SetNum(faceCount);
-    TSet<int> processedFaces;
+    TSet<int32> processedFaces;
 	processedFaces.Reserve(faceCount/4);
-	TSet<int> discardedWrapAroundFaces;
+	TSet<int32> discardedWrapAroundFaces;
 	discardedWrapAroundFaces.Reserve(faceCount/16);
-	TArray<int> faceStep;
+	TArray<int32> faceStep;
 	faceStep.SetNum(faceCount);
 
     // Map vertices to the one they are collapsed to because they are very similar, if they aren't collapsed then they are mapped to themselves
@@ -2211,14 +2211,14 @@ void MeshProject_Optimised_Wrapping(const Mesh* pMesh,
     MeshCreateCollapsedVertexMap(pMesh, collapsedVertexMap, collapsedVertsMap, vertices);
 
     // Used to speed up connectivity building
-	TMultiMap<int, int> vertToFacesMap;
+	TMultiMap<int32, int32> vertToFacesMap;
 	vertToFacesMap.Reserve(vertexCount);
 
-    for (int f = 0; f < faceCount; ++f)
+    for (int32 f = 0; f < faceCount; ++f)
     {
-        int i0 = pIndices[f * 3 + 0];
-        int i1 = pIndices[f * 3 + 1];
-        int i2 = pIndices[f * 3 + 2];
+        int32 i0 = pIndices[f * 3 + 0];
+        int32 i1 = pIndices[f * 3 + 1];
+        int32 i2 = pIndices[f * 3 + 2];
 
         vertToFacesMap.Add(collapsedVertexMap[i0], f);
         vertToFacesMap.Add(collapsedVertexMap[i1], f);
@@ -2230,11 +2230,11 @@ void MeshProject_Optimised_Wrapping(const Mesh* pMesh,
 
     // Trace a ray in the projection direction to find the face that will be projected planarly and be the root of the unfolding
     // Also build face connectivity information
-    for (int f = 0; f < faceCount; ++f)
+    for (int32 f = 0; f < faceCount; ++f)
     {
-        int i0 = pIndices[f * 3 + 0];
-        int i1 = pIndices[f * 3 + 1];
-        int i2 = pIndices[f * 3 + 2];
+        int32 i0 = pIndices[f * 3 + 0];
+        int32 i1 = pIndices[f * 3 + 1];
+        int32 i2 = pIndices[f * 3 + 2];
 
 		FVector3f rayStart = projectorPosition;
 		FVector3f rayEnd = projectorPosition + projectorDirection * maxDist;
@@ -2261,32 +2261,32 @@ void MeshProject_Optimised_Wrapping(const Mesh* pMesh,
         }
 
         // Face connectivity info
-        for (int i = 0; i < 3; ++i)
+        for (int32 i = 0; i < 3; ++i)
         {
-            int v = collapsedVertexMap[pIndices[f * 3 + i]];
-            //int v = pIndices[f * 3 + i];
+            int32 v = collapsedVertexMap[pIndices[f * 3 + i]];
+            //int32 v = pIndices[f * 3 + i];
 
-			TArray<int,TInlineAllocator<32>> FoundValues;
+			TArray<int32,TInlineAllocator<32>> FoundValues;
 			vertToFacesMap.MultiFind(v, FoundValues);
-            for (int f2 : FoundValues )
+            for (int32 f2 : FoundValues )
             {
                 if (f != f2 && f2 >= 0)
                 {
-                    int commonVertices = 0;
-                    int commonVerticesDifferentIsland = 0;
+                    int32 commonVertices = 0;
+                    int32 commonVerticesDifferentIsland = 0;
                     bool commonMask[3] = { false, false, false };
 
-                    for (int ii = 0; ii < 3; ++ii)
+                    for (int32 ii = 0; ii < 3; ++ii)
                     {
-                        int f_vertex_orig = pIndices[f * 3 + ii];
-                        int f_vertex = collapsedVertexMap[f_vertex_orig];
-                        //int f_vertex = pIndices[f * 3 + ii];
+                        int32 f_vertex_orig = pIndices[f * 3 + ii];
+                        int32 f_vertex = collapsedVertexMap[f_vertex_orig];
+                        //int32 f_vertex = pIndices[f * 3 + ii];
 
-                        for (int j = 0; j < 3; ++j)
+                        for (int32 j = 0; j < 3; ++j)
                         {
-                            int f2_vertex_orig = pIndices[f2 * 3 + j];
-                            int f2_vertex = collapsedVertexMap[f2_vertex_orig];
-                            //int f2_vertex = pIndices[f2 * 3 + j];
+                            int32 f2_vertex_orig = pIndices[f2 * 3 + j];
+                            int32 f2_vertex = collapsedVertexMap[f2_vertex_orig];
+                            //int32 f2_vertex = pIndices[f2 * 3 + j];
 
                             if (f_vertex != -1 && f_vertex == f2_vertex)
                             {
@@ -2305,9 +2305,9 @@ void MeshProject_Optimised_Wrapping(const Mesh* pMesh,
 #endif
                     if (commonVertices == 2)
                     {
-                        int newVertex = -1;
+                        int32 newVertex = -1;
 
-                        for (int j = 0; j < 3; ++j)
+                        for (int32 j = 0; j < 3; ++j)
                         {
                             if (!commonMask[j])
                             {
@@ -2369,7 +2369,7 @@ void MeshProject_Optimised_Wrapping(const Mesh* pMesh,
 		pendingFaces.Reserve(faceCount/64);
         
 		// \TODO: Bool array to speed up?
-		TSet<int> pendingFacesUnique; // Used to quickly check uniqueness in the pendingFaces queue
+		TSet<int32> pendingFacesUnique; // Used to quickly check uniqueness in the pendingFaces queue
 		pendingFacesUnique.Reserve(faceCount / 64);
 
         //FVector3f hitFaceProjectedNormal;
@@ -2385,18 +2385,18 @@ void MeshProject_Optimised_Wrapping(const Mesh* pMesh,
         pendingFaces.HeapPush(neighborFace);
         pendingFacesUnique.Add(intersectedFace);
 
-        int step = 0;
+        int32 step = 0;
         float totalUVAreaCovered = 0.f;
-        TArray<int> oldVertices;
+        TArray<int32> oldVertices;
         oldVertices.Reserve(3);
 
         while (!pendingFaces.IsEmpty())
         {
 			NeighborFace currentFaceStruct;
 			pendingFaces.HeapPop(currentFaceStruct);
-            int currentFace = currentFaceStruct.neighborFace;
+            int32 currentFace = currentFaceStruct.neighborFace;
 #ifdef DEBUG_PROJECTION
-            int newVertexFromQueue = currentFaceStruct.newVertex;
+            int32 newVertexFromQueue = currentFaceStruct.newVertex;
 #endif
             pendingFacesUnique.Remove(currentFace);
 
@@ -2412,9 +2412,9 @@ void MeshProject_Optimised_Wrapping(const Mesh* pMesh,
 #endif
 
                 FVector3f outIntersectionBaricentric;
-                int i0 = pIndices[currentFace * 3 + 0];
-                int i1 = pIndices[currentFace * 3 + 1];
-                int i2 = pIndices[currentFace * 3 + 2];
+                int32 i0 = pIndices[currentFace * 3 + 0];
+                int32 i1 = pIndices[currentFace * 3 + 1];
+                int32 i2 = pIndices[currentFace * 3 + 2];
 				FVector3f v3_0 = pVertices[i0].Position;
 				FVector3f v3_1 = pVertices[i1].Position;
 				FVector3f v3_2 = pVertices[i2].Position;
@@ -2460,14 +2460,14 @@ void MeshProject_Optimised_Wrapping(const Mesh* pMesh,
                 OrigUVs_Origin = outIntersectionUV;
                 GetBarycentricCoords(outIntersectionUV, pVertices[i0].Uv, pVertices[i1].Uv, pVertices[i2].Uv, outIntersectionBaricentric[0], outIntersectionBaricentric[1], outIntersectionBaricentric[2]);
 
-                for (int i = 0; i < 3; ++i)
+                for (int32 i = 0; i < 3; ++i)
                 {
-                    int v = pIndices[currentFace * 3 + i];
+                    int32 v = pIndices[currentFace * 3 + i];
 #ifdef DEBUG_PROJECTION
                     assert(v >= 0 && v < vertexCount);
 #endif
 
-                    int collapsedVert = collapsedVertexMap[v];
+                    int32 collapsedVert = collapsedVertexMap[v];
 
                     processedVertices.Add(collapsedVert);
 
@@ -2494,10 +2494,10 @@ void MeshProject_Optimised_Wrapping(const Mesh* pMesh,
                 //hitFaceProjectedNormal = (cross<float>(v1 - v0, v2 - v0));
                 hitFaceHasPositiveArea = currentTriangleArea >= 0.f;
 
-                for (int i = 0; i < 3; ++i)
+                for (int32 i = 0; i < 3; ++i)
                 {
-                    int v = pIndices[currentFace * 3 + i];
-                    int collapsedVert = collapsedVertexMap[v];
+                    int32 v = pIndices[currentFace * 3 + i];
+                    int32 collapsedVert = collapsedVertexMap[v];
 
                     projectedPositions[collapsedVert].pos0 = ((projectedPositions[collapsedVert].pos0 - 0.5f) * triangleAreaFactor) + 0.5f;
                     projectedPositions[collapsedVert].pos1 = ((projectedPositions[collapsedVert].pos1 - 0.5f) * triangleAreaFactor) + 0.5f;
@@ -2506,9 +2506,9 @@ void MeshProject_Optimised_Wrapping(const Mesh* pMesh,
                                                            projectedPositions[collapsedVert].pos1 >= 0.0f && projectedPositions[collapsedVert].pos1 <= 1.0f;
 
                     // Copy the new info to all the vertices that collapse to the same vertex
-					TArray<int> FoundValues;
+					TArray<int32> FoundValues;
 					collapsedVertsMap.MultiFind(collapsedVert,FoundValues);
-                    for (int otherVert:FoundValues)
+                    for (int32 otherVert:FoundValues)
                     {
                         processedVertices.Add(otherVert);
 
@@ -2534,14 +2534,14 @@ void MeshProject_Optimised_Wrapping(const Mesh* pMesh,
 #ifdef DEBUG_PROJECTION
                 assert(newVertexFromQueue >= 0);
 #endif
-                int newVertex = -1;
-                //int oldVertices[2] = { -1, -1 };
+                int32 newVertex = -1;
+                //int32 oldVertices[2] = { -1, -1 };
                 oldVertices.Empty();
                 bool reverseOrder = false;
 
-                for (int i = 0; i < 3; ++i)
+                for (int32 i = 0; i < 3; ++i)
                 {
-                    int v = pIndices[currentFace * 3 + i];
+                    int32 v = pIndices[currentFace * 3 + i];
 
                     if (!processedVertices.Contains(collapsedVertexMap[v]))
                     {
@@ -2554,7 +2554,7 @@ void MeshProject_Optimised_Wrapping(const Mesh* pMesh,
                     }
                     else
                     {
-                        //int index = oldVertices[0] == -1 ? 0 : 1;
+                        //int32 index = oldVertices[0] == -1 ? 0 : 1;
                         //oldVertices[index] = v;
                         //oldVertices.push_back(collapsedVertexMap[v]);
                         oldVertices.Add(v);
@@ -2563,7 +2563,7 @@ void MeshProject_Optimised_Wrapping(const Mesh* pMesh,
 
                 if (reverseOrder && oldVertices.Num() == 2)
                 {
-                    int aux = oldVertices[0];
+                    int32 aux = oldVertices[0];
                     oldVertices[0] = oldVertices[1];
                     oldVertices[1] = aux;
                 }
@@ -2578,15 +2578,15 @@ void MeshProject_Optimised_Wrapping(const Mesh* pMesh,
                     assert(newVertex < vertexCount && oldVertices[0] < vertexCount && oldVertices[1] < vertexCount);
 #endif
 
-                    int previousFace = currentFaceStruct.previousFace;
+                    int32 previousFace = currentFaceStruct.previousFace;
 #ifdef DEBUG_PROJECTION
                     assert(previousFace >= 0);
                     assert(processedFaces.count(previousFace) == 1);
 #endif
                     // Previous face indices
-                    int pi0 = pIndices[previousFace * 3 + 0];
-                    int pi1 = pIndices[previousFace * 3 + 1];
-                    int pi2 = pIndices[previousFace * 3 + 2];
+                    int32 pi0 = pIndices[previousFace * 3 + 0];
+                    int32 pi1 = pIndices[previousFace * 3 + 1];
+                    int32 pi2 = pIndices[previousFace * 3 + 2];
 
                     // Previous face uv coords
 					FVector2f pv0;
@@ -2598,13 +2598,13 @@ void MeshProject_Optimised_Wrapping(const Mesh* pMesh,
 					FVector2f pv1_proj;
 					FVector2f pv2_proj;
 
-                    int oldVertices_previousFace[2] = { -1, -1 };
-                    int oldVertex = -1;
+                    int32 oldVertices_previousFace[2] = { -1, -1 };
+                    int32 oldVertex = -1;
 
-                    for (int i = 0; i < 3; ++i)
+                    for (int32 i = 0; i < 3; ++i)
                     {
-                        int v = pIndices[previousFace * 3 + i];
-                        int collapsed_v = collapsedVertexMap[v];
+                        int32 v = pIndices[previousFace * 3 + i];
+                        int32 collapsed_v = collapsedVertexMap[v];
 
                         if (collapsed_v == collapsedVertexMap[oldVertices[0]])
                         {
@@ -2706,7 +2706,7 @@ void MeshProject_Optimised_Wrapping(const Mesh* pMesh,
                         projectedPositions[collapsedVertexMap[oldVertices[0]]].xy(), projectedPositions[collapsedVertexMap[oldVertices[1]]].xy()));
 #endif
 
-                    int collapsedNewVertex = collapsedVertexMap[newVertex];
+                    int32 collapsedNewVertex = collapsedVertexMap[newVertex];
                     processedVertices.Add(collapsedNewVertex);
 
                     projectedPositions[collapsedNewVertex].pos0 = newVertex_proj[0];
@@ -2716,9 +2716,9 @@ void MeshProject_Optimised_Wrapping(const Mesh* pMesh,
                                                                 newVertex_proj[1] >= 0.0f && newVertex_proj[1] <= 1.0f;
 
                     // Copy the new info to all the vertices that collapse to the same vertex
-					TArray<int> FoundValues;
+					TArray<int32> FoundValues;
 					collapsedVertsMap.MultiFind(collapsedNewVertex,FoundValues);
-                    for (int otherVert : FoundValues)
+                    for (int32 otherVert : FoundValues)
                     {
                         processedVertices.Add(otherVert);
 
@@ -2740,9 +2740,9 @@ void MeshProject_Optimised_Wrapping(const Mesh* pMesh,
 #endif
                 }
 
-                int i0 = collapsedVertexMap[pIndices[currentFace * 3 + 0]];
-                int i1 = collapsedVertexMap[pIndices[currentFace * 3 + 1]];
-                int i2 = collapsedVertexMap[pIndices[currentFace * 3 + 2]];
+                int32 i0 = collapsedVertexMap[pIndices[currentFace * 3 + 0]];
+                int32 i1 = collapsedVertexMap[pIndices[currentFace * 3 + 1]];
+                int32 i2 = collapsedVertexMap[pIndices[currentFace * 3 + 2]];
 
                 FVector2f v0 = projectedPositions[i0].xy();
                 FVector2f v1 = projectedPositions[i1].xy();
@@ -2760,9 +2760,9 @@ void MeshProject_Optimised_Wrapping(const Mesh* pMesh,
 
             bool anyVertexInUVSpace = false;
 
-            for (int i = 0; i < 3; ++i)
+            for (int32 i = 0; i < 3; ++i)
             {
-                int v = collapsedVertexMap[pIndices[currentFace * 3 + i]];
+                int32 v = collapsedVertexMap[pIndices[currentFace * 3 + i]];
 
                 if (projectedPositions[v].mask3 == 1) // Is it inside?
                 {
@@ -2773,9 +2773,9 @@ void MeshProject_Optimised_Wrapping(const Mesh* pMesh,
 
             bool anyVertexInObjSpaceRange = false;
 
-            for (int i = 0; i < 3; ++i)
+            for (int32 i = 0; i < 3; ++i)
             {
-                int v = pIndices[currentFace * 3 + i];
+                int32 v = pIndices[currentFace * 3 + i];
                 FVector3f r = pVertices[v].Position - out_intersection;
                 float squaredDist = FVector3f::DotProduct(r, r);
 
@@ -2793,10 +2793,10 @@ void MeshProject_Optimised_Wrapping(const Mesh* pMesh,
 #endif
                 totalUVAreaCovered += fabs(currentTriangleArea);
 
-                for(int i = 0; i < 3; ++i)
+                for(int32 i = 0; i < 3; ++i)
                 {
-                    int neighborFace2 = faceConnectivity[currentFace].faces[i];
-                    int newVertex = faceConnectivity[currentFace].newVertices[i];
+                    int32 neighborFace2 = faceConnectivity[currentFace].faces[i];
+                    int32 newVertex = faceConnectivity[currentFace].newVertices[i];
                     bool changesUVIsland = faceConnectivity[currentFace].changesUVIsland[i];
 
                     if (neighborFace2 >= 0 && !processedFaces.Contains(neighborFace2) && pendingFacesUnique.Contains(neighborFace2) == 0)
@@ -2837,7 +2837,7 @@ void MeshProject_Optimised_Wrapping(const Mesh* pMesh,
 	oldToNewVertex.Init(-1,vertexCount);
 
     // Add the projected face
-    for(int f : processedFaces)
+    for(int32 f : processedFaces)
     {
         if (discardedWrapAroundFaces.Contains(f))
         {
@@ -2845,9 +2845,9 @@ void MeshProject_Optimised_Wrapping(const Mesh* pMesh,
         }
 
 #ifdef DEBUG_PROJECTION
-        int i0 = pIndices[f * 3 + 0];
-        int i1 = pIndices[f * 3 + 1];
-        int i2 = pIndices[f * 3 + 2];
+        int32 i0 = pIndices[f * 3 + 0];
+        int32 i1 = pIndices[f * 3 + 1];
+        int32 i2 = pIndices[f * 3 + 2];
 
         assert(processedVertices.count(i0));
         assert(processedVertices.count(i1));
@@ -2864,9 +2864,9 @@ void MeshProject_Optimised_Wrapping(const Mesh* pMesh,
         //     projectedPositions[i2][3] > 0.0f )
         {
             // This face is required.
-            for (int v = 0; v < 3; ++v)
+            for (int32 v = 0; v < 3; ++v)
             {
-                int i = pIndices[f * 3 + v];
+                int32 i = pIndices[f * 3 + v];
 #ifdef DEBUG_PROJECTION
                 assert(i >= 0 && i < vertexCount);
 #endif
@@ -2921,58 +2921,11 @@ void MeshProject_Optimised_Wrapping(const Mesh* pMesh,
             }
         }
     }
-
-#ifdef DEBUG_PROJECTION
-    std::ofstream outfile;
-    outfile.open("[REDACTED_USER_HOME]/Documents/uvs.obj", std::ios::out | std::ios::trunc);
-
-    //for (int i = 0; i < currentVertex; ++i)
-    //{
-    //	LogDebug("v %f %f %f\n", pResultVertices[i].Position[0], pResultVertices[i].Position[1], pResultVertices[i].Position[2]);
-    //}
-
-    //for (int i = 0; i < currentIndex; ++i)
-    //{
-    //	LogDebug("f %d\n", pResultIndices[i]);
-    //}
-
-    if (outfile)
-    {
-        for (int i = 0; i < currentVertex; ++i)
-        {
-            outfile << "v " << pResultVertices[i].Position[0] << " " << pResultVertices[i].Position[1] << " " << pResultVertices[i].Position[2] << "\n";
-        }
-
-        for (int i = 0; i < currentVertex; ++i)
-        {
-            outfile << "vt " << pResultVertices[i].Position[0] << " " << 1.f - pResultVertices[i].Position[1] << "\n";
-        }
-
-        for (int i = 0; i < currentIndex / 3; ++i)
-        {
-            int i0 = pResultIndices[3 * i] + 1;
-            int i1 = pResultIndices[3 * i + 1] + 1;
-            int i2 = pResultIndices[3 * i + 2] + 1;
-
-            outfile << "f " << i0 << "/" << i0 << " " << i1 << "/" << i1 << " " << i2 << "/" << i2 << "\n";
-
-            FVector2f v0 = pResultVertices[i0].Position.xy();
-            FVector2f v1 = pResultVertices[i1].Position.xy();
-            FVector2f v2 = pResultVertices[i2].Position.xy();
-
-            outfile << "# Triangle " << i << ", area = " << getTriangleArea(v0, v1, v2) << "\n";
-        }
-
-        outfile << std::endl;
-
-        outfile.close();
-    }
-#endif
-
 }
 
+
 //-------------------------------------------------------------------------------------------------
-void MeshProject_Optimised(Mesh* Result, const Mesh* pMesh, const FProjector& projector, bool& bOutSuccess)
+void MeshProject_Optimised(FMesh* Result, const FMesh* pMesh, const FProjector& projector, bool& bOutSuccess)
 {
 	MUTABLE_CPUPROFILER_SCOPE(MeshProject_Optimised);
 
@@ -2984,17 +2937,16 @@ void MeshProject_Optimised(Mesh* Result, const Mesh* pMesh, const FProjector& pr
     FVector3f projectorScale( projector.scale[0], projector.scale[1], projector.scale[2] );
 
     // Create with worse case, shrink later.
-    int vertexCount = pMesh->GetVertexCount();
-    int indexCount = pMesh->GetIndexCount();
-    int currentVertex = 0;
-    int currentIndex = 0;
+    int32 vertexCount = pMesh->GetVertexCount();
+    int32 indexCount = pMesh->GetIndexCount();
+    int32 currentVertex = 0;
+    int32 currentIndex = 0;
 
     // At this point the code generation ensures we are working with layout 0
-    const int layout = 0;
+    const int32 layout = 0;
 
     switch (projector.type)
     {
-
     case EProjectorType::Planar:
     case EProjectorType::Cylindrical:
     {
@@ -3011,7 +2963,7 @@ void MeshProject_Optimised(Mesh* Result, const Mesh* pMesh, const FProjector& pr
         // Get the indices
         check(pMesh->GetIndexBuffers().GetElementSize(0) == 4);
 		const uint32* pIndices = reinterpret_cast<const uint32*>(pMesh->GetIndexBuffers().GetBufferData(0));
-        int faceCount = pMesh->GetFaceCount();
+        int32 faceCount = pMesh->GetFaceCount();
 
         if (projector.type==EProjectorType::Planar)
         {
@@ -3080,7 +3032,7 @@ UE_ENABLE_OPTIMIZATION
 
 
 //-------------------------------------------------------------------------------------------------
-void MeshProject(Mesh* Result, const Mesh* pMesh, const FProjector& projector, bool& bOutSuccess)
+void MeshProject(FMesh* Result, const FMesh* pMesh, const FProjector& projector, bool& bOutSuccess)
 {
 	MUTABLE_CPUPROFILER_SCOPE(MeshProject);
 
@@ -3111,7 +3063,7 @@ void MeshProject(Mesh* Result, const Mesh* pMesh, const FProjector& projector, b
 
 
 //-------------------------------------------------------------------------------------------------
-void CreateMeshOptimisedForProjection(Mesh* Result, int32 layout)
+void CreateMeshOptimisedForProjection(FMesh* Result, int32 layout)
 {
     Result->GetVertexBuffers().SetBufferCount( 1 );
     Result->GetIndexBuffers().SetBufferCount( 1 );
@@ -3142,7 +3094,7 @@ void CreateMeshOptimisedForProjection(Mesh* Result, int32 layout)
 
 
 //-------------------------------------------------------------------------------------------------
-void CreateMeshOptimisedForWrappingProjection(Mesh* Result, int32 LayoutIndex)
+void CreateMeshOptimisedForWrappingProjection(FMesh* Result, int32 LayoutIndex)
 {
     Result->GetVertexBuffers().SetBufferCount( 2 );
     Result->GetIndexBuffers().SetBufferCount( 1 );
diff --git a/Engine/Plugins/Mutable/Source/MutableRuntime/Private/MuR/OpImageRasterMesh.h b/Engine/Plugins/Mutable/Source/MutableRuntime/Private/MuR/OpImageRasterMesh.h
--- a/Engine/Plugins/Mutable/Source/MutableRuntime/Private/MuR/OpImageRasterMesh.h
+++ b/Engine/Plugins/Mutable/Source/MutableRuntime/Private/MuR/OpImageRasterMesh.h
@@ -27,7 +27,7 @@ namespace mu
 	};
 
 
-	inline void ImageRasterMesh( const Mesh* pMesh, Image* pImage, int32 LayoutIndex, uint64 BlockId,
+	inline void ImageRasterMesh( const FMesh* pMesh, FImage* pImage, int32 LayoutIndex, uint64 BlockId,
 		UE::Math::TIntVector2<uint16> CropMin, UE::Math::TIntVector2<uint16> UncroppedSize )
 	{
 		MUTABLE_CPUPROFILER_SCOPE(ImageRasterMesh);
diff --git a/Engine/Plugins/Mutable/Source/MutableRuntime/Private/MuR/OpImageResize.cpp b/Engine/Plugins/Mutable/Source/MutableRuntime/Private/MuR/OpImageResize.cpp
--- a/Engine/Plugins/Mutable/Source/MutableRuntime/Private/MuR/OpImageResize.cpp
+++ b/Engine/Plugins/Mutable/Source/MutableRuntime/Private/MuR/OpImageResize.cpp
@@ -28,7 +28,7 @@ namespace mu
 {
 	/** */
 	template< int32 NC >
-	void ImageMagnifyX(Image* Dest, const Image* Base)
+	void ImageMagnifyX(FImage* Dest, const FImage* Base)
 	{
 		int32 BaseSizeX = Base->GetSizeX();
 		int32 DestSizeX = Dest->GetSizeX();
@@ -74,7 +74,7 @@ namespace mu
 
 
 	/** */
-	void ImageMagnifyX(Image* Dest, const Image* Base)
+	void ImageMagnifyX(FImage* Dest, const FImage* Base)
 	{
 		MUTABLE_CPUPROFILER_SCOPE(ImageMagnifyX)
 
@@ -112,7 +112,7 @@ namespace mu
 
 	/** General image minimisation. */
 	template< int32 NC >
-	void ImageMinifyX(Image* Dest, const Image* Base)
+	void ImageMinifyX(FImage* Dest, const FImage* Base)
 	{
 		int32 BaseSizeX = Base->GetSizeX();
 		int32 DestSizeX = Dest->GetSizeX();
@@ -193,7 +193,7 @@ namespace mu
 	//! Optimised for whole factors
 	//---------------------------------------------------------------------------------------------
 	template< int32 NC, int32 FACTOR >
-	void ImageMinifyX_Exact(Image* Dest, const Image* Base)
+	void ImageMinifyX_Exact(FImage* Dest, const FImage* Base)
 	{
 		int32 BaseSizeX = Base->GetSizeX();
 		int32 DestSizeX = Dest->GetSizeX();
@@ -242,7 +242,7 @@ namespace mu
 	}
 
 	template<>
-	void ImageMinifyX_Exact<4, 2>(Image* Dest, const Image* Base)
+	void ImageMinifyX_Exact<4, 2>(FImage* Dest, const FImage* Base)
 	{
 		int32 BaseSizeX = Base->GetSizeX();
 		int32 DestSizeX = Dest->GetSizeX();
@@ -316,7 +316,7 @@ namespace mu
 	//---------------------------------------------------------------------------------------------
 	//! Image minify X version hub.
 	//---------------------------------------------------------------------------------------------
-	void ImageMinifyX(Image* Dest, const Image* Base)
+	void ImageMinifyX(FImage* Dest, const FImage* Base)
 	{
 		MUTABLE_CPUPROFILER_SCOPE(ImageMinifyX)
 
@@ -389,7 +389,7 @@ namespace mu
 	//!
 	//---------------------------------------------------------------------------------------------
 	template< int32 NC >
-	void ImageMagnifyY(Image* Dest, const Image* Base)
+	void ImageMagnifyY(FImage* Dest, const FImage* Base)
 	{
 		if (!Base || !Dest ||
 			!Base->GetSizeX() || !Base->GetSizeY() || !Dest->GetSizeX() || !Dest->GetSizeY())
@@ -467,7 +467,7 @@ namespace mu
 	}
 
 
-	void ImageMagnifyY(Image* Dest, const Image* Base)
+	void ImageMagnifyY(FImage* Dest, const FImage* Base)
 	{
 		check(Dest->GetSizeY() > Base->GetSizeY());
 		check(Dest->GetSizeX() == Base->GetSizeX());
@@ -507,7 +507,7 @@ namespace mu
 	//!
 	//---------------------------------------------------------------------------------------------
 	template< int32 NC >
-	void ImageMinifyY(Image* Dest, const Image* Base)
+	void ImageMinifyY(FImage* Dest, const FImage* Base)
 	{
 		int32 BaseSizeY = Base->GetSizeY();
 		int32 DestSizeY = Dest->GetSizeY();
@@ -589,7 +589,7 @@ namespace mu
 	//!
 	//---------------------------------------------------------------------------------------------
 	template< int32 NC, int32 FACTOR >
-	void ImageMinifyY_Exact(Image* Dest, const Image* Base)
+	void ImageMinifyY_Exact(FImage* Dest, const FImage* Base)
 	{
 		int32 DestSizeY = Dest->GetSizeY();
 		int32 SizeX = Base->GetSizeX();
@@ -637,7 +637,7 @@ namespace mu
 	//---------------------------------------------------------------------------------------------
 	//! Image minify Y version hub.
 	//---------------------------------------------------------------------------------------------
-	void ImageMinifyY(Image* Dest, const Image* Base)
+	void ImageMinifyY(FImage* Dest, const FImage* Base)
 	{
 		check(Dest->GetSizeY() < Base->GetSizeY());
 		check(Dest->GetSizeX() == Base->GetSizeX());
@@ -905,7 +905,7 @@ namespace
 	}
 
 	template<bool bUseVectorImpl>
-	void ImageResizeLinearDownByLessThanTwo(Image* Dest, const Image* InBase)
+	void ImageResizeLinearDownByLessThanTwo(FImage* Dest, const FImage* InBase)
 	{
 		MUTABLE_CPUPROFILER_SCOPE(ImageResizeLinearDownByLessThanTwo);
 
@@ -1136,11 +1136,11 @@ namespace
 	}
 }
 
-	void FImageOperator::ImageResizeLinear(Image* Dest, int32 ImageCompressionQuality, const Image* InBase)
+	void FImageOperator::ImageResizeLinear(FImage* Dest, int32 ImageCompressionQuality, const FImage* InBase)
 	{
 		MUTABLE_CPUPROFILER_SCOPE(ImageResizeLinear);
 
-		check(!(InBase->m_flags & Image::IF_CANNOT_BE_SCALED));
+		check(!(InBase->Flags & FImage::IF_CANNOT_BE_SCALED));
 
 		const FImageSize BaseSize = FImageSize(InBase->GetSizeX(), InBase->GetSizeY());
 		const FImageSize DestSize = FImageSize(Dest->GetSizeX(), Dest->GetSizeY());
@@ -1168,36 +1168,38 @@ namespace
 			const EImageFormat BaseUncompressedFormat = GetUncompressedFormat(InBase->GetFormat());
 			const EImageFormat DestFormat = Dest->GetFormat();
 			
-			Ptr<Image> DecompressedDest = Dest;
+			TSharedPtr<FImage> TempImage;
+			FImage* DecompressedDest = Dest;
 			if (DestFormat != BaseUncompressedFormat)
 			{
-				DecompressedDest = CreateImage(DestSize.X, DestSize.Y, 1, BaseUncompressedFormat, EInitializationType::NotInitialized);
+				TempImage = CreateImage(DestSize.X, DestSize.Y, 1, BaseUncompressedFormat, EInitializationType::NotInitialized);
+				DecompressedDest = TempImage.Get();
 			}
 		
 			if (bEnableVectorImplementationForSmallResizes)
 			{
 				constexpr bool bUseVectorImpl = true;
-				ImageResizeLinearDownByLessThanTwo<bUseVectorImpl>(DecompressedDest.get(), InBase);
+				ImageResizeLinearDownByLessThanTwo<bUseVectorImpl>(DecompressedDest, InBase);
 			}
 			else
 			{
 				constexpr bool bUseVectorImpl = false;
-				ImageResizeLinearDownByLessThanTwo<bUseVectorImpl>(DecompressedDest.get(), InBase);
+				ImageResizeLinearDownByLessThanTwo<bUseVectorImpl>(DecompressedDest, InBase);
 			}
 
 			if (DestFormat != BaseUncompressedFormat)
 			{
 				bool bSuccess = false;
-				ImagePixelFormat(bSuccess, ImageCompressionQuality, Dest, DecompressedDest.get());
+				ImagePixelFormat(bSuccess, ImageCompressionQuality, Dest, DecompressedDest);
 				check(bSuccess);
 
-				ReleaseImage(DecompressedDest);
+				ReleaseImage(TempImage);
 			}
 		}
 		else
 		{
-			ImagePtrConst Base = InBase;
-			ImagePtr BaseReformat;
+			const FImage* Base = InBase;
+			TSharedPtr<FImage> BaseReformat;
 			
 			// Shouldn't happen! But if it does...
 			const EImageFormat SourceFormat = Base->GetFormat();
@@ -1205,50 +1207,50 @@ namespace
 			if (SourceFormat != UncompressedFormat)
 			{
 				BaseReformat = ImagePixelFormat(ImageCompressionQuality, InBase, UncompressedFormat);
-				Base = BaseReformat;
+				Base = BaseReformat.Get();
 			}
 
 			const EImageFormat BaseFormat = Base->GetFormat();
 
 			// First resize X
-			ImagePtr Temp;
+			TSharedPtr<FImage> Temp;
 			if (DestSize[0] > BaseSize[0])
 			{
 				Temp = CreateImage(DestSize[0], BaseSize[1], 1, BaseFormat, EInitializationType::NotInitialized);
-				ImageMagnifyX(Temp.get(), Base.get());
+				ImageMagnifyX(Temp.Get(), Base);
 			}
 			else if (DestSize[0] < BaseSize[0])
 			{
 				Temp = CreateImage(DestSize[0], BaseSize[1], 1, BaseFormat, EInitializationType::NotInitialized);
-				ImageMinifyX(Temp.get(), Base.get());
+				ImageMinifyX(Temp.Get(), Base);
 			}
 			else
 			{
-				Temp = CloneImage(Base.get());
+				Temp = CloneImage(Base);
 			}
 
 			// No lonegr needed. Should be only set if a reformat has been done. 
 			check(!(SourceFormat == UncompressedFormat) || !BaseReformat);
 			if (BaseReformat)
 			{
-				check(Base == BaseReformat);
+				check(Base == BaseReformat.Get());
 				ReleaseImage(BaseReformat);
 				BaseReformat = nullptr;
 				Base = nullptr;
 			}
 
 			// Now resize Y
-			ImagePtr Temp2;
+			TSharedPtr<FImage> Temp2;
 			if (DestSize[1] > BaseSize[1])
 			{
 				Temp2 = CreateImage(DestSize[0], DestSize[1], 1, BaseFormat, EInitializationType::NotInitialized);
-				ImageMagnifyY(Temp2.get(), Temp.get());
+				ImageMagnifyY(Temp2.Get(), Temp.Get());
 				ReleaseImage(Temp);
 			}
 			else if (DestSize[1] < BaseSize[1])
 			{
 				Temp2 = CreateImage(DestSize[0], DestSize[1], 1, BaseFormat, EInitializationType::NotInitialized);
-				ImageMinifyY(Temp2.get(), Temp.get());
+				ImageMinifyY(Temp2.Get(), Temp.Get());
 				ReleaseImage(Temp);
 			}
 			else
@@ -1261,20 +1263,20 @@ namespace
 			// Reset format if it was changed to scale
 			if (SourceFormat != UncompressedFormat)
 			{
-				Ptr<Image> OldTemp = Temp2;
-				Temp2 = ImagePixelFormat(ImageCompressionQuality, Temp2.get(), SourceFormat);
+				TSharedPtr<FImage> OldTemp = Temp2;
+				Temp2 = ImagePixelFormat(ImageCompressionQuality, Temp2.Get(), SourceFormat);
 				ReleaseImage(OldTemp);
 			}
 
-			Dest->CopyMove(Temp2.get());
+			Dest->CopyMove(Temp2.Get());
 			ReleaseImage(Temp2);
 			Temp2 = nullptr;
 		}
 		
 		// Update the relevancy data of the image.
-		if (InBase->m_flags & Image::EImageFlags::IF_HAS_RELEVANCY_MAP)
+		if (InBase->Flags & FImage::EImageFlags::IF_HAS_RELEVANCY_MAP)
 		{
-			Dest->m_flags |= Image::EImageFlags::IF_HAS_RELEVANCY_MAP;
+			Dest->Flags |= FImage::EImageFlags::IF_HAS_RELEVANCY_MAP;
 
 			float FactorY = float(DestSize[1]) / float(BaseSize[1]);
 
diff --git a/Engine/Plugins/Mutable/Source/MutableRuntime/Private/MuR/OpImageSaturate.h b/Engine/Plugins/Mutable/Source/MutableRuntime/Private/MuR/OpImageSaturate.h
--- a/Engine/Plugins/Mutable/Source/MutableRuntime/Private/MuR/OpImageSaturate.h
+++ b/Engine/Plugins/Mutable/Source/MutableRuntime/Private/MuR/OpImageSaturate.h
@@ -9,7 +9,7 @@
 namespace mu
 {
 	template<bool bUseVectorImpl = false>
-	inline void ImageSaturate(Image* AImage, float Factor)
+	inline void ImageSaturate(FImage* AImage, float Factor)
 	{
 		MUTABLE_CPUPROFILER_SCOPE(ImageSaturateInplace);
 
diff --git a/Engine/Plugins/Mutable/Source/MutableRuntime/Private/MuR/OpImageSwizzle.cpp b/Engine/Plugins/Mutable/Source/MutableRuntime/Private/MuR/OpImageSwizzle.cpp
--- a/Engine/Plugins/Mutable/Source/MutableRuntime/Private/MuR/OpImageSwizzle.cpp
+++ b/Engine/Plugins/Mutable/Source/MutableRuntime/Private/MuR/OpImageSwizzle.cpp
@@ -6,7 +6,7 @@
 
 namespace mu
 {
-	void ImageSwizzle(Image* Result, const Ptr<const Image> Sources[], const uint8 Channels[])
+	void ImageSwizzle(FImage* Result, const TSharedPtr<const FImage> Sources[], const uint8 Channels[])
 	{
 		MUTABLE_CPUPROFILER_SCOPE(ImageSwizzle);
 
@@ -44,7 +44,7 @@ namespace mu
 
 		for (int32 Channel = 0; Channel < NumDestChannels; ++Channel)
 		{
-			const Image* Src = Sources[Channel].get();
+			const FImage* Src = Sources[Channel].Get();
 			
 			const int32 DestChannel = Format == EImageFormat::BGRA_UByte && Channel < 3
 								    ? FMath::Abs(Channel - 2)
@@ -218,7 +218,7 @@ namespace mu
 		}
 	}
 
-	Ptr<Image> FImageOperator::ImageSwizzle(EImageFormat Format, const Ptr<const Image> Sources[], const uint8 Channels[])
+	TSharedPtr<FImage> FImageOperator::ImageSwizzle(EImageFormat Format, const TSharedPtr<const FImage> Sources[], const uint8 Channels[])
 	{
 		MUTABLE_CPUPROFILER_SCOPE(ImageSwizzle);
 
@@ -240,9 +240,9 @@ namespace mu
 		const FImageSize ResultSize = Sources[FirstValidSourceIndex]->GetSize();
 		const int32 ResultNumLODs = Sources[FirstValidSourceIndex]->GetLODCount();
 
-		Ptr<Image> Dest = CreateImage(ResultSize.X, ResultSize.Y, ResultNumLODs, Format, EInitializationType::Black);
+		TSharedPtr<FImage> Dest = CreateImage(ResultSize.X, ResultSize.Y, ResultNumLODs, Format, EInitializationType::Black);
 
-		mu::ImageSwizzle(Dest.get(), Sources, Channels);
+		mu::ImageSwizzle(Dest.Get(), Sources, Channels);
 
 		return Dest;
 	}
diff --git a/Engine/Plugins/Mutable/Source/MutableRuntime/Private/MuR/OpImageTransform.h b/Engine/Plugins/Mutable/Source/MutableRuntime/Private/MuR/OpImageTransform.h
--- a/Engine/Plugins/Mutable/Source/MutableRuntime/Private/MuR/OpImageTransform.h
+++ b/Engine/Plugins/Mutable/Source/MutableRuntime/Private/MuR/OpImageTransform.h
@@ -13,7 +13,7 @@
 namespace mu
 {
 	void ImageTransform(
-			Image* pDestImage, const Image* pImage,
+			FImage* pDestImage, const FImage* pImage,
 			FTransform2f Transform, float MipFactor, EAddressMode AddressMode, bool bUseVectorImplCvar)
     {
 		MUTABLE_CPUPROFILER_SCOPE(ImageTransform);
@@ -52,7 +52,7 @@ namespace mu
 
 		if (AddressMode == EAddressMode::ClampToBlack)
 		{
-			pDestImage->m_flags |= Image::IF_HAS_RELEVANCY_MAP;
+			pDestImage->Flags |= FImage::IF_HAS_RELEVANCY_MAP;
 			pDestImage->RelevancyMinY = DestCropRect.Min.Y;
 			pDestImage->RelevancyMaxY = FMath::Max(DestCropRect.Max.Y - 1, DestCropRect.Min.Y);
 		}
diff --git a/Engine/Plugins/Mutable/Source/MutableRuntime/Private/MuR/OpLayoutPack.cpp b/Engine/Plugins/Mutable/Source/MutableRuntime/Private/MuR/OpLayoutPack.cpp
--- a/Engine/Plugins/Mutable/Source/MutableRuntime/Private/MuR/OpLayoutPack.cpp
+++ b/Engine/Plugins/Mutable/Source/MutableRuntime/Private/MuR/OpLayoutPack.cpp
@@ -385,7 +385,7 @@ namespace mu
 	}
 
 
-    void LayoutPack3( Layout* pResult, const Layout* pSourceLayout )
+    void LayoutPack3( FLayout* pResult, const FLayout* pSourceLayout )
     {
 		MUTABLE_CPUPROFILER_SCOPE(MeshLayoutPack3);
 
diff --git a/Engine/Plugins/Mutable/Source/MutableRuntime/Private/MuR/OpLayoutPack.h b/Engine/Plugins/Mutable/Source/MutableRuntime/Private/MuR/OpLayoutPack.h
--- a/Engine/Plugins/Mutable/Source/MutableRuntime/Private/MuR/OpLayoutPack.h
+++ b/Engine/Plugins/Mutable/Source/MutableRuntime/Private/MuR/OpLayoutPack.h
@@ -5,8 +5,8 @@
 
 namespace mu
 {
-	class Layout;
+	class FLayout;
 
-	extern void LayoutPack3(Layout* pResult, const Layout* pSourceLayout);
+	extern void LayoutPack3(FLayout* Result, const FLayout* SourceLayout);
 
 }
diff --git a/Engine/Plugins/Mutable/Source/MutableRuntime/Private/MuR/OpLayoutRemoveBlocks.h b/Engine/Plugins/Mutable/Source/MutableRuntime/Private/MuR/OpLayoutRemoveBlocks.h
--- a/Engine/Plugins/Mutable/Source/MutableRuntime/Private/MuR/OpLayoutRemoveBlocks.h
+++ b/Engine/Plugins/Mutable/Source/MutableRuntime/Private/MuR/OpLayoutRemoveBlocks.h
@@ -11,15 +11,15 @@ namespace mu
 {
 
 	//---------------------------------------------------------------------------------------------
-	inline Ptr<Layout> LayoutFromMesh_RemoveBlocks(const Mesh* InMesh, int32 InLayoutIndex)
+	inline TSharedPtr<FLayout> LayoutFromMesh_RemoveBlocks(const FMesh* InMesh, int32 InLayoutIndex)
 	{
 		if (!InMesh || InMesh->GetLayoutCount() <= InLayoutIndex)
 		{
 			return nullptr;
 		}
 
-		const Layout* Source = InMesh->GetLayout(InLayoutIndex);
-		Ptr<Layout> pResult;
+		TSharedPtr<const FLayout> Source = InMesh->GetLayout(InLayoutIndex);
+		TSharedPtr<FLayout> pResult;
 
 		UntypedMeshBufferIteratorConst itBlocks(InMesh->GetVertexBuffers(), EMeshBufferSemantic::LayoutBlock, InLayoutIndex);
 		if (itBlocks.GetFormat() == EMeshBufferFormat::UInt16)
@@ -118,10 +118,10 @@ namespace mu
 
 
 	//---------------------------------------------------------------------------------------------
-	inline Ptr<Layout> LayoutRemoveBlocks(const Layout* Source, const Layout* ReferenceLayout)
+	inline TSharedPtr<FLayout> LayoutRemoveBlocks(const FLayout* Source, const FLayout* ReferenceLayout)
 	{
 		// Remove blocks that are not in the mesh
-		Ptr<Layout> pResult = Source->Clone();
+		TSharedPtr<FLayout> pResult = Source->Clone();
 		int32 DestBlockIndex = 0;
 		for (int32 BlockIndexInSource = 0; BlockIndexInSource < pResult->Blocks.Num(); ++BlockIndexInSource)
 		{
@@ -141,9 +141,9 @@ namespace mu
 
 
 	//---------------------------------------------------------------------------------------------
-	inline Ptr<Layout> LayoutMerge(const Layout* pA, const Layout* pB )
+	inline TSharedPtr<FLayout> LayoutMerge(const FLayout* pA, const FLayout* pB )
 	{
-		Ptr<Layout> pResult = pA->Clone();
+		TSharedPtr<FLayout> pResult = pA->Clone();
 
 		for ( const FLayoutBlock& Block: pB->Blocks )
 		{
diff --git a/Engine/Plugins/Mutable/Source/MutableRuntime/Private/MuR/OpMeshApplyLayout.h b/Engine/Plugins/Mutable/Source/MutableRuntime/Private/MuR/OpMeshApplyLayout.h
--- a/Engine/Plugins/Mutable/Source/MutableRuntime/Private/MuR/OpMeshApplyLayout.h
+++ b/Engine/Plugins/Mutable/Source/MutableRuntime/Private/MuR/OpMeshApplyLayout.h
@@ -12,7 +12,7 @@ namespace mu
 {
 
 
-    inline void MeshApplyLayout( Mesh* Applied, const Layout* InLayout, int32 TexCoordsIndex )
+    inline void MeshApplyLayout( FMesh* Applied, const FLayout* InLayout, int32 TexCoordsIndex )
 	{
 		MUTABLE_CPUPROFILER_SCOPE(MeshApplyLayout);
 
diff --git a/Engine/Plugins/Mutable/Source/MutableRuntime/Private/MuR/OpMeshApplyPose.h b/Engine/Plugins/Mutable/Source/MutableRuntime/Private/MuR/OpMeshApplyPose.h
--- a/Engine/Plugins/Mutable/Source/MutableRuntime/Private/MuR/OpMeshApplyPose.h
+++ b/Engine/Plugins/Mutable/Source/MutableRuntime/Private/MuR/OpMeshApplyPose.h
@@ -14,13 +14,13 @@ namespace mu
 	/**
     * Reference version
     */
-    inline void MeshApplyPose(Mesh* Result, const Mesh* BaseMesh, const Mesh* PoseMesh, bool& bOutSuccess)
+    inline void MeshApplyPose(FMesh* Result, const FMesh* BaseMesh, const FMesh* PoseMesh, bool& bOutSuccess)
     {
         MUTABLE_CPUPROFILER_SCOPE(MeshApplyPose);
 
 		bOutSuccess = true;
 
-		Ptr<const Skeleton> Skeleton = BaseMesh->GetSkeleton();
+		TSharedPtr<const FSkeleton> Skeleton = BaseMesh->GetSkeleton();
 		if (!Skeleton)
 		{
 			bOutSuccess = false;
diff --git a/Engine/Plugins/Mutable/Source/MutableRuntime/Private/MuR/OpMeshBind.h b/Engine/Plugins/Mutable/Source/MutableRuntime/Private/MuR/OpMeshBind.h
--- a/Engine/Plugins/Mutable/Source/MutableRuntime/Private/MuR/OpMeshBind.h
+++ b/Engine/Plugins/Mutable/Source/MutableRuntime/Private/MuR/OpMeshBind.h
@@ -363,7 +363,7 @@ namespace mu
 	//! Colour {1,1,1,1} is reserved for the non rigid cluster
     //---------------------------------------------------------------------------------------------
 	
-	inline void FindRigidClusters(const Mesh* Mesh, const FMeshBindColorChannelUsageMasks& ColorUsageMasks, TArray<TArray<int32>>& OutClusters, int32& OutNonRigidClusterIdx)
+	inline void FindRigidClusters(const FMesh* Mesh, const FMeshBindColorChannelUsageMasks& ColorUsageMasks, TArray<TArray<int32>>& OutClusters, int32& OutNonRigidClusterIdx)
 	{
 		MUTABLE_CPUPROFILER_SCOPE(MeshFindRigidClusters);
 		
@@ -421,7 +421,7 @@ namespace mu
 	}
 	
 	inline void FindBindingForCluster( 
-			const Mesh* BaseMesh, FShapeMeshTree& ShapeMeshTree,
+			const FMesh* BaseMesh, FShapeMeshTree& ShapeMeshTree,
 			const TArray<int32>& Cluster, FReshapeVertexBindingData& OutBindingData, float BindTolerance)
 	{
 		const UntypedMeshBufferIteratorConst ItPositionBase(BaseMesh->GetVertexBuffers(), EMeshBufferSemantic::Position);
@@ -441,7 +441,7 @@ namespace mu
 	}
 
 	inline TTuple<TArray<FReshapePointBindingData>, TArray<int32>, TArray<int32>> BindPhysicsBodies( 
-			TArray<const PhysicsBody*> PhysicsBodies, FShapeMeshTree& ShapeMeshTree, const Mesh* pMesh, 
+			TArray<const FPhysicsBody*> PhysicsBodies, FShapeMeshTree& ShapeMeshTree, const FMesh* pMesh, 
 			const TArray<FBoneName>& PhysicsToDeform )
 	{
 		TTuple<TArray<FReshapePointBindingData>, TArray<int32>, TArray<int32>> ReturnValue;
@@ -457,7 +457,7 @@ namespace mu
 		BodiesToDeformOffsets[0] = 0;
 		for (int32 PhysicsBodyIndex = 0; PhysicsBodyIndex < NumPhysicsBodies; ++PhysicsBodyIndex)
 		{
-			if (const PhysicsBody* Body = PhysicsBodies[PhysicsBodyIndex])
+			if (const FPhysicsBody* Body = PhysicsBodies[PhysicsBodyIndex])
 			{
 				const int32 NumBodies = PhysicsBodies[PhysicsBodyIndex]->GetBodyCount();
 				for (int32 BodyIndex = 0; BodyIndex < NumBodies; ++BodyIndex)
@@ -481,7 +481,7 @@ namespace mu
 			TArrayView<int32> BodyIndices = TArrayView<int32>(
 				BodiesToDeformIndices.GetData() + IndicesBegin, IndicesEnd - IndicesBegin);
 
-			const PhysicsBody* Body = PhysicsBodies[PhysicsBodyIndex];
+			const FPhysicsBody* Body = PhysicsBodies[PhysicsBodyIndex];
 		
 			int32 BodyNumPoints = 0;
 			for (const int32 I : BodyIndices)
@@ -518,7 +518,7 @@ namespace mu
 
 		for (int32 PhysicsBodyIndex = 0; PhysicsBodyIndex < NumPhysicsBodies; ++PhysicsBodyIndex)
 		{
-			const PhysicsBody* Body = PhysicsBodies[PhysicsBodyIndex];
+			const FPhysicsBody* Body = PhysicsBodies[PhysicsBodyIndex];
 
 			const int32 IndicesBegin = BodiesToDeformOffsets[PhysicsBodyIndex];
 			const int32 IndicesEnd = BodiesToDeformOffsets[PhysicsBodyIndex + 1];
@@ -703,7 +703,7 @@ namespace mu
 	}
 
 	inline TArray<FReshapeVertexBindingData> BindVerticesReshape(
-			const Mesh* BaseMesh, FShapeMeshTree& ShapeMeshTree, const FMeshBindColorChannelUsageMasks& ColorUsageMasks)	
+			const FMesh* BaseMesh, FShapeMeshTree& ShapeMeshTree, const FMeshBindColorChannelUsageMasks& ColorUsageMasks)	
 	{
 	
 		UE::Geometry::FAxisAlignedBox3d ShapeAABBox = ShapeMeshTree.GetBoundingBox();
@@ -799,7 +799,7 @@ namespace mu
 		return BindData;
 	}
 
-	inline void GenerateAndAddLaplacianData(Mesh& InOutMesh)
+	inline void GenerateAndAddLaplacianData(FMesh& InOutMesh)
 	{
 		MUTABLE_CPUPROFILER_SCOPE(GenerateLaplacianData);
 
@@ -917,7 +917,7 @@ namespace mu
 
 
 	inline TTuple<TArray<FReshapePointBindingData>, TArray<int32>> BindPose(
-			const Mesh* Mesh, FShapeMeshTree& ShapeMeshTree, const TArray<FBoneName>& BonesToDeform )
+			const FMesh* Mesh, FShapeMeshTree& ShapeMeshTree, const TArray<FBoneName>& BonesToDeform )
 	{
 		UE::Geometry::FAxisAlignedBox3d ShapeAABBox = ShapeMeshTree.GetBoundingBox();
 
@@ -996,8 +996,8 @@ namespace mu
     //! Generate the mesh-shape binding data
     //---------------------------------------------------------------------------------------------
     inline void MeshBindShapeReshape(
-			Mesh* Result,
-			const Mesh* BaseMesh, const Mesh* ShapeMesh, 
+			FMesh* Result,
+			const FMesh* BaseMesh, const FMesh* ShapeMesh, 
 			const TArray<FBoneName>& BonesToDeform, const TArray<FBoneName>& PhysicsToDeform,
 			EMeshBindShapeFlags BindFlags, FMeshBindColorChannelUsages ColorChannelUsages,
 			bool& bOutSuccess)
@@ -1168,21 +1168,21 @@ namespace mu
 			Result->AdditionalBuffers.Emplace(EMeshBufferType::SkeletonDeformBinding, MoveTemp(SkeletonBuffer));
 		}
 
-		const PhysicsBody* ResultPhysicsBody = Result->PhysicsBody.get();
+		const FPhysicsBody* ResultPhysicsBody = Result->PhysicsBody.Get();
 		if (bReshapePhysics && (ResultPhysicsBody || Result->AdditionalPhysicsBodies.Num()) && PhysicsToDeform.Num())
 		{
 			MUTABLE_CPUPROFILER_SCOPE(BindPhysicsBody);
 
 			// Gather bodies respecting order, firts main physics body then additional bodies.
 			// Null entries are need to be able to mantin that order.  
-			TArray<const PhysicsBody*> PhysicsBodiesToBind;
+			TArray<const FPhysicsBody*> PhysicsBodiesToBind;
 			{
 				PhysicsBodiesToBind.Reserve(Result->AdditionalPhysicsBodies.Num() + 1);
 
 				PhysicsBodiesToBind.Add(ResultPhysicsBody);
-				for (const Ptr<const PhysicsBody>& Body : Result->AdditionalPhysicsBodies)
+				for (const TSharedPtr<const FPhysicsBody>& Body : Result->AdditionalPhysicsBodies)
 				{
-					PhysicsBodiesToBind.Add(Body.get());
+					PhysicsBodiesToBind.Add(Body.Get());
 				}
 			}
 
@@ -1373,7 +1373,7 @@ namespace mu
 	}
 	
 	inline TArray<FClipDeformVertexBindingData> BindVerticesClipDeform(
-			const Mesh* BaseMesh, 
+			const FMesh* BaseMesh, 
 			FShapeMeshTree& ShapeMeshTree, 
 			EShapeBindingMethod BindingMethod )
 	{	
@@ -1443,7 +1443,7 @@ namespace mu
 		return BindData;
 	}
 
-    inline void MeshBindShapeClipDeform(Mesh* Result, const Mesh* BaseMesh, const Mesh* ShapeMesh, EShapeBindingMethod BindingMethod, bool& bOutSuccess)
+    inline void MeshBindShapeClipDeform(FMesh* Result, const FMesh* BaseMesh, const FMesh* ShapeMesh, EShapeBindingMethod BindingMethod, bool& bOutSuccess)
     {
 		MUTABLE_CPUPROFILER_SCOPE(MeshBindShapeClipDeform);
 	
diff --git a/Engine/Plugins/Mutable/Source/MutableRuntime/Private/MuR/OpMeshClipDeform.h b/Engine/Plugins/Mutable/Source/MutableRuntime/Private/MuR/OpMeshClipDeform.h
--- a/Engine/Plugins/Mutable/Source/MutableRuntime/Private/MuR/OpMeshClipDeform.h
+++ b/Engine/Plugins/Mutable/Source/MutableRuntime/Private/MuR/OpMeshClipDeform.h
@@ -100,7 +100,7 @@ namespace mu
 		return true;
 	}
 
-	inline void MeshClipDeform(Mesh* Result, const Mesh* BaseMesh, const Mesh* ShapeMesh, const float ClipWeightThreshold, bool bRemoveIfAllVerticesCulled, bool& bOutSuccess )
+	inline void MeshClipDeform(FMesh* Result, const FMesh* BaseMesh, const FMesh* ShapeMesh, const float ClipWeightThreshold, bool bRemoveIfAllVerticesCulled, bool& bOutSuccess )
 	{
 		MUTABLE_CPUPROFILER_SCOPE(ClipDeform);
 		bOutSuccess = true;
diff --git a/Engine/Plugins/Mutable/Source/MutableRuntime/Private/MuR/OpMeshClipMorphPlane.h b/Engine/Plugins/Mutable/Source/MutableRuntime/Private/MuR/OpMeshClipMorphPlane.h
--- a/Engine/Plugins/Mutable/Source/MutableRuntime/Private/MuR/OpMeshClipMorphPlane.h
+++ b/Engine/Plugins/Mutable/Source/MutableRuntime/Private/MuR/OpMeshClipMorphPlane.h
@@ -71,7 +71,7 @@ namespace mu
 	//---------------------------------------------------------------------------------------------
 	//! Reference version
 	//---------------------------------------------------------------------------------------------
-	inline void MeshClipMorphPlane(Mesh* Result, const Mesh* pBase, const FVector3f& Origin, const FVector3f& Normal, float Dist, float Factor, float Radius,
+	inline void MeshClipMorphPlane(FMesh* Result, const FMesh* pBase, const FVector3f& Origin, const FVector3f& Normal, float Dist, float Factor, float Radius,
 		float Radius2, float Angle, const FShape& SelectionShape, bool bRemoveIfAllVerticesCulled, bool& bOutSuccess, 
 		const FBoneName* BoneId = nullptr, float VertexSelectionBoneMaxRadius = -1.f)
 	{
@@ -101,7 +101,7 @@ namespace mu
 
 		TArray<FVertexBoneInfo> VertexInfo;
 
-        Ptr<const Skeleton> BaseSkeleton = pBase->GetSkeleton();
+        TSharedPtr<const FSkeleton> BaseSkeleton = pBase->GetSkeleton();
 
 		TBitArray<> AffectedBoneMapIndices;
 		const int32 BaseBoneIndex = BaseSkeleton && BoneId ? BaseSkeleton->FindBone(*BoneId) : INDEX_NONE;
diff --git a/Engine/Plugins/Mutable/Source/MutableRuntime/Private/MuR/OpMeshClipWithMesh.cpp b/Engine/Plugins/Mutable/Source/MutableRuntime/Private/MuR/OpMeshClipWithMesh.cpp
--- a/Engine/Plugins/Mutable/Source/MutableRuntime/Private/MuR/OpMeshClipWithMesh.cpp
+++ b/Engine/Plugins/Mutable/Source/MutableRuntime/Private/MuR/OpMeshClipWithMesh.cpp
@@ -26,7 +26,7 @@ namespace mu { namespace
     //! Create a map from vertices into vertices, collapsing vertices that have the same position, 
 	//! This version uses UE Containers to return.	
     //---------------------------------------------------------------------------------------------
-	void MeshCreateCollapsedVertexMap(const Mesh* pMesh, TArray<int32>& OutCollapsedVertices, TArray<FVector3f>& OutVertices)
+	void MeshCreateCollapsedVertexMap(const FMesh* pMesh, TArray<int32>& OutCollapsedVertices, TArray<FVector3f>& OutVertices)
 	{
 		MUTABLE_CPUPROFILER_SCOPE(MeshCreateCollapsedVertexMap);
 
@@ -80,7 +80,7 @@ namespace mu { namespace
     //---------------------------------------------------------------------------------------------
     //! Return true if a mesh is closed
     //---------------------------------------------------------------------------------------------
-  //  bool MeshIsClosed( const Mesh* pMesh )
+  //  bool MeshIsClosed( const FMesh* pMesh )
   //  {
   //      MUTABLE_CPUPROFILER_SCOPE(MeshIsClosed);
 
@@ -148,7 +148,7 @@ namespace mu { namespace
     //---------------------------------------------------------------------------------------------
     //! Remove all unused vertices from a mesh, and fix its index buffers.
     //---------------------------------------------------------------------------------------------
-    void MeshRemoveUnusedVertices( Mesh* pMesh )
+    void MeshRemoveUnusedVertices( FMesh* pMesh )
     {
         MUTABLE_CPUPROFILER_SCOPE(MeshRemoveUnusedVertices);
 
@@ -320,7 +320,7 @@ namespace mu { namespace
         return NumIntersections;
     }
 
-	void MeshUVMaskClassifyVertices(TBitArray<>& VertexClipped, const Mesh* Base, const Image* Mask, const uint8 LayoutIndex)
+	void MeshUVMaskClassifyVertices(TBitArray<>& VertexClipped, const FMesh* Base, const FImage* Mask, const uint8 LayoutIndex)
 	{
 		using namespace UE::Geometry;
 
@@ -371,7 +371,7 @@ namespace mu { namespace
 	}
 
 
-	void MeshLayoutMaskClassifyVertices(TBitArray<>& VertexClipped, const Mesh* Base, const Layout* Mask, const uint8 LayoutIndex)
+	void MeshLayoutMaskClassifyVertices(TBitArray<>& VertexClipped, const FMesh* Base, const FLayout* Mask, const uint8 LayoutIndex)
 	{
 		using namespace UE::Geometry;
 
@@ -434,7 +434,7 @@ namespace mu { namespace
 
 
 	/** Make a mask with the indices of the vertices with 0 in the IncludedVertices array. */
-	void CreateMask(Mesh* Result, const Mesh* Base, const TArray<uint8>& IncludedVertices)
+	void CreateMask(FMesh* Result, const FMesh* Base, const TArray<uint8>& IncludedVertices)
 	{
 		int32 MaskVertexCount = 0;
 		for (uint8 b : IncludedVertices)
@@ -500,7 +500,7 @@ namespace mu { namespace
 namespace mu
 {
     /** Core Geometry version */
-    void MeshClipMeshClassifyVertices(TBitArray<>& VertexInClipMesh, const Mesh* pBase, const Mesh* pClipMesh)
+    void MeshClipMeshClassifyVertices(TBitArray<>& VertexInClipMesh, const FMesh* pBase, const FMesh* pClipMesh)
     {
 		using namespace UE::Geometry;
 
@@ -685,7 +685,7 @@ namespace mu
         }
     }	
 
-    void MeshClipWithMesh(Mesh* Result, const Mesh* pBase, const Mesh* pClipMesh, bool& bOutSuccess)
+    void MeshClipWithMesh(FMesh* Result, const FMesh* pBase, const FMesh* pClipMesh, bool& bOutSuccess)
     {
         MUTABLE_CPUPROFILER_SCOPE(MeshClipWithMesh);
 		bOutSuccess = true;
@@ -740,7 +740,7 @@ namespace mu
     }
 
 
-	void MeshMaskClipMesh(Mesh* Result, const Mesh* pBase, const Mesh* pClipMesh, bool& bOutSuccess)
+	void MeshMaskClipMesh(FMesh* Result, const FMesh* pBase, const FMesh* pClipMesh, bool& bOutSuccess)
     {
         MUTABLE_CPUPROFILER_SCOPE(MeshMaskClipMesh);
 
@@ -787,7 +787,7 @@ namespace mu
     }
 
 
-	void MakeMeshMaskFromUVMask(Mesh* Result, const Mesh* Base, const Mesh* BaseForUVs, const Image* Mask, uint8 LayoutIndex, bool& bOutSuccess)
+	void MakeMeshMaskFromUVMask(FMesh* Result, const FMesh* Base, const FMesh* BaseForUVs, const FImage* Mask, uint8 LayoutIndex, bool& bOutSuccess)
 	{
 		MUTABLE_CPUPROFILER_SCOPE(MeshMaskUVMask);
 
@@ -832,7 +832,7 @@ namespace mu
 	}
 
 
-	void MakeMeshMaskFromLayout(Mesh* Result, const Mesh* Base, const Mesh* BaseForUVs, const Layout* Mask, uint8 LayoutIndex, bool& bOutSuccess)
+	void MakeMeshMaskFromLayout(FMesh* Result, const FMesh* Base, const FMesh* BaseForUVs, const FLayout* Mask, uint8 LayoutIndex, bool& bOutSuccess)
 	{
 		MUTABLE_CPUPROFILER_SCOPE(MakeMeshMaskFromLayout);
 
@@ -877,7 +877,7 @@ namespace mu
 	}
 
 
-	void MeshMaskDiff(Mesh* Result, const Mesh* pBase, const Mesh* pFragment, bool& bOutSuccess)
+	void MeshMaskDiff(FMesh* Result, const FMesh* pBase, const FMesh* pFragment, bool& bOutSuccess)
     {
         MUTABLE_CPUPROFILER_SCOPE(MeshMaskDiff);
 		bOutSuccess = true;
@@ -911,7 +911,7 @@ namespace mu
             }
         }
         float tolerance = 1e-5f * aabbox.size.Length();
-        Mesh::FVertexMatchMap vertexMap;
+        FMesh::FVertexMatchMap vertexMap;
         pFragment->GetVertexMap( *pBase, vertexMap, tolerance );
 
 
diff --git a/Engine/Plugins/Mutable/Source/MutableRuntime/Private/MuR/OpMeshClipWithMesh.h b/Engine/Plugins/Mutable/Source/MutableRuntime/Private/MuR/OpMeshClipWithMesh.h
--- a/Engine/Plugins/Mutable/Source/MutableRuntime/Private/MuR/OpMeshClipWithMesh.h
+++ b/Engine/Plugins/Mutable/Source/MutableRuntime/Private/MuR/OpMeshClipWithMesh.h
@@ -6,26 +6,26 @@
 
 namespace mu
 {
-	class Mesh;
-	class Image;
-	class Layout;
+	class FMesh;
+	class FImage;
+	class FLayout;
 
     /** Generate a classification list for which vertex of pBase is fully contained in pClipMesh */
-	extern void MeshClipMeshClassifyVertices(TBitArray<>& VertexInClipMesh, const Mesh* pBase, const Mesh* pClipMesh);
+	extern void MeshClipMeshClassifyVertices(TBitArray<>& VertexInClipMesh, const FMesh* Base, const FMesh* ClipMesh);
 
     /**  */
-	extern void MeshClipWithMesh(Mesh* Result, const Mesh* pBase, const Mesh* pClipMesh, bool& bOutSuccess);
+	extern void MeshClipWithMesh(FMesh* Result, const FMesh* Base, const FMesh* ClipMesh, bool& bOutSuccess);
 
     /** Generate a mask mesh with the faces of the base mesh inside the clip mesh. */
-	extern void MeshMaskClipMesh(Mesh* Result, const Mesh* pBase, const Mesh* pClipMesh, bool& bOutSuccess);
+	extern void MeshMaskClipMesh(FMesh* Result, const FMesh* Base, const FMesh* ClipMesh, bool& bOutSuccess);
 
 	/** Generate a mask mesh with the faces of the base mesh that have all 3 vertices marked in the given mask. */
-	extern void MakeMeshMaskFromUVMask(Mesh* Result, const Mesh* Base, const Mesh* BaseForUVs, const Image* Mask, uint8 LayoutIndex, bool& bOutSuccess);
+	extern void MakeMeshMaskFromUVMask(FMesh* Result, const FMesh* Base, const FMesh* BaseForUVs, const FImage* Mask, uint8 LayoutIndex, bool& bOutSuccess);
 
     /** Generate a mask mesh with the faces of the base mesh matching the fragment. */
-	extern void MeshMaskDiff(Mesh* Result, const Mesh* pBase, const Mesh* pFragment, bool& bOutSuccess);
+	extern void MeshMaskDiff(FMesh* Result, const FMesh* Base, const FMesh* Fragment, bool& bOutSuccess);
 
 	/** Generate a mask mesh with the faces of the base mesh that have all 3 vertices inside any block of the given layout. */
-	extern void MakeMeshMaskFromLayout(Mesh* Result, const Mesh* Base, const Mesh* BaseForUVs, const Layout* Mask, uint8 LayoutIndex, bool& bOutSuccess);
+	extern void MakeMeshMaskFromLayout(FMesh* Result, const FMesh* Base, const FMesh* BaseForUVs, const FLayout* Mask, uint8 LayoutIndex, bool& bOutSuccess);
 
 }
diff --git a/Engine/Plugins/Mutable/Source/MutableRuntime/Private/MuR/OpMeshComputeNormals.h b/Engine/Plugins/Mutable/Source/MutableRuntime/Private/MuR/OpMeshComputeNormals.h
--- a/Engine/Plugins/Mutable/Source/MutableRuntime/Private/MuR/OpMeshComputeNormals.h
+++ b/Engine/Plugins/Mutable/Source/MutableRuntime/Private/MuR/OpMeshComputeNormals.h
@@ -44,7 +44,7 @@ namespace mu
 		}
 	}
 
-	void ComputeMeshNormals(Mesh& DestMesh)
+	void ComputeMeshNormals(FMesh& DestMesh)
 	{
 		MUTABLE_CPUPROFILER_SCOPE(ComputeMeshNormals);
 
diff --git a/Engine/Plugins/Mutable/Source/MutableRuntime/Private/MuR/OpMeshDifference.h b/Engine/Plugins/Mutable/Source/MutableRuntime/Private/MuR/OpMeshDifference.h
--- a/Engine/Plugins/Mutable/Source/MutableRuntime/Private/MuR/OpMeshDifference.h
+++ b/Engine/Plugins/Mutable/Source/MutableRuntime/Private/MuR/OpMeshDifference.h
@@ -12,7 +12,7 @@ namespace mu
 	//! Create a diff from the mesh vertices. The meshes must have the same amount of vertices.
 	//! If the channel list is empty all the channels will be compared.
 	//---------------------------------------------------------------------------------------------
-    inline void MeshDifference( Mesh* Result, const Mesh* pBase, const Mesh* pTarget,
+    inline void MeshDifference( FMesh* Result, const FMesh* pBase, const FMesh* pTarget,
                                    int32 numChannels,
                                    const EMeshBufferSemantic* semantics,
                                    const int32* semanticIndices,
