[mutable] More enforcement of coding standards.

diff --git a/Engine/Plugins/Mutable/Source/MutableRuntime/Private/MuR/OpMeshExtractLayoutBlock.h b/Engine/Plugins/Mutable/Source/MutableRuntime/Private/MuR/OpMeshExtractLayoutBlock.h
--- a/Engine/Plugins/Mutable/Source/MutableRuntime/Private/MuR/OpMeshExtractLayoutBlock.h
+++ b/Engine/Plugins/Mutable/Source/MutableRuntime/Private/MuR/OpMeshExtractLayoutBlock.h
@@ -9,8 +9,8 @@
 
 namespace mu
 {
-    inline void MeshExtractFromVertices( const Mesh* Source,
-                                         Mesh* Result,
+    inline void MeshExtractFromVertices( const FMesh* Source,
+                                         FMesh* Result,
                                          const TArray<int32>& OldToNew,
                                          const TArray<int32>& NewToOld )
     {
@@ -123,7 +123,7 @@ namespace mu
 
 
 	/** Extract all vertices that have a layout block from the passed list on the given channel. */
-    inline void MeshExtractLayoutBlock(Mesh* Result, const Mesh* Source,
+    inline void MeshExtractLayoutBlock(FMesh* Result, const FMesh* Source,
                                            uint32 LayoutIndex,
                                            uint16 BlockCount,
                                            const uint64* BlockIds, bool& bOutSuccess)
@@ -203,7 +203,7 @@ namespace mu
 	}
 
 	/** Extract all vertices that have a valid layout block on the given channel. */
-	inline void MeshExtractLayoutBlock(Mesh* Result, const Mesh* Source, uint32 LayoutIndex, bool& bOutSuccess)
+	inline void MeshExtractLayoutBlock(FMesh* Result, const FMesh* Source, uint32 LayoutIndex, bool& bOutSuccess)
 	{
 		check(Source);
 		bOutSuccess = true;
diff --git a/Engine/Plugins/Mutable/Source/MutableRuntime/Private/MuR/OpMeshFormat.cpp b/Engine/Plugins/Mutable/Source/MutableRuntime/Private/MuR/OpMeshFormat.cpp
--- a/Engine/Plugins/Mutable/Source/MutableRuntime/Private/MuR/OpMeshFormat.cpp
+++ b/Engine/Plugins/Mutable/Source/MutableRuntime/Private/MuR/OpMeshFormat.cpp
@@ -526,20 +526,20 @@ namespace mu
 	//-------------------------------------------------------------------------------------------------
 	void MeshFormat
 	(
-		Mesh* Result, 
-		const Mesh* PureSource,
-		const Mesh* Format,
-		bool keepSystemBuffers,
-		bool formatVertices,
-		bool formatIndices,
-		bool ignoreMissingChannels,
+		FMesh* Result, 
+		const FMesh* Source,
+		const FMesh* Format,
+		bool bKeepSystemBuffers,
+		bool bFormatVertices,
+		bool bFormatIndices,
+		bool bIgnoreMissingChannels,
 		bool& bOutSuccess
 	)
 	{
 		MUTABLE_CPUPROFILER_SCOPE(MeshFormat);
 		bOutSuccess = true;
 
-		if (!PureSource) 
+		if (!Source) 
 		{
 			check(false);
 			bOutSuccess = false;	
@@ -553,13 +553,11 @@ namespace mu
 			return;
 		}
 
-		Ptr<const Mesh> Source = PureSource;
-
 		Result->CopyFrom(*Format);
 		Result->MeshIDPrefix = Source->MeshIDPrefix;
 
 		// Make sure that the bone indices will fit in this format, or extend it.
-		if (formatVertices)
+		if (bFormatVertices)
 		{
 			const FMeshBufferSet& VertexBuffers = Source->GetVertexBuffers();
 
@@ -621,21 +619,21 @@ namespace mu
 			}
 		}
 
-		if (formatVertices)
+		if (bFormatVertices)
 		{
 			FormatBufferSet(Source->GetVertexBuffers(), Result->GetVertexBuffers(),
-				keepSystemBuffers, ignoreMissingChannels, true, Source->MeshIDPrefix);
+				bKeepSystemBuffers, bIgnoreMissingChannels, true, Source->MeshIDPrefix);
 		}
 		else
 		{
 			Result->VertexBuffers = Source->GetVertexBuffers();
 		}
 
-		if (formatIndices)
+		if (bFormatIndices)
 		{
 			// \todo Make sure that the vertex indices will fit in this format, or extend it.
-			FormatBufferSet(Source->GetIndexBuffers(), Result->GetIndexBuffers(), keepSystemBuffers,
-				ignoreMissingChannels, false);
+			FormatBufferSet(Source->GetIndexBuffers(), Result->GetIndexBuffers(), bKeepSystemBuffers,
+				bIgnoreMissingChannels, false);
 		}
 		else
 		{
@@ -647,7 +645,7 @@ namespace mu
 		Result->SetPhysicsBody(Source->GetPhysicsBody());
 
 		Result->Layouts.Empty();
-		for (const Ptr<const Layout>& Layout : Source->Layouts)
+		for (const TSharedPtr<const FLayout>& Layout : Source->Layouts)
 		{
 			Result->Layouts.Add(Layout->Clone());
 		}
@@ -672,7 +670,7 @@ namespace mu
 	}
 
 
-	void MeshOptimizeBuffers( Mesh* InMesh )
+	void MeshOptimizeBuffers( FMesh* InMesh )
 	{
 		if (!InMesh)
 		{
diff --git a/Engine/Plugins/Mutable/Source/MutableRuntime/Private/MuR/OpMeshGeometryOperation.h b/Engine/Plugins/Mutable/Source/MutableRuntime/Private/MuR/OpMeshGeometryOperation.h
--- a/Engine/Plugins/Mutable/Source/MutableRuntime/Private/MuR/OpMeshGeometryOperation.h
+++ b/Engine/Plugins/Mutable/Source/MutableRuntime/Private/MuR/OpMeshGeometryOperation.h
@@ -20,7 +20,7 @@ namespace mu
 	//! It only converts position and normal data.
 	//! Reference version, that can be optimized with some assumptions on how Vertex IDs are generated.
 	//---------------------------------------------------------------------------------------------
-	inline TSharedPtr<UE::Geometry::FDynamicMesh3> MutableToDynamicMesh(const Mesh* MutableMesh, TArray<int>& VertexMutableToDyn)
+	inline TSharedPtr<UE::Geometry::FDynamicMesh3> MutableToDynamicMesh(const FMesh* MutableMesh, TArray<int>& VertexMutableToDyn)
 	{
 		if (!MutableMesh)
 		{
@@ -78,7 +78,7 @@ namespace mu
 	//---------------------------------------------------------------------------------------------
 	//! Apply an example geometric operation (twist)
 	//---------------------------------------------------------------------------------------------
-	inline void UpdateMutableMesh(Mesh* MutableMesh, const UE::Geometry::FDynamicMesh3* DynMesh, const TArray<int>& VertexMutableToDyn)
+	inline void UpdateMutableMesh(FMesh* MutableMesh, const UE::Geometry::FDynamicMesh3* DynMesh, const TArray<int>& VertexMutableToDyn)
 	{
 		MUTABLE_CPUPROFILER_SCOPE(DynamicMeshToMutable);
 
@@ -105,7 +105,7 @@ namespace mu
     //---------------------------------------------------------------------------------------------
     //! Apply an example geometric operation (twist)
     //---------------------------------------------------------------------------------------------
-    inline void MeshGeometryOperation(Mesh* Result, const Mesh* MeshA, const Mesh* MeshB, float ScalarA, float ScalarB, bool& bOutSuccess)
+    inline void MeshGeometryOperation(FMesh* Result, const FMesh* MeshA, const FMesh* MeshB, float ScalarA, float ScalarB, bool& bOutSuccess)
     {
 		MUTABLE_CPUPROFILER_SCOPE(MeshGeometryOperation);
 		bOutSuccess = true;
diff --git a/Engine/Plugins/Mutable/Source/MutableRuntime/Private/MuR/OpMeshMorph.h b/Engine/Plugins/Mutable/Source/MutableRuntime/Private/MuR/OpMeshMorph.h
--- a/Engine/Plugins/Mutable/Source/MutableRuntime/Private/MuR/OpMeshMorph.h
+++ b/Engine/Plugins/Mutable/Source/MutableRuntime/Private/MuR/OpMeshMorph.h
@@ -13,7 +13,7 @@ namespace mu
 	/**
 	 * Optimized linear factor version for morphing 2 targets
 	 */
-    inline void MeshMorph2(Mesh* BaseMesh, const Mesh* MinMesh, const Mesh* MaxMesh, const float Factor)
+    inline void MeshMorph2(FMesh* BaseMesh, const FMesh* MinMesh, const FMesh* MaxMesh, const float Factor)
     {
         MUTABLE_CPUPROFILER_SCOPE(MeshMorph2);
 
@@ -228,7 +228,7 @@ namespace mu
 		const int32 MinNum = MinMesh ? MinMesh->GetVertexBuffers().GetElementCount() : 0;
 		const int32 MaxNum = MaxMesh ? MaxMesh->GetVertexBuffers().GetElementCount() : 0;
 		const int32 BaseNum = BaseMesh ? BaseMesh->GetVertexBuffers().GetElementCount() : 0;
-		const Mesh* RefTarget = MinNum > 0 ? MinMesh : MaxMesh;
+		const FMesh* RefTarget = MinNum > 0 ? MinMesh : MaxMesh;
 
 		if (BaseNum == 0 || (MinNum + MaxNum) == 0)
 		{
@@ -326,13 +326,13 @@ namespace mu
     }
 
 	//TODO Optimized linear factor version
-	inline void MeshMorph(Mesh* BaseMesh, const Mesh* MorphMesh, float Factor)
+	inline void MeshMorph(FMesh* BaseMesh, const FMesh* MorphMesh, float Factor)
 	{
 		MeshMorph2(BaseMesh, nullptr, MorphMesh, Factor);
 	}
 
     //TODO Optimized Factor-less version
-	inline void MeshMorph(Mesh* BaseMesh, const Mesh* MorphMesh)
+	inline void MeshMorph(FMesh* BaseMesh, const FMesh* MorphMesh)
     {
         // Trust the compiler to remove the factor
 		MeshMorph(BaseMesh, MorphMesh, 1.0f);
diff --git a/Engine/Plugins/Mutable/Source/MutableRuntime/Private/MuR/OpMeshOptimizeSkinning.h b/Engine/Plugins/Mutable/Source/MutableRuntime/Private/MuR/OpMeshOptimizeSkinning.h
--- a/Engine/Plugins/Mutable/Source/MutableRuntime/Private/MuR/OpMeshOptimizeSkinning.h
+++ b/Engine/Plugins/Mutable/Source/MutableRuntime/Private/MuR/OpMeshOptimizeSkinning.h
@@ -14,7 +14,7 @@ namespace mu
 	//---------------------------------------------------------------------------------------------
 	//!
 	//---------------------------------------------------------------------------------------------
-	inline void MeshOptimizeSkinning(Mesh* Result, const Mesh* InMesh, bool& bOutSuccess)
+	inline void MeshOptimizeSkinning(FMesh* Result, const FMesh* InMesh, bool& bOutSuccess)
 	{
 		MUTABLE_CPUPROFILER_SCOPE(MeshOptimizeSkinning);
 
diff --git a/Engine/Plugins/Mutable/Source/MutableRuntime/Private/MuR/OpMeshRemove.cpp b/Engine/Plugins/Mutable/Source/MutableRuntime/Private/MuR/OpMeshRemove.cpp
--- a/Engine/Plugins/Mutable/Source/MutableRuntime/Private/MuR/OpMeshRemove.cpp
+++ b/Engine/Plugins/Mutable/Source/MutableRuntime/Private/MuR/OpMeshRemove.cpp
@@ -9,7 +9,7 @@
 
 namespace mu
 {
-    void MeshRemoveRecreateSurface(Mesh* Result, const TBitArray<>& UsedVertices, const TBitArray<>& UsedFaces)
+    void MeshRemoveRecreateSurface(FMesh* Result, const TBitArray<>& UsedVertices, const TBitArray<>& UsedFaces)
     {
         MUTABLE_CPUPROFILER_SCOPE(MeshRemoveRecreateSurface);
 		TArray<FSurfaceSubMesh, TInlineAllocator<32>> OrigSubMeshes;
@@ -77,19 +77,19 @@ namespace mu
     };
 
 
-    void ExtractVertexIndexIntervals( TArray<FIdInterval>& intervals, const Mesh* pSource )
+    void ExtractVertexIndexIntervals( TArray<FIdInterval>& intervals, const FMesh* Source )
     {
-		MeshVertexIdIteratorConst itVI( pSource );
+		MeshVertexIdIteratorConst itVI(Source);
 		FIdInterval current;
-		current.idStart = Mesh::InvalidVertexId;
+		current.idStart = FMesh::InvalidVertexId;
 		current.idPosition = 0;
 		current.size = 0;
-		for ( int32 sv=0; sv<pSource->GetVertexBuffers().GetElementCount(); ++sv )
+		for ( int32 sv=0; sv< Source->GetVertexBuffers().GetElementCount(); ++sv )
         {
             uint64 id = itVI.Get();
             ++itVI;
 
-            if (current.idStart== Mesh::InvalidVertexId)
+            if (current.idStart== FMesh::InvalidVertexId)
             {
                 current.idStart = id;
                 current.idPosition = sv;
@@ -111,7 +111,7 @@ namespace mu
             }
         }
 
-        if (current.idStart!= Mesh::InvalidVertexId)
+        if (current.idStart!= FMesh::InvalidVertexId)
         {
             intervals.Add(current);
         }
@@ -133,7 +133,7 @@ namespace mu
 
     // TODO: Remove once the new implementation is proven to be faster for all cases and does not present 
     // any problem. For now keep it as a reference.
-	//void MeshRemoveVerticesWithMap( Mesh* Result, const TBitArray<>& RemovedVertices)
+	//void MeshRemoveVerticesWithMap( FMesh* Result, const TBitArray<>& RemovedVertices)
 	//{
     //    uint32 RemovedVertexCount = RemovedVertices.Num();
 
@@ -337,7 +337,7 @@ namespace mu
     //    MeshRemoveRecreateSurface(Result, UsedVertices, UsedFaces);
     //}
 
-    void MeshRemoveVerticesWithCullSet(Mesh* Result, const TBitArray<>& VerticesToCull, bool bRemoveIfAllVerticesCulled)
+    void MeshRemoveVerticesWithCullSet(FMesh* Result, const TBitArray<>& VerticesToCull, bool bRemoveIfAllVerticesCulled)
     { 
         MUTABLE_CPUPROFILER_SCOPE(MeshRemoveVerticesWithCullSet);
 		
@@ -572,7 +572,7 @@ namespace mu
 		MeshRemoveRecreateSurface(Result, UsedVertices, UsedFaces);
     }
 
-	void MeshRemoveMask(Mesh* Result, const Mesh* Source, const Mesh* Mask, bool bRemoveIfAllVerticesCulled, bool& bOutSuccess)
+	void MeshRemoveMask(FMesh* Result, const FMesh* Source, const FMesh* Mask, bool bRemoveIfAllVerticesCulled, bool& bOutSuccess)
     {
         MUTABLE_CPUPROFILER_SCOPE(MeshRemoveMask);
 		bOutSuccess = true;
diff --git a/Engine/Plugins/Mutable/Source/MutableRuntime/Private/MuR/OpMeshRemove.h b/Engine/Plugins/Mutable/Source/MutableRuntime/Private/MuR/OpMeshRemove.h
--- a/Engine/Plugins/Mutable/Source/MutableRuntime/Private/MuR/OpMeshRemove.h
+++ b/Engine/Plugins/Mutable/Source/MutableRuntime/Private/MuR/OpMeshRemove.h
@@ -9,13 +9,13 @@
 
 namespace mu
 {
-	class Mesh;
+	class FMesh;
 
 	/** Remove a list of vertices and related faces from a mesh. The list of vertices is stored in a specially formatted Mask mesh. 
 	 * If bRemoveIfAllVerticesCulled is true, a face is remove if all its vertices have the bit set in VerticesToCull.
 	 * If bRemoveIfAllVerticesCulled is false, remove a face if at least one vertex is removed.
 	 */
-	extern void MeshRemoveMask(Mesh* Result, const Mesh* Source, const Mesh* Mask, bool bRemoveIfAllVerticesCulled, bool& bOutSuccess);
+	extern void MeshRemoveMask(FMesh* Result, const FMesh* Source, const FMesh* Mask, bool bRemoveIfAllVerticesCulled, bool& bOutSuccess);
 
 	/** 
 	 * Remove a set of vertices and related faces from a mesh in-place. VertexToCull is a bitset where if bit i-th is set, 
@@ -23,10 +23,10 @@ namespace mu
 	 * If bRemoveIfAllVerticesCulled is true, a face is remove if all its vertices have the bit set in VerticesToCull. 
 	 * If bRemoveIfAllVerticesCulled is false, remove a face if at least one vertex is removed.
 	 */
-	extern void MeshRemoveVerticesWithCullSet(Mesh* Result, const TBitArray<>& VerticesToCull, bool bRemoveIfAllVerticesCulled);
+	extern void MeshRemoveVerticesWithCullSet(FMesh* Result, const TBitArray<>& VerticesToCull, bool bRemoveIfAllVerticesCulled);
 
 	/**
 	 * Recreates the Surface and Surfaces Submeshes given a set of vertices and faces remaining after mesh removal.
 	 */
-	extern void MeshRemoveRecreateSurface(Mesh* Result, const TBitArray<>& UsedVertices, const TBitArray<>& UsedFaces);
+	extern void MeshRemoveRecreateSurface(FMesh* Result, const TBitArray<>& UsedVertices, const TBitArray<>& UsedFaces);
 }
diff --git a/Engine/Plugins/Mutable/Source/MutableRuntime/Private/MuR/OpMeshReshape.h b/Engine/Plugins/Mutable/Source/MutableRuntime/Private/MuR/OpMeshReshape.h
--- a/Engine/Plugins/Mutable/Source/MutableRuntime/Private/MuR/OpMeshReshape.h
+++ b/Engine/Plugins/Mutable/Source/MutableRuntime/Private/MuR/OpMeshReshape.h
@@ -369,7 +369,7 @@ namespace mu
 		OutL = FMath::Max(0.0f, (Centroids[0] - Centroids[1]).Length() - (OutR0 + OutR1));
 	}
 
-	inline void ApplyToVertices(Mesh* Mesh, TArrayView<const FReshapeVertexBindingData> BindingData, const FShapeMeshDescriptorApply& Shape, bool bSkipNormalReshape)
+	inline void ApplyToVertices(FMesh* Mesh, TArrayView<const FReshapeVertexBindingData> BindingData, const FShapeMeshDescriptorApply& Shape, bool bSkipNormalReshape)
 	{
 		check(Mesh);
 		check(Mesh->GetVertexCount() == BindingData.Num());
@@ -469,7 +469,7 @@ namespace mu
 		}
 	}
 
-	inline void ApplyToPose(Mesh* Result, 
+	inline void ApplyToPose(FMesh* Result, 
 			TArrayView<const FReshapePointBindingData> BindingData, 
 			TArrayView<const int32> BoneIndices,
 			const FShapeMeshDescriptorApply& Shape)
@@ -522,8 +522,8 @@ namespace mu
 	}
 
 	inline void ApplyToPhysicsBodies(
-			PhysicsBody& PBody, int32& InOutNumProcessedBindPoints, 
-			const Mesh& BaseMesh, 
+			FPhysicsBody& PBody, int32& InOutNumProcessedBindPoints, 
+			const FMesh& BaseMesh, 
 			TArrayView<const FReshapePointBindingData> BindingData, 
 			TArrayView<const int32> UsedIndices, 
 			const FShapeMeshDescriptorApply& Shape)
@@ -662,7 +662,7 @@ namespace mu
 	}
 
 	inline void ApplyToAllPhysicsBodies(
-		Mesh& OutNewMesh, const Mesh& BaseMesh, 
+		FMesh& OutNewMesh, const FMesh& BaseMesh, 
 		TArrayView<const FReshapePointBindingData> BindingData,
 		TArrayView<const int32> UsedIndices,
 		TArrayView<const int32> BodyOffsets,
@@ -675,7 +675,7 @@ namespace mu
 
 		int32 NumProcessedBindPoints = 0;
 
-		auto ApplyPhysicsBody = [&](PhysicsBody& OutBody, int32 IndicesBegin, int32 IndicesEnd) -> void
+		auto ApplyPhysicsBody = [&](FPhysicsBody& OutBody, int32 IndicesBegin, int32 IndicesEnd) -> void
 		{
 			TArrayView<const int32> BodyUsedIndices(UsedIndices.GetData() + IndicesBegin, IndicesEnd - IndicesBegin);
 
@@ -690,7 +690,7 @@ namespace mu
 		// Apply main physics body
 		if (BaseMesh.PhysicsBody)
 		{
-			Ptr<PhysicsBody> NewBody = BaseMesh.PhysicsBody->Clone();
+			TSharedPtr<FPhysicsBody> NewBody = BaseMesh.PhysicsBody->Clone();
 			ApplyPhysicsBody(*NewBody, 0, BodyOffsets[1]);
 			OutNewMesh.PhysicsBody = NewBody;
 		}
@@ -699,7 +699,7 @@ namespace mu
 		// Apply additional physics bodies
 		for (int32 I = 1; I < PhysicsBodiesNum; ++I)
 		{	
-			Ptr<PhysicsBody> NewBody = BaseMesh.AdditionalPhysicsBodies[I - 1]->Clone();
+			TSharedPtr<FPhysicsBody> NewBody = BaseMesh.AdditionalPhysicsBodies[I - 1]->Clone();
 			ApplyPhysicsBody(*NewBody, BodyOffsets[I], BodyOffsets[I + 1]);
 			OutNewMesh.AdditionalPhysicsBodies[I - 1] = NewBody;
 		}
@@ -711,7 +711,7 @@ namespace mu
 	//! Rebuild the (previously bound) mesh data for a new shape.
 	//! Proof-of-concept implementation.
 	//---------------------------------------------------------------------------------------------
-	inline void MeshApplyShape(Mesh* Result, const Mesh* BaseMesh, const Mesh* ShapeMesh, EMeshBindShapeFlags BindFlags, bool& bOutSuccess)
+	inline void MeshApplyShape(FMesh* Result, const FMesh* BaseMesh, const FMesh* ShapeMesh, EMeshBindShapeFlags BindFlags, bool& bOutSuccess)
 	{
 		MUTABLE_CPUPROFILER_SCOPE(MeshApplyReshape);
 
@@ -748,13 +748,13 @@ namespace mu
 		const FMeshBufferSet& VB = BaseMesh->GetVertexBuffers();
 		VB.FindChannel(EMeshBufferSemantic::BarycentricCoords, BindingDataIndex, &BarycentricDataBuffer, &BarycentricDataChannel);
 		
-		Ptr<Mesh> TempMesh = nullptr;
-		Mesh* VerticesReshapeMesh = Result;
+		TSharedPtr<FMesh> TempMesh = nullptr;
+		FMesh* VerticesReshapeMesh = Result;
 		
 		if (bApplyLaplacian)
 		{
-			TempMesh = new Mesh();
-			VerticesReshapeMesh = TempMesh.get();
+			TempMesh = MakeShared<FMesh>();
+			VerticesReshapeMesh = TempMesh.Get();
 		}
 
 		// Copy Without VertexBuffers or AdditionalBuffers
@@ -920,7 +920,7 @@ namespace mu
 		// order of operation is important.
 
 		// Transform physics volumes based on the deformed sampling points.
-		const PhysicsBody* OldPhysicsBody = Result->PhysicsBody.get();
+		const FPhysicsBody* OldPhysicsBody = Result->PhysicsBody.Get();
 		const int32 AdditionalPhysicsBodiesNum = Result->AdditionalPhysicsBodies.Num();
 
 		if (bReshapePhysicsVolumes && (OldPhysicsBody || AdditionalPhysicsBodiesNum))
diff --git a/Engine/Plugins/Mutable/Source/MutableRuntime/Private/MuR/OpMeshSmoothing.h b/Engine/Plugins/Mutable/Source/MutableRuntime/Private/MuR/OpMeshSmoothing.h
--- a/Engine/Plugins/Mutable/Source/MutableRuntime/Private/MuR/OpMeshSmoothing.h
+++ b/Engine/Plugins/Mutable/Source/MutableRuntime/Private/MuR/OpMeshSmoothing.h
@@ -444,7 +444,7 @@ namespace mu
 		return MakeTuple(MoveTemp(VertexRingsOffsets), MoveTemp(VertexRings));
 	}
 
-	inline void SmoothMeshLaplacian(Mesh& DstMesh, const Mesh& SrcMesh, int32 Iterations = 1)
+	inline void SmoothMeshLaplacian(FMesh& DstMesh, const FMesh& SrcMesh, int32 Iterations = 1)
 	{
 		MUTABLE_CPUPROFILER_SCOPE(MeshSmoothing_SmoothMeshLaplacian);
 
diff --git a/Engine/Plugins/Mutable/Source/MutableRuntime/Private/MuR/OpMeshTransform.h b/Engine/Plugins/Mutable/Source/MutableRuntime/Private/MuR/OpMeshTransform.h
--- a/Engine/Plugins/Mutable/Source/MutableRuntime/Private/MuR/OpMeshTransform.h
+++ b/Engine/Plugins/Mutable/Source/MutableRuntime/Private/MuR/OpMeshTransform.h
@@ -13,7 +13,7 @@ namespace mu
     //---------------------------------------------------------------------------------------------
     //! Reference version
     //---------------------------------------------------------------------------------------------
-    inline void MeshTransform(Mesh* Result, const Mesh* pBase, const FMatrix44f& transform, bool& bOutSuccess)
+    inline void MeshTransform(FMesh* Result, const FMesh* pBase, const FMatrix44f& transform, bool& bOutSuccess)
 	{
 		bOutSuccess = true;
 
diff --git a/Engine/Plugins/Mutable/Source/MutableRuntime/Private/MuR/OpMeshTransformWithMesh.cpp b/Engine/Plugins/Mutable/Source/MutableRuntime/Private/MuR/OpMeshTransformWithMesh.cpp
--- a/Engine/Plugins/Mutable/Source/MutableRuntime/Private/MuR/OpMeshTransformWithMesh.cpp
+++ b/Engine/Plugins/Mutable/Source/MutableRuntime/Private/MuR/OpMeshTransformWithMesh.cpp
@@ -10,7 +10,7 @@
 
 namespace mu
 {
-void MeshTransformWithMesh(Mesh* Result, const Mesh* SourceMesh, const Mesh* BoundingMesh, const FMatrix44f& Transform, bool& bOutSuccess)
+void MeshTransformWithMesh(FMesh* Result, const FMesh* SourceMesh, const FMesh* BoundingMesh, const FMatrix44f& Transform, bool& bOutSuccess)
 {
 	MUTABLE_CPUPROFILER_SCOPE(MeshClipWithMesh);
 	bOutSuccess = true;
diff --git a/Engine/Plugins/Mutable/Source/MutableRuntime/Private/MuR/OpMeshTransformWithMesh.h b/Engine/Plugins/Mutable/Source/MutableRuntime/Private/MuR/OpMeshTransformWithMesh.h
--- a/Engine/Plugins/Mutable/Source/MutableRuntime/Private/MuR/OpMeshTransformWithMesh.h
+++ b/Engine/Plugins/Mutable/Source/MutableRuntime/Private/MuR/OpMeshTransformWithMesh.h
@@ -6,9 +6,9 @@
 
 namespace mu
 {
-class Mesh;
+class FMesh;
 
 /**  */
-extern void MeshTransformWithMesh(Mesh* Result, const Mesh* SourceMesh, const Mesh* BoundingMesh, const FMatrix44f& Transform, bool& bOutSuccess);
+extern void MeshTransformWithMesh(FMesh* Result, const FMesh* SourceMesh, const FMesh* BoundingMesh, const FMatrix44f& Transform, bool& bOutSuccess);
 
 }
diff --git a/Engine/Plugins/Mutable/Source/MutableRuntime/Private/MuR/Operations.cpp b/Engine/Plugins/Mutable/Source/MutableRuntime/Private/MuR/Operations.cpp
--- a/Engine/Plugins/Mutable/Source/MutableRuntime/Private/MuR/Operations.cpp
+++ b/Engine/Plugins/Mutable/Source/MutableRuntime/Private/MuR/Operations.cpp
@@ -149,10 +149,10 @@ namespace mu
 
     // clang-format on
 
-    static_assert( sizeof(s_opDescs)/sizeof(FOpDesc) == (int32)OP_TYPE::COUNT, "OperationDescMismatch");
+    static_assert( sizeof(s_opDescs)/sizeof(FOpDesc) == (int32)EOpType::COUNT, "OperationDescMismatch");
 
 	//---------------------------------------------------------------------------------------------
-	const FOpDesc& GetOpDesc( OP_TYPE type )
+	const FOpDesc& GetOpDesc( EOpType type )
 	{
         return s_opDescs[ (int32)type ];
 	}
@@ -164,104 +164,104 @@ namespace mu
 		// Only operations that still use ASTOpFixed should be handled here.
         switch ( op.type )
         {
-        case OP_TYPE::NONE:
-        case OP_TYPE::BO_CONSTANT:
-        case OP_TYPE::NU_CONSTANT:
-        case OP_TYPE::SC_CONSTANT:
-        case OP_TYPE::CO_CONSTANT:
-        case OP_TYPE::IM_CONSTANT:
-        case OP_TYPE::ME_CONSTANT:
-        case OP_TYPE::LA_CONSTANT:
-        case OP_TYPE::PR_CONSTANT:
-        case OP_TYPE::BO_PARAMETER:
-        case OP_TYPE::NU_PARAMETER:
-        case OP_TYPE::SC_PARAMETER:
-        case OP_TYPE::CO_PARAMETER:
-        case OP_TYPE::PR_PARAMETER:
-        case OP_TYPE::IM_PARAMETER:
+        case EOpType::NONE:
+        case EOpType::BO_CONSTANT:
+        case EOpType::NU_CONSTANT:
+        case EOpType::SC_CONSTANT:
+        case EOpType::CO_CONSTANT:
+        case EOpType::IM_CONSTANT:
+        case EOpType::ME_CONSTANT:
+        case EOpType::LA_CONSTANT:
+        case EOpType::PR_CONSTANT:
+        case EOpType::BO_PARAMETER:
+        case EOpType::NU_PARAMETER:
+        case EOpType::SC_PARAMETER:
+        case EOpType::CO_PARAMETER:
+        case EOpType::PR_PARAMETER:
+        case EOpType::IM_PARAMETER:
             break;
 
         //-------------------------------------------------------------------------------------
-        case OP_TYPE::BO_LESS:
+        case EOpType::BO_LESS:
             f(&op.args.BoolLess.a );
             f(&op.args.BoolLess.b );
             break;
 
         //-------------------------------------------------------------------------------------
-        case OP_TYPE::BO_EQUAL_INT_CONST:
+        case EOpType::BO_EQUAL_INT_CONST:
             f(&op.args.BoolEqualScalarConst.value );
             break;
 
         //-------------------------------------------------------------------------------------
-        case OP_TYPE::BO_AND:
-        case OP_TYPE::BO_OR:
+        case EOpType::BO_AND:
+        case EOpType::BO_OR:
             f(&op.args.BoolBinary.a );
             f(&op.args.BoolBinary.b );
             break;
 
         //-------------------------------------------------------------------------------------
-        case OP_TYPE::BO_NOT:
+        case EOpType::BO_NOT:
             f(&op.args.BoolNot.source );
             break;
 
         //-------------------------------------------------------------------------------------
-        case OP_TYPE::SC_MULTIPLYADD:
+        case EOpType::SC_MULTIPLYADD:
             f(&op.args.ScalarMultiplyAdd.factor0 );
             f(&op.args.ScalarMultiplyAdd.factor1 );
             f(&op.args.ScalarMultiplyAdd.add );
             break;
 
         //-------------------------------------------------------------------------------------
-        case OP_TYPE::SC_ARITHMETIC:
+        case EOpType::SC_ARITHMETIC:
             f(&op.args.ScalarArithmetic.a );
             f(&op.args.ScalarArithmetic.b );
             break;
 
         //-------------------------------------------------------------------------------------
-        case OP_TYPE::CO_SAMPLEIMAGE:
+        case EOpType::CO_SAMPLEIMAGE:
             f(&op.args.ColourSampleImage.image );
             f(&op.args.ColourSampleImage.x );
             f(&op.args.ColourSampleImage.y );
             break;
 
-        case OP_TYPE::CO_SWIZZLE:
+        case EOpType::CO_SWIZZLE:
             for (int32 t=0;t<MUTABLE_OP_MAX_SWIZZLE_CHANNELS;++t)
             {
                 f(&op.args.ColourSwizzle.sources[t] );
             }
             break;
 
-        case OP_TYPE::CO_FROMSCALARS:
+        case EOpType::CO_FROMSCALARS:
 			for (int32 t = 0; t < MUTABLE_OP_MAX_SWIZZLE_CHANNELS; ++t)
 			{
 				f(&op.args.ColourFromScalars.v[t]);
 			}
             break;
 
-        case OP_TYPE::CO_ARITHMETIC:
+        case EOpType::CO_ARITHMETIC:
             f(&op.args.ColourArithmetic.a);
             f(&op.args.ColourArithmetic.b);
             break;
 
         //-------------------------------------------------------------------------------------
-        case OP_TYPE::IM_RESIZE:
+        case EOpType::IM_RESIZE:
             f(&op.args.ImageResize.source );
             break;
 
-        case OP_TYPE::IM_RESIZELIKE:
+        case EOpType::IM_RESIZELIKE:
             f(&op.args.ImageResizeLike.source );
             f(&op.args.ImageResizeLike.sizeSource );
             break;
 
-        case OP_TYPE::IM_RESIZEREL:
+        case EOpType::IM_RESIZEREL:
             f(&op.args.ImageResizeRel.source );
             break;
 
-        case OP_TYPE::IM_BLANKLAYOUT:
+        case EOpType::IM_BLANKLAYOUT:
             f(&op.args.ImageBlankLayout.layout );
             break;
 
-        case OP_TYPE::IM_INTERPOLATE:
+        case EOpType::IM_INTERPOLATE:
             f(&op.args.ImageInterpolate.factor );
 
             for (int t=0;t<MUTABLE_OP_MAX_INTERPOLATE_COUNT;++t)
@@ -270,56 +270,56 @@ namespace mu
             }
             break;
 
-        case OP_TYPE::IM_SATURATE:
+        case EOpType::IM_SATURATE:
             f(&op.args.ImageSaturate.base );
             f(&op.args.ImageSaturate.factor );
             break;
 
-        case OP_TYPE::IM_LUMINANCE:
+        case EOpType::IM_LUMINANCE:
             f(&op.args.ImageLuminance.base );
             break;
 
-        case OP_TYPE::IM_COLOURMAP:
+        case EOpType::IM_COLOURMAP:
             f(&op.args.ImageColourMap.base );
             f(&op.args.ImageColourMap.mask );
             f(&op.args.ImageColourMap.map );
             break;
 
-        case OP_TYPE::IM_GRADIENT:
+        case EOpType::IM_GRADIENT:
             f(&op.args.ImageGradient.colour0 );
             f(&op.args.ImageGradient.colour1 );
             break;
 
-        case OP_TYPE::IM_BINARISE:
+        case EOpType::IM_BINARISE:
             f(&op.args.ImageBinarise.base );
             f(&op.args.ImageBinarise.threshold );
             break;
 
-        case OP_TYPE::IM_PLAINCOLOUR:
+        case EOpType::IM_PLAINCOLOUR:
             f(&op.args.ImagePlainColour.colour );
             break;
 
-        case OP_TYPE::IM_DISPLACE:
+        case EOpType::IM_DISPLACE:
             f(&op.args.ImageDisplace.source );
             f(&op.args.ImageDisplace.displacementMap );
             break;
 
-		case OP_TYPE::IM_INVERT:
+		case EOpType::IM_INVERT:
 			f(&op.args.ImageInvert.base);
 			break;
 
         //-------------------------------------------------------------------------------------
-        case OP_TYPE::ME_APPLYLAYOUT:
+        case EOpType::ME_APPLYLAYOUT:
             f(&op.args.MeshApplyLayout.layout );
             f(&op.args.MeshApplyLayout.mesh );
             break;
 
-        case OP_TYPE::ME_MERGE:
+        case EOpType::ME_MERGE:
             f(&op.args.MeshMerge.base );
             f(&op.args.MeshMerge.added );
             break;
 
-        case OP_TYPE::ME_INTERPOLATE:
+        case EOpType::ME_INTERPOLATE:
             f(&op.args.MeshInterpolate.factor );
             f(&op.args.MeshInterpolate.base );
 
@@ -329,26 +329,26 @@ namespace mu
             }
             break;
 
-        case OP_TYPE::ME_MASKDIFF:
+        case EOpType::ME_MASKDIFF:
             f(&op.args.MeshMaskDiff.source );
             f(&op.args.MeshMaskDiff.fragment );
             break;
 
-        case OP_TYPE::ME_CLIPMORPHPLANE:
+        case EOpType::ME_CLIPMORPHPLANE:
             f(&op.args.MeshClipMorphPlane.source);
             break;
 
-        case OP_TYPE::ME_CLIPWITHMESH :
+        case EOpType::ME_CLIPWITHMESH :
             f(&op.args.MeshClipWithMesh.source);
             f(&op.args.MeshClipWithMesh.clipMesh);
             break;
 
-        case OP_TYPE::ME_SETSKELETON :
+        case EOpType::ME_SETSKELETON :
             f(&op.args.MeshSetSkeleton.source);
             f(&op.args.MeshSetSkeleton.skeleton);
             break;
 
-		case OP_TYPE::ME_PROJECT:
+		case EOpType::ME_PROJECT:
 			f(&op.args.MeshProject.mesh);
 			f(&op.args.MeshProject.projector);
 			break;
@@ -365,47 +365,47 @@ namespace mu
     //---------------------------------------------------------------------------------------------
     void ForEachReference( const FProgram& program, OP::ADDRESS at, const TFunctionRef<void(OP::ADDRESS)> f )
     {
-        OP_TYPE type = program.GetOpType(at);
+        EOpType type = program.GetOpType(at);
         switch ( type )
         {
-        case OP_TYPE::NONE:
-        case OP_TYPE::BO_CONSTANT:
-        case OP_TYPE::NU_CONSTANT:
-        case OP_TYPE::SC_CONSTANT:
-        case OP_TYPE::ST_CONSTANT:
-        case OP_TYPE::CO_CONSTANT:
-        case OP_TYPE::IM_CONSTANT:
-        case OP_TYPE::ME_CONSTANT:
-        case OP_TYPE::LA_CONSTANT:
-        case OP_TYPE::PR_CONSTANT:
-		case OP_TYPE::ED_CONSTANT:
-		case OP_TYPE::MA_CONSTANT:
-        case OP_TYPE::BO_PARAMETER:
-        case OP_TYPE::NU_PARAMETER:
-        case OP_TYPE::SC_PARAMETER:
-        case OP_TYPE::CO_PARAMETER:
-        case OP_TYPE::PR_PARAMETER:
-        case OP_TYPE::IM_PARAMETER:
-        case OP_TYPE::MA_PARAMETER:
-		case OP_TYPE::IM_REFERENCE:
-		case OP_TYPE::ME_REFERENCE:
+        case EOpType::NONE:
+        case EOpType::BO_CONSTANT:
+        case EOpType::NU_CONSTANT:
+        case EOpType::SC_CONSTANT:
+        case EOpType::ST_CONSTANT:
+        case EOpType::CO_CONSTANT:
+        case EOpType::IM_CONSTANT:
+        case EOpType::ME_CONSTANT:
+        case EOpType::LA_CONSTANT:
+        case EOpType::PR_CONSTANT:
+		case EOpType::ED_CONSTANT:
+		case EOpType::MA_CONSTANT:
+        case EOpType::BO_PARAMETER:
+        case EOpType::NU_PARAMETER:
+        case EOpType::SC_PARAMETER:
+        case EOpType::CO_PARAMETER:
+        case EOpType::PR_PARAMETER:
+        case EOpType::IM_PARAMETER:
+        case EOpType::MA_PARAMETER:
+		case EOpType::IM_REFERENCE:
+		case EOpType::ME_REFERENCE:
 			break;
 
-        case OP_TYPE::SC_CURVE:
+        case EOpType::SC_CURVE:
         {
 			OP::ScalarCurveArgs args = program.GetOpArgs<OP::ScalarCurveArgs>(at);
             f(args.time );
             break;
         }
 
-        case OP_TYPE::NU_CONDITIONAL:
-        case OP_TYPE::SC_CONDITIONAL:
-        case OP_TYPE::CO_CONDITIONAL:
-        case OP_TYPE::IM_CONDITIONAL:
-        case OP_TYPE::ME_CONDITIONAL:
-        case OP_TYPE::LA_CONDITIONAL:
-        case OP_TYPE::IN_CONDITIONAL:
-		case OP_TYPE::ED_CONDITIONAL:
+        case EOpType::NU_CONDITIONAL:
+        case EOpType::SC_CONDITIONAL:
+        case EOpType::CO_CONDITIONAL:
+        case EOpType::IM_CONDITIONAL:
+        case EOpType::ME_CONDITIONAL:
+        case EOpType::LA_CONDITIONAL:
+        case EOpType::IN_CONDITIONAL:
+		case EOpType::ED_CONDITIONAL:
         {
 			OP::ConditionalArgs args = program.GetOpArgs<OP::ConditionalArgs>(at);
             f(args.condition );
@@ -414,14 +414,14 @@ namespace mu
             break;
         }
 
-        case OP_TYPE::NU_SWITCH:
-        case OP_TYPE::SC_SWITCH:
-        case OP_TYPE::CO_SWITCH:
-        case OP_TYPE::IM_SWITCH:
-        case OP_TYPE::ME_SWITCH:
-        case OP_TYPE::LA_SWITCH:
-        case OP_TYPE::IN_SWITCH:
-		case OP_TYPE::ED_SWITCH:
+        case EOpType::NU_SWITCH:
+        case EOpType::SC_SWITCH:
+        case EOpType::CO_SWITCH:
+        case EOpType::IM_SWITCH:
+        case EOpType::ME_SWITCH:
+        case EOpType::LA_SWITCH:
+        case EOpType::IN_SWITCH:
+		case EOpType::ED_SWITCH:
         {
 			const uint8_t* data = program.GetOpArgsPointer(at);
 			
@@ -460,7 +460,7 @@ namespace mu
         }
 
         //-------------------------------------------------------------------------------------
-        case OP_TYPE::BO_LESS:
+        case EOpType::BO_LESS:
         {
 			OP::BoolLessArgs args = program.GetOpArgs<OP::BoolLessArgs>(at);
             f(args.a );
@@ -469,7 +469,7 @@ namespace mu
         }
 
         //-------------------------------------------------------------------------------------
-        case OP_TYPE::BO_EQUAL_INT_CONST:
+        case EOpType::BO_EQUAL_INT_CONST:
         {
 			OP::BoolEqualScalarConstArgs args = program.GetOpArgs<OP::BoolEqualScalarConstArgs>(at);
             f(args.value );
@@ -477,8 +477,8 @@ namespace mu
         }
 
         //-------------------------------------------------------------------------------------
-        case OP_TYPE::BO_AND:
-        case OP_TYPE::BO_OR:
+        case EOpType::BO_AND:
+        case EOpType::BO_OR:
         {
 			OP::BoolBinaryArgs args = program.GetOpArgs<OP::BoolBinaryArgs>(at);
             f(args.a );
@@ -487,7 +487,7 @@ namespace mu
         }
 
         //-------------------------------------------------------------------------------------
-        case OP_TYPE::BO_NOT:
+        case EOpType::BO_NOT:
         {
 			OP::BoolNotArgs args = program.GetOpArgs<OP::BoolNotArgs>(at);
             f(args.source );
@@ -495,7 +495,7 @@ namespace mu
         }
 
         //-------------------------------------------------------------------------------------
-        case OP_TYPE::SC_MULTIPLYADD:
+        case EOpType::SC_MULTIPLYADD:
         {
 			OP::ScalarMultiplyAddArgs args = program.GetOpArgs<OP::ScalarMultiplyAddArgs>(at);
             f(args.factor0 );
@@ -505,7 +505,7 @@ namespace mu
         }
 
         //-------------------------------------------------------------------------------------
-        case OP_TYPE::SC_ARITHMETIC:
+        case EOpType::SC_ARITHMETIC:
         {
 			OP::ArithmeticArgs args = program.GetOpArgs<OP::ArithmeticArgs>(at);
             f(args.a );
@@ -514,7 +514,7 @@ namespace mu
         }
 
         //-------------------------------------------------------------------------------------
-        case OP_TYPE::CO_SAMPLEIMAGE:
+        case EOpType::CO_SAMPLEIMAGE:
         {
 			OP::ColourSampleImageArgs args = program.GetOpArgs<OP::ColourSampleImageArgs>(at);
             f(args.image );
@@ -523,7 +523,7 @@ namespace mu
             break;
         }
 
-        case OP_TYPE::CO_SWIZZLE:
+        case EOpType::CO_SWIZZLE:
         {
 			OP::ColourSwizzleArgs args = program.GetOpArgs<OP::ColourSwizzleArgs>(at);
             for (int t=0;t<MUTABLE_OP_MAX_SWIZZLE_CHANNELS;++t)
@@ -533,7 +533,7 @@ namespace mu
             break;
         }
 
-        case OP_TYPE::CO_FROMSCALARS:
+        case EOpType::CO_FROMSCALARS:
         {
 			OP::ColourFromScalarsArgs args = program.GetOpArgs<OP::ColourFromScalarsArgs>(at);
 			for (int32 t = 0; t < MUTABLE_OP_MAX_SWIZZLE_CHANNELS; ++t)
@@ -543,7 +543,7 @@ namespace mu
             break;
         }
 
-        case OP_TYPE::CO_ARITHMETIC:
+        case EOpType::CO_ARITHMETIC:
         {
 			OP::ArithmeticArgs args = program.GetOpArgs<OP::ArithmeticArgs>(at);
             f(args.a);
@@ -552,7 +552,7 @@ namespace mu
         }
 
         //-------------------------------------------------------------------------------------
-        case OP_TYPE::IM_LAYER:
+        case EOpType::IM_LAYER:
         {
 			OP::ImageLayerArgs args = program.GetOpArgs<OP::ImageLayerArgs>(at);
             f(args.base );
@@ -564,7 +564,7 @@ namespace mu
             break;
         }
 
-        case OP_TYPE::IM_LAYERCOLOUR:
+        case EOpType::IM_LAYERCOLOUR:
         {
 			OP::ImageLayerColourArgs args = program.GetOpArgs<OP::ImageLayerColourArgs>(at);
             f(args.base );
@@ -576,7 +576,7 @@ namespace mu
             break;
         }
 
-        case OP_TYPE::IM_MULTILAYER:
+        case EOpType::IM_MULTILAYER:
         {
 			OP::ImageMultiLayerArgs args = program.GetOpArgs<OP::ImageMultiLayerArgs>(at);
             f(args.rangeSize );
@@ -589,7 +589,7 @@ namespace mu
             break;
         }
 
-		case OP_TYPE::IM_NORMALCOMPOSITE:
+		case EOpType::IM_NORMALCOMPOSITE:
 		{
 			OP::ImageNormalCompositeArgs args = program.GetOpArgs<OP::ImageNormalCompositeArgs>(at);
 			f(args.base);
@@ -598,28 +598,28 @@ namespace mu
 			break;
 		}
 
-        case OP_TYPE::IM_PIXELFORMAT:
+        case EOpType::IM_PIXELFORMAT:
         {
 			OP::ImagePixelFormatArgs args = program.GetOpArgs<OP::ImagePixelFormatArgs>(at);
             f(args.source );
             break;
         }
 
-        case OP_TYPE::IM_MIPMAP:
+        case EOpType::IM_MIPMAP:
         {
 			OP::ImageMipmapArgs args = program.GetOpArgs<OP::ImageMipmapArgs>(at);
             f(args.source );
             break;
         }
 
-        case OP_TYPE::IM_RESIZE:
+        case EOpType::IM_RESIZE:
         {
 			OP::ImageResizeArgs args = program.GetOpArgs<OP::ImageResizeArgs>(at);
             f(args.source );
             break;
         }
 
-        case OP_TYPE::IM_RESIZELIKE:
+        case EOpType::IM_RESIZELIKE:
         {
 			OP::ImageResizeLikeArgs args = program.GetOpArgs<OP::ImageResizeLikeArgs>(at);
             f(args.source );
@@ -627,21 +627,21 @@ namespace mu
             break;
         }
 
-        case OP_TYPE::IM_RESIZEREL:
+        case EOpType::IM_RESIZEREL:
         {
 			OP::ImageResizeRelArgs args = program.GetOpArgs<OP::ImageResizeRelArgs>(at);
             f(args.source );
             break;
         }
 
-        case OP_TYPE::IM_BLANKLAYOUT:
+        case EOpType::IM_BLANKLAYOUT:
         {
 			OP::ImageBlankLayoutArgs args = program.GetOpArgs<OP::ImageBlankLayoutArgs>(at);
             f(args.layout );
             break;
         }
 
-        case OP_TYPE::IM_COMPOSE:
+        case EOpType::IM_COMPOSE:
         {
 			OP::ImageComposeArgs args = program.GetOpArgs<OP::ImageComposeArgs>(at);
             f(args.layout );
@@ -651,7 +651,7 @@ namespace mu
             break;
         }
 
-        case OP_TYPE::IM_INTERPOLATE:
+        case EOpType::IM_INTERPOLATE:
         {
 			OP::ImageInterpolateArgs args = program.GetOpArgs<OP::ImageInterpolateArgs>(at);
             f(args.factor );
@@ -663,7 +663,7 @@ namespace mu
             break;
         }
 
-        case OP_TYPE::IM_SWIZZLE:
+        case EOpType::IM_SWIZZLE:
         {
 			OP::ImageSwizzleArgs args = program.GetOpArgs<OP::ImageSwizzleArgs>(at);
             for (int t=0;t<MUTABLE_OP_MAX_SWIZZLE_CHANNELS;++t)
@@ -673,7 +673,7 @@ namespace mu
             break;
         }
 
-        case OP_TYPE::IM_SATURATE:
+        case EOpType::IM_SATURATE:
         {
 			OP::ImageSaturateArgs args = program.GetOpArgs<OP::ImageSaturateArgs>(at);
             f(args.base );
@@ -681,14 +681,14 @@ namespace mu
             break;
         }
 
-        case OP_TYPE::IM_LUMINANCE:
+        case EOpType::IM_LUMINANCE:
         {
 			OP::ImageLuminanceArgs args = program.GetOpArgs<OP::ImageLuminanceArgs>(at);
             f(args.base );
             break;
         }
 
-        case OP_TYPE::IM_COLOURMAP:
+        case EOpType::IM_COLOURMAP:
         {
 			OP::ImageColourMapArgs args = program.GetOpArgs<OP::ImageColourMapArgs>(at);
             f(args.base );
@@ -697,7 +697,7 @@ namespace mu
             break;
         }
 
-        case OP_TYPE::IM_GRADIENT:
+        case EOpType::IM_GRADIENT:
         {
 			OP::ImageGradientArgs args = program.GetOpArgs<OP::ImageGradientArgs>(at);
             f(args.colour0 );
@@ -705,7 +705,7 @@ namespace mu
             break;
         }
 
-        case OP_TYPE::IM_BINARISE:
+        case EOpType::IM_BINARISE:
         {
 			OP::ImageBinariseArgs args = program.GetOpArgs<OP::ImageBinariseArgs>(at);
             f(args.base );
@@ -713,21 +713,21 @@ namespace mu
             break;
         }
 
-        case OP_TYPE::IM_PLAINCOLOUR:
+        case EOpType::IM_PLAINCOLOUR:
         {
 			OP::ImagePlainColourArgs args = program.GetOpArgs<OP::ImagePlainColourArgs>(at);
             f(args.colour );
             break;
         }
 
-        case OP_TYPE::IM_CROP:
+        case EOpType::IM_CROP:
         {
 			OP::ImageCropArgs args = program.GetOpArgs<OP::ImageCropArgs>(at);
             f(args.source );
             break;
         }
 
-        case OP_TYPE::IM_PATCH:
+        case EOpType::IM_PATCH:
         {
 			OP::ImagePatchArgs args = program.GetOpArgs<OP::ImagePatchArgs>(at);
             f(args.base );
@@ -735,7 +735,7 @@ namespace mu
             break;
         }
 
-        case OP_TYPE::IM_RASTERMESH:
+        case EOpType::IM_RASTERMESH:
         {
 			OP::ImageRasterMeshArgs args = program.GetOpArgs<OP::ImageRasterMeshArgs>(at);
             f(args.mesh );
@@ -746,14 +746,14 @@ namespace mu
             break;
         }
 
-        case OP_TYPE::IM_MAKEGROWMAP:
+        case EOpType::IM_MAKEGROWMAP:
         {
 			OP::ImageMakeGrowMapArgs args = program.GetOpArgs<OP::ImageMakeGrowMapArgs>(at);
             f(args.mask );
             break;
         }
 
-        case OP_TYPE::IM_DISPLACE:
+        case EOpType::IM_DISPLACE:
         {
 			OP::ImageDisplaceArgs args = program.GetOpArgs<OP::ImageDisplaceArgs>(at);
             f(args.source );
@@ -761,14 +761,14 @@ namespace mu
             break;
         }
 
-		case OP_TYPE::IM_INVERT:
+		case EOpType::IM_INVERT:
 		{
 			OP::ImageInvertArgs args = program.GetOpArgs<OP::ImageInvertArgs>(at);
 			f(args.base);
 			break;
 		}
 
-		case OP_TYPE::IM_TRANSFORM:
+		case EOpType::IM_TRANSFORM:
 		{
 			OP::ImageTransformArgs Args = program.GetOpArgs<OP::ImageTransformArgs>(at);
 			f(Args.Base);
@@ -781,7 +781,7 @@ namespace mu
 		}
 
         //-------------------------------------------------------------------------------------
-        case OP_TYPE::ME_APPLYLAYOUT:
+        case EOpType::ME_APPLYLAYOUT:
         {
 			OP::MeshApplyLayoutArgs args = program.GetOpArgs<OP::MeshApplyLayoutArgs>(at);
             f(args.layout );
@@ -789,7 +789,7 @@ namespace mu
             break;
         }
 
-        case OP_TYPE::ME_DIFFERENCE:
+        case EOpType::ME_DIFFERENCE:
         {
 			const uint8_t* data = program.GetOpArgsPointer(at);
 
@@ -803,7 +803,7 @@ namespace mu
 			break;
         }
 
-        case OP_TYPE::ME_MORPH:
+        case EOpType::ME_MORPH:
         {
 			const uint8_t* data = program.GetOpArgsPointer(at);
 
@@ -821,7 +821,7 @@ namespace mu
 			break;
         }
 
-        case OP_TYPE::ME_MERGE:
+        case EOpType::ME_MERGE:
         {
 			OP::MeshMergeArgs args = program.GetOpArgs<OP::MeshMergeArgs>(at);
             f(args.base );
@@ -829,7 +829,7 @@ namespace mu
             break;
         }
 
-        case OP_TYPE::ME_INTERPOLATE:
+        case EOpType::ME_INTERPOLATE:
         {
 			OP::MeshInterpolateArgs args = program.GetOpArgs<OP::MeshInterpolateArgs>(at);
             f(args.factor );
@@ -842,7 +842,7 @@ namespace mu
             break;
         }
 
-		case OP_TYPE::ME_MASKCLIPMESH:
+		case EOpType::ME_MASKCLIPMESH:
 		{
 			OP::MeshMaskClipMeshArgs args = program.GetOpArgs<OP::MeshMaskClipMeshArgs>(at);
 			f(args.source);
@@ -850,7 +850,7 @@ namespace mu
 			break;
 		}
 
-		case OP_TYPE::ME_MASKCLIPUVMASK:
+		case EOpType::ME_MASKCLIPUVMASK:
 		{
 			OP::MeshMaskClipUVMaskArgs args = program.GetOpArgs<OP::MeshMaskClipUVMaskArgs>(at);
 			f(args.Source);
@@ -859,7 +859,7 @@ namespace mu
 			break;
 		}
 
-        case OP_TYPE::ME_MASKDIFF:
+        case EOpType::ME_MASKDIFF:
         {
 			OP::MeshMaskDiffArgs args = program.GetOpArgs<OP::MeshMaskDiffArgs>(at);
             f(args.source );
@@ -867,7 +867,7 @@ namespace mu
             break;
         }
 
-        case OP_TYPE::ME_REMOVEMASK:
+        case EOpType::ME_REMOVEMASK:
         {
             const uint8* data = program.GetOpArgsPointer(at);
             mu::OP::ADDRESS source;
@@ -893,7 +893,7 @@ namespace mu
             break;
         }
 
-		case OP_TYPE::ME_ADDTAGS:
+		case EOpType::ME_ADDTAGS:
 		{
 			const uint8_t* data = program.GetOpArgsPointer(at);
 			mu::OP::ADDRESS source;
@@ -902,7 +902,7 @@ namespace mu
 			break;
 		}
 
-        case OP_TYPE::ME_FORMAT:
+        case EOpType::ME_FORMAT:
         {
 			OP::MeshFormatArgs args = program.GetOpArgs<OP::MeshFormatArgs>(at);
             f(args.source );
@@ -910,14 +910,14 @@ namespace mu
             break;
         }
 
-        case OP_TYPE::ME_TRANSFORM:
+        case EOpType::ME_TRANSFORM:
         {
 			OP::MeshTransformArgs args = program.GetOpArgs<OP::MeshTransformArgs>(at);
             f(args.source );
             break;
         }
 
-        case OP_TYPE::ME_EXTRACTLAYOUTBLOCK:
+        case EOpType::ME_EXTRACTLAYOUTBLOCK:
         {
             const uint8_t* data = program.GetOpArgsPointer(at);
             mu::OP::ADDRESS source;
@@ -926,14 +926,14 @@ namespace mu
             break;
         }
 
-        case OP_TYPE::ME_CLIPMORPHPLANE:
+        case EOpType::ME_CLIPMORPHPLANE:
         {
 			OP::MeshClipMorphPlaneArgs args = program.GetOpArgs<OP::MeshClipMorphPlaneArgs>(at);
             f(args.source);
             break;
         }
 
-        case OP_TYPE::ME_CLIPWITHMESH :
+        case EOpType::ME_CLIPWITHMESH :
         {
 			OP::MeshClipWithMeshArgs args = program.GetOpArgs<OP::MeshClipWithMeshArgs>(at);
             f(args.source);
@@ -941,7 +941,7 @@ namespace mu
             break;
         }
 
-        case OP_TYPE::ME_CLIPDEFORM:
+        case EOpType::ME_CLIPDEFORM:
         {
 			OP::MeshClipDeformArgs args = program.GetOpArgs<OP::MeshClipDeformArgs>(at);
             f(args.mesh);
@@ -949,7 +949,7 @@ namespace mu
             break;
         }
        
-		case OP_TYPE::ME_MORPHRESHAPE:
+		case EOpType::ME_MORPHRESHAPE:
 		{
 			OP::MeshMorphReshapeArgs Args = program.GetOpArgs<OP::MeshMorphReshapeArgs>(at);
 			f(Args.Morph);
@@ -958,7 +958,7 @@ namespace mu
 			break;
 		}
  
-        case OP_TYPE::ME_SETSKELETON :
+        case EOpType::ME_SETSKELETON :
         {
 			OP::MeshSetSkeletonArgs args = program.GetOpArgs<OP::MeshSetSkeletonArgs>(at);
             f(args.source);
@@ -966,7 +966,7 @@ namespace mu
             break;
         }
 
-        case OP_TYPE::ME_PROJECT :
+        case EOpType::ME_PROJECT :
         {
 			OP::MeshProjectArgs args = program.GetOpArgs<OP::MeshProjectArgs>(at);
             f(args.mesh);
@@ -974,7 +974,7 @@ namespace mu
             break;
         }
 
-		case OP_TYPE::ME_APPLYPOSE:
+		case EOpType::ME_APPLYPOSE:
 		{
 			OP::MeshApplyPoseArgs args = program.GetOpArgs<OP::MeshApplyPoseArgs>(at);
 			f(args.base);
@@ -982,7 +982,7 @@ namespace mu
 			break;
 		}
 
-		case OP_TYPE::ME_GEOMETRYOPERATION:
+		case EOpType::ME_GEOMETRYOPERATION:
 		{
 			OP::MeshGeometryOperationArgs args = program.GetOpArgs<OP::MeshGeometryOperationArgs>(at);
 			f(args.meshA);
@@ -992,7 +992,7 @@ namespace mu
 			break;
 		}
 
-		case OP_TYPE::ME_BINDSHAPE:
+		case EOpType::ME_BINDSHAPE:
 		{
 			OP::MeshBindShapeArgs args = program.GetOpArgs<OP::MeshBindShapeArgs>(at);
 			f(args.mesh);
@@ -1000,7 +1000,7 @@ namespace mu
 			break;
 		}
 
-		case OP_TYPE::ME_APPLYSHAPE:
+		case EOpType::ME_APPLYSHAPE:
 		{
 			OP::MeshApplyShapeArgs args = program.GetOpArgs<OP::MeshApplyShapeArgs>(at);
 			f(args.mesh);
@@ -1008,14 +1008,14 @@ namespace mu
 			break;
 		}
 
-		case OP_TYPE::ME_OPTIMIZESKINNING:
+		case EOpType::ME_OPTIMIZESKINNING:
 		{
 			OP::MeshOptimizeSkinningArgs args = program.GetOpArgs<OP::MeshOptimizeSkinningArgs>(at);
 			f(args.source);
 			break;
 		}
 
-        case OP_TYPE::ME_TRANSFORMWITHMESH :
+        case EOpType::ME_TRANSFORMWITHMESH :
         {
         	OP::MeshTransformWithinMeshArgs args = program.GetOpArgs<OP::MeshTransformWithinMeshArgs>(at);
         	f(args.sourceMesh);
@@ -1025,13 +1025,13 @@ namespace mu
         }
         	
         //-------------------------------------------------------------------------------------
-        case OP_TYPE::IN_ADDMESH:
-        case OP_TYPE::IN_ADDIMAGE:
-        case OP_TYPE::IN_ADDVECTOR:
-        case OP_TYPE::IN_ADDSCALAR:
-        case OP_TYPE::IN_ADDSTRING:
-        case OP_TYPE::IN_ADDCOMPONENT:
-        case OP_TYPE::IN_ADDSURFACE:
+        case EOpType::IN_ADDMESH:
+        case EOpType::IN_ADDIMAGE:
+        case EOpType::IN_ADDVECTOR:
+        case EOpType::IN_ADDSCALAR:
+        case EOpType::IN_ADDSTRING:
+        case EOpType::IN_ADDCOMPONENT:
+        case EOpType::IN_ADDSURFACE:
         {
 			OP::InstanceAddArgs args = program.GetOpArgs<OP::InstanceAddArgs>(at);
             f(args.instance );
@@ -1039,7 +1039,7 @@ namespace mu
             break;
         }
 
-        case OP_TYPE::IN_ADDLOD:
+        case EOpType::IN_ADDLOD:
         {
 			const uint8* Data = program.GetOpArgsPointer(at);
 
@@ -1058,7 +1058,7 @@ namespace mu
             break;
         }
 
-		case OP_TYPE::IN_ADDEXTENSIONDATA:
+		case EOpType::IN_ADDEXTENSIONDATA:
 		{
 			const OP::InstanceAddExtensionDataArgs Args = program.GetOpArgs<OP::InstanceAddExtensionDataArgs>(at);
 
@@ -1069,14 +1069,14 @@ namespace mu
 		}
 
         //-------------------------------------------------------------------------------------
-        case OP_TYPE::LA_PACK:
+        case EOpType::LA_PACK:
         {
 			OP::LayoutPackArgs args = program.GetOpArgs<OP::LayoutPackArgs>(at);
             f(args.Source );
             break;
         }
 
-        case OP_TYPE::LA_MERGE:
+        case EOpType::LA_MERGE:
         {
 			OP::LayoutMergeArgs args = program.GetOpArgs<OP::LayoutMergeArgs>(at);
             f(args.Base );
@@ -1084,7 +1084,7 @@ namespace mu
             break;
         }
 
-		case OP_TYPE::LA_REMOVEBLOCKS:
+		case EOpType::LA_REMOVEBLOCKS:
 		{
 			OP::LayoutRemoveBlocksArgs args = program.GetOpArgs<OP::LayoutRemoveBlocksArgs>(at);
 			f(args.Source);
@@ -1092,7 +1092,7 @@ namespace mu
 			break;
 		}
 
-		case OP_TYPE::LA_FROMMESH:
+		case EOpType::LA_FROMMESH:
 		{
 			OP::LayoutFromMeshArgs args = program.GetOpArgs<OP::LayoutFromMeshArgs>(at);
 			f(args.Mesh);
