[mutable] More enforcement of coding standards.

diff --git a/Engine/Plugins/Mutable/Source/MutableRuntime/Private/MuR/Parameters.cpp b/Engine/Plugins/Mutable/Source/MutableRuntime/Private/MuR/Parameters.cpp
--- a/Engine/Plugins/Mutable/Source/MutableRuntime/Private/MuR/Parameters.cpp
+++ b/Engine/Plugins/Mutable/Source/MutableRuntime/Private/MuR/Parameters.cpp
@@ -21,15 +21,13 @@ namespace mu
     MUTABLE_IMPLEMENT_ENUM_SERIALISABLE(EProjectorType)
 
 	
-	//---------------------------------------------------------------------------------------------
-	Parameters::Parameters()
+	FParameters::FParameters()
 	{
 		m_pD = new Private();
 	}
 
 
-	//---------------------------------------------------------------------------------------------
-	Parameters::~Parameters()
+	FParameters::~FParameters()
 	{
         check( m_pD );
 		delete m_pD;
@@ -37,196 +35,181 @@ namespace mu
 	}
 
 
-	//---------------------------------------------------------------------------------------------
-	void Parameters::Serialise( const Parameters* p, OutputArchive& arch )
+	void FParameters::Serialise( const FParameters* p, FOutputArchive& arch )
 	{
 		arch << *p->m_pD;
     }
 
 
-	//---------------------------------------------------------------------------------------------
-	ParametersPtr Parameters::StaticUnserialise( InputArchive& arch )
+	TSharedPtr<FParameters> FParameters::StaticUnserialise( FInputArchive& arch )
 	{
 		LLM_SCOPE_BYNAME(TEXT("MutableRuntime"));
 
-		ParametersPtr pResult = new Parameters();
+		TSharedPtr<FParameters> pResult = MakeShared<FParameters>();
 		arch >> *pResult->m_pD;
 		return pResult;
 	}
 
 
-	//---------------------------------------------------------------------------------------------
-	Parameters::Private* Parameters::GetPrivate() const
+	FParameters::Private* FParameters::GetPrivate() const
 	{
 		return m_pD;
 	}
 
 
-	//---------------------------------------------------------------------------------------------
-	ParametersPtr Parameters::Clone() const
+	TSharedPtr<FParameters> FParameters::Clone() const
 	{
 		LLM_SCOPE_BYNAME(TEXT("MutableRuntime"));
 
-		ParametersPtr pRes = new Parameters();
+		TSharedPtr<FParameters> pRes = MakeShared<FParameters>();
 
-		pRes->m_pD->m_pModel = m_pD->m_pModel;
-		pRes->m_pD->m_values = m_pD->m_values;
-		pRes->m_pD->m_multiValues = m_pD->m_multiValues;
+		pRes->m_pD->Model = m_pD->Model;
+		pRes->m_pD->Values = m_pD->Values;
+		pRes->m_pD->MultiValues = m_pD->MultiValues;
 
 		return pRes;
 	}
 
 
-	//---------------------------------------------------------------------------------------------
-	int32 Parameters::GetCount() const
+	int32 FParameters::GetCount() const
 	{
-		return (int32)m_pD->m_values.Num();
+		return (int32)m_pD->Values.Num();
 	}
 
 
-	//---------------------------------------------------------------------------------------------
-	const FString& Parameters::GetName( int32 index ) const
+	const FString& FParameters::GetName( int32 Index ) const
 	{
-		const FProgram& program = m_pD->m_pModel->GetPrivate()->m_program;
-		check( index>=0 && index<(int32)program.m_parameters.Num() );
+		const FProgram& Program = m_pD->Model->GetPrivate()->Program;
+		check( Index>=0 && Index<(int32)Program.Parameters.Num() );
 
-		return program.m_parameters[index].m_name;
+		return Program.Parameters[Index].Name;
 	}
 
 
-	//---------------------------------------------------------------------------------------------
-	const FGuid& Parameters::GetUid( int32 index ) const
+	const FGuid& FParameters::GetUid( int32 Index ) const
 	{
-		const FProgram& program = m_pD->m_pModel->GetPrivate()->m_program;
-		check( index>=0 && index<(int32)program.m_parameters.Num() );
+		const FProgram& Program = m_pD->Model->GetPrivate()->Program;
+		check( Index>=0 && Index<Program.Parameters.Num() );
 
-		return program.m_parameters[index].m_uid;
+		return Program.Parameters[Index].UID;
 	}
 
 
-	//---------------------------------------------------------------------------------------------
-	int32 Parameters::Find( const FString& strName ) const
+	int32 FParameters::Find( const FString& strName ) const
 	{
 		return m_pD->Find( strName );
 	}
 
 
-	//---------------------------------------------------------------------------------------------
-	EParameterType Parameters::GetType( int32 index ) const
+	EParameterType FParameters::GetType( int32 Index ) const
 	{
-		const FProgram& program = m_pD->m_pModel->GetPrivate()->m_program;
-		check( index>=0 && index<(int32)program.m_parameters.Num() );
+		const FProgram& Program = m_pD->Model->GetPrivate()->Program;
+		check( Index>=0 && Index<Program.Parameters.Num() );
 
-		return program.m_parameters[index].m_type;
+		return Program.Parameters[Index].Type;
 	}
 
 
-    //---------------------------------------------------------------------------------------------
-    Ptr<RangeIndex> Parameters::NewRangeIndex( int32 paramIndex ) const
+    TSharedPtr<FRangeIndex> FParameters::NewRangeIndex( int32 ParamIndex ) const
     {
 		LLM_SCOPE_BYNAME(TEXT("MutableRuntime"));
 
-        const FProgram& program = m_pD->m_pModel->GetPrivate()->m_program;
-        check( paramIndex>=0 && paramIndex<int32(program.m_parameters.Num()) );
-        Ptr<RangeIndex> range;
+        const FProgram& Program = m_pD->Model->GetPrivate()->Program;
+        check( ParamIndex>=0 && ParamIndex<Program.Parameters.Num() );
+		TSharedPtr<FRangeIndex> Range;
 
-        if ( paramIndex>=0 && paramIndex<int32(program.m_parameters.Num()) )
+        if ( ParamIndex>=0 && ParamIndex<int32(Program.Parameters.Num()) )
         {
-            size_t rangeCount = program.m_parameters[paramIndex].m_ranges.Num();
-            if ( rangeCount>0 )
+            int32 RangeCount = Program.Parameters[ParamIndex].Ranges.Num();
+            if (RangeCount >0 )
             {
-                range = new RangeIndex;
-                range->m_pD->m_pParameters = this;
-                range->m_pD->m_parameter = paramIndex;
-                range->m_pD->m_values.SetNumZeroed( rangeCount );
+				Range = MakeShared<FRangeIndex>();
+				Range->Parameters = this->AsShared();
+				Range->Parameter = ParamIndex;
+				Range->Values.SetNumZeroed(RangeCount);
             }
         }
 
-        return range;
+        return Range;
     }
 
 
-    //---------------------------------------------------------------------------------------------
-    int32 Parameters::GetValueCount( int32 paramIndex ) const
+    int32 FParameters::GetValueCount( int32 ParamIndex ) const
     {
-        if (paramIndex<0 || paramIndex >= int32(m_pD->m_multiValues.Num()) )
+        if (ParamIndex<0 || ParamIndex >= int32(m_pD->MultiValues.Num()) )
         {
             return 0;
         }
 
-        return int32( m_pD->m_multiValues[paramIndex].Num() );
+        return int32( m_pD->MultiValues[ParamIndex].Num() );
     }
 
 
-    //---------------------------------------------------------------------------------------------
-    Ptr<RangeIndex> Parameters::GetValueIndex( int32 paramIndex, int32 valueIndex ) const
+    TSharedPtr<FRangeIndex> FParameters::GetValueIndex( int32 ParamIndex, int32 ValueIndex ) const
     {
-        if (paramIndex<0 || paramIndex >= int32(m_pD->m_multiValues.Num()) )
+        if (ParamIndex<0 || ParamIndex >= m_pD->MultiValues.Num() )
         {
             return nullptr;
         }
 
-        if (valueIndex<0 || valueIndex >= int32(m_pD->m_multiValues[paramIndex].Num()) )
+        if (ValueIndex <0 || ValueIndex >= int32(m_pD->MultiValues[ParamIndex].Num()) )
         {
             return nullptr;
         }
 
-		TMap< TArray<int32>, FParameterValue >::TRangedForIterator it = m_pD->m_multiValues[paramIndex].begin();
-        for ( int32 i=0; i<valueIndex; ++i )
+		TMap< TArray<int32>, FParameterValue >::TRangedForIterator it = m_pD->MultiValues[ParamIndex].begin();
+        for ( int32 i=0; i< ValueIndex; ++i )
         {
             ++it;
         }
 
-        Ptr<RangeIndex> result = NewRangeIndex( paramIndex );
-        result->m_pD->m_values = it->Key;
+		TSharedPtr<FRangeIndex> result = NewRangeIndex( ParamIndex );
+        result->Values = it->Key;
 
         return result;
     }
 
 
-    //---------------------------------------------------------------------------------------------
-    void Parameters::ClearAllValues( int32 paramIndex )
+    void FParameters::ClearAllValues( int32 ParamIndex )
     {
-        if (paramIndex<0 || paramIndex >= int32(m_pD->m_multiValues.Num()) )
+        if (ParamIndex<0 || ParamIndex >= int32(m_pD->MultiValues.Num()) )
         {
             return;
         }
 
-        m_pD->m_multiValues[paramIndex].Empty();
+        m_pD->MultiValues[ParamIndex].Empty();
     }
 
 
-	//---------------------------------------------------------------------------------------------
-    bool Parameters::GetBoolValue( int32 index,
-                                   const Ptr<const RangeIndex>& pos ) const
+    bool FParameters::GetBoolValue( int32 Index, const FRangeIndex* Pos ) const
 	{
-		check( index>=0 && index<(int32)m_pD->m_values.Num() );
-        check( GetType(index)==EParameterType::Bool );
+		check( Index>=0 && Index<(int32)m_pD->Values.Num() );
+        check( GetType(Index)==EParameterType::Bool );
 
         // Early out in case of invalid parameters
-        if ( index < 0
+        if ( Index < 0
              ||
-             index >= (int32)m_pD->m_values.Num()
+             Index >= (int32)m_pD->Values.Num()
              ||
-             GetType(index) != EParameterType::Bool )
+             GetType(Index) != EParameterType::Bool )
         {
             return false;
         }
 
         // Single value case
-        if (!pos)
+        if (!Pos)
         {
             // Return the single value
-            return m_pD->m_values[index].Get<FParamBoolType>();
+            return m_pD->Values[Index].Get<FParamBoolType>();
         }
 
         // Multivalue case
-        check( pos->m_pD->m_parameter==index );
+        check( Pos->Parameter==Index );
 
-        if ( index<int32(m_pD->m_multiValues.Num()))
+        if ( Index<int32(m_pD->MultiValues.Num()))
         {
-            const TMap< TArray<int32>, FParameterValue >& m = m_pD->m_multiValues[index];
-			const FParameterValue* it = m.Find(pos->m_pD->m_values);
+            const TMap< TArray<int32>, FParameterValue >& m = m_pD->MultiValues[Index];
+			const FParameterValue* it = m.Find(Pos->Values);
             if (it)
             {
                 return it->Get<FParamBoolType>();
@@ -234,89 +217,85 @@ namespace mu
         }
 
         // Multivalue parameter, but no multivalue set. Return single value.
-        return m_pD->m_values[index].Get<FParamBoolType>();
+        return m_pD->Values[Index].Get<FParamBoolType>();
 	}
 
 
-	//---------------------------------------------------------------------------------------------
-    void Parameters::SetBoolValue( int32 index, bool value,
-                                   const Ptr<const RangeIndex>& pos )
+    void FParameters::SetBoolValue( int32 Index, bool value,
+                                   const FRangeIndex* Pos )
 	{
-		check( index>=0 && index<(int32)m_pD->m_values.Num() );
-		check( GetType(index)== EParameterType::Bool );
+		check( Index>=0 && Index<(int32)m_pD->Values.Num() );
+		check( GetType(Index)== EParameterType::Bool );
 
         // Early out in case of invalid parameters
-        if ( index < 0
+        if ( Index < 0
              ||
-             index >= (int32)m_pD->m_values.Num()
+             Index >= (int32)m_pD->Values.Num()
              ||
-             GetType(index) != EParameterType::Bool )
+             GetType(Index) != EParameterType::Bool )
         {
             return;
         }
 
         // Single value case
-        if (!pos)
+        if (!Pos)
         {
             // Clear multivalue, if set.
-            if (index<int32(m_pD->m_multiValues.Num()))
+            if (Index<int32(m_pD->MultiValues.Num()))
             {
-                m_pD->m_multiValues[index].Empty();
+                m_pD->MultiValues[Index].Empty();
             }
 
-            m_pD->m_values[index].Set<FParamBoolType>(value);
+            m_pD->Values[Index].Set<FParamBoolType>(value);
         }
 
         // Multivalue case
         else
         {
-            check( pos->m_pD->m_parameter==index );
+            check( Pos->Parameter==Index );
 
-            if ( index>=int32(m_pD->m_multiValues.Num()))
+            if ( Index>=int32(m_pD->MultiValues.Num()))
             {
-                m_pD->m_multiValues.SetNum(index+1);
+                m_pD->MultiValues.SetNum(Index+1);
             }
 
-			TMap< TArray<int32>, FParameterValue >& m = m_pD->m_multiValues[index];
-			FParameterValue& it = m.FindOrAdd(pos->m_pD->m_values);
+			TMap< TArray<int32>, FParameterValue >& m = m_pD->MultiValues[Index];
+			FParameterValue& it = m.FindOrAdd(Pos->Values);
             it.Set<FParamBoolType>(value);
         }
 	}
 
 
-	//---------------------------------------------------------------------------------------------
-	int32 Parameters::GetIntPossibleValueCount( int32 paramIndex ) const
+	int32 FParameters::GetIntPossibleValueCount( int32 ParamIndex ) const
 	{
-		const FProgram& program = m_pD->m_pModel->GetPrivate()->m_program;
-		check( paramIndex>=0 && paramIndex<(int32)program.m_parameters.Num() );
-		return (int32)program.m_parameters[paramIndex].m_possibleValues.Num();
+		const FProgram& Program = m_pD->Model->GetPrivate()->Program;
+		check( ParamIndex>=0 && ParamIndex<(int32)Program.Parameters.Num() );
+		return (int32)Program.Parameters[ParamIndex].PossibleValues.Num();
 	}
 
 
-	//---------------------------------------------------------------------------------------------
-	int32 Parameters::GetIntPossibleValue( int32 paramIndex, int32 valueIndex ) const
+	int32 FParameters::GetIntPossibleValue( int32 ParamIndex, int32 valueIndex ) const
 	{
-		const FProgram& program = m_pD->m_pModel->GetPrivate()->m_program;
-		check( paramIndex>=0
-				&& paramIndex<(int32)program.m_parameters.Num() );
+		const FProgram& Program = m_pD->Model->GetPrivate()->Program;
+		check( ParamIndex>=0
+				&& ParamIndex<(int32)Program.Parameters.Num() );
 		check( valueIndex>=0
-				&& valueIndex<(int32)program.m_parameters[paramIndex].m_possibleValues.Num() );
+				&& valueIndex<(int32)Program.Parameters[ParamIndex].PossibleValues.Num() );
 
-		return (int32)program.m_parameters[paramIndex].m_possibleValues[valueIndex].m_value;
+		return (int32)Program.Parameters[ParamIndex].PossibleValues[valueIndex].Value;
 	}
 
 
-	//---------------------------------------------------------------------------------------------
-	int32 Parameters::GetIntValueIndex(int32 paramIndex, const FString& ValueName) const
+	int32 FParameters::GetIntValueIndex(int32 ParamIndex, const FString& ValueName) const
 	{
-		const FProgram& program = m_pD->m_pModel->GetPrivate()->m_program;
-		check(paramIndex >= 0
-			&& paramIndex < (int32)program.m_parameters.Num());
+		const FProgram& Program = m_pD->Model->GetPrivate()->Program;
+		check(ParamIndex >= 0
+			&& ParamIndex < (int32)Program.Parameters.Num());
 
 		int32 result = -1;
-		for (size_t v = 0; v < program.m_parameters[paramIndex].m_possibleValues.Num(); ++v)
+		for (size_t v = 0; v < Program.Parameters[ParamIndex].PossibleValues.Num(); ++v)
 		{
-			if (program.m_parameters[paramIndex].m_possibleValues[v].m_name == ValueName)
+			if (Program.Parameters[ParamIndex].PossibleValues[v].Name == ValueName)
 			{
 				result = (int32)v;
 				break;
@@ -326,16 +305,15 @@ namespace mu
 	}
 
 
-	//---------------------------------------------------------------------------------------------
-	int32 Parameters::GetIntValueIndex(int32 paramIndex, int32 Value) const
+	int32 FParameters::GetIntValueIndex(int32 ParamIndex, int32 Value) const
 	{
-		const FProgram& program = m_pD->m_pModel->GetPrivate()->m_program;
-		check(paramIndex >= 0
-			&& paramIndex < (int32)program.m_parameters.Num());
+		const FProgram& Program = m_pD->Model->GetPrivate()->Program;
+		check(ParamIndex >= 0
+			&& ParamIndex < (int32)Program.Parameters.Num());
 
-		for (size_t v = 0; v < program.m_parameters[paramIndex].m_possibleValues.Num(); ++v)
+		for (size_t v = 0; v < Program.Parameters[ParamIndex].PossibleValues.Num(); ++v)
 		{
-			if (program.m_parameters[paramIndex].m_possibleValues[v].m_value == Value)
+			if (Program.Parameters[ParamIndex].PossibleValues[v].Value == Value)
 			{
 				return (int32)v;
 			}
@@ -344,49 +322,47 @@ namespace mu
 	}
 
 
-	//---------------------------------------------------------------------------------------------
-	const FString& Parameters::GetIntPossibleValueName( int32 paramIndex, int32 valueIndex ) const
+	const FString& FParameters::GetIntPossibleValueName( int32 ParamIndex, int32 valueIndex ) const
 	{
-		const FProgram& program = m_pD->m_pModel->GetPrivate()->m_program;
-		check( paramIndex>=0
-				&& paramIndex<(int32)program.m_parameters.Num() );
+		const FProgram& Program = m_pD->Model->GetPrivate()->Program;
+		check( ParamIndex>=0
+				&& ParamIndex<(int32)Program.Parameters.Num() );
 		check( valueIndex>=0
-				&& valueIndex<(int32)program.m_parameters[paramIndex].m_possibleValues.Num() );
+				&& valueIndex<(int32)Program.Parameters[ParamIndex].PossibleValues.Num() );
 
-		return program.m_parameters[paramIndex].m_possibleValues[valueIndex].m_name;
+		return Program.Parameters[ParamIndex].PossibleValues[valueIndex].Name;
 	}
 
 
-	//---------------------------------------------------------------------------------------------
-    int32 Parameters::GetIntValue( int32 index, const Ptr<const RangeIndex>& pos ) const
+    int32 FParameters::GetIntValue( int32 Index, const FRangeIndex* Pos ) const
 	{
-		check( index>=0 && index<(int32)m_pD->m_values.Num() );
-		check( GetType(index)== EParameterType::Int );
+		check( Index>=0 && Index<(int32)m_pD->Values.Num() );
+		check( GetType(Index)== EParameterType::Int );
 
         // Early out in case of invalid parameters
-        if ( index < 0
+        if ( Index < 0
              ||
-             index >= (int32)m_pD->m_values.Num()
+             Index >= (int32)m_pD->Values.Num()
              ||
-             GetType(index) != EParameterType::Int )
+             GetType(Index) != EParameterType::Int )
         {
             return 0;
         }
 
         // Single value case
-        if (!pos)
+        if (!Pos)
         {
             // Return the single value
-            return m_pD->m_values[index].Get<FParamIntType>();
+            return m_pD->Values[Index].Get<FParamIntType>();
         }
 
         // Multivalue case
-        check( pos->m_pD->m_parameter==index );
+        check( Pos->Parameter==Index );
 
-        if ( index<int32(m_pD->m_multiValues.Num()))
+        if ( Index<int32(m_pD->MultiValues.Num()))
         {
-            const TMap< TArray<int32>, FParameterValue >& m = m_pD->m_multiValues[index];
-			const FParameterValue* it = m.Find(pos->m_pD->m_values);
+            const TMap< TArray<int32>, FParameterValue >& m = m_pD->MultiValues[Index];
+			const FParameterValue* it = m.Find(Pos->Values);
             if (it)
             {
                 return it->Get<FParamIntType>();
@@ -394,87 +370,85 @@ namespace mu
         }
 
         // Multivalue parameter, but no multivalue set. Return single value.
-        return m_pD->m_values[index].Get<FParamIntType>();
+        return m_pD->Values[Index].Get<FParamIntType>();
     }
 
 
-	//---------------------------------------------------------------------------------------------
-    void Parameters::SetIntValue( int32 index, int32 value,
-                                  const Ptr<const RangeIndex>& pos )
+    void FParameters::SetIntValue( int32 Index, int32 value,
+                                  const FRangeIndex* Pos )
 	{
-		check( index>=0 && index<(int32)m_pD->m_values.Num() );
-		check( GetType(index)== EParameterType::Int );
+		check( Index>=0 && Index<(int32)m_pD->Values.Num() );
+		check( GetType(Index)== EParameterType::Int );
 
         // Early out in case of invalid parameters
-        if ( index < 0
+        if ( Index < 0
              ||
-             index >= (int32)m_pD->m_values.Num()
+             Index >= (int32)m_pD->Values.Num()
              ||
-             GetType(index) != EParameterType::Int )
+             GetType(Index) != EParameterType::Int )
         {
             return;
         }
 
         // Single value case
-        if (!pos)
+        if (!Pos)
         {
             // Clear multivalue, if set.
-            if (index<int32(m_pD->m_multiValues.Num()))
+            if (Index<int32(m_pD->MultiValues.Num()))
             {
-                m_pD->m_multiValues[index].Empty();
+                m_pD->MultiValues[Index].Empty();
             }
 
-            m_pD->m_values[index].Set<FParamIntType>(value);
+            m_pD->Values[Index].Set<FParamIntType>(value);
         }
 
         // Multivalue case
         else
         {
-            check( pos->m_pD->m_parameter==index );
+            check( Pos->Parameter==Index );
 
-            if ( index>=int32(m_pD->m_multiValues.Num()))
+            if ( Index>=int32(m_pD->MultiValues.Num()))
             {
-                m_pD->m_multiValues.SetNum(index+1);
+                m_pD->MultiValues.SetNum(Index+1);
             }
 
-			TMap< TArray<int32>, FParameterValue >& m = m_pD->m_multiValues[index];
-			FParameterValue& it = m.FindOrAdd(pos->m_pD->m_values);
+			TMap< TArray<int32>, FParameterValue >& m = m_pD->MultiValues[Index];
+			FParameterValue& it = m.FindOrAdd(Pos->Values);
             it.Set<FParamIntType>(value);
         }
     }
 
 
-	//---------------------------------------------------------------------------------------------
-    float Parameters::GetFloatValue( int32 index,
-                                     const Ptr<const RangeIndex>& pos ) const
+    float FParameters::GetFloatValue( int32 Index,
+                                     const FRangeIndex* Pos ) const
 	{
-		check( index>=0 && index<(int32)m_pD->m_values.Num() );
-		check( GetType(index)== EParameterType::Float );
+		check( Index>=0 && Index<(int32)m_pD->Values.Num() );
+		check( GetType(Index)== EParameterType::Float );
 
         // Early out in case of invalid parameters
-        if ( index < 0
+        if ( Index < 0
              ||
-             index >= (int32)m_pD->m_values.Num()
+             Index >= (int32)m_pD->Values.Num()
              ||
-             GetType(index) != EParameterType::Float )
+             GetType(Index) != EParameterType::Float )
         {
             return 0.0f;
         }
 
         // Single value case
-        if (!pos)
+        if (!Pos)
         {
             // Return the single value
-            return m_pD->m_values[index].Get<FParamFloatType>();
+            return m_pD->Values[Index].Get<FParamFloatType>();
         }
 
         // Multivalue case
-        check( pos->m_pD->m_parameter==index );
+        check( Pos->Parameter==Index );
 
-        if ( index<int32(m_pD->m_multiValues.Num()))
+        if ( Index<int32(m_pD->MultiValues.Num()))
         {
-            const TMap< TArray<int32>, FParameterValue >& m = m_pD->m_multiValues[index];
-			const FParameterValue* it = m.Find(pos->m_pD->m_values);
+            const TMap< TArray<int32>, FParameterValue >& m = m_pD->MultiValues[Index];
+			const FParameterValue* it = m.Find(Pos->Values);
             if (it)
             {
                 return it->Get<FParamFloatType>();
@@ -482,86 +456,84 @@ namespace mu
         }
 
         // Multivalue parameter, but no multivalue set. Return single value.
-        return m_pD->m_values[index].Get<FParamFloatType>();
+        return m_pD->Values[Index].Get<FParamFloatType>();
     }
 
 
-	//---------------------------------------------------------------------------------------------
-    void Parameters::SetFloatValue( int32 index, float value,
-                                    const Ptr<const RangeIndex>& pos )
+    void FParameters::SetFloatValue( int32 Index, float value,
+                                    const FRangeIndex* Pos )
 	{
-		check( index>=0 && index<(int32)m_pD->m_values.Num() );
-		check( GetType(index)== EParameterType::Float );
+		check( Index>=0 && Index<(int32)m_pD->Values.Num() );
+		check( GetType(Index)== EParameterType::Float );
 
         // Early out in case of invalid parameters
-        if ( index < 0
+        if ( Index < 0
              ||
-             index >= (int32)m_pD->m_values.Num()
+             Index >= (int32)m_pD->Values.Num()
              ||
-             GetType(index) != EParameterType::Float )
+             GetType(Index) != EParameterType::Float )
         {
             return;
         }
 
         // Single value case
-        if (!pos)
+        if (!Pos)
         {
             // Clear multivalue, if set.
-            if (index<int32(m_pD->m_multiValues.Num()))
+            if (Index<int32(m_pD->MultiValues.Num()))
             {
-                m_pD->m_multiValues[index].Empty();
+                m_pD->MultiValues[Index].Empty();
             }
 
-            m_pD->m_values[index].Set<FParamFloatType>(value);
+            m_pD->Values[Index].Set<FParamFloatType>(value);
         }
 
         // Multivalue case
         else
         {
-            check( pos->m_pD->m_parameter==index );
+            check( Pos->Parameter==Index );
 
-            if ( index>=int32(m_pD->m_multiValues.Num()))
+            if ( Index>=int32(m_pD->MultiValues.Num()))
             {
-                m_pD->m_multiValues.SetNum(index+1);
+                m_pD->MultiValues.SetNum(Index+1);
             }
 
-			TMap< TArray<int32>, FParameterValue >& m = m_pD->m_multiValues[index];
-			FParameterValue& it = m.FindOrAdd(pos->m_pD->m_values);
+			TMap< TArray<int32>, FParameterValue >& m = m_pD->MultiValues[Index];
+			FParameterValue& it = m.FindOrAdd(Pos->Values);
             it.Set<FParamFloatType>(value);
         }
     }
 
 
-    //---------------------------------------------------------------------------------------------
-    void Parameters::GetColourValue( int32 index, FVector4f& OutValue, const Ptr<const RangeIndex>& pos ) const
+    void FParameters::GetColourValue( int32 Index, FVector4f& OutValue, const FRangeIndex* Pos ) const
     {
-        check( index>=0 && index<(int32)m_pD->m_values.Num() );
-        check( GetType(index)== EParameterType::Color );
+        check( Index>=0 && Index<(int32)m_pD->Values.Num() );
+        check( GetType(Index)== EParameterType::Color );
 
         // Early out in case of invalid parameters
-        if ( index < 0
+        if ( Index < 0
              ||
-             index >= (int32)m_pD->m_values.Num()
+             Index >= (int32)m_pD->Values.Num()
              ||
-             GetType(index) != EParameterType::Color )
+             GetType(Index) != EParameterType::Color )
         {
             return;
         }
 
         // Single value case
-        if (!pos)
+        if (!Pos)
         {
-			OutValue = m_pD->m_values[index].Get<FParamColorType>();
+			OutValue = m_pD->Values[Index].Get<FParamColorType>();
             return;
         }
 
         // Multivalue case
-        check( pos->m_pD->m_parameter==index );
+        check( Pos->Parameter==Index );
 
-        if ( index<int32(m_pD->m_multiValues.Num()))
+        if ( Index<int32(m_pD->MultiValues.Num()))
         {
-            const TMap< TArray<int32>, FParameterValue >& m = m_pD->m_multiValues[index];
-			const FParameterValue* it = m.Find(pos->m_pD->m_values);
+            const TMap< TArray<int32>, FParameterValue >& m = m_pD->MultiValues[Index];
+			const FParameterValue* it = m.Find(Pos->Values);
             if (it)
             {
 				OutValue = it->Get<FParamColorType>();
@@ -570,86 +542,84 @@ namespace mu
         }
 
         // Multivalue parameter, but no multivalue set. Return single value.
-		OutValue = m_pD->m_values[index].Get<FParamColorType>();
+		OutValue = m_pD->Values[Index].Get<FParamColorType>();
         return;
     }
 
 
-    //---------------------------------------------------------------------------------------------
-    void Parameters::SetColourValue( int32 index, FVector4f InValue, const Ptr<const RangeIndex>& pos )
+    void FParameters::SetColourValue( int32 Index, FVector4f InValue, const FRangeIndex* Pos )
     {
-        check( index>=0 && index<(int32)m_pD->m_values.Num() );
-        check( GetType(index)== EParameterType::Color );
+        check( Index>=0 && Index<(int32)m_pD->Values.Num() );
+        check( GetType(Index)== EParameterType::Color );
 
         // Early out in case of invalid parameters
-        if ( index < 0
+        if ( Index < 0
              ||
-             index >= (int32)m_pD->m_values.Num()
+             Index >= (int32)m_pD->Values.Num()
              ||
-             GetType(index) != EParameterType::Color )
+             GetType(Index) != EParameterType::Color )
         {
             return;
         }
 
         // Single value case
-        if (!pos)
+        if (!Pos)
         {
             // Clear multivalue, if set.
-            if (index<int32(m_pD->m_multiValues.Num()))
+            if (Index<int32(m_pD->MultiValues.Num()))
             {
-                m_pD->m_multiValues[index].Empty();
+                m_pD->MultiValues[Index].Empty();
             }
 
-        	m_pD->m_values[index].Set<FParamColorType>(InValue);
+        	m_pD->Values[Index].Set<FParamColorType>(InValue);
         }
 
         // Multivalue case
         else
         {
-            check( pos->m_pD->m_parameter==index );
+            check( Pos->Parameter==Index );
 
-            if ( index>=int32(m_pD->m_multiValues.Num()))
+            if ( Index>=int32(m_pD->MultiValues.Num()))
             {
-                m_pD->m_multiValues.SetNum(index+1);
+                m_pD->MultiValues.SetNum(Index+1);
             }
 
-			TMap< TArray<int32>, FParameterValue >& m = m_pD->m_multiValues[index];
-			FParameterValue& it = m.FindOrAdd(pos->m_pD->m_values);
+			TMap< TArray<int32>, FParameterValue >& m = m_pD->MultiValues[Index];
+			FParameterValue& it = m.FindOrAdd(Pos->Values);
 
         	it.Set<FParamColorType>(InValue);
         }
     }
 
 
-    //---------------------------------------------------------------------------------------------
-	FName Parameters::GetImageValue( int32 index, const Ptr<const RangeIndex>& pos ) const
+	FName FParameters::GetImageValue( int32 Index, const FRangeIndex* Pos ) const
     {
-        check( index >= 0 && index < (int32)m_pD->m_values.Num() );
-        check( GetType( index ) == EParameterType::Image );
+        check( Index >= 0 && Index < (int32)m_pD->Values.Num() );
+        check( GetType( Index ) == EParameterType::Image );
 
 		// Early out in case of invalid parameters
-        if ( index < 0
+        if ( Index < 0
              ||
-             index >= (int32)m_pD->m_values.Num()
+             Index >= (int32)m_pD->Values.Num()
              ||
-             GetType(index) != EParameterType::Image )
+             GetType(Index) != EParameterType::Image )
         {
             return {};
         }
 		
 		// Single value case
-		if (!pos)
+		if (!Pos)
 		{
-			return m_pD->m_values[index].Get<FParamImageType>();
+			return m_pD->Values[Index].Get<FParamImageType>();
 		}
 
 		// Multivalue case
-		check(pos->m_pD->m_parameter == index);
+		check(Pos->Parameter == Index);
 
-		if (index < int32(m_pD->m_multiValues.Num()))
+		if (Index < int32(m_pD->MultiValues.Num()))
 		{
-			const TMap< TArray<int32>, FParameterValue >& m = m_pD->m_multiValues[index];
-			const FParameterValue* it = m.Find(pos->m_pD->m_values);
+			const TMap< TArray<int32>, FParameterValue >& m = m_pD->MultiValues[Index];
+			const FParameterValue* it = m.Find(Pos->Values);
 			if (it)
 			{
 				return it->Get<FParamImageType>();
@@ -657,69 +627,67 @@ namespace mu
 		}
 
 		// Multivalue parameter, but no multivalue set. Return single value.
-		return m_pD->m_values[index].Get<FParamImageType>();
+		return m_pD->Values[Index].Get<FParamImageType>();
 
     }
 
 
-    //---------------------------------------------------------------------------------------------
-    void Parameters::SetImageValue( int32 index, FName id, const Ptr<const RangeIndex>& pos )
+    void FParameters::SetImageValue( int32 Index, FName id, const FRangeIndex* Pos )
     {
-        check( index >= 0 && index < (int32)m_pD->m_values.Num() );
-        check( GetType( index ) == EParameterType::Image );
+        check( Index >= 0 && Index < (int32)m_pD->Values.Num() );
+        check( GetType( Index ) == EParameterType::Image );
 
 		// Single value case
-		if (!pos)
+		if (!Pos)
 		{
-			m_pD->m_values[index].Set<FParamImageType>(id);
+			m_pD->Values[Index].Set<FParamImageType>(id);
 		}
 
 		// Multivalue case
 		else
 		{
-			check(pos->m_pD->m_parameter == index);
+			check(Pos->Parameter == Index);
 
-			if (index >= int32(m_pD->m_multiValues.Num()))
+			if (Index >= int32(m_pD->MultiValues.Num()))
 			{
-				m_pD->m_multiValues.SetNum(index + 1);
+				m_pD->MultiValues.SetNum(Index + 1);
 			}
 
-			TMap< TArray<int32>, FParameterValue >& m = m_pD->m_multiValues[index];
-			FParameterValue& it = m.FindOrAdd(pos->m_pD->m_values);
+			TMap< TArray<int32>, FParameterValue >& m = m_pD->MultiValues[Index];
+			FParameterValue& it = m.FindOrAdd(Pos->Values);
 			it.Set<FParamImageType>(id);
 		}
     }
 
 
-    //---------------------------------------------------------------------------------------------
-    void Parameters::GetStringValue( int32 index, FString& OutValue, const Ptr<const RangeIndex>& pos ) const
+    void FParameters::GetStringValue( int32 Index, FString& OutValue, const FRangeIndex* Pos ) const
     {
-        check( index >= 0 && index < (int32)m_pD->m_values.Num() );
-        check( GetType( index ) == EParameterType::String );
+        check( Index >= 0 && Index < (int32)m_pD->Values.Num() );
+        check( GetType( Index ) == EParameterType::String );
 
         // Early out in case of invalid parameters
-        if ( index < 0 || index >= (int32)m_pD->m_values.Num() ||
-             GetType( index ) != EParameterType::String )
+        if ( Index < 0 || Index >= (int32)m_pD->Values.Num() ||
+             GetType( Index ) != EParameterType::String )
         {
             OutValue = TEXT("");
 			return;
         }
 
         // Single value case
-        if ( !pos )
+        if ( !Pos )
         {
             // Return the single value
-			OutValue = m_pD->m_values[index].Get<FParamStringType>();
+			OutValue = m_pD->Values[Index].Get<FParamStringType>();
             return;
         }
 
         // Multivalue case
-        check( pos->m_pD->m_parameter == index );
+        check( Pos->Parameter == Index );
 
-        if ( index < int32( m_pD->m_multiValues.Num() ) )
+        if ( Index < int32( m_pD->MultiValues.Num() ) )
         {
-            const TMap< TArray<int32>, FParameterValue >& m = m_pD->m_multiValues[index];
-			const FParameterValue* it = m.Find( pos->m_pD->m_values );
+            const TMap< TArray<int32>, FParameterValue >& m = m_pD->MultiValues[Index];
+			const FParameterValue* it = m.Find( Pos->Values );
             if ( it )
             {
 				OutValue = it->Get<FParamStringType>();
@@ -728,82 +696,81 @@ namespace mu
         }
 
         // Multivalue parameter, but no multivalue set. Return single value.
-		OutValue = m_pD->m_values[index].Get<FParamStringType>();
+		OutValue = m_pD->Values[Index].Get<FParamStringType>();
         return;
     }
 
 
-    //---------------------------------------------------------------------------------------------
-    void Parameters::SetStringValue( int32 index, const FString& Value, const Ptr<const RangeIndex>& pos )
+    void FParameters::SetStringValue( int32 Index, const FString& Value, const FRangeIndex* Pos )
     {
 		LLM_SCOPE_BYNAME(TEXT("MutableRuntime"));
 
-        check( index >= 0 && index < (int32)m_pD->m_values.Num() );
-        check( GetType( index ) == EParameterType::Float );
+        check( Index >= 0 && Index < (int32)m_pD->Values.Num() );
+        check( GetType( Index ) == EParameterType::Float );
 
         // Early out in case of invalid parameters
-        if ( index < 0 || index >= (int32)m_pD->m_values.Num() ||
-             GetType( index ) != EParameterType::Float )
+        if ( Index < 0 || Index >= (int32)m_pD->Values.Num() ||
+             GetType( Index ) != EParameterType::Float )
         {
             return;
         }
 
         // Single value case
-        if ( !pos )
+        if ( !Pos )
         {
             // Clear multivalue, if set.
-            if ( index < int32( m_pD->m_multiValues.Num() ) )
+            if ( Index < int32( m_pD->MultiValues.Num() ) )
             {
-                m_pD->m_multiValues[index].Empty();
+                m_pD->MultiValues[Index].Empty();
             }
 
-        	m_pD->m_values[index].Set<FParamStringType>(Value);
+        	m_pD->Values[Index].Set<FParamStringType>(Value);
         }
 
         // Multivalue case
         else
         {
-            check( pos->m_pD->m_parameter == index );
+            check( Pos->Parameter == Index );
 
-            if ( index >= int32( m_pD->m_multiValues.Num() ) )
+            if ( Index >= int32( m_pD->MultiValues.Num() ) )
             {
-                m_pD->m_multiValues.SetNum( index + 1 );
+                m_pD->MultiValues.SetNum( Index + 1 );
             }
 
-			TMap< TArray<int32>, FParameterValue >& m = m_pD->m_multiValues[index];
-			FParameterValue& it = m.FindOrAdd( pos->m_pD->m_values );
+			TMap< TArray<int32>, FParameterValue >& m = m_pD->MultiValues[Index];
+			FParameterValue& it = m.FindOrAdd( Pos->Values );
         	it.Set<FParamStringType>(Value);
         }
     }
 
-    void Parameters::GetMatrixValue(int32 index, FMatrix44f& OutValue, const Ptr<const RangeIndex>& pos) const
+    void FParameters::GetMatrixValue(int32 Index, FMatrix44f& OutValue, const FRangeIndex* Pos) const
     {
-    	check( index >= 0 && index < (int32)m_pD->m_values.Num() );
-    	check( GetType( index ) == EParameterType::Matrix );
+    	check( Index >= 0 && Index < (int32)m_pD->Values.Num() );
+    	check( GetType( Index ) == EParameterType::Matrix );
 
     	// Early out in case of invalid parameters
-    	if ( index < 0 || index >= (int32)m_pD->m_values.Num() ||
-			 GetType( index ) != EParameterType::Matrix )
+    	if ( Index < 0 || Index >= (int32)m_pD->Values.Num() ||
+			 GetType( Index ) != EParameterType::Matrix )
     	{
     		OutValue = FMatrix44f::Identity;
     		return;
     	}
 
     	// Single value case
-    	if ( !pos )
+    	if ( !Pos )
     	{
     		// Return the single value
-    		OutValue = m_pD->m_values[index].Get<FParamMatrixType>();
+    		OutValue = m_pD->Values[Index].Get<FParamMatrixType>();
     		return;
     	}
 
     	// Multivalue case
-    	check( pos->m_pD->m_parameter == index );
+    	check( Pos->Parameter == Index );
 
-    	if ( index < int32( m_pD->m_multiValues.Num() ) )
+    	if ( Index < int32( m_pD->MultiValues.Num() ) )
     	{
-    		const TMap< TArray<int32>, FParameterValue >& m = m_pD->m_multiValues[index];
-    		const FParameterValue* it = m.Find( pos->m_pD->m_values );
+    		const TMap< TArray<int32>, FParameterValue >& m = m_pD->MultiValues[Index];
+    		const FParameterValue* it = m.Find( Pos->Values );
     		if ( it )
     		{
     			OutValue = it->Get<FParamMatrixType>();
@@ -812,129 +779,118 @@ namespace mu
     	}
 
     	// Multivalue parameter, but no multivalue set. Return single value.
-    	OutValue = m_pD->m_values[index].Get<FParamMatrixType>();
+    	OutValue = m_pD->Values[Index].Get<FParamMatrixType>();
     }
 
-    void Parameters::SetMatrixValue(int32 index, const FMatrix44f& Value, const Ptr<const RangeIndex>& pos)
+    void FParameters::SetMatrixValue(int32 Index, const FMatrix44f& Value, const FRangeIndex* Pos)
     {
     	LLM_SCOPE_BYNAME(TEXT("MutableRuntime"));
 
-    	check( index >= 0 && index < (int32)m_pD->m_values.Num() );
-    	check( GetType( index ) == EParameterType::Matrix );
+    	check( Index >= 0 && Index < (int32)m_pD->Values.Num() );
+    	check( GetType( Index ) == EParameterType::Matrix );
 
     	// Early out in case of invalid parameters
-    	if ( index < 0 || index >= (int32)m_pD->m_values.Num() ||
-			 GetType( index ) != EParameterType::Matrix )
+    	if ( Index < 0 || Index >= (int32)m_pD->Values.Num() ||
+			 GetType( Index ) != EParameterType::Matrix )
     	{
     		return;
     	}
 
     	// Single value case
-    	if ( !pos )
+    	if ( !Pos )
     	{
     		// Clear multivalue, if set.
-    		if ( index < int32( m_pD->m_multiValues.Num() ) )
+    		if ( Index < int32( m_pD->MultiValues.Num() ) )
     		{
-    			m_pD->m_multiValues[index].Empty();
+    			m_pD->MultiValues[Index].Empty();
     		}
 
-    		m_pD->m_values[index].Set<FParamMatrixType>(Value);
+    		m_pD->Values[Index].Set<FParamMatrixType>(Value);
     	}
 
     	// Multivalue case
     	else
     	{
-    		check( pos->m_pD->m_parameter == index );
+    		check( Pos->Parameter == Index );
 
-    		if ( index >= int32( m_pD->m_multiValues.Num() ) )
+    		if ( Index >= int32( m_pD->MultiValues.Num() ) )
     		{
-    			m_pD->m_multiValues.SetNum( index + 1 );
+    			m_pD->MultiValues.SetNum( Index + 1 );
     		}
 
-    		TMap< TArray<int32>, FParameterValue >& m = m_pD->m_multiValues[index];
-    		FParameterValue& it = m.FindOrAdd( pos->m_pD->m_values );
+    		TMap< TArray<int32>, FParameterValue >& m = m_pD->MultiValues[Index];
+    		FParameterValue& it = m.FindOrAdd( Pos->Values );
     		it.Set<FParamMatrixType>(Value);
     	}
     }
 
 
-    //---------------------------------------------------------------------------------------------
-    FProjector Parameters::Private::GetProjectorValue( int32 index, const Ptr<const RangeIndex>& pos ) const
+    FProjector FParameters::Private::GetProjectorValue( int32 Index, const FRangeIndex* Pos ) const
     {
-
-        const FProgram& program = m_pModel->GetPrivate()->m_program;
+        const FProgram& Program = Model->GetPrivate()->Program;
 
         // Early out in case of invalid parameters
-        if ( index < 0
+        if ( Index < 0
              ||
-             index >= (int32)m_values.Num()
+             Index >= Values.Num()
 			 || 
-			 index >= (int32)program.m_parameters.Num()
+			 Index >= Program.Parameters.Num()
              ||
-             program.m_parameters[index].m_type != EParameterType::Projector )
+             Program.Parameters[Index].Type != EParameterType::Projector )
         {
 			check(false);
             return FProjector();
         }
 
-        const FProjector* result = nullptr;
-
         // Single value case
-        if (!pos)
+        if (!Pos)
         {
             // Return the single value
-            result = &m_values[index].Get<FParamProjectorType>();
+			return Values[Index].Get<FParamProjectorType>();
         }
 
+		const FProjector* Result = nullptr;
+
         // Multivalue case
-        if (!result)
-        {
-            check( pos->m_pD->m_parameter==index );
+        check( Pos->Parameter==Index );
 
-            if ( index<m_multiValues.Num())
+        if ( Index<MultiValues.Num())
+        {
+            const FParameterValue* it = MultiValues[Index].Find(Pos->Values);
+            if (it)
             {
-                const FParameterValue* it = m_multiValues[index].Find(pos->m_pD->m_values);
-                if (it)
-                {
-                    result = &it->Get<FParamProjectorType>();
-                }
+				return it->Get<FParamProjectorType>();
             }
         }
 
         // Multivalue parameter, but no multivalue set. Return single value.
-        if (!result)
-        {
-            result = &m_values[index].Get<FParamProjectorType>();
-        }
-
-        return *result;
+        return Values[Index].Get<FParamProjectorType>();
     }
 
 
-	//---------------------------------------------------------------------------------------------
-	void Parameters::GetProjectorValue( int32 ParameterIndex,
+	void FParameters::GetProjectorValue( int32 ParameterIndex,
 		EProjectorType* OutType,
 		FVector3f* OutPos,
 		FVector3f* OutDir,
 		FVector3f* OutUp,
 		FVector3f* OutScale,
 		float* OutProjectionAngle,
-		const Ptr<const RangeIndex>& pos ) const
+		const FRangeIndex* Pos ) const
 	{
-		check(ParameterIndex >=0 && ParameterIndex <m_pD->m_values.Num() );
+		check(ParameterIndex >=0 && ParameterIndex <m_pD->Values.Num() );
 		check( GetType(ParameterIndex)== EParameterType::Projector );
 
         // Early out in case of invalid parameters
         if (ParameterIndex < 0
              ||
-			ParameterIndex >= m_pD->m_values.Num()
+			ParameterIndex >= m_pD->Values.Num()
              ||
              GetType(ParameterIndex) != EParameterType::Projector )
         {
             return;
         }
 
-        FProjector result = m_pD->GetProjectorValue(ParameterIndex, pos );
+        FProjector result = m_pD->GetProjectorValue(ParameterIndex, Pos );
 
         // Copy results
         if (OutType) *OutType = result.type;
@@ -948,39 +904,38 @@ namespace mu
     }
 
 
-	//---------------------------------------------------------------------------------------------
-	void Parameters::SetProjectorValue( int32 ParameterIndex,
-		const FVector3f& pos,
+	void FParameters::SetProjectorValue( int32 ParameterIndex,
+		const FVector3f& Pos,
 		const FVector3f& dir,
 		const FVector3f& up,
 		const FVector3f& scale,
         float projectionAngle,
-        const Ptr<const RangeIndex>& RangePosition)
+        const FRangeIndex* RangePosition)
 	{
-		check(ParameterIndex >=0 && ParameterIndex <m_pD->m_values.Num() );
+		check(ParameterIndex >=0 && ParameterIndex <m_pD->Values.Num() );
 		check( GetType(ParameterIndex)== EParameterType::Projector );
 
         // Early out in case of invalid parameters
         if (ParameterIndex < 0
              ||
-			ParameterIndex >= m_pD->m_values.Num()
+			ParameterIndex >= m_pD->Values.Num()
              ||
              GetType(ParameterIndex) != EParameterType::Projector )
         {
             return;
         }
 
-	    // Parameters cannot change the projector type anymore
+	    // FParameters cannot change the projector type anymore
 		FProjector Value;
         Value.type = EProjectorType::Count;
-        if (m_pD->m_pModel)
+        if (m_pD->Model)
         {
-            const FProgram& program = m_pD->m_pModel->GetPrivate()->m_program;
-            const FProjector& Projector = program.m_parameters[ParameterIndex].m_defaultValue.Get<FParamProjectorType>();
+            const FProgram& Program = m_pD->Model->GetPrivate()->Program;
+            const FProjector& Projector = Program.Parameters[ParameterIndex].DefaultValue.Get<FParamProjectorType>();
             Value.type = Projector.type;
         }
 
-        Value.position = pos;
+        Value.position = Pos;
         Value.direction = dir;
         Value.up = up;        
         Value.scale = scale;
@@ -991,34 +946,33 @@ namespace mu
         if (!RangePosition)
         {
             // Clear multivalue, if set.
-            if (ParameterIndex<m_pD->m_multiValues.Num())
+            if (ParameterIndex<m_pD->MultiValues.Num())
             {
-                m_pD->m_multiValues[ParameterIndex].Empty();
+                m_pD->MultiValues[ParameterIndex].Empty();
             }
 
-            m_pD->m_values[ParameterIndex].Set<FParamProjectorType>(Value);
+            m_pD->Values[ParameterIndex].Set<FParamProjectorType>(Value);
         }
 
         // Multivalue case
         else
         {
-            check(RangePosition->m_pD->m_parameter== ParameterIndex);
+            check(RangePosition->Parameter== ParameterIndex);
 
-            if (ParameterIndex >=m_pD->m_multiValues.Num())
+            if (ParameterIndex >=m_pD->MultiValues.Num())
             {
-                m_pD->m_multiValues.SetNum(ParameterIndex+1);
+                m_pD->MultiValues.SetNum(ParameterIndex+1);
             }
 
-			TMap< TArray<int32>, FParameterValue >& m = m_pD->m_multiValues[ParameterIndex];
-			FParameterValue& it = m.FindOrAdd(RangePosition->m_pD->m_values);
+			TMap< TArray<int32>, FParameterValue >& m = m_pD->MultiValues[ParameterIndex];
+			FParameterValue& it = m.FindOrAdd(RangePosition->Values);
             it.Set<FParamProjectorType>(Value);
         }
     }
 
 
-    //---------------------------------------------------------------------------------------------
-    bool Parameters::HasSameValue( int32 thisParamIndex,
-                                   const ParametersPtrConst& other,
+    bool FParameters::HasSameValue( int32 thisParamIndex,
+                                   const TSharedPtr<const FParameters>& other,
                                    int32 otherParamIndex ) const
     {
         if ( GetType(thisParamIndex) != other->GetType(otherParamIndex) )
@@ -1026,23 +980,23 @@ namespace mu
             return false;
         }
 
-        if ( !(m_pD->m_values[thisParamIndex]==other->m_pD->m_values[otherParamIndex] ) )
+        if ( !(m_pD->Values[thisParamIndex]==other->m_pD->Values[otherParamIndex] ) )
         {
             return false;
         }
 
         size_t thisNumMultiValues = 0;
-        bool thisHasMultiValues = int32(m_pD->m_multiValues.Num()) > thisParamIndex;
+        bool thisHasMultiValues = int32(m_pD->MultiValues.Num()) > thisParamIndex;
         if (thisHasMultiValues)
         {
-            thisNumMultiValues = m_pD->m_multiValues[thisParamIndex].Num();
+            thisNumMultiValues = m_pD->MultiValues[thisParamIndex].Num();
         }
 
         size_t otherNumMultiValues = 0;
-        bool otherHasMultiValues = int32(other->m_pD->m_multiValues.Num()) > otherParamIndex;
+        bool otherHasMultiValues = int32(other->m_pD->MultiValues.Num()) > otherParamIndex;
         if (otherHasMultiValues)
         {
-            otherNumMultiValues = other->m_pD->m_multiValues[otherParamIndex].Num();
+            otherNumMultiValues = other->m_pD->MultiValues[otherParamIndex].Num();
         }
 
         if ( thisNumMultiValues != otherNumMultiValues )
@@ -1054,7 +1008,7 @@ namespace mu
              &&
              otherHasMultiValues
              &&
-             !(m_pD->m_multiValues[thisParamIndex]==other->m_pD->m_multiValues[otherParamIndex] ))
+             !(m_pD->MultiValues[thisParamIndex]==other->m_pD->MultiValues[otherParamIndex] ))
         {
             return false;
         }
@@ -1063,18 +1017,17 @@ namespace mu
     }
 
 
-	//---------------------------------------------------------------------------------------------
-	int32 Parameters::Private::Find( const FString& Name ) const
+	int32 FParameters::Private::Find( const FString& Name ) const
 	{
-		const FProgram& program = m_pModel->GetPrivate()->m_program;
+		const FProgram& Program = Model->GetPrivate()->Program;
 
 		int32 result = -1;
 
-		for( int32 i=0; result<0 && i<(int32)program.m_parameters.Num(); ++i )
+		for( int32 i=0; result<0 && i<(int32)Program.Parameters.Num(); ++i )
 		{
-			const FParameterDesc& p = program.m_parameters[i];
+			const FParameterDesc& p = Program.Parameters[i];
 
-			if (p.m_name == Name)
+			if (p.Name == Name)
 			{
 				result = i;
 			}
@@ -1087,79 +1040,52 @@ namespace mu
     //---------------------------------------------------------------------------------------------
     //---------------------------------------------------------------------------------------------
     //---------------------------------------------------------------------------------------------
-    RangeIndex::RangeIndex()
+     int32 FRangeIndex::GetRangeCount() const
     {
-        m_pD = new Private();
+        return int32( Values.Num() );
     }
 
 
-    //---------------------------------------------------------------------------------------------
-    RangeIndex::~RangeIndex()
-    {
-        check( m_pD );
-        delete m_pD;
-        m_pD = 0;
-    }
-
-
-    //---------------------------------------------------------------------------------------------
-    RangeIndex::Private* RangeIndex::GetPrivate() const
-    {
-        return m_pD;
-    }
-
-
-    //---------------------------------------------------------------------------------------------
-    int32 RangeIndex::GetRangeCount() const
-    {
-        return int32( m_pD->m_values.Num() );
-    }
-
-
-    //---------------------------------------------------------------------------------------------
-    const FString& RangeIndex::GetRangeName( int32 index ) const
+    const FString& FRangeIndex::GetRangeName( int32 Index ) const
     {
-        check( index >= 0 && index < GetRangeCount() );
-        check( index < int32(m_pD->m_pParameters->GetPrivate()->m_values.Num()) );
-        TSharedPtr<const Model> pModel = m_pD->m_pParameters->GetPrivate()->m_pModel;
-        check(pModel);
-        check( index < int32(pModel->GetPrivate()->m_program.m_ranges.Num()) );
+        check( Index >= 0 && Index < GetRangeCount() );
+        check( Index < int32(Parameters->GetPrivate()->Values.Num()) );
+        TSharedPtr<const FModel> Model = Parameters->GetPrivate()->Model;
+        check(Model);
+        check( Index < int32(Model->GetPrivate()->Program.Ranges.Num()) );
 
-        return pModel->GetPrivate()->m_program.m_ranges[index].m_name;
+        return Model->GetPrivate()->Program.Ranges[Index].Name;
     }
 
 
-    //---------------------------------------------------------------------------------------------
-    const FString& RangeIndex::GetRangeUid( int32 index ) const
+    const FString& FRangeIndex::GetRangeUid( int32 Index ) const
     {
-        check( index >= 0 && index < GetRangeCount() );
-        check( index < int32(m_pD->m_pParameters->GetPrivate()->m_values.Num()) );
-		TSharedPtr<const Model> pModel = m_pD->m_pParameters->GetPrivate()->m_pModel;
-		check(pModel);
-        check( index < int32(pModel->GetPrivate()->m_program.m_ranges.Num()) );
+        check( Index >= 0 && Index < GetRangeCount() );
+        check( Index < int32(Parameters->GetPrivate()->Values.Num()) );
+		TSharedPtr<const FModel> Model = Parameters->GetPrivate()->Model;
+		check(Model);
+        check( Index < int32(Model->GetPrivate()->Program.Ranges.Num()) );
 
-        return pModel->GetPrivate()->m_program.m_ranges[index].m_uid;
+        return Model->GetPrivate()->Program.Ranges[Index].UID;
     }
 
 
-    //---------------------------------------------------------------------------------------------
-    void RangeIndex::SetPosition( int32 index, int32 position )
+    void FRangeIndex::SetPosition( int32 Index, int32 Position )
     {
-        check( index >= 0 && index < GetRangeCount() );
-        if ( index >= 0 && index < GetRangeCount() )
+        check( Index >= 0 && Index < GetRangeCount() );
+        if ( Index >= 0 && Index < GetRangeCount() )
         {
-            m_pD->m_values[index] = position;
+            Values[Index] = Position;
         }
     }
 
 
-    //---------------------------------------------------------------------------------------------
-    int32 RangeIndex::GetPosition( int32 index ) const
+    int32 FRangeIndex::GetPosition( int32 Index ) const
     {
-        check( index >= 0 && index < GetRangeCount() );
-        if ( index >= 0 && index < GetRangeCount() )
+        check( Index >= 0 && Index < GetRangeCount() );
+        if ( Index >= 0 && Index < GetRangeCount() )
         {
-            return m_pD->m_values[index];
+            return Values[Index];
         }
         return 0;
     }
diff --git a/Engine/Plugins/Mutable/Source/MutableRuntime/Private/MuR/PhysicsBody.cpp b/Engine/Plugins/Mutable/Source/MutableRuntime/Private/MuR/PhysicsBody.cpp
--- a/Engine/Plugins/Mutable/Source/MutableRuntime/Private/MuR/PhysicsBody.cpp
+++ b/Engine/Plugins/Mutable/Source/MutableRuntime/Private/MuR/PhysicsBody.cpp
@@ -9,89 +9,62 @@ namespace mu
 
 
 	//-------------------------------------------------------------------------------------------------
-	void PhysicsBody::Serialise(const PhysicsBody* p, OutputArchive& arch)
+	void FPhysicsBody::Serialise(const FPhysicsBody* p, FOutputArchive& arch)
 	{
 		arch << *p;
 	}
 
 
 	//-------------------------------------------------------------------------------------------------
-	mu::Ptr<PhysicsBody> PhysicsBody::StaticUnserialise(InputArchive& arch)
+	TSharedPtr<FPhysicsBody> FPhysicsBody::StaticUnserialise(FInputArchive& arch)
 	{
-		mu::Ptr<PhysicsBody> pResult = new PhysicsBody();
+		TSharedPtr<FPhysicsBody> pResult = MakeShared<FPhysicsBody>();
 		arch >> *pResult;
 		return pResult;
 	}
 
 
 	//-------------------------------------------------------------------------------------------------
-	void FBodyShape::Serialise(OutputArchive& arch) const
+	void FBodyShape::Serialise(FOutputArchive& arch) const
 	{
-		const uint32 ver = 1;
-		arch << ver;
-
 		arch << Name;
 		arch << Flags;
 	}
 
 
 	//-------------------------------------------------------------------------------------------------
-	void FBodyShape::Unserialise(InputArchive& arch)
+	void FBodyShape::Unserialise(FInputArchive& arch)
 	{
-		uint32 ver;
-		arch >> ver;
-		check(ver <= 1);
-
-		if (ver == 0)
-		{
-			std::string Temp;
-			arch >> Temp;
-			Name = Temp.c_str();
-		}
-		else
-		{
-			arch >> Name;
-		}
-
+		arch >> Name;
 		arch >> Flags;
 	}
 
 
 	//-------------------------------------------------------------------------------------------------
-	void FSphereBody::Serialise(OutputArchive& arch) const
+	void FSphereBody::Serialise(FOutputArchive& arch) const
 	{
 		FBodyShape::Serialise(arch);
 
-		const uint32 ver = 0;
-		arch << ver;
-
 		arch << Position;
 		arch << Radius;
 	}
 
 
 	//-------------------------------------------------------------------------------------------------
-	void FSphereBody::Unserialise(InputArchive& arch)
+	void FSphereBody::Unserialise(FInputArchive& arch)
 	{
 		FBodyShape::Unserialise(arch);
 
-		uint32 ver;
-		arch >> ver;
-		check(ver == 0);
-
 		arch >> Position;
 		arch >> Radius;
 	}
 
 
 	//-------------------------------------------------------------------------------------------------
-	void FBoxBody::Serialise(OutputArchive& arch) const
+	void FBoxBody::Serialise(FOutputArchive& arch) const
 	{
 		FBodyShape::Serialise( arch );
 
-		const uint32 ver = 0;
-		arch << ver;
-
 		arch << Position;
 		arch << Orientation;
 		arch << Size;
@@ -99,14 +72,10 @@ namespace mu
 
 
 	//-------------------------------------------------------------------------------------------------
-	void FBoxBody::Unserialise(InputArchive& arch)
+	void FBoxBody::Unserialise(FInputArchive& arch)
 	{
 		FBodyShape::Unserialise( arch );
 
-		uint32 ver;
-		arch >> ver;
-		check(ver == 0);
-
 		arch >> Position;
 		arch >> Orientation;
 		arch >> Size;
@@ -114,13 +83,10 @@ namespace mu
 
 	
 	//-------------------------------------------------------------------------------------------------
-	void FSphylBody::Serialise(OutputArchive& arch) const
+	void FSphylBody::Serialise(FOutputArchive& arch) const
 	{
 		FBodyShape::Serialise( arch );
 
-		const uint32 ver = 0;
-		arch << ver;
-
 		arch << Position;
 		arch << Orientation;
 		arch << Radius;
@@ -129,14 +95,10 @@ namespace mu
 
 
 	//-------------------------------------------------------------------------------------------------
-	void FSphylBody::Unserialise(InputArchive& arch)
+	void FSphylBody::Unserialise(FInputArchive& arch)
 	{
 		FBodyShape::Unserialise( arch );
 
-		uint32 ver;
-		arch >> ver;
-		check(ver == 0);
-
 		arch >> Position;
 		arch >> Orientation;
 		arch >> Radius;
@@ -145,13 +107,10 @@ namespace mu
 
 
 	//-------------------------------------------------------------------------------------------------
-	void FTaperedCapsuleBody::Serialise(OutputArchive& arch) const
+	void FTaperedCapsuleBody::Serialise(FOutputArchive& arch) const
 	{
 		FBodyShape::Serialise( arch );
 
-		const uint32 ver = 0;
-		arch << ver;
-
 		arch << Position;
 		arch << Orientation;
 		arch << Radius0;
@@ -161,13 +120,9 @@ namespace mu
 
 
 	//-------------------------------------------------------------------------------------------------
-	void FTaperedCapsuleBody::Unserialise(InputArchive& arch)
+	void FTaperedCapsuleBody::Unserialise(FInputArchive& arch)
 	{
 		FBodyShape::Unserialise( arch );
-			
-		uint32 ver;
-		arch >> ver;
-		check(ver == 0);
 
 		arch >> Position;
 		arch >> Orientation;
@@ -178,12 +133,9 @@ namespace mu
 
 
 	//-------------------------------------------------------------------------------------------------
-	void FConvexBody::Serialise(OutputArchive& arch) const
+	void FConvexBody::Serialise(FOutputArchive& arch) const
 	{
 		FBodyShape::Serialise( arch );
-	
-		uint32 ver = 0;
-		arch << ver;
 
 		arch << Vertices;
 		arch << Indices;
@@ -192,14 +144,10 @@ namespace mu
 
 
 	//-------------------------------------------------------------------------------------------------
-	void FConvexBody::Unserialise(InputArchive& arch)
+	void FConvexBody::Unserialise(FInputArchive& arch)
 	{
 		FBodyShape::Unserialise( arch );
 
-		uint32 ver;
-		arch >> ver;
-		check(ver == 0);
-
 		arch >> Vertices;
 		arch >> Indices;
 		arch >> Transform;
@@ -207,11 +155,8 @@ namespace mu
 
 
 	//-------------------------------------------------------------------------------------------------
-	void FPhysicsBodyAggregate::Serialise(OutputArchive& arch) const
+	void FPhysicsBodyAggregate::Serialise(FOutputArchive& arch) const
 	{
-		uint32 ver = 0;
-		arch << ver;
-		
 		arch << Spheres;
 		arch << Boxes;
 		arch << Convex;
@@ -221,12 +166,8 @@ namespace mu
 
 
 	//-------------------------------------------------------------------------------------------------
-	void FPhysicsBodyAggregate::Unserialise(InputArchive& arch)
+	void FPhysicsBodyAggregate::Unserialise(FInputArchive& arch)
 	{
-		uint32 ver;
-		arch >> ver;
-		check(ver == 0);
-		
 		arch >> Spheres;
 		arch >> Boxes;
 		arch >> Convex;
@@ -236,9 +177,9 @@ namespace mu
 
 
 	//-------------------------------------------------------------------------------------------------
-	mu::Ptr<PhysicsBody> PhysicsBody::Clone() const
+	TSharedPtr<FPhysicsBody> FPhysicsBody::Clone() const
 	{
-		mu::Ptr<PhysicsBody> pResult = new PhysicsBody();
+		TSharedPtr<FPhysicsBody> pResult = MakeShared<FPhysicsBody>();
 
 		pResult->CustomId = CustomId;
 
@@ -253,19 +194,19 @@ namespace mu
 	}
 
 	//-------------------------------------------------------------------------------------------------
-	void PhysicsBody::SetCustomId(int32 InCustomId)
+	void FPhysicsBody::SetCustomId(int32 InCustomId)
 	{
 		CustomId = InCustomId;
 	}
 
 	//-------------------------------------------------------------------------------------------------
-	int32 PhysicsBody::GetCustomId() const
+	int32 FPhysicsBody::GetCustomId() const
 	{
 		return CustomId;
 	}
 
 	//-------------------------------------------------------------------------------------------------
-	void PhysicsBody::SetBodyCount(int32 Count)
+	void FPhysicsBody::SetBodyCount(int32 Count)
 	{
 		Bodies.SetNum(Count);
 		BoneIds.SetNum(Count);
@@ -273,27 +214,27 @@ namespace mu
 	}
 
 	//-------------------------------------------------------------------------------------------------
-	int32 PhysicsBody::GetBodyCount() const
+	int32 FPhysicsBody::GetBodyCount() const
 	{
 		return Bodies.Num();
 	}
 
 	//-------------------------------------------------------------------------------------------------
-	const FBoneName& PhysicsBody::GetBodyBoneId(int32 B) const
+	const FBoneName& FPhysicsBody::GetBodyBoneId(int32 B) const
 	{
 		check(BoneIds.IsValidIndex(B));
 		return BoneIds[B];
 	}
 
 	//-------------------------------------------------------------------------------------------------
-	void PhysicsBody::SetBodyBoneId(int32 B, const FBoneName& BoneId)
+	void FPhysicsBody::SetBodyBoneId(int32 B, const FBoneName& BoneId)
 	{
 		check(BoneIds.IsValidIndex(B));
 		BoneIds[B] = BoneId;
 	}
 
 	//-------------------------------------------------------------------------------------------------
-	int32 PhysicsBody::GetBodyCustomId(int32 B) const
+	int32 FPhysicsBody::GetBodyCustomId(int32 B) const
 	{
 		check(B >= 0 && B < BodiesCustomIds.Num());
 
@@ -301,7 +242,7 @@ namespace mu
 	}
 
 	//-------------------------------------------------------------------------------------------------
-	void PhysicsBody::SetBodyCustomId(int32 B, int32 BodyCustomId)
+	void FPhysicsBody::SetBodyCustomId(int32 B, int32 BodyCustomId)
 	{
 		check(B >= 0 && B < BodiesCustomIds.Num());
 
@@ -309,67 +250,67 @@ namespace mu
 	}
 
 	//-------------------------------------------------------------------------------------------------
-	int32 PhysicsBody::GetSphereCount(int32 B) const
+	int32 FPhysicsBody::GetSphereCount(int32 B) const
 	{
 		return Bodies[B].Spheres.Num();
 	}
 
 	//-------------------------------------------------------------------------------------------------
-	int32 PhysicsBody::GetBoxCount(int32 B) const
+	int32 FPhysicsBody::GetBoxCount(int32 B) const
 	{
 		return Bodies[B].Boxes.Num();
 	}
 
 	//-------------------------------------------------------------------------------------------------
-	int32 PhysicsBody::GetConvexCount(int32 B) const
+	int32 FPhysicsBody::GetConvexCount(int32 B) const
 	{
 		return Bodies[B].Convex.Num();
 	}
 
 	//-------------------------------------------------------------------------------------------------
-	int32 PhysicsBody::GetSphylCount(int32 B) const
+	int32 FPhysicsBody::GetSphylCount(int32 B) const
 	{
 		return Bodies[B].Sphyls.Num();
 	}
 
 	//-------------------------------------------------------------------------------------------------
-	int32 PhysicsBody::GetTaperedCapsuleCount(int32 B) const
+	int32 FPhysicsBody::GetTaperedCapsuleCount(int32 B) const
 	{
 		return Bodies[B].TaperedCapsules.Num();
 	}
 
 	//-------------------------------------------------------------------------------------------------
-	void PhysicsBody::SetSphereCount(int32 B, int32 Count)
+	void FPhysicsBody::SetSphereCount(int32 B, int32 Count)
 	{
 		Bodies[B].Spheres.SetNum(Count);
 	}
 
 	//-------------------------------------------------------------------------------------------------
-	void PhysicsBody::SetBoxCount(int32 B, int32 Count)
+	void FPhysicsBody::SetBoxCount(int32 B, int32 Count)
 	{
 		Bodies[B].Boxes.SetNum(Count);
 	}
 
 	//-------------------------------------------------------------------------------------------------
-	void PhysicsBody::SetConvexCount(int32 B, int32 Count)
+	void FPhysicsBody::SetConvexCount(int32 B, int32 Count)
 	{
 		Bodies[B].Convex.SetNum(Count);
 	}
 
 	//-------------------------------------------------------------------------------------------------
-	void PhysicsBody::SetSphylCount(int32 B, int32 Count)
+	void FPhysicsBody::SetSphylCount(int32 B, int32 Count)
 	{
 		Bodies[B].Sphyls.SetNum(Count);
 	}
 
 	//-------------------------------------------------------------------------------------------------
-	void PhysicsBody::SetTaperedCapsuleCount(int32 B, int32 Count)
+	void FPhysicsBody::SetTaperedCapsuleCount(int32 B, int32 Count)
 	{
 		Bodies[B].TaperedCapsules.SetNum(Count);
 	}
 
 	//-------------------------------------------------------------------------------------------------
-	void PhysicsBody::SetSphere(
+	void FPhysicsBody::SetSphere(
 		int32 B, int32 I,
 		FVector3f Position, float Radius)
 	{
@@ -381,7 +322,7 @@ namespace mu
 	}
 
 	//-------------------------------------------------------------------------------------------------
-	void PhysicsBody::SetBox(
+	void FPhysicsBody::SetBox(
 		int32 B, int32 I,
 		FVector3f Position, FQuat4f Orientation, FVector3f Size)
 	{
@@ -394,7 +335,7 @@ namespace mu
 	}
 
 	//-------------------------------------------------------------------------------------------------
-	void PhysicsBody::SetConvexMesh(
+	void FPhysicsBody::SetConvexMesh(
 		int32 B, int32 I,
 		TArrayView<const FVector3f> Vertices, TArrayView<const int32> Indices)
 	{
@@ -406,7 +347,7 @@ namespace mu
 	}
 
 	//-------------------------------------------------------------------------------------------------
-	void PhysicsBody::SetConvexTransform(
+	void FPhysicsBody::SetConvexTransform(
 		int32 B, int32 I,
 		const FTransform3f& Transform)
 	{
@@ -418,7 +359,7 @@ namespace mu
 
 
 	//-------------------------------------------------------------------------------------------------
-	void PhysicsBody::SetSphyl(
+	void FPhysicsBody::SetSphyl(
 		int32 B, int32 I,
 		FVector3f Position, FQuat4f Orientation, float Radius, float Length)
 	{
@@ -432,7 +373,7 @@ namespace mu
 	}
 
 	//-------------------------------------------------------------------------------------------------
-	void PhysicsBody::SetSphereFlags(int32 B, int32 I, uint32 Flags)
+	void FPhysicsBody::SetSphereFlags(int32 B, int32 I, uint32 Flags)
 	{
 		check(B >= 0 && B < Bodies.Num());
 		check(I >= 0 && I < Bodies[B].Spheres.Num());
@@ -441,7 +382,7 @@ namespace mu
 	}
 
 	//-------------------------------------------------------------------------------------------------
-	void PhysicsBody::SetBoxFlags(int32 B, int32 I, uint32 Flags)
+	void FPhysicsBody::SetBoxFlags(int32 B, int32 I, uint32 Flags)
 	{
 		check(B >= 0 && B < Bodies.Num());
 		check(I >= 0 && I < Bodies[B].Boxes.Num());
@@ -450,7 +391,7 @@ namespace mu
 	}
 
 	//-------------------------------------------------------------------------------------------------
-	void PhysicsBody::SetConvexFlags(int32 B, int32 I, uint32 Flags)
+	void FPhysicsBody::SetConvexFlags(int32 B, int32 I, uint32 Flags)
 	{
 		check(B >= 0 && B < Bodies.Num());
 		check(I >= 0 && I < Bodies[B].Convex.Num());
@@ -459,7 +400,7 @@ namespace mu
 	}
 
 	//-------------------------------------------------------------------------------------------------
-	void PhysicsBody::SetSphylFlags(int32 B, int32 I, uint32 Flags)
+	void FPhysicsBody::SetSphylFlags(int32 B, int32 I, uint32 Flags)
 	{
 		check(B >= 0 && B < Bodies.Num());
 		check(I >= 0 && I < Bodies[B].Sphyls.Num());
@@ -468,7 +409,7 @@ namespace mu
 	}
 
 	//-------------------------------------------------------------------------------------------------
-	void PhysicsBody::SetTaperedCapsuleFlags(int32 B, int32 I, uint32 Flags)
+	void FPhysicsBody::SetTaperedCapsuleFlags(int32 B, int32 I, uint32 Flags)
 	{
 		check(B >= 0 && B < Bodies.Num());
 		check(I >= 0 && I < Bodies[B].TaperedCapsules.Num());
@@ -477,7 +418,7 @@ namespace mu
 	}
 
 	//-------------------------------------------------------------------------------------------------
-	void PhysicsBody::SetSphereName(int32 B, int32 I, const char* Name)
+	void FPhysicsBody::SetSphereName(int32 B, int32 I, const char* Name)
 	{
 		check(B >= 0 && B < Bodies.Num());
 		check(I >= 0 && I < Bodies[B].Spheres.Num());
@@ -486,7 +427,7 @@ namespace mu
 	}
 
 	//-------------------------------------------------------------------------------------------------
-	void PhysicsBody::SetBoxName(int32 B, int32 I, const char* Name)
+	void FPhysicsBody::SetBoxName(int32 B, int32 I, const char* Name)
 	{
 		check(B >= 0 && B < Bodies.Num());
 		check(I >= 0 && I < Bodies[B].Boxes.Num());
@@ -495,7 +436,7 @@ namespace mu
 	}
 
 	//-------------------------------------------------------------------------------------------------
-	void PhysicsBody::SetConvexName(int32 B, int32 I, const char* Name)
+	void FPhysicsBody::SetConvexName(int32 B, int32 I, const char* Name)
 	{
 		check(B >= 0 && B < Bodies.Num());
 		check(I >= 0 && I < Bodies[B].Convex.Num());
@@ -504,7 +445,7 @@ namespace mu
 	}
 
 	//-------------------------------------------------------------------------------------------------
-	void PhysicsBody::SetSphylName(int32 B, int32 I, const char* Name)
+	void FPhysicsBody::SetSphylName(int32 B, int32 I, const char* Name)
 	{
 		check(B >= 0 && B < Bodies.Num());
 		check(I >= 0 && I < Bodies[B].Sphyls.Num());
@@ -513,7 +454,7 @@ namespace mu
 	}
 
 	//-------------------------------------------------------------------------------------------------
-	void PhysicsBody::SetTaperedCapsuleName(int32 B, int32 I, const char* Name)
+	void FPhysicsBody::SetTaperedCapsuleName(int32 B, int32 I, const char* Name)
 	{
 		check(B >= 0 && B < Bodies.Num());
 		check(I >= 0 && I < Bodies[B].TaperedCapsules.Num());
@@ -522,7 +463,7 @@ namespace mu
 	}
 
 	//-------------------------------------------------------------------------------------------------
-	void PhysicsBody::SetTaperedCapsule(
+	void FPhysicsBody::SetTaperedCapsule(
 		int32 B, int32 I,
 		FVector3f Position, FQuat4f Orientation,
 		float Radius0, float Radius1, float Length)
@@ -538,7 +479,7 @@ namespace mu
 	}
 
 	//-------------------------------------------------------------------------------------------------
-	void PhysicsBody::GetSphere(
+	void FPhysicsBody::GetSphere(
 		int32 B, int32 I,
 		FVector3f& OutPosition, float& OutRadius) const
 	{
@@ -550,7 +491,7 @@ namespace mu
 	}
 
 	//-------------------------------------------------------------------------------------------------
-	void PhysicsBody::GetBox(
+	void FPhysicsBody::GetBox(
 		int32 B, int32 I,
 		FVector3f& OutPosition, FQuat4f& OutOrientation, FVector3f& OutSize) const
 	{
@@ -563,7 +504,7 @@ namespace mu
 	}
 
 	//-------------------------------------------------------------------------------------------------
-	void PhysicsBody::GetConvex(
+	void FPhysicsBody::GetConvex(
 		int32 B, int32 I,
 		TArrayView<const FVector3f>& OutVertices, TArrayView<const int32>& OutIndices,
 		FTransform3f& OutTransform) const
@@ -576,7 +517,7 @@ namespace mu
 		OutTransform = Bodies[B].Convex[I].Transform;
 	}
 
-	void PhysicsBody::GetConvexMeshView(int32 B, int32 I,
+	void FPhysicsBody::GetConvexMeshView(int32 B, int32 I,
 		TArrayView<FVector3f>& OutVerticesView, TArrayView<int32>& OutIndicesView)
 	{
 		check(B >= 0 && B < Bodies.Num());
@@ -586,7 +527,7 @@ namespace mu
 		OutIndicesView = TArrayView<int32>(Bodies[B].Convex[I].Indices.GetData(), Bodies[B].Convex[I].Indices.Num());
 	}
 
-	void PhysicsBody::GetConvexTransform(int32 B, int32 I, FTransform3f& OutTransform) const
+	void FPhysicsBody::GetConvexTransform(int32 B, int32 I, FTransform3f& OutTransform) const
 	{
 		check(B >= 0 && B < Bodies.Num());
 		check(I >= 0 && I < Bodies[B].Convex.Num());
@@ -595,7 +536,7 @@ namespace mu
 	}
 
 	//-------------------------------------------------------------------------------------------------
-	void PhysicsBody::GetSphyl(
+	void FPhysicsBody::GetSphyl(
 		int32 B, int32 I,
 		FVector3f& OutPosition, FQuat4f& OutOrientation, float& OutRadius, float& OutLength) const
 	{
@@ -609,7 +550,7 @@ namespace mu
 	}
 
 	//-------------------------------------------------------------------------------------------------
-	void PhysicsBody::GetTaperedCapsule(
+	void FPhysicsBody::GetTaperedCapsule(
 		int32 B, int32 I,
 		FVector3f& OutPosition, FQuat4f& OutOrientation,
 		float& OutRadius0, float& OutRadius1, float& OutLength) const
@@ -626,7 +567,7 @@ namespace mu
 
 
 	//-------------------------------------------------------------------------------------------------
-	uint32 PhysicsBody::GetSphereFlags(int32 B, int32 I) const
+	uint32 FPhysicsBody::GetSphereFlags(int32 B, int32 I) const
 	{
 		check(B >= 0 && B < Bodies.Num());
 		check(I >= 0 && I < Bodies[B].Spheres.Num());
@@ -635,7 +576,7 @@ namespace mu
 	}
 
 	//-------------------------------------------------------------------------------------------------
-	uint32 PhysicsBody::GetBoxFlags(int32 B, int32 I) const
+	uint32 FPhysicsBody::GetBoxFlags(int32 B, int32 I) const
 	{
 		check(B >= 0 && B < Bodies.Num());
 		check(I >= 0 && I < Bodies[B].Boxes.Num());
@@ -644,7 +585,7 @@ namespace mu
 	}
 
 	//-------------------------------------------------------------------------------------------------
-	uint32 PhysicsBody::GetConvexFlags(int32 B, int32 I) const
+	uint32 FPhysicsBody::GetConvexFlags(int32 B, int32 I) const
 	{
 		check(B >= 0 && B < Bodies.Num());
 		check(I >= 0 && I < Bodies[B].Convex.Num());
@@ -653,7 +594,7 @@ namespace mu
 	}
 
 	//-------------------------------------------------------------------------------------------------
-	uint32 PhysicsBody::GetSphylFlags(int32 B, int32 I) const
+	uint32 FPhysicsBody::GetSphylFlags(int32 B, int32 I) const
 	{
 		check(B >= 0 && B < Bodies.Num());
 		check(I >= 0 && I < Bodies[B].Sphyls.Num());
@@ -662,7 +603,7 @@ namespace mu
 	}
 
 	//-------------------------------------------------------------------------------------------------
-	uint32 PhysicsBody::GetTaperedCapsuleFlags(int32 B, int32 I) const
+	uint32 FPhysicsBody::GetTaperedCapsuleFlags(int32 B, int32 I) const
 	{
 		check(B >= 0 && B < Bodies.Num());
 		check(I >= 0 && I < Bodies[B].TaperedCapsules.Num());
@@ -671,7 +612,7 @@ namespace mu
 	}
 
 	//-------------------------------------------------------------------------------------------------
-	const FString& PhysicsBody::GetSphereName(int32 B, int32 I) const
+	const FString& FPhysicsBody::GetSphereName(int32 B, int32 I) const
 	{
 		check(B >= 0 && B < Bodies.Num());
 		check(I >= 0 && I < Bodies[B].Spheres.Num());
@@ -680,7 +621,7 @@ namespace mu
 	}
 
 	//-------------------------------------------------------------------------------------------------
-	const FString& PhysicsBody::GetBoxName(int32 B, int32 I) const
+	const FString& FPhysicsBody::GetBoxName(int32 B, int32 I) const
 	{
 		check(B >= 0 && B < Bodies.Num());
 		check(I >= 0 && I < Bodies[B].Boxes.Num());
@@ -689,7 +630,7 @@ namespace mu
 	}
 
 	//-------------------------------------------------------------------------------------------------
-	const FString& PhysicsBody::GetConvexName(int32 B, int32 I) const
+	const FString& FPhysicsBody::GetConvexName(int32 B, int32 I) const
 	{
 		check(B >= 0 && B < Bodies.Num());
 		check(I >= 0 && I < Bodies[B].Convex.Num());
@@ -698,7 +639,7 @@ namespace mu
 	}
 
 	//-------------------------------------------------------------------------------------------------
-	const FString& PhysicsBody::GetSphylName(int32 B, int32 I) const
+	const FString& FPhysicsBody::GetSphylName(int32 B, int32 I) const
 	{
 		check(B >= 0 && B < Bodies.Num());
 		check(I >= 0 && I < Bodies[B].Sphyls.Num());
@@ -707,7 +648,7 @@ namespace mu
 	}
 
 	//-------------------------------------------------------------------------------------------------
-	const FString& PhysicsBody::GetTaperedCapsuleName(int32 B, int32 I) const
+	const FString& FPhysicsBody::GetTaperedCapsuleName(int32 B, int32 I) const
 	{
 		check(B >= 0 && B < Bodies.Num());
 		check(I >= 0 && I < Bodies[B].TaperedCapsules.Num());
@@ -717,11 +658,8 @@ namespace mu
 
 
 	//-------------------------------------------------------------------------------------------------
-	void PhysicsBody::Serialise(OutputArchive& arch) const
+	void FPhysicsBody::Serialise(FOutputArchive& arch) const
 	{
-		uint32 ver = 4;
-		arch << ver;
-        	
 		arch << CustomId;
 		arch << Bodies;
 		arch << BoneIds;
@@ -731,54 +669,12 @@ namespace mu
 
 
 	//-------------------------------------------------------------------------------------------------
-	void PhysicsBody::Unserialise(InputArchive& arch)
+	void FPhysicsBody::Unserialise(FInputArchive& arch)
 	{
-		uint32 ver;
-		arch >> ver;
-		check(ver <= 4);
-        	
-		if (ver >= 2)
-		{
-			arch >> CustomId;
-		}
-        	
+		arch >> CustomId;        	
 		arch >> Bodies;
-
-		if (ver >= 4)
-		{
-			arch >> BoneIds;
-		}
-		else if (ver == 3)
-		{
-			TArray<uint16> BoneIds_DEPRECATED;
-			arch >> BoneIds_DEPRECATED;
-
-			const int32 NumBoneNames = BoneIds_DEPRECATED.Num();
-			BoneIds.SetNum(NumBoneNames);
-			for (int32 BoneIndex = 0; BoneIndex < NumBoneNames; ++BoneIndex)
-			{
-				BoneIds[BoneIndex].Id = BoneIds_DEPRECATED[BoneIndex];
-			}
-		}
-		else
-		{
-			TArray<std::string> BoneNames;
-			arch >> BoneNames;
-
-			const int32 NumBoneNames = BoneNames.Num();
-			TArray<std::string> UniqueBoneNames;
-			UniqueBoneNames.Reserve(NumBoneNames);
-			BoneIds.Reserve(NumBoneNames);
-			for (int32 BoneIndex = 0; BoneIndex < NumBoneNames; ++BoneIndex)
-			{
-				BoneIds.Add(FBoneName(BoneIndex));
-			}
-		}
+		arch >> BoneIds;
 		arch >> BodiesCustomIds;
-
-		if (ver >= 1)
-		{
-			arch >> bBodiesModified;
-		}
+		arch >> bBodiesModified;
 	}
 }
diff --git a/Engine/Plugins/Mutable/Source/MutableRuntime/Private/MuR/Serialisation.cpp b/Engine/Plugins/Mutable/Source/MutableRuntime/Private/MuR/Serialisation.cpp
--- a/Engine/Plugins/Mutable/Source/MutableRuntime/Private/MuR/Serialisation.cpp
+++ b/Engine/Plugins/Mutable/Source/MutableRuntime/Private/MuR/Serialisation.cpp
@@ -39,14 +39,14 @@ namespace mu
 	MUTABLE_IMPLEMENT_POD_SERIALISABLE(FRichCurveKey);
 
 	
-    void operator<<(OutputArchive& arch, const FString& t)
+    void operator<<(FOutputArchive& arch, const FString& t)
     {
 	    const TArray<TCHAR>& Data = t.GetCharArray();
     	arch << Data;
     }
 
 
-    void operator>>(InputArchive& arch, FString& t)
+    void operator>>(FInputArchive& arch, FString& t)
     {
 		TArray<TCHAR> Data;
     	arch >> Data;
@@ -55,52 +55,40 @@ namespace mu
     }
 
 
-	void operator<<(OutputArchive& arch, const FRichCurve& t)
+	void operator<<(FOutputArchive& arch, const FRichCurve& t)
 	{
 		arch << t.Keys;
 	}
 
 
-	void operator>>(InputArchive& arch, FRichCurve& t)
+	void operator>>(FInputArchive& arch, FRichCurve& t)
 	{
 		arch >> t.Keys;
 	}
 
 
-    void operator<<(OutputArchive& arch, const FName& v)
+    void operator<<(FOutputArchive& arch, const FName& v)
     {
 	    arch << v.ToString();
     }
 
 
-    void operator>>(InputArchive& arch, FName& v)
+    void operator>>(FInputArchive& arch, FName& v)
     {
 	    FString Temp;
 	    arch >> Temp;
 	    v = FName(Temp);
     }
 	
-
-	void operator>> ( InputArchive& arch, std::string& v )
-    {
-    	uint32 size;
-    	arch >> size;
-    	v.resize( size );
-    	if (size)
-    	{
-    		arch.Stream->Read( &v[0], (unsigned)size*sizeof(char) );
-    	}
-    }
-
 	
-	void operator<<(OutputArchive& Arch, const bool& T)
+	void operator<<(FOutputArchive& Arch, const bool& T)
     {
     	uint8 S = T ? 1 : 0;
     	Arch.Stream->Write(&S, sizeof(uint8));
     }
 
 
-	void operator>>(InputArchive& Arch, bool& T)
+	void operator>>(FInputArchive& Arch, bool& T)
     {
     	uint8 S;
     	Arch.Stream->Read(&S, sizeof(uint8));
@@ -108,18 +96,14 @@ namespace mu
     }
 
 	
-    //---------------------------------------------------------------------------------------------
-    //---------------------------------------------------------------------------------------------
-    //---------------------------------------------------------------------------------------------
-    InputMemoryStream::InputMemoryStream( const void* InBuffer, uint64 InSize )
+    FInputMemoryStream::FInputMemoryStream( const void* InBuffer, uint64 InSize )
     {
         Buffer = InBuffer;
         Size = InSize;
     }
 
 
-    //---------------------------------------------------------------------------------------------
-    void InputMemoryStream::Read( void* Data, uint64 InSize)
+    void FInputMemoryStream::Read( void* Data, uint64 InSize)
     {
         if (InSize)
         {
@@ -132,10 +116,7 @@ namespace mu
     }
 
 
-    //---------------------------------------------------------------------------------------------
-    //---------------------------------------------------------------------------------------------
-    //---------------------------------------------------------------------------------------------
-    OutputMemoryStream::OutputMemoryStream(uint64 Reserve )
+    FOutputMemoryStream::FOutputMemoryStream(uint64 Reserve )
     {
         if (Reserve)
         {
@@ -144,8 +125,7 @@ namespace mu
     }
 
 
-    //---------------------------------------------------------------------------------------------
-    void OutputMemoryStream::Write( const void* Data, uint64 Size)
+    void FOutputMemoryStream::Write( const void* Data, uint64 Size)
     {
         if (Size)
         {
@@ -156,141 +136,71 @@ namespace mu
     }
 
 
-    //---------------------------------------------------------------------------------------------
-    const void* OutputMemoryStream::GetBuffer() const
+    const uint8* FOutputMemoryStream::GetBuffer() const
     {
-        const void* pResult = 0;
+        const uint8* Result = nullptr;
 
         if (Buffer.Num() )
         {
-            pResult = Buffer.GetData();
+            Result = Buffer.GetData();
         }
 
-        return pResult;
+        return Result;
     }
 
 
-	//---------------------------------------------------------------------------------------------
-	uint64 OutputMemoryStream::GetBufferSize() const
+	int32 FOutputMemoryStream::GetBufferSize() const
     {
         return Buffer.Num();
     }
 
 
-	//---------------------------------------------------------------------------------------------
-	void OutputMemoryStream::Reset()
+	void FOutputMemoryStream::Reset()
 	{
 		Buffer.SetNum(0,EAllowShrinking::No);
 	}
 
 
-    //-------------------------------------------------------------------------------------------------
-    //-------------------------------------------------------------------------------------------------
-    //-------------------------------------------------------------------------------------------------
-    void OutputSizeStream::Write( const void*, uint64 size )
+	void FOutputSizeStream::Write( const void*, uint64 size )
     {
         WrittenBytes += size;
     }
 
-	uint64 OutputSizeStream::GetBufferSize() const
+
+	uint64 FOutputSizeStream::GetBufferSize() const
     {
         return WrittenBytes;
     }
 
 
-	//-------------------------------------------------------------------------------------------------
-	void OutputHashStream::Write(const void* Data, uint64 size)
+	void FOutputHashStream::Write(const void* Data, uint64 size)
 	{
 		Hash = CityHash64WithSeed(reinterpret_cast<const char*>(Data), size, Hash);
 	}
 
 
-	uint64 OutputHashStream::GetHash() const
+	uint64 FOutputHashStream::GetHash() const
 	{
 		return Hash;
 	}
 
 
-    //---------------------------------------------------------------------------------------------
-	//---------------------------------------------------------------------------------------------
-	//---------------------------------------------------------------------------------------------
-	InputArchive::InputArchive( InputStream* InStream )
+ 	FInputArchive::FInputArchive( FInputStream* InStream )
 	{
 		Stream = InStream;
 	}
 
 
-    //---------------------------------------------------------------------------------------------
-    Ptr<ResourceProxy<Image>> InputArchive::NewImageProxy()
+    Ptr<TResourceProxy<FImage>> FInputArchive::NewImageProxy()
     {
         return nullptr;
     }
 
 
-	//---------------------------------------------------------------------------------------------
-	//---------------------------------------------------------------------------------------------
-	//---------------------------------------------------------------------------------------------
-	OutputArchive::OutputArchive( OutputStream* InStream )
+	FOutputArchive::FOutputArchive( FOutputStream* InStream )
 	{
 		Stream = InStream;
 	}
 
-
-    //---------------------------------------------------------------------------------------------
-    //---------------------------------------------------------------------------------------------
-    //---------------------------------------------------------------------------------------------
-    InputArchiveWithProxies::InputArchiveWithProxies( InputStream* s, ProxyFactory* f )
-        : InputArchive( s )
-    {
-        Factory = f;
-    }
-
-
-    //---------------------------------------------------------------------------------------------
-    Ptr<ResourceProxy<Image>> InputArchiveWithProxies::NewImageProxy()
-    {
-        // Similar to Ptr serisalisation in SerialisationPrivate
-        Ptr<ResourceProxy<Image>> p;
-        {
-            int32 id;
-            (*this) >> id;
-
-            if ( id == -1 )
-            {
-                // We consumed the serialisation, so we need to return something.
-                class ImageProxyNull : public ResourceProxy<Image>
-                {
-                public:
-                    ImagePtrConst Get() override
-                    {
-                        return nullptr;
-                    }
-                };
-                p = new ImageProxyNull();
-            }
-            else
-            {
-                if ( id < ProxyHistory.Num() )
-                {
-                    p = ProxyHistory[id];
-
-                    // If the pointer was nullptr it means the position in history is used, but not set
-                    // yet: we have a smart pointer loop which is very bad.
-                    check( p );
-                }
-                else
-                {
-                    // Ids come in order.
-                    ProxyHistory.SetNum(id+1);
-
-                    p = Factory->NewImageProxy(*this);
-                    ProxyHistory[id] = p;
-                }
-            }
-        }
-
-        return p;
-    }
-
 }
 
diff --git a/Engine/Plugins/Mutable/Source/MutableRuntime/Private/MuR/Settings.cpp b/Engine/Plugins/Mutable/Source/MutableRuntime/Private/MuR/Settings.cpp
--- a/Engine/Plugins/Mutable/Source/MutableRuntime/Private/MuR/Settings.cpp
+++ b/Engine/Plugins/Mutable/Source/MutableRuntime/Private/MuR/Settings.cpp
@@ -2,37 +2,22 @@
 
 #include "MuR/Settings.h"
 
-#include "Misc/AssertionMacros.h"
-
 namespace mu
 {
-  
-	Settings::Settings()
-    {
-    }
-
-
-    Settings::~Settings()
-    {
-    }
-
 
-    //---------------------------------------------------------------------------------------------
-    void Settings::SetProfile( bool bEnabled )
+    void FSettings::SetProfile( bool bEnabled )
     {
         bProfile = bEnabled;
     }
 
 
-    //---------------------------------------------------------------------------------------------
-    void Settings::SetWorkingMemoryBytes( uint64 Bytes )
+    void FSettings::SetWorkingMemoryBytes( uint64 Bytes )
     {
         WorkingMemoryBytes = Bytes;
     }
 
 
-    //---------------------------------------------------------------------------------------------
-    void Settings::SetImageCompressionQuality( int32 Quality )
+    void FSettings::SetImageCompressionQuality( int32 Quality )
     {
         ImageCompressionQuality = Quality;
     }
diff --git a/Engine/Plugins/Mutable/Source/MutableRuntime/Private/MuR/Skeleton.cpp b/Engine/Plugins/Mutable/Source/MutableRuntime/Private/MuR/Skeleton.cpp
--- a/Engine/Plugins/Mutable/Source/MutableRuntime/Private/MuR/Skeleton.cpp
+++ b/Engine/Plugins/Mutable/Source/MutableRuntime/Private/MuR/Skeleton.cpp
@@ -9,47 +9,41 @@
 namespace mu
 {
 
-
-	//---------------------------------------------------------------------------------------------
-	void Skeleton::Serialise(const Skeleton* p, OutputArchive& arch)
+	void FSkeleton::Serialise(const FSkeleton* In, FOutputArchive& Arch)
 	{
-		arch << *p;
+		Arch << *In;
 	}
 
 
-	//---------------------------------------------------------------------------------------------
-	mu::Ptr<Skeleton> Skeleton::StaticUnserialise(InputArchive& arch)
+	TSharedPtr<FSkeleton> FSkeleton::StaticUnserialise(FInputArchive& Arch)
 	{
-		SkeletonPtr pResult = new Skeleton();
-		arch >> *pResult;
-		return pResult;
+		TSharedPtr<FSkeleton> Result = MakeShared<FSkeleton>();
+		Arch >> *Result;
+		return Result;
 	}
 
 
-	//---------------------------------------------------------------------------------------------
-	SkeletonPtr Skeleton::Clone() const
+	TSharedPtr<FSkeleton> FSkeleton::Clone() const
 	{
-		SkeletonPtr pResult = new Skeleton();
+		TSharedPtr<FSkeleton> Result = MakeShared<FSkeleton>();
 
-		pResult->BoneIds = BoneIds;
-		pResult->BoneParents = BoneParents;
+		Result->BoneIds = BoneIds;
+		Result->BoneParents = BoneParents;
 
 		// For debug
-		pResult->DebugBoneNames = DebugBoneNames;
+		Result->DebugBoneNames = DebugBoneNames;
 
-		return pResult;
+		return Result;
 	}
 
 
-	//---------------------------------------------------------------------------------------------
-	int32 Skeleton::GetBoneCount() const
+	int32 FSkeleton::GetBoneCount() const
 	{
 		return BoneIds.Num();
 	}
 
 
-	//---------------------------------------------------------------------------------------------
-	void Skeleton::SetBoneCount(int32 NumBones)
+	void FSkeleton::SetBoneCount(int32 NumBones)
 	{
 		DebugBoneNames.SetNum(NumBones);
 		BoneIds.SetNum(NumBones);
@@ -57,8 +51,7 @@ namespace mu
 	}
 
 
-	//---------------------------------------------------------------------------------------------
-	const FName Skeleton::GetDebugName(int32 Index) const
+	const FName FSkeleton::GetDebugName(int32 Index) const
 	{
 		if (DebugBoneNames.IsValidIndex(Index))
 		{
@@ -69,8 +62,7 @@ namespace mu
 	}
 
 
-	//---------------------------------------------------------------------------------------------
-	void Skeleton::SetDebugName(const int32 Index, const FName BoneName)
+	void FSkeleton::SetDebugName(const int32 Index, const FName BoneName)
 	{
 		if (DebugBoneNames.IsValidIndex(Index))
 		{
@@ -79,8 +71,7 @@ namespace mu
 	}
 
 
-	//---------------------------------------------------------------------------------------------
-	int32 Skeleton::GetBoneParent(int32 Index) const
+	int32 FSkeleton::GetBoneParent(int32 Index) const
 	{
 		if (BoneParents.IsValidIndex(Index))
 		{
@@ -91,8 +82,7 @@ namespace mu
 	}
 
 
-	//-------------------------------------------------------------------------------------------------
-	void Skeleton::SetBoneParent(int32 Index, int32 ParentIndex)
+	void FSkeleton::SetBoneParent(int32 Index, int32 ParentIndex)
 	{
 		check(ParentIndex >= -1 && ParentIndex < GetBoneCount() && ParentIndex < 0xffff);
 		check(BoneParents.IsValidIndex(Index));
@@ -104,16 +94,14 @@ namespace mu
 	}
 
 
-	//-------------------------------------------------------------------------------------------------
-	const FBoneName& Skeleton::GetBoneName(int32 Index) const
+	const FBoneName& FSkeleton::GetBoneName(int32 Index) const
 	{
 		check(BoneIds.IsValidIndex(Index));
 		return BoneIds[Index];
 	}
 
 	
-	//-------------------------------------------------------------------------------------------------
-	void Skeleton::SetBoneName(int32 Index, const FBoneName& BoneName)
+	void FSkeleton::SetBoneName(int32 Index, const FBoneName& BoneName)
 	{
 		check(BoneIds.IsValidIndex(Index));
 		if (BoneIds.IsValidIndex(Index))
@@ -123,88 +111,22 @@ namespace mu
 	}
 
 
-	//-------------------------------------------------------------------------------------------------
-	int32 Skeleton::FindBone(const FBoneName& BoneName) const
+	int32 FSkeleton::FindBone(const FBoneName& BoneName) const
 	{
 		return BoneIds.Find(BoneName);
 	}
 
 
-	//-------------------------------------------------------------------------------------------------
-	void Skeleton::Serialise(OutputArchive& arch) const
+	void FSkeleton::Serialise(FOutputArchive& Arch) const
 	{
-		uint32 ver = 7;
-		arch << ver;
-
-		arch << BoneIds;
-		arch << BoneParents;
+		Arch << BoneIds;
+		Arch << BoneParents;
 	}
 
 
-	//-------------------------------------------------------------------------------------------------
-	void Skeleton::Unserialise(InputArchive& arch)
+	void FSkeleton::Unserialise(FInputArchive& Arch)
 	{
-		uint32 ver;
-		arch >> ver;
-		check(ver >= 3);
-
-		if (ver >= 7)
-		{
-			arch >> BoneIds;
-		}
-		else if(ver == 6)
-		{
-			TArray<uint16> BoneIds_DEPRECATED;
-			arch >> BoneIds_DEPRECATED;
-
-			const int32 NumBones = BoneIds_DEPRECATED.Num();
-			BoneIds.Reserve(NumBones);
-			for (int32 Index = 0; Index < NumBones; ++Index)
-			{
-				BoneIds.Add({ BoneIds_DEPRECATED[Index]});
-			}
-		}
-		else
-		{
-			TArray<std::string> OldBoneNames;
-			arch >> OldBoneNames;
-
-			const int32 NumBones = OldBoneNames.Num();
-
-			BoneIds.SetNum(NumBones);
-			for (int32 Index = 0; Index < NumBones; ++Index)
-			{
-				BoneIds[Index].Id = Index;
-			}
-		}
-
-		if (ver == 3)
-		{
-			arch >> m_boneTransforms_DEPRECATED;
-		}
-
-		arch >> BoneParents;
-
-		if (ver < 6) // before BoneIds 
-		{
-			int16 ParentIndex = INDEX_NONE;
-			for (int16& BoneParent : BoneParents)
-			{
-				BoneParent = ParentIndex;
-				ParentIndex++;
-			}
-		}
-
-		if (ver <= 4)
-		{
-			TArray<int32> BoneIds_DEPRECATED;
-			arch >> BoneIds_DEPRECATED;
-		}
-
-		if (ver == 3)
-		{
-			bool bBoneTransformModified;
-			arch >> bBoneTransformModified;
-		}
+		Arch >> BoneIds;
+		Arch >> BoneParents;
 	}
 }
