[mutable] More enforcement of coding standards.

diff --git a/Engine/Plugins/Mutable/Source/MutableRuntime/Private/MuR/System.cpp b/Engine/Plugins/Mutable/Source/MutableRuntime/Private/MuR/System.cpp
--- a/Engine/Plugins/Mutable/Source/MutableRuntime/Private/MuR/System.cpp
+++ b/Engine/Plugins/Mutable/Source/MutableRuntime/Private/MuR/System.cpp
@@ -73,24 +73,17 @@ namespace mu
 
 
     //---------------------------------------------------------------------------------------------
-    System::System(const Ptr<Settings>& InSettings, const TSharedPtr<ExtensionDataStreamer>& DataStreamer)
+    FSystem::FSystem(const FSettings& InSettings, const TSharedPtr<FExtensionDataStreamer>& DataStreamer)
     {
 		LLM_SCOPE_BYNAME(TEXT("MutableRuntime"));
 
-        Ptr<Settings> pSettings = InSettings;
-
-        if ( !pSettings )
-        {
-            pSettings = new Settings;
-        }
-
         // Choose the implementation
-        m_pD = new System::Private( pSettings, DataStreamer );
+        m_pD = new FSystem::Private( InSettings, DataStreamer );
     }
 
 
     //---------------------------------------------------------------------------------------------
-    System::~System()
+    FSystem::~FSystem()
     {
 		LLM_SCOPE_BYNAME(TEXT("MutableRuntime"));
 		MUTABLE_CPUPROFILER_SCOPE(SystemDestructor);
@@ -103,20 +96,20 @@ namespace mu
 
 
     //---------------------------------------------------------------------------------------------
-    System::Private* System::GetPrivate() const
+    FSystem::Private* FSystem::GetPrivate() const
     {
         return m_pD;
     }
 
 
     //---------------------------------------------------------------------------------------------
-    System::Private::Private( Ptr<mu::Settings> InSettings, const TSharedPtr<mu::ExtensionDataStreamer>& InDataStreamer)
+    FSystem::Private::Private( const FSettings& InSettings, const TSharedPtr<mu::FExtensionDataStreamer>& InDataStreamer)
     {
 		LLM_SCOPE_BYNAME(TEXT("MutableRuntime"));
 
         Settings = InSettings;
 	
-		WorkingMemoryManager.BudgetBytes = Settings->WorkingMemoryBytes;
+		WorkingMemoryManager.BudgetBytes = Settings.WorkingMemoryBytes;
 		WorkingMemoryManager.GeneratedResources.Reserve(WorkingMemoryManager.MaxGeneratedResourceCacheSize);
 
 		ExtensionDataStreamer = InDataStreamer;
@@ -130,7 +123,7 @@ namespace mu
 
 
     //---------------------------------------------------------------------------------------------
-    System::Private::~Private()
+    FSystem::Private::~Private()
     {
 		LLM_SCOPE_BYNAME(TEXT("MutableRuntime"));
 		MUTABLE_CPUPROFILER_SCOPE(SystemPrivateDestructor);
@@ -143,7 +136,7 @@ namespace mu
 
 
     //---------------------------------------------------------------------------------------------
-    void System::SetStreamingInterface(const TSharedPtr<ModelReader>& InInterface )
+    void FSystem::SetStreamingInterface(const TSharedPtr<FModelReader>& InInterface )
     {
 		LLM_SCOPE_BYNAME(TEXT("MutableRuntime"));
        
@@ -152,7 +145,7 @@ namespace mu
 
 
 	//---------------------------------------------------------------------------------------------
-	void System::SetWorkingMemoryBytes(uint64 InBytes)
+	void FSystem::SetWorkingMemoryBytes(uint64 InBytes)
 	{
 		LLM_SCOPE_BYNAME(TEXT("MutableRuntime"));
 		MUTABLE_CPUPROFILER_SCOPE(SetWorkingMemoryBytes);
@@ -165,7 +158,7 @@ namespace mu
 
 
     //---------------------------------------------------------------------------------------------
-    void System::SetGeneratedCacheSize( uint32 InCount )
+    void FSystem::SetGeneratedCacheSize( uint32 InCount )
     {
 		LLM_SCOPE_BYNAME(TEXT("MutableRuntime"));
 		MUTABLE_CPUPROFILER_SCOPE(SetGeneratedCacheSize);
@@ -181,18 +174,18 @@ namespace mu
 
 	
 	//---------------------------------------------------------------------------------------------
-	void System::ClearWorkingMemory()
+	void FSystem::ClearWorkingMemory()
     {
 		LLM_SCOPE_BYNAME(TEXT("MutableRuntime"));
 
 		// rom caches
 		for (FWorkingMemoryManager::FModelCacheEntry& ModelCache : m_pD->WorkingMemoryManager.CachePerModel)
 		{
-			if (const TSharedPtr<const Model> CacheModel = ModelCache.Model.Pin())
+			if (const TSharedPtr<const FModel> CacheModel = ModelCache.Model.Pin())
 			{
-				FProgram& Program = CacheModel->GetPrivate()->m_program;
+				FProgram& Program = CacheModel->GetPrivate()->Program;
 
-				for (int32 RomIndex = 0; RomIndex < Program.m_roms.Num(); ++RomIndex)
+				for (int32 RomIndex = 0; RomIndex < Program.Roms.Num(); ++RomIndex)
 				{
 					Program.UnloadRom(RomIndex);
 				}
@@ -209,7 +202,7 @@ namespace mu
 
 	
     //---------------------------------------------------------------------------------------------
-    void System::SetExternalResourceProvider(const TSharedPtr<ExternalResourceProvider>& InInterface )
+    void FSystem::SetExternalResourceProvider(const TSharedPtr<FExternalResourceProvider>& InInterface )
     {
 		LLM_SCOPE_BYNAME(TEXT("MutableRuntime"));
 
@@ -218,7 +211,7 @@ namespace mu
 
 
 	//---------------------------------------------------------------------------------------------
-	void System::SetImagePixelConversionOverride(const FImageOperator::FImagePixelFormatFunc& FormatFunc)
+	void FSystem::SetImagePixelConversionOverride(const FImageOperator::FImagePixelFormatFunc& FormatFunc)
 	{
 		LLM_SCOPE_BYNAME(TEXT("MutableRuntime"));
 
@@ -227,7 +220,7 @@ namespace mu
 
 
     //---------------------------------------------------------------------------------------------
-    Instance::ID System::NewInstance( const TSharedPtr<const Model>& InModel )
+    FInstance::FID FSystem::NewInstance( const TSharedPtr<const FModel>& InModel )
     {
 		LLM_SCOPE_BYNAME(TEXT("MutableRuntime"));
 		MUTABLE_CPUPROFILER_SCOPE(NewInstance);
@@ -247,8 +240,8 @@ namespace mu
 
 
     //---------------------------------------------------------------------------------------------
-    Ptr<const Instance> System::BeginUpdate( Instance::ID InInstanceID,
-                                     const Ptr<const Parameters>& InParams,
+    TSharedPtr<const FInstance> FSystem::BeginUpdate( FInstance::FID InInstanceID,
+                                     const TSharedPtr<const FParameters>& InParams,
                                      int32 InStateIndex,
                                      uint32 InLodMask )
     {
@@ -271,9 +264,9 @@ namespace mu
 
 		m_pD->WorkingMemoryManager.CurrentInstanceCache = pLiveInstance->Cache;
 
-		FProgram& program = pLiveInstance->Model->GetPrivate()->m_program;
+		FProgram& program = pLiveInstance->Model->GetPrivate()->Program;
 
-		bool validState = InStateIndex >= 0 && InStateIndex < (int)program.m_states.Num();
+		bool validState = InStateIndex >= 0 && InStateIndex < (int)program.States.Num();
 		if (!validState)
 		{
 			UE_LOG(LogMutableCore, Error, TEXT("Invalid state in mutable update."));
@@ -290,7 +283,7 @@ namespace mu
 		// If we changed parameters that are not in this state, we need to rebuild all.
 		if (!fullBuild)
 		{
-			fullBuild = m_pD->CheckUpdatedParameters(pLiveInstance, InParams.get(), pLiveInstance->UpdatedParameters);
+			fullBuild = m_pD->CheckUpdatedParameters(pLiveInstance, InParams, pLiveInstance->UpdatedParameters);
 		}
 
 		// Remove cached data
@@ -302,7 +295,7 @@ namespace mu
 
 		m_pD->WorkingMemoryManager.BeginRunnerThread();
 
-		OP::ADDRESS rootAt = pLiveInstance->Model->GetPrivate()->m_program.m_states[InStateIndex].m_root;
+		OP::ADDRESS rootAt = pLiveInstance->Model->GetPrivate()->Program.States[InStateIndex].Root;
 
 		// Prepare instance cache
 		m_pD->PrepareCache(pLiveInstance->Model.Get(), InStateIndex);
@@ -311,9 +304,9 @@ namespace mu
 		// Ensure the model cache has been created
 		m_pD->WorkingMemoryManager.FindOrAddModelCache(pLiveInstance->Model);
 
-		m_pD->RunCode(pLiveInstance->Model, InParams.get(), rootAt, InLodMask);
+		m_pD->RunCode(pLiveInstance->Model, InParams.Get(), rootAt, InLodMask);
 
-		Ptr<const Instance> Result = pLiveInstance->Cache->GetInstance(FCacheAddress(rootAt, 0, 0));
+		TSharedPtr<const FInstance> Result = pLiveInstance->Cache->GetInstance(FCacheAddress(rootAt, 0, 0));
 
 		// Debug check to see if we managed the op-hit-counts correctly
 		pLiveInstance->Cache->CheckHitCountsCleared();
@@ -326,7 +319,7 @@ namespace mu
 		else
 		{
 			// In case of failure return an empty instance, to prevent following code to have to check it every time
-			Result = new Instance;
+			Result = MakeShared<FInstance>();
 		}
 
 		m_pD->WorkingMemoryManager.EndRunnerThread();
@@ -337,12 +330,12 @@ namespace mu
 	}
 
 
-	Ptr<const Image> System::GetImageInline(Instance::ID instanceID, FResourceID ImageId, int32 MipsToSkip, int32 InImageLOD)
+	TSharedPtr<const FImage> FSystem::GetImageInline(FInstance::FID instanceID, FResourceID ImageId, int32 MipsToSkip, int32 InImageLOD)
 	{
 		LLM_SCOPE_BYNAME(TEXT("MutableRuntime"));
 		MUTABLE_CPUPROFILER_SCOPE(SystemGetImage);
 
-		Ptr<const Image> pResult;
+		TSharedPtr<const FImage> pResult;
 
 		// Find the live instance
 		FLiveInstance* pLiveInstance = m_pD->FindLiveInstance(instanceID);
@@ -350,12 +343,12 @@ namespace mu
 		m_pD->WorkingMemoryManager.CurrentInstanceCache = pLiveInstance->Cache;
 
 		OP::ADDRESS RootAddress = GetResourceIDRoot(ImageId);
-		pResult = m_pD->BuildImage(pLiveInstance->Model, pLiveInstance->OldParameters.get(), RootAddress, MipsToSkip, InImageLOD);
+		pResult = m_pD->BuildImage(pLiveInstance->Model, pLiveInstance->OldParameters.Get(), RootAddress, MipsToSkip, InImageLOD);
 
 		// We always need to return something valid.
 		if (!pResult)
 		{
-			pResult = new mu::Image(16, 16, 1, EImageFormat::RGBA_UByte, EInitializationType::Black);
+			pResult = MakeShared<mu::FImage>(16, 16, 1, EImageFormat::RGBA_UByte, EInitializationType::Black);
 		}
 
 		m_pD->WorkingMemoryManager.CurrentInstanceCache = nullptr;
@@ -363,12 +356,12 @@ namespace mu
 		return pResult;
 	}
 
-	UE::Tasks::TTask<Ptr<const Image>> System::GetImage(Instance::ID instanceID, FResourceID ImageId, int32 MipsToSkip, int32 InImageLOD)
+	UE::Tasks::TTask<TSharedPtr<const FImage>> FSystem::GetImage(FInstance::FID instanceID, FResourceID ImageId, int32 MipsToSkip, int32 InImageLOD)
 	{
 		LLM_SCOPE_BYNAME(TEXT("MutableRuntime"));
 		MUTABLE_CPUPROFILER_SCOPE(SystemGetImage);
 
-		Ptr<const Image> pResult;
+		TSharedPtr<const FImage> pResult;
 
 		// Find the live instance
 		FLiveInstance* pLiveInstance = m_pD->FindLiveInstance(instanceID);
@@ -379,27 +372,27 @@ namespace mu
 		
 		m_pD->WorkingMemoryManager.BeginRunnerThread();
 		
-		mu::OP_TYPE OpType = pLiveInstance->Model->GetPrivate()->m_program.GetOpType(RootAddress);
+		mu::EOpType OpType = pLiveInstance->Model->GetPrivate()->Program.GetOpType(RootAddress);
 		if (GetOpDataType(OpType) != DT_IMAGE)
 		{
 			m_pD->WorkingMemoryManager.EndRunnerThread();
 			m_pD->WorkingMemoryManager.CurrentInstanceCache = nullptr;
 
-			return UE::Tasks::MakeCompletedTask<Ptr<const Image>>(
-					new mu::Image(16, 16, 1, EImageFormat::RGBA_UByte, EInitializationType::Black));
+			return UE::Tasks::MakeCompletedTask<TSharedPtr<const FImage>>(
+					MakeShared<mu::FImage>(16, 16, 1, EImageFormat::RGBA_UByte, EInitializationType::Black));
 		}
 		
 		TSharedRef<CodeRunner> Runner = CodeRunner::Create(
-				m_pD->Settings, m_pD, EExecutionStrategy::MinimizeMemory, pLiveInstance->Model, pLiveInstance->OldParameters.get(), RootAddress, System::AllLODs, MipsToSkip, InImageLOD, FScheduledOp::EType::Full);
+				m_pD->Settings, m_pD, EExecutionStrategy::MinimizeMemory, pLiveInstance->Model, pLiveInstance->OldParameters.Get(), RootAddress, FSystem::AllLODs, MipsToSkip, InImageLOD, FScheduledOp::EType::Full);
 	
 		
 		constexpr bool bForceInlineExecution = false;
 		UE::Tasks::FTask RunnerCompletionEvent = Runner->StartRun(bForceInlineExecution);
 
-		return UE::Tasks::Launch(TEXT("System::GetImageResultTask"),
-				[SystemPrivate = m_pD, Runner, RootAddress, MipsToSkip]() -> Ptr<const Image>
+		return UE::Tasks::Launch(TEXT("FSystem::GetImageResultTask"),
+				[SystemPrivate = m_pD, Runner, RootAddress, MipsToSkip]() -> TSharedPtr<const FImage>
 				{
-					Ptr<const Image> Result;
+					TSharedPtr<const FImage> Result;
 
 					SystemPrivate->bUnrecoverableError = Runner->bUnrecoverableError;
 					if (!Runner->bUnrecoverableError)
@@ -409,7 +402,7 @@ namespace mu
 
 					if (!Result)
 					{
-						Result = new mu::Image(16, 16, 1, EImageFormat::RGBA_UByte, EInitializationType::Black);
+						Result = MakeShared<mu::FImage>(16, 16, 1, EImageFormat::RGBA_UByte, EInitializationType::Black);
 					}
 
 					SystemPrivate->WorkingMemoryManager.EndRunnerThread();
@@ -431,7 +424,7 @@ namespace mu
 		TEXT("If different than 0, clear the image desc cache at every image."),
 		ECVF_Scalability);
 
-	void System::GetImageDescInline(Instance::ID instanceID, FResourceID ImageId, FImageDesc& OutDesc)
+	void FSystem::GetImageDescInline(FInstance::FID instanceID, FResourceID ImageId, FImageDesc& OutDesc)
 	{
 		LLM_SCOPE_BYNAME(TEXT("MutableRuntime"));
 		MUTABLE_CPUPROFILER_SCOPE(SystemGetImageDesc);
@@ -445,17 +438,17 @@ namespace mu
 
 		OP::ADDRESS RootAddress = GetResourceIDRoot(ImageId);
 
-		const mu::Model* Model = pLiveInstance->Model.Get();
-		const mu::FProgram& program = Model->GetPrivate()->m_program;
+		const mu::FModel* Model = pLiveInstance->Model.Get();
+		const mu::FProgram& program = Model->GetPrivate()->Program;
 
-		// TODO: It should be possible to reuse this data if cleared in the correct places only, together with m_heapImageDesc.
+		// TODO: It should be possible to reuse this data if cleared in the correct places only, together with HeapImageDesc.
 		int32 VarValue = CVarClearImageDescCache.GetValueOnAnyThread();
 		if (VarValue != 0)
 		{
 			m_pD->WorkingMemoryManager.CurrentInstanceCache->ClearDescCache();
 		}
 
-		mu::OP_TYPE opType = program.GetOpType(RootAddress);
+		mu::EOpType opType = program.GetOpType(RootAddress);
 		if (GetOpDataType(opType) == DT_IMAGE)
 		{
 			// GetImageDesc may call normal execution paths where meshes are computed.
@@ -463,7 +456,7 @@ namespace mu
 					
 			int8 executionOptions = 0;
 			TSharedRef<CodeRunner> Runner = CodeRunner::Create(
-					m_pD->Settings, m_pD, EExecutionStrategy::MinimizeMemory, pLiveInstance->Model, pLiveInstance->OldParameters.get(), RootAddress, System::AllLODs, executionOptions, 0, FScheduledOp::EType::ImageDesc);
+					m_pD->Settings, m_pD, EExecutionStrategy::MinimizeMemory, pLiveInstance->Model, pLiveInstance->OldParameters.Get(), RootAddress, FSystem::AllLODs, executionOptions, 0, FScheduledOp::EType::ImageDesc);
 			
 			constexpr bool bForceInlineExecution = true;
 			UE::Tasks::FTask CompletionEvent = Runner->StartRun(bForceInlineExecution);
@@ -478,7 +471,7 @@ namespace mu
 	}
 
 
-	UE::Tasks::TTask<FImageDesc> System::GetImageDesc(Instance::ID instanceID, FResourceID ImageId)
+	UE::Tasks::TTask<FImageDesc> FSystem::GetImageDesc(FInstance::FID instanceID, FResourceID ImageId)
 	{
 		LLM_SCOPE_BYNAME(TEXT("MutableRuntime"));
 		MUTABLE_CPUPROFILER_SCOPE(SystemGetImageDesc);
@@ -490,10 +483,10 @@ namespace mu
 
 		OP::ADDRESS RootAddress = GetResourceIDRoot(ImageId);
 
-		const mu::Model* Model = pLiveInstance->Model.Get();
-		const mu::FProgram& program = Model->GetPrivate()->m_program;
+		const mu::FModel* Model = pLiveInstance->Model.Get();
+		const mu::FProgram& program = Model->GetPrivate()->Program;
 
-		// TODO: It should be possible to reuse this data if cleared in the correct places only, together with m_heapImageDesc.
+		// TODO: It should be possible to reuse this data if cleared in the correct places only, together with HeapImageDesc.
 		int32 VarValue = CVarClearImageDescCache.GetValueOnAnyThread();
 		if (VarValue != 0)
 		{
@@ -501,7 +494,7 @@ namespace mu
 		}
 
 		m_pD->WorkingMemoryManager.BeginRunnerThread();
-		mu::OP_TYPE OpType = program.GetOpType(RootAddress);
+		mu::EOpType OpType = program.GetOpType(RootAddress);
 		if (GetOpDataType(OpType) != DT_IMAGE)
 		{
 			m_pD->WorkingMemoryManager.EndRunnerThread();
@@ -513,12 +506,12 @@ namespace mu
 		// GetImageDesc may call normal execution paths where meshes are computed.
 		int8 ExecutionOptions = 0;
 		TSharedRef<CodeRunner> Runner = CodeRunner::Create(
-				m_pD->Settings, m_pD, EExecutionStrategy::MinimizeMemory, pLiveInstance->Model, pLiveInstance->OldParameters.get(), RootAddress, System::AllLODs, ExecutionOptions, 0, FScheduledOp::EType::ImageDesc);
+				m_pD->Settings, m_pD, EExecutionStrategy::MinimizeMemory, pLiveInstance->Model, pLiveInstance->OldParameters.Get(), RootAddress, FSystem::AllLODs, ExecutionOptions, 0, FScheduledOp::EType::ImageDesc);
 		
 		constexpr bool bForceInlineExecution = false;
 		UE::Tasks::FTask RunnerCompletionEvent = Runner->StartRun(bForceInlineExecution);
 		
-		return UE::Tasks::Launch(TEXT("System::GetImageDescResultTask"),
+		return UE::Tasks::Launch(TEXT("FSystem::GetImageDescResultTask"),
 				[SystemPrivate = m_pD, Runner]() -> FImageDesc
 				{
 					FImageDesc Result;
@@ -536,12 +529,12 @@ namespace mu
 	}
 
 
-    Ptr<const Mesh> System::GetMeshInline(Instance::ID instanceID, FResourceID MeshId)
+    TSharedPtr<const FMesh> FSystem::GetMeshInline(FInstance::FID instanceID, FResourceID MeshId)
     {
 		LLM_SCOPE_BYNAME(TEXT("MutableRuntime"));
 		MUTABLE_CPUPROFILER_SCOPE(SystemGetMesh);
 
-		Ptr<const Mesh> Result;
+		TSharedPtr<const FMesh> Result;
 
 		// Find the live instance
 		FLiveInstance* pLiveInstance = m_pD->FindLiveInstance(instanceID);
@@ -549,24 +542,24 @@ namespace mu
 		m_pD->WorkingMemoryManager.CurrentInstanceCache = pLiveInstance->Cache;
 
 		OP::ADDRESS RootAddress = GetResourceIDRoot(MeshId);
-		Result = m_pD->BuildMesh(pLiveInstance->Model, pLiveInstance->OldParameters.get(), RootAddress);
+		Result = m_pD->BuildMesh(pLiveInstance->Model, pLiveInstance->OldParameters.Get(), RootAddress);
 
 		// If the mesh is null it means empty, but we still need to return a valid one
 		if (!Result)
 		{
-			Result = new Mesh();
+			Result = MakeShared<FMesh>();
 		}
 
 		m_pD->WorkingMemoryManager.CurrentInstanceCache = nullptr;
 		return Result;
 	}
 
-	UE::Tasks::TTask<Ptr<const Mesh>> System::GetMesh(Instance::ID instanceID, FResourceID MeshId)
+	UE::Tasks::TTask<TSharedPtr<const FMesh>> FSystem::GetMesh(FInstance::FID instanceID, FResourceID MeshId)
     {
 		LLM_SCOPE_BYNAME(TEXT("MutableRuntime"));
 		MUTABLE_CPUPROFILER_SCOPE(SystemGetMesh);
 
-		Ptr<const Mesh> ResultMesh;
+		TSharedPtr<const FMesh> ResultMesh;
 
 		// Find the live instance
 		FLiveInstance* pLiveInstance = m_pD->FindLiveInstance(instanceID);
@@ -576,25 +569,25 @@ namespace mu
 		
 		OP::ADDRESS RootAddress = GetResourceIDRoot(MeshId);
 			
-		mu::OP_TYPE OpType = pLiveInstance->Model->GetPrivate()->m_program.GetOpType(RootAddress);
+		mu::EOpType OpType = pLiveInstance->Model->GetPrivate()->Program.GetOpType(RootAddress);
 		if (GetOpDataType(OpType) != DT_MESH)
 		{
 			m_pD->WorkingMemoryManager.EndRunnerThread();
 			m_pD->WorkingMemoryManager.CurrentInstanceCache = nullptr;
 
-			return UE::Tasks::MakeCompletedTask<Ptr<const Mesh>>(new Mesh());
+			return UE::Tasks::MakeCompletedTask<TSharedPtr<const FMesh>>(new FMesh());
 		}
 		
 		TSharedRef<CodeRunner> Runner = CodeRunner::Create(
-				m_pD->Settings, m_pD, EExecutionStrategy::MinimizeMemory, pLiveInstance->Model, pLiveInstance->OldParameters.get(), RootAddress, System::AllLODs, 0, 0, FScheduledOp::EType::Full);
+				m_pD->Settings, m_pD, EExecutionStrategy::MinimizeMemory, pLiveInstance->Model, pLiveInstance->OldParameters.Get(), RootAddress, FSystem::AllLODs, 0, 0, FScheduledOp::EType::Full);
 		
 		constexpr bool bForceInlineExecution = false;
 		UE::Tasks::FTask RunnerCompletionEvent = Runner->StartRun(bForceInlineExecution);
 
-		return UE::Tasks::Launch(TEXT("System::GetMeshResultTask"),
-				[SystemPrivate = m_pD, Runner, RootAddress]() -> Ptr<const Mesh>
+		return UE::Tasks::Launch(TEXT("FSystem::GetMeshResultTask"),
+				[SystemPrivate = m_pD, Runner, RootAddress]() -> TSharedPtr<const FMesh>
 				{
-					Ptr<const Mesh> Result;
+					TSharedPtr<const FMesh> Result;
 
 					SystemPrivate->bUnrecoverableError = Runner->bUnrecoverableError;
 					if (!Runner->bUnrecoverableError)
@@ -614,7 +607,7 @@ namespace mu
 
 
     //---------------------------------------------------------------------------------------------
-	void System::EndUpdate(Instance::ID instanceID)
+	void FSystem::EndUpdate(FInstance::FID instanceID)
 	{
 		LLM_SCOPE_BYNAME(TEXT("MutableRuntime"));
 		MUTABLE_CPUPROFILER_SCOPE(EndUpdate);
@@ -654,7 +647,7 @@ namespace mu
 
 
     //---------------------------------------------------------------------------------------------
-    void System::ReleaseInstance( Instance::ID instanceID )
+    void FSystem::ReleaseInstance( FInstance::FID instanceID )
     {
 		LLM_SCOPE_BYNAME(TEXT("MutableRuntime"));
 		MUTABLE_CPUPROFILER_SCOPE(ReleaseInstance);
@@ -665,7 +658,7 @@ namespace mu
 			if (Instance.InstanceID == instanceID)
 			{
 				// Make sure all the resources cached in the instance are removed from the tracking list
-				for (const FProgramCache::TResourceResult<Image>& Data : Instance.Cache->ImageResults)
+				for (const FProgramCache::TResourceResult<FImage>& Data : Instance.Cache->ImageResults)
 				{
 					if (Data.Value)
 					{
@@ -673,7 +666,7 @@ namespace mu
 					}
 				}
 
-				for (const FProgramCache::TResourceResult<Mesh>& Data : Instance.Cache->MeshResults)
+				for (const FProgramCache::TResourceResult<FMesh>& Data : Instance.Cache->MeshResults)
 				{
 					if (Data.Value)
 					{
@@ -704,18 +697,18 @@ namespace mu
 
         RelevantParameterVisitor
             (
-                System::Private* InSystem,
-				const TSharedPtr<const Model>& InModel,
-                const Ptr<const Parameters>& InParams,
+                FSystem::Private* InSystem,
+				const TSharedPtr<const FModel>& InModel,
+                const TSharedPtr<const FParameters>& InParams,
                 bool* InFlags
             )
-            : UniqueDiscreteCoveredCodeVisitor<>( InSystem, InModel, InParams, System::AllLODs )
+            : UniqueDiscreteCoveredCodeVisitor<>( InSystem, InModel, InParams, FSystem::AllLODs )
         {
             Flags = InFlags;
 
             FMemory::Memset( Flags, 0, sizeof(bool)*InParams->GetCount() );
 
-            OP::ADDRESS at = InModel->GetPrivate()->m_program.m_states[0].m_root;
+            OP::ADDRESS at = InModel->GetPrivate()->Program.States[0].Root;
 
             Run( at );
         }
@@ -725,13 +718,13 @@ namespace mu
         {
             switch ( program.GetOpType(at) )
             {
-            case OP_TYPE::BO_PARAMETER:
-            case OP_TYPE::NU_PARAMETER:
-            case OP_TYPE::SC_PARAMETER:
-            case OP_TYPE::CO_PARAMETER:
-            case OP_TYPE::PR_PARAMETER:
-            case OP_TYPE::IM_PARAMETER:
-            case OP_TYPE::MA_PARAMETER:
+            case EOpType::BO_PARAMETER:
+            case EOpType::NU_PARAMETER:
+            case EOpType::SC_PARAMETER:
+            case EOpType::CO_PARAMETER:
+            case EOpType::PR_PARAMETER:
+            case EOpType::IM_PARAMETER:
+            case EOpType::MA_PARAMETER:
             {
 				OP::ParameterArgs args = program.GetOpArgs<OP::ParameterArgs>(at);
 				OP::ADDRESS ParamIndex = args.variable;
@@ -755,8 +748,8 @@ namespace mu
 
 
     //---------------------------------------------------------------------------------------------
-    void System::GetParameterRelevancy( Instance::ID InstanceID,
-                                        const Ptr<const Parameters>& Parameters,
+    void FSystem::GetParameterRelevancy( FInstance::FID InstanceID,
+                                        const TSharedPtr<const FParameters>& FParameters,
                                         bool* Flags )
     {
 		LLM_SCOPE_BYNAME(TEXT("MutableRuntime"));
@@ -766,14 +759,14 @@ namespace mu
 		check(pLiveInstance);
 		m_pD->WorkingMemoryManager.CurrentInstanceCache = pLiveInstance->Cache;
 		
-		RelevantParameterVisitor visitor( m_pD, pLiveInstance->Model, Parameters, Flags );
+		RelevantParameterVisitor visitor( m_pD, pLiveInstance->Model, FParameters, Flags );
 
 		m_pD->WorkingMemoryManager.CurrentInstanceCache = nullptr;
     }
 
 
     //---------------------------------------------------------------------------------------------
-	inline FLiveInstance* System::Private::FindLiveInstance(Instance::ID id)
+	inline FLiveInstance* FSystem::Private::FindLiveInstance(FInstance::FID id)
 	{
 		for (int32 i = 0; i < WorkingMemoryManager.LiveInstances.Num(); ++i)
 		{
@@ -787,8 +780,8 @@ namespace mu
 
 
 	//---------------------------------------------------------------------------------------------
-	bool System::Private::CheckUpdatedParameters( const FLiveInstance* LiveInstance,
-                                 const Ptr<const Parameters>& Params,
+	bool FSystem::Private::CheckUpdatedParameters( const FLiveInstance* LiveInstance,
+                                 const TSharedPtr<const FParameters>& Params,
                                  uint64& UpdatedParameters)
     {
         bool bFullBuild = false;
@@ -801,15 +794,15 @@ namespace mu
 
         // check what parameters have changed
 		UpdatedParameters = 0;
-        const FProgram& program = LiveInstance->Model->GetPrivate()->m_program;
-        const TArray<int>& runtimeParams = program.m_states[ LiveInstance->State ].m_runtimeParameters;
+        const FProgram& program = LiveInstance->Model->GetPrivate()->Program;
+        const TArray<int>& runtimeParams = program.States[ LiveInstance->State ].m_runtimeParameters;
 
-        check( Params->GetCount() == (int)program.m_parameters.Num() );
+        check( Params->GetCount() == (int)program.Parameters.Num() );
         check( !LiveInstance->OldParameters
 			||
 			Params->GetCount() == LiveInstance->OldParameters->GetCount() );
 
-        for ( int32 p=0; p<program.m_parameters.Num() && !bFullBuild; ++p )
+        for ( int32 p=0; p<program.Parameters.Num() && !bFullBuild; ++p )
         {
             bool isRuntime = runtimeParams.Contains( p );
             bool changed = !Params->HasSameValue( p, LiveInstance->OldParameters, p );
@@ -833,12 +826,12 @@ namespace mu
 
 
 	//---------------------------------------------------------------------------------------------
-	void System::Private::BeginBuild(const TSharedPtr<const Model>& InModel)
+	void FSystem::Private::BeginBuild(const TSharedPtr<const FModel>& InModel)
 	{
 		// We don't have a FLiveInstance, let's create the memory
 		// \TODO: There is no clear moment to remove this... EndBuild?
 		WorkingMemoryManager.CurrentInstanceCache = MakeShared<FProgramCache>();
-		WorkingMemoryManager.CurrentInstanceCache->Init(InModel->GetPrivate()->m_program.m_opAddress.Num());
+		WorkingMemoryManager.CurrentInstanceCache->Init(InModel->GetPrivate()->Program.OpAddress.Num());
 
 		// Ensure the model cache has been created
 		WorkingMemoryManager.FindOrAddModelCache(InModel);
@@ -848,15 +841,15 @@ namespace mu
 
 
 	//---------------------------------------------------------------------------------------------
-	void System::Private::EndBuild()
+	void FSystem::Private::EndBuild()
 	{
 		WorkingMemoryManager.CurrentInstanceCache = nullptr;
 	}
 
 
 	//---------------------------------------------------------------------------------------------
-	void System::Private::RunCode(const TSharedPtr<const Model>& InModel,
-		const Parameters* InParameters, OP::ADDRESS InCodeRoot, uint32 InLODs, uint8 executionOptions, int32 InImageLOD)
+	void FSystem::Private::RunCode(const TSharedPtr<const FModel>& InModel,
+		const FParameters* InParameters, OP::ADDRESS InCodeRoot, uint32 InLODs, uint8 executionOptions, int32 InImageLOD)
 	{
 		TSharedRef<CodeRunner> Runner = CodeRunner::Create(Settings, this, EExecutionStrategy::MinimizeMemory, InModel, InParameters, InCodeRoot, InLODs,
 			executionOptions, InImageLOD, FScheduledOp::EType::Full);
@@ -870,7 +863,7 @@ namespace mu
 
 
 	//---------------------------------------------------------------------------------------------
-	bool System::Private::BuildBool(const TSharedPtr<const Model>& pModel, const Parameters* Params, OP::ADDRESS at)
+	bool FSystem::Private::BuildBool(const TSharedPtr<const FModel>& pModel, const FParameters* Params, OP::ADDRESS at)
 	{
 		WorkingMemoryManager.BeginRunnerThread();
 
@@ -889,7 +882,7 @@ namespace mu
 
 
 	//---------------------------------------------------------------------------------------------
-	float System::Private::BuildScalar(const TSharedPtr<const Model>& pModel, const Parameters* Params, OP::ADDRESS at)
+	float FSystem::Private::BuildScalar(const TSharedPtr<const FModel>& pModel, const FParameters* Params, OP::ADDRESS at)
 	{
 		WorkingMemoryManager.BeginRunnerThread();
 
@@ -908,7 +901,7 @@ namespace mu
 
 
 	//---------------------------------------------------------------------------------------------
-	int32 System::Private::BuildInt(const TSharedPtr<const Model>& pModel, const Parameters* Params, OP::ADDRESS at)
+	int32 FSystem::Private::BuildInt(const TSharedPtr<const FModel>& pModel, const FParameters* Params, OP::ADDRESS at)
 	{
 		WorkingMemoryManager.BeginRunnerThread();
 
@@ -927,13 +920,13 @@ namespace mu
 
 
 	//---------------------------------------------------------------------------------------------
-	FVector4f System::Private::BuildColour(const TSharedPtr<const Model>& pModel, const Parameters* Params, OP::ADDRESS at)
+	FVector4f FSystem::Private::BuildColour(const TSharedPtr<const FModel>& pModel, const FParameters* Params, OP::ADDRESS at)
 	{
 		WorkingMemoryManager.BeginRunnerThread();
 
 		FVector4f Result(0,0,0,1);
 
-		mu::OP_TYPE opType = pModel->GetPrivate()->m_program.GetOpType(at);
+		mu::EOpType opType = pModel->GetPrivate()->Program.GetOpType(at);
 		if (GetOpDataType(opType) == DT_COLOUR)
 		{
 			RunCode(pModel, Params, at);
@@ -950,7 +943,7 @@ namespace mu
 
 	
 	//---------------------------------------------------------------------------------------------
-	FProjector System::Private::BuildProjector(const TSharedPtr<const Model>& pModel, const Parameters* Params, OP::ADDRESS at)
+	FProjector FSystem::Private::BuildProjector(const TSharedPtr<const FModel>& pModel, const FParameters* Params, OP::ADDRESS at)
 	{
 		WorkingMemoryManager.BeginRunnerThread();
 
@@ -968,16 +961,16 @@ namespace mu
 	}
 
 	
-	Ptr<const Image> System::Private::BuildImage(const TSharedPtr<const Model>& pModel, const Parameters* Params, OP::ADDRESS at, int32 MipsToSkip, int32 InImageLOD)
+	TSharedPtr<const FImage> FSystem::Private::BuildImage(const TSharedPtr<const FModel>& pModel, const FParameters* Params, OP::ADDRESS at, int32 MipsToSkip, int32 InImageLOD)
 	{
 		WorkingMemoryManager.BeginRunnerThread();
 
-		Ptr<const Image> Result;
+		TSharedPtr<const FImage> Result;
 
-		mu::OP_TYPE opType = pModel->GetPrivate()->m_program.GetOpType(at);
+		mu::EOpType opType = pModel->GetPrivate()->Program.GetOpType(at);
 		if (GetOpDataType(opType) == DT_IMAGE)
 		{
-			RunCode(pModel, Params, at, System::AllLODs, uint8(MipsToSkip), InImageLOD);
+			RunCode(pModel, Params, at, FSystem::AllLODs, uint8(MipsToSkip), InImageLOD);
 			if (!bUnrecoverableError)
 			{
 				Result = WorkingMemoryManager.LoadImage(FCacheAddress(at, 0, MipsToSkip), true);
@@ -990,13 +983,13 @@ namespace mu
 	}
 
 
-	Ptr<const Mesh> System::Private::BuildMesh(const TSharedPtr<const Model>& pModel, const Parameters* Params, OP::ADDRESS at)
+	TSharedPtr<const FMesh> FSystem::Private::BuildMesh(const TSharedPtr<const FModel>& pModel, const FParameters* Params, OP::ADDRESS at)
 	{
 		WorkingMemoryManager.BeginRunnerThread();
 
-		Ptr<const Mesh> Result;
+		TSharedPtr<const FMesh> Result;
 
-		mu::OP_TYPE opType = pModel->GetPrivate()->m_program.GetOpType(at);
+		mu::EOpType opType = pModel->GetPrivate()->Program.GetOpType(at);
 		if (GetOpDataType(opType) == DT_MESH)
 		{
 			RunCode(pModel, Params, at);
@@ -1012,13 +1005,13 @@ namespace mu
 	}
 
 
-	Ptr<const Instance> System::Private::BuildInstance(const TSharedPtr<const Model>& pModel, const Parameters* Params, OP::ADDRESS at)
+	TSharedPtr<const FInstance> FSystem::Private::BuildInstance(const TSharedPtr<const FModel>& pModel, const FParameters* Params, OP::ADDRESS at)
 	{
 		WorkingMemoryManager.BeginRunnerThread();
 
-		Ptr<const Instance> Result;
+		TSharedPtr<const FInstance> Result;
 
-		mu::OP_TYPE opType = pModel->GetPrivate()->m_program.GetOpType(at);
+		mu::EOpType opType = pModel->GetPrivate()->Program.GetOpType(at);
 		if (GetOpDataType(opType) == DT_INSTANCE)
 		{
 			RunCode(pModel, Params, at);
@@ -1035,15 +1028,15 @@ namespace mu
 
 
 	//---------------------------------------------------------------------------------------------
-	Ptr<const Layout> System::Private::BuildLayout(const TSharedPtr<const Model>& pModel, const Parameters* Params, OP::ADDRESS at)
+	TSharedPtr<const FLayout> FSystem::Private::BuildLayout(const TSharedPtr<const FModel>& pModel, const FParameters* Params, OP::ADDRESS at)
 	{
 		WorkingMemoryManager.BeginRunnerThread();
 
-		Ptr<const Layout> Result;
+		TSharedPtr<const FLayout> Result;
 
-		if (pModel->GetPrivate()->m_program.m_states[0].m_root)
+		if (pModel->GetPrivate()->Program.States[0].Root)
 		{
-			mu::OP_TYPE opType = pModel->GetPrivate()->m_program.GetOpType(at);
+			mu::EOpType opType = pModel->GetPrivate()->Program.GetOpType(at);
 			if (GetOpDataType(opType) == DT_LAYOUT)
 			{
 				RunCode(pModel, Params, at);
@@ -1061,15 +1054,15 @@ namespace mu
 
 
 	//---------------------------------------------------------------------------------------------
-	Ptr<const String> System::Private::BuildString(const TSharedPtr<const Model>& pModel, const Parameters* Params, OP::ADDRESS at)
+	TSharedPtr<const String> FSystem::Private::BuildString(const TSharedPtr<const FModel>& pModel, const FParameters* Params, OP::ADDRESS at)
 	{
 		WorkingMemoryManager.BeginRunnerThread();
 
-		Ptr<const String> Result;
+		TSharedPtr<const String> Result;
 
-		if (pModel->GetPrivate()->m_program.m_states[0].m_root)
+		if (pModel->GetPrivate()->Program.States[0].Root)
 		{
-			mu::OP_TYPE opType = pModel->GetPrivate()->m_program.GetOpType(at);
+			mu::EOpType opType = pModel->GetPrivate()->Program.GetOpType(at);
 			if (GetOpDataType(opType) == DT_STRING)
 			{
 				RunCode(pModel, Params, at);
@@ -1087,12 +1080,12 @@ namespace mu
 
 
 	//---------------------------------------------------------------------------------------------
-	void System::Private::PrepareCache( const Model* InModel, int32 InState)
+	void FSystem::Private::PrepareCache( const FModel* InModel, int32 InState)
 	{
 		MUTABLE_CPUPROFILER_SCOPE(PrepareCache);
 
-		FProgram& Program = InModel->GetPrivate()->m_program;
-		int32 OpCount = Program.m_opAddress.Num();
+		FProgram& Program = InModel->GetPrivate()->Program;
+		int32 OpCount = Program.OpAddress.Num();
 		WorkingMemoryManager.CurrentInstanceCache->Init(OpCount);
 
 		// Clear cache flags of existing data
@@ -1104,9 +1097,9 @@ namespace mu
 		}
 
 		// Mark the resources that have to be cached to update the instance in this state
-		if (InState >= 0 && InState< Program.m_states.Num())
+		if (InState >= 0 && InState< Program.States.Num())
 		{
-			const FProgram::FState& State = Program.m_states[InState];
+			const FProgram::FState& State = Program.States[InState];
 			for (uint32 Address : State.m_updateCache)
 			{
 				WorkingMemoryManager.CurrentInstanceCache->SetForceCached(Address);
@@ -1116,7 +1109,7 @@ namespace mu
 
 
 	//---------------------------------------------------------------------------------------------
-	void System::Private::UpdateStats()
+	void FSystem::Private::UpdateStats()
 	{
 		// Updater stats
 		int32 WorkingMemoryKb = WorkingMemoryManager.BudgetBytes / 1024;
@@ -1176,13 +1169,13 @@ namespace mu
 
 
 	//---------------------------------------------------------------------------------------------
-	FWorkingMemoryManager::FModelCacheEntry* FWorkingMemoryManager::FindModelCache(const Model* InModel)
+	FWorkingMemoryManager::FModelCacheEntry* FWorkingMemoryManager::FindModelCache(const FModel* InModel)
 	{
 		check(InModel);
 
 		for (FModelCacheEntry& c : CachePerModel)
 		{
-			TSharedPtr<const Model> pCandidate = c.Model.Pin();
+			TSharedPtr<const FModel> pCandidate = c.Model.Pin();
 			if (pCandidate)
 			{
 				if (pCandidate.Get() == InModel)
@@ -1196,7 +1189,7 @@ namespace mu
 	}
 
     //---------------------------------------------------------------------------------------------
-	FWorkingMemoryManager::FModelCacheEntry& FWorkingMemoryManager::FindOrAddModelCache(const TSharedPtr<const Model>& InModel)
+	FWorkingMemoryManager::FModelCacheEntry& FWorkingMemoryManager::FindOrAddModelCache(const TSharedPtr<const FModel>& InModel)
     {
         check(InModel);
 
@@ -1222,9 +1215,9 @@ namespace mu
 
         // Not found. Add new
 		FModelCacheEntry n;
-        n.Model = TWeakPtr<const Model>(InModel);
-		n.PendingOpsPerRom.SetNum(InModel->GetPrivate()->m_program.m_roms.Num());
-    	n.RomWeights.SetNumZeroed(InModel->GetPrivate()->m_program.m_roms.Num());
+        n.Model = TWeakPtr<const FModel>(InModel);
+		n.PendingOpsPerRom.SetNum(InModel->GetPrivate()->Program.Roms.Num());
+    	n.RomWeights.SetNumZeroed(InModel->GetPrivate()->Program.Roms.Num());
     	CachePerModel.Add(n);
         return CachePerModel.Last();
     }
@@ -1312,7 +1305,7 @@ namespace mu
 
 			struct FRomRef
 			{
-				const Model* Model = nullptr;
+				const FModel* Model = nullptr;
 				int32 RomIndex = 0;
 			};
 
@@ -1321,19 +1314,19 @@ namespace mu
 
 			for (FModelCacheEntry& ModelCache : CachePerModel)
 			{
-				TSharedPtr<const Model> CacheModel = ModelCache.Model.Pin();
+				TSharedPtr<const FModel> CacheModel = ModelCache.Model.Pin();
 				if (CacheModel)
 				{
-					mu::FProgram& Program = CacheModel->GetPrivate()->m_program;
-					check(ModelCache.RomWeights.Num() == Program.m_roms.Num());
+					mu::FProgram& Program = CacheModel->GetPrivate()->Program;
+					check(ModelCache.RomWeights.Num() == Program.Roms.Num());
 	
-					check(Program.LoadedMemTrackedRoms.GetMaxIndex() <= Program.m_roms.Num());
+					check(Program.LoadedMemTrackedRoms.GetMaxIndex() <= Program.Roms.Num());
 
 					for (TSparseArray<uint8>::TConstIterator Iter = Program.LoadedMemTrackedRoms.CreateConstIterator(); Iter; ++Iter)
 					{
 						const int32 RomIndex = Iter.GetIndex();
 
-						const FRomData& Rom = Program.m_roms[RomIndex];
+						const FRomData& Rom = Program.Roms[RomIndex];
 						check(Program.IsRomLoaded(RomIndex));
 						check(Rom.ResourceType == (uint32)*Iter);
 
@@ -1365,7 +1358,7 @@ namespace mu
 				Candidates.HeapPop(Candidate, CompareCandidates, EAllowShrinking::No);
 
 				// UE_LOG(LogMutableCore,Log, "Unloading rom because of memory budget: %d.", lowestPriorityRom);
-				int32 UnloadedSize = Candidate.Value.Model->GetPrivate()->m_program.UnloadRom(Candidate.Value.RomIndex);
+				int32 UnloadedSize = Candidate.Value.Model->GetPrivate()->Program.UnloadRom(Candidate.Value.RomIndex);
 				TotalBytes -= UnloadedSize;
 				bFinished = TotalBytes <= BudgetBytes;
 			}
@@ -1376,7 +1369,7 @@ namespace mu
 		{
 			MUTABLE_CPUPROFILER_SCOPE(EnsureBudgetBelow_FreeCached);
 
-			TSet<const Resource*> RemovedResources;
+			TSet<const FResource*> RemovedResources;
 			RemovedResources.Reserve(1024);
 			
 			// From other live instances first
@@ -1394,19 +1387,19 @@ namespace mu
 
 					RemovedResources.Reset();
 					int32 ImageRemovedBytes = 0;
-					for (const FProgramCache::TResourceResult<Image>& ImageResult : Instance.Cache->ImageResults)
+					for (const FProgramCache::TResourceResult<FImage>& ImageResult : Instance.Cache->ImageResults)
 					{
 						if (!ImageResult.Value)
 						{
 							continue;
 						}
 
-						if (!RemovedResources.Find(ImageResult.Value.get()))
+						if (!RemovedResources.Find(ImageResult.Value.Get()))
 						{
 							ImageRemovedBytes += ImageResult.Value->GetDataSize();
 							bFinished = TotalBytes - ImageRemovedBytes <= BudgetBytes;
 
-							RemovedResources.Add(ImageResult.Value.get());
+							RemovedResources.Add(ImageResult.Value.Get());
 							CacheResources.Remove(ImageResult.Value);
 						}
 
@@ -1418,9 +1411,9 @@ namespace mu
 
 					if (ImageRemovedBytes > 0)
 					{
-						for (const FProgramCache::TResourceResult<Image>& ImageResult : Instance.Cache->ImageResults)
+						for (const FProgramCache::TResourceResult<FImage>& ImageResult : Instance.Cache->ImageResults)
 						{
-							if (RemovedResources.Find(ImageResult.Value.get()))
+							if (RemovedResources.Find(ImageResult.Value.Get()))
 							{
 								Instance.Cache->SetUnused(Instance.Cache->OpExecutionData[ImageResult.OpAddress]);
 							}
@@ -1435,19 +1428,19 @@ namespace mu
 
 					int32 MeshRemovedBytes = 0;
 					RemovedResources.Reset();
-					for (const FProgramCache::TResourceResult<Mesh>& MeshResult : Instance.Cache->MeshResults)
+					for (const FProgramCache::TResourceResult<FMesh>& MeshResult : Instance.Cache->MeshResults)
 					{
 						if (!MeshResult.Value)
 						{
 							continue;
 						}
 
-						if (!RemovedResources.Find(MeshResult.Value.get()))
+						if (!RemovedResources.Find(MeshResult.Value.Get()))
 						{
 							MeshRemovedBytes = MeshResult.Value->GetDataSize();
 							bFinished = TotalBytes - MeshRemovedBytes <= BudgetBytes;
 
-							RemovedResources.Add(MeshResult.Value.get());
+							RemovedResources.Add(MeshResult.Value.Get());
 							CacheResources.Remove(MeshResult.Value);
 						}
 
@@ -1459,9 +1452,9 @@ namespace mu
 
 					if (MeshRemovedBytes > 0)
 					{
-						for (const FProgramCache::TResourceResult<Mesh>& MeshResult : Instance.Cache->MeshResults)
+						for (const FProgramCache::TResourceResult<FMesh>& MeshResult : Instance.Cache->MeshResults)
 						{
-							if (RemovedResources.Find(MeshResult.Value.get()))
+							if (RemovedResources.Find(MeshResult.Value.Get()))
 							{
 								Instance.Cache->SetUnused(Instance.Cache->OpExecutionData[MeshResult.OpAddress]);
 							}
@@ -1579,7 +1572,7 @@ namespace mu
 			{
 				FMemory::Memzero(SearchFlags.GetData(), MaxNumResources*sizeof(ESearchFlags));
 				
-				const TArrayView<FProgramCache::TResourceResult<Image>>& Images = MakeArrayView(CurrentInstanceCache->ImageResults);
+				const TArrayView<FProgramCache::TResourceResult<FImage>>& Images = MakeArrayView(CurrentInstanceCache->ImageResults);
 				const int32 RemovedBytes = SearchResourcesToRemove(Images, SearchFlags); 
 				
 				if (RemovedBytes > 0)
@@ -1611,7 +1604,7 @@ namespace mu
 			{
 				FMemory::Memzero(SearchFlags.GetData(), MaxNumResources*sizeof(ESearchFlags));
 				
-				TArrayView<FProgramCache::TResourceResult<Mesh>> Meshes = MakeArrayView(CurrentInstanceCache->MeshResults);
+				TArrayView<FProgramCache::TResourceResult<FMesh>> Meshes = MakeArrayView(CurrentInstanceCache->MeshResults);
 				const int32 RemovedBytes = SearchResourcesToRemove(Meshes, SearchFlags); 
 
 				if (RemovedBytes > 0)
@@ -1665,7 +1658,7 @@ namespace mu
     }
 
     //---------------------------------------------------------------------------------------------
-    void FWorkingMemoryManager::MarkRomUsed( int32 romIndex, const TSharedPtr<const Model>& pModel )
+    void FWorkingMemoryManager::MarkRomUsed( int32 romIndex, const TSharedPtr<const FModel>& pModel )
     {
         check(pModel);
 
@@ -1703,7 +1696,7 @@ namespace mu
 
 
 	//---------------------------------------------------------------------------------------------
-	FResourceID FWorkingMemoryManager::GetResourceKey(const TSharedPtr<const Model>& Model, const Parameters* Params, uint32 ParamListIndex, OP::ADDRESS RootAt)
+	FResourceID FWorkingMemoryManager::GetResourceKey(const TSharedPtr<const FModel>& Model, const FParameters* Params, uint32 ParamListIndex, OP::ADDRESS RootAt)
 	{
 		MUTABLE_CPUPROFILER_SCOPE(GetResourceKey);
 
@@ -1714,13 +1707,13 @@ namespace mu
 			return ErrorId;
 		}
 
-		const FProgram& Program = Model->GetPrivate()->m_program;
+		const FProgram& Program = Model->GetPrivate()->Program;
 
 		// Find the list of relevant parameters
 		const TArray<uint16>* RelevantParams = nullptr;
-		if (ParamListIndex < (uint32)Program.m_parameterLists.Num())
+		if (ParamListIndex < (uint32)Program.ParameterLists.Num())
 		{
-			RelevantParams = &Program.m_parameterLists[ParamListIndex];
+			RelevantParams = &Program.ParameterLists[ParamListIndex];
 		}
 		check(RelevantParams);
 		if (!RelevantParams)
@@ -1731,7 +1724,7 @@ namespace mu
 		// Generate the relevant parameters blob
 		FBitWriter Blob( 2048*8, true );
 
-		const TArray<FParameterDesc>& ParamDescs = Params->GetPrivate()->m_pModel->GetPrivate()->m_program.m_parameters;
+		const TArray<FParameterDesc>& ParamDescs = Params->GetPrivate()->Model->GetPrivate()->Program.Parameters;
 
 		// First make a mask with a bit for each relevant parameter. It will be on for parameters included in the blob.
 		// A parameter will be excluded from the blob if it has the default value, and no multivalues.
@@ -1745,9 +1738,9 @@ namespace mu
 				if (!bInclude)
 				{
 					bInclude = !(
-						Params->GetPrivate()->m_values[ParamIndex]
+						Params->GetPrivate()->Values[ParamIndex]
 						==
-						ParamDescs[ParamIndex].m_defaultValue
+						ParamDescs[ParamIndex].DefaultValue
 						);
 				}
 
@@ -1767,15 +1760,15 @@ namespace mu
 
 			int32 DataSize = 0;
 
-			switch (Program.m_parameters[ParamIndex].m_type)
+			switch (Program.Parameters[ParamIndex].Type)
 			{
 			case EParameterType::Bool:
-				Blob.WriteBit(Params->GetPrivate()->m_values[ParamIndex].Get<FParamBoolType>() ? 1 : 0);
+				Blob.WriteBit(Params->GetPrivate()->Values[ParamIndex].Get<FParamBoolType>() ? 1 : 0);
 
 				// Multi-values
 				if (Params->GetPrivate()->HasMultipleValues(ParamIndex))
 				{
-					const TMap< TArray<int32>, FParameterValue >& Multi = Params->GetPrivate()->m_multiValues[ParamIndex];
+					const TMap< TArray<int32>, FParameterValue >& Multi = Params->GetPrivate()->MultiValues[ParamIndex];
 					AddMultiValueKeys(Blob, Multi);
 					for (const TPair< TArray<int32>, FParameterValue >& v : Multi)
 					{
@@ -1786,8 +1779,8 @@ namespace mu
 
 			case EParameterType::Int:
 			{
-				int32 MaxValue = ParamDescs[ParamIndex].m_possibleValues.Num();
-				int32 Value = Params->GetPrivate()->m_values[ParamIndex].Get<FParamIntType>();
+				int32 MaxValue = ParamDescs[ParamIndex].PossibleValues.Num();
+				int32 Value = Params->GetPrivate()->Values[ParamIndex].Get<FParamIntType>();
 				if (MaxValue)
 				{
 					// It is an enum
@@ -1804,7 +1797,7 @@ namespace mu
 				// Multi-values
 				if (Params->GetPrivate()->HasMultipleValues(ParamIndex))
 				{
-					const TMap< TArray<int32>, FParameterValue >& Multi = Params->GetPrivate()->m_multiValues[ParamIndex];
+					const TMap< TArray<int32>, FParameterValue >& Multi = Params->GetPrivate()->MultiValues[ParamIndex];
 					AddMultiValueKeys(Blob, Multi);
 					for (const TPair< TArray<int32>, FParameterValue >& v : Multi)
 					{
@@ -1828,12 +1821,12 @@ namespace mu
 
 			case EParameterType::Float:
 				DataSize = sizeof(FParamFloatType);
-				Blob.Serialize(&Params->GetPrivate()->m_values[ParamIndex].Get<FParamFloatType>(), DataSize);
+				Blob.Serialize(&Params->GetPrivate()->Values[ParamIndex].Get<FParamFloatType>(), DataSize);
 
 				// Multi-values
 				if (Params->GetPrivate()->HasMultipleValues(ParamIndex))
 				{
-					const TMap< TArray<int32>, FParameterValue >& Multi = Params->GetPrivate()->m_multiValues[ParamIndex];
+					const TMap< TArray<int32>, FParameterValue >& Multi = Params->GetPrivate()->MultiValues[ParamIndex];
 					AddMultiValueKeys(Blob, Multi);
 					for (const TPair< TArray<int32>, FParameterValue >& v : Multi)
 					{
@@ -1844,12 +1837,12 @@ namespace mu
 
 			case EParameterType::Color:
 				DataSize = sizeof(FParamColorType);
-				Blob.Serialize(&Params->GetPrivate()->m_values[ParamIndex].Get<FParamColorType>(), DataSize);
+				Blob.Serialize(&Params->GetPrivate()->Values[ParamIndex].Get<FParamColorType>(), DataSize);
 
 				// Multi-values
 				if (Params->GetPrivate()->HasMultipleValues(ParamIndex))
 				{
-					const TMap< TArray<int32>, FParameterValue >& Multi = Params->GetPrivate()->m_multiValues[ParamIndex];
+					const TMap< TArray<int32>, FParameterValue >& Multi = Params->GetPrivate()->MultiValues[ParamIndex];
 					AddMultiValueKeys(Blob, Multi);
 					for (const TPair< TArray<int32>, FParameterValue >& v : Multi)
 					{
@@ -1861,7 +1854,7 @@ namespace mu
 			case EParameterType::Projector:
 			{
 				FPackedNormal TempVec;
-				FProjector& Value = Params->GetPrivate()->m_values[ParamIndex].Get<FParamProjectorType>();
+				FProjector& Value = Params->GetPrivate()->Values[ParamIndex].Get<FParamProjectorType>();
 				Blob.Serialize((void*)&Value.position, sizeof(FVector3f));
 				TempVec = FPackedNormal(Value.direction);
 				Blob.Serialize(&TempVec, sizeof(FPackedNormal));
@@ -1873,7 +1866,7 @@ namespace mu
 				// Multi-values
 				if (Params->GetPrivate()->HasMultipleValues(ParamIndex))
 				{
-					const TMap< TArray<int32>, FParameterValue >& Multi = Params->GetPrivate()->m_multiValues[ParamIndex];
+					const TMap< TArray<int32>, FParameterValue >& Multi = Params->GetPrivate()->MultiValues[ParamIndex];
 					AddMultiValueKeys(Blob, Multi);
 					for (const TPair< TArray<int32>, FParameterValue >& v : Multi)
 					{
@@ -1893,12 +1886,12 @@ namespace mu
 			case EParameterType::Image:
 			{
 				DataSize = sizeof(FName);
-				Blob.Serialize((void*)&Params->GetPrivate()->m_values[ParamIndex].Get<FParamImageType>(), DataSize);
+				Blob.Serialize((void*)&Params->GetPrivate()->Values[ParamIndex].Get<FParamImageType>(), DataSize);
 
 				// Multi-values
 				if (Params->GetPrivate()->HasMultipleValues(ParamIndex))
 				{
-					const TMap< TArray<int32>, FParameterValue >& Multi = Params->GetPrivate()->m_multiValues[ParamIndex];
+					const TMap< TArray<int32>, FParameterValue >& Multi = Params->GetPrivate()->MultiValues[ParamIndex];
 					AddMultiValueKeys(Blob, Multi);
 					for (const TPair< TArray<int32>, FParameterValue >& v : Multi)
 					{
diff --git a/Engine/Plugins/Mutable/Source/MutableRuntime/Public/MuR/ExtensionData.h b/Engine/Plugins/Mutable/Source/MutableRuntime/Public/MuR/ExtensionData.h
--- a/Engine/Plugins/Mutable/Source/MutableRuntime/Public/MuR/ExtensionData.h
+++ b/Engine/Plugins/Mutable/Source/MutableRuntime/Public/MuR/ExtensionData.h
@@ -9,20 +9,15 @@
 
 namespace mu
 {
-	//! ExtensionData represents data types that Mutable doesn't support natively.
-	//!
-	//! Extensions can provide data, and functionality to operate on that data, without Mutable
-	//! needing to know what the data refers to.
-
-	class ExtensionData;
-	typedef Ptr<ExtensionData> ExtensionDataPtr;
-	typedef Ptr<const ExtensionData> ExtensionDataPtrConst;
-
-	class MUTABLERUNTIME_API ExtensionData : public Resource
+	/** ExtensionData represents data types that Mutable doesn't support natively.
+	* Extensions can provide data, and functionality to operate on that data, without Mutable
+	* needing to know what the data refers to.
+	*/
+	class MUTABLERUNTIME_API FExtensionData : public FResource
 	{
 	public:
-		static void Serialise(const ExtensionData* Data, OutputArchive& Archive);
-		static ExtensionDataPtr StaticUnserialise(InputArchive& Archive);
+		static void Serialise(const FExtensionData* Data, FOutputArchive& Archive);
+		static TSharedPtr<FExtensionData> StaticUnserialise(FInputArchive& Archive);
 
 		// Resource interface
 		int32 GetDataSize() const override;
@@ -30,10 +25,10 @@ namespace mu
 		//! A stable hash of the contents
 		uint32 Hash() const;
 
-		void Serialise(OutputArchive& Archive) const;
-		void Unserialise(InputArchive& Archive);
+		void Serialise(FOutputArchive& Archive) const;
+		void Unserialise(FInputArchive& Archive);
 
-		inline bool operator==(const ExtensionData& Other) const
+		inline bool operator==(const FExtensionData& Other) const
 		{
 			const bool bResult =
 				Other.Index == Index
diff --git a/Engine/Plugins/Mutable/Source/MutableRuntime/Public/MuR/ExtensionDataStreamer.h b/Engine/Plugins/Mutable/Source/MutableRuntime/Public/MuR/ExtensionDataStreamer.h
--- a/Engine/Plugins/Mutable/Source/MutableRuntime/Public/MuR/ExtensionDataStreamer.h
+++ b/Engine/Plugins/Mutable/Source/MutableRuntime/Public/MuR/ExtensionDataStreamer.h
@@ -12,59 +12,56 @@
 namespace mu
 {
 
-//! Allows runtime code to track the progress of an Extension Data load
-class FExtensionDataLoadHandle
-{
-public:
-	ExtensionDataPtrConst Data;
-
-	enum class ELoadState : uint8
+	//! Allows runtime code to track the progress of an Extension Data load
+	class FExtensionDataLoadHandle
 	{
-		Pending,
-		FailedToLoad,
-		Loaded
+	public:
+		TSharedPtr<const FExtensionData> Data;
+
+		enum class ELoadState : uint8
+		{
+			Pending,
+			FailedToLoad,
+			Loaded
+		};
+		ELoadState LoadState;
 	};
-	ELoadState LoadState;
-};
 
-//! \brief Interface to allow the runtime to load ExtensionData for use
-//!
-//! An object implementing this interface must be provided to the System for any Models that use
-//! ExtensionData.
-//! 
-//! All functions must be safe to call from the Mutable thread.
-//! 
-//! \ingroup runtime
-class MUTABLERUNTIME_API ExtensionDataStreamer
-{
-public:
+	/** Interface to allow the runtime to load ExtensionData for use
+	* An object implementing this interface must be provided to the System for any Models that use
+	* ExtensionData.
+	* All functions must be safe to call from the Mutable thread.
+	*/
+	class MUTABLERUNTIME_API FExtensionDataStreamer
+	{
+	public:
 
-	//! Ensure virtual destruction.
-	virtual ~ExtensionDataStreamer() = default;
+		//! Ensure virtual destruction.
+		virtual ~FExtensionDataStreamer() = default;
 
-	//! Make a copy of the actual data referenced by the ExtensionData, and return a new
-	//! runtime-generated ExtensionData that references it.
-	//! 
-	//! Source can either be a constant or another runtime-generated ExtensionData.
-	//! 
-	//! The new ExtensionData should be completely independent of the source ExtensionData, such
-	//! that if the source instance is deleted, the new instance should still be valid.
-	virtual ExtensionDataPtr CloneExtensionData(const ExtensionDataPtrConst& Source) = 0;
+		//! Make a copy of the actual data referenced by the ExtensionData, and return a new
+		//! runtime-generated ExtensionData that references it.
+		//! 
+		//! Source can either be a constant or another runtime-generated ExtensionData.
+		//! 
+		//! The new ExtensionData should be completely independent of the source ExtensionData, such
+		//! that if the source instance is deleted, the new instance should still be valid.
+		virtual TSharedPtr<FExtensionData> CloneExtensionData(const TSharedPtr<const FExtensionData>& Source) = 0;
 	
-	//! Start asynchronously loading the streamed constant data referenced by the given
-	//! ExtensionData.
-	//! 
-	//! \param Data The ExtensionData to load. Its Origin must be ConstantStreamed.
-	//! 
-	//! \param OutUnloadedConstants An array of ConstantStreamed ExtensionData that was unloaded
-	//! to make space for the requested data. Guaranteed to only contain valid pointers to
-	//! ConstantStreamed data from this program.
-	//! 
-	//! \return A handle that allows the caller to poll the load state. When the state is no longer
-	//! Pending, the load is either complete or has failed.
-	virtual TSharedRef<const FExtensionDataLoadHandle> StartLoad(
-		const ExtensionDataPtrConst& Data,
-		TArray<ExtensionDataPtrConst>& OutUnloadedConstants) = 0;
-};
+		//! Start asynchronously loading the streamed constant data referenced by the given
+		//! ExtensionData.
+		//! 
+		//! \param Data The ExtensionData to load. Its Origin must be ConstantStreamed.
+		//! 
+		//! \param OutUnloadedConstants An array of ConstantStreamed ExtensionData that was unloaded
+		//! to make space for the requested data. Guaranteed to only contain valid pointers to
+		//! ConstantStreamed data from this program.
+		//! 
+		//! \return A handle that allows the caller to poll the load state. When the state is no longer
+		//! Pending, the load is either complete or has failed.
+		virtual TSharedRef<const FExtensionDataLoadHandle> StartLoad(
+			const TSharedPtr<const FExtensionData>& Data,
+			TArray<TSharedPtr<const FExtensionData>>& OutUnloadedConstants) = 0;
+	};
 
 }
diff --git a/Engine/Plugins/Mutable/Source/MutableRuntime/Public/MuR/Image.h b/Engine/Plugins/Mutable/Source/MutableRuntime/Public/MuR/Image.h
--- a/Engine/Plugins/Mutable/Source/MutableRuntime/Public/MuR/Image.h
+++ b/Engine/Plugins/Mutable/Source/MutableRuntime/Public/MuR/Image.h
@@ -18,15 +18,9 @@
 
 namespace mu
 {
-	// Forward declarations
-	class Image;
-
-	typedef Ptr<Image> ImagePtr;
-	typedef Ptr<const Image> ImagePtrConst;
-
-    //! \brief 2D image resource with mipmaps.
-	//! \ingroup runtime
-    class MUTABLERUNTIME_API Image : public Resource
+    /** 2D image resource with mipmaps.
+	*/
+    class MUTABLERUNTIME_API FImage : public FResource
 	{
 	public:
 
@@ -34,7 +28,7 @@ namespace mu
 		// Life cycle
 		//-----------------------------------------------------------------------------------------
 
-		Image();
+		FImage();
 
 		//! Constructor from the size and format.
 		//! It will allocate the memory for the image, but its content will be undefined.
@@ -44,14 +38,14 @@ namespace mu
         //!         base level. It must be a number between 1 and the maximum possible levels, which
         //!         depends on the image size.
         //! \param format Pixel format.
-        Image(uint32 SizeX, uint32 SizeY, uint32 lods, EImageFormat Format, EInitializationType InitType);
+		FImage(uint32 SizeX, uint32 SizeY, uint32 Lods, EImageFormat Format, EInitializationType InitType);
 
 		/** Create a new empty image that repreents an external resource image. */
-		static Ptr<Image> CreateAsReference(uint32 ID, const FImageDesc& Desc, bool bForceLoad);
+		static TSharedPtr<FImage> CreateAsReference(uint32 ID, const FImageDesc& Desc, bool bForceLoad);
 
 		//! Serialisation
-		static void Serialise( const Image* p, OutputArchive& arch );
-		static Ptr<Image> StaticUnserialise( InputArchive& arch );
+		static void Serialise( const FImage*, FOutputArchive& );
+		static TSharedPtr<FImage> StaticUnserialise( FInputArchive& );
 
 		// Resource interface
 		int32 GetDataSize() const override;
@@ -96,13 +90,7 @@ namespace mu
 
         //! Return the size in bytes of a specific LOD of the image.
         int32 GetLODDataSize(int32 LODIndex) const;
-
-
-	protected:
-
-		//! Forbidden. Manage with the Ptr<> template.
-		~Image() {}
-
+		
 	public:
 
 		// This used to be the data in the private implementation of the image interface
@@ -130,7 +118,7 @@ namespace mu
 		} EImageFlags;
 
 		/** Persistent flags with some image properties. The meaning will depend of every context. */
-		mutable uint8 m_flags = 0;
+		mutable uint8 Flags = 0;
 
 		/** Non-persistent relevancy map. */
 		uint16 RelevancyMinY = 0;
@@ -146,24 +134,24 @@ namespace mu
 		//-----------------------------------------------------------------------------------------
 
 		//! Deep clone.
-		Ptr<Image> Clone() const
+		TSharedPtr<FImage> Clone() const
 		{
 			LLM_SCOPE_BYNAME(TEXT("MutableRuntime"));
 			MUTABLE_CPUPROFILER_SCOPE(ImageClone)
 
-			Ptr<Image> pResult = new Image();
+			TSharedPtr<FImage> Result = MakeShared<FImage>();
 
-			pResult->m_flags = m_flags;
-			pResult->RelevancyMinY = RelevancyMinY;
-			pResult->RelevancyMaxY = RelevancyMaxY;
-			pResult->DataStorage = DataStorage;
-			pResult->ReferenceID = ReferenceID;
+			Result->Flags = Flags;
+			Result->RelevancyMinY = RelevancyMinY;
+			Result->RelevancyMaxY = RelevancyMaxY;
+			Result->DataStorage = DataStorage;
+			Result->ReferenceID = ReferenceID;
 
-			return pResult;
+			return Result;
 		}
 
 		//! Copy another image.
-		void Copy(const Image* Other)
+		void Copy(const FImage* Other)
 		{
 			LLM_SCOPE_BYNAME(TEXT("MutableRuntime"));
 			MUTABLE_CPUPROFILER_SCOPE(Copy);
@@ -173,14 +161,14 @@ namespace mu
 				return;
 			}
 
-			m_flags = Other->m_flags;
+			Flags = Other->Flags;
 			RelevancyMinY = Other->RelevancyMinY;
 			RelevancyMaxY = Other->RelevancyMaxY;
 			DataStorage = Other->DataStorage;
 			ReferenceID = Other->ReferenceID;
 		}
 
-		void CopyMove(Image* Other)
+		void CopyMove(FImage* Other)
 		{
 			LLM_SCOPE_BYNAME(TEXT("MutableRuntime"));
 			MUTABLE_CPUPROFILER_SCOPE(CopyMove);
@@ -190,7 +178,7 @@ namespace mu
 				return;
 			}
 
-			m_flags = Other->m_flags;
+			Flags = Other->Flags;
 			RelevancyMinY = Other->RelevancyMinY;
 			RelevancyMaxY = Other->RelevancyMaxY;
 			ReferenceID = Other->ReferenceID;
@@ -200,13 +188,13 @@ namespace mu
 
 
 		//!
-		void Serialise(OutputArchive& arch) const;
+		void Serialise(FOutputArchive&) const;
 
 		//!
-		void Unserialise(InputArchive& arch);
+		void Unserialise(FInputArchive&);
 
 		//!
-		inline bool operator==(const Image& Other) const
+		inline bool operator==(const FImage& Other) const
 		{
 			return 
 				(DataStorage == Other.DataStorage) &&
@@ -216,7 +204,7 @@ namespace mu
 		//-----------------------------------------------------------------------------------------
 
 		//! Sample the image and return an RGB colour.
-		FVector4f Sample(FVector2f coords) const;
+		FVector4f Sample(FVector2f Coords) const;
 
 		//! Calculate the size of the image data in bytes, regardless of what is allocated in
 		//! m_data, only using the image descriptions. For non-block-compressed images, it returns
@@ -236,7 +224,7 @@ namespace mu
 		int32 GetMipsDataSize() const;
 
 		//! See if all the pixels in the image are equal and return the colour.
-		bool IsPlainColour(FVector4f& colour) const;
+		bool IsPlainColour(FVector4f& Colour) const;
 
 		//! See if all the pixels in the alpha channel are the max value of the pixel format
 		//! (white).
@@ -251,8 +239,8 @@ namespace mu
 		static int32 GetMipmapCount(int32 SizeX, int32 SizeY);
 
 		//! Get the rect inside the image bounding the non-black content of the image.
-		void GetNonBlackRect(FImageRect& rect) const;
-		void GetNonBlackRect_Reference(FImageRect& rect) const;
+		void GetNonBlackRect(FImageRect& OutRect) const;
+		void GetNonBlackRect_Reference(FImageRect& OutRect) const;
 	};
 
 	/** This struct contains image operations that may need to allocate and free temporary images to work.
@@ -261,9 +249,9 @@ namespace mu
 	struct FImageOperator
 	{
 		/** Common callback used for functions that can create temporary images. */
-		typedef TFunction<Ptr<Image>(int32, int32, int32, EImageFormat, EInitializationType)> FImageCreateFunc;
-		typedef TFunction<void(Ptr<Image>&)> FImageReleaseFunc;
-		typedef TFunction<Ptr<Image>(const Image*)> FImageCloneFunc;
+		typedef TFunction<TSharedPtr<FImage>(int32, int32, int32, EImageFormat, EInitializationType)> FImageCreateFunc;
+		typedef TFunction<void(TSharedPtr<FImage>&)> FImageReleaseFunc;
+		typedef TFunction<TSharedPtr<FImage>(const FImage*)> FImageCloneFunc;
 
 		FImageCreateFunc CreateImage;
 		FImageReleaseFunc ReleaseImage;
@@ -272,7 +260,7 @@ namespace mu
 		/** Interface to override the internal mutable image pixel format conversion functions.
 		* Arguments match the FImageOperator::ImagePixelFormat function.
 		*/
-		typedef TFunction<void(bool&, int32, Image*, const Image*, int32)> FImagePixelFormatFunc;
+		typedef TFunction<void(bool&, int32, FImage*, const FImage*, int32)> FImagePixelFormatFunc;
 		FImagePixelFormatFunc FormatImageOverride;
 
 		FImageOperator(FImageCreateFunc Create, FImageReleaseFunc Release, FImageCloneFunc Clone, FImagePixelFormatFunc FormatOverride) : CreateImage(Create), ReleaseImage(Release), CloneImage(Clone), FormatImageOverride(FormatOverride) {};
@@ -281,9 +269,9 @@ namespace mu
 		static FImageOperator GetDefault( const FImagePixelFormatFunc& InFormatOverride )
 		{
 			return FImageOperator(
-				[](int32 x, int32 y, int32 m, EImageFormat f, EInitializationType i) { return new Image(x, y, m, f, i); },
-				[](Ptr<Image>& i) {i = nullptr; },
-				[](const Image* i) { return i->Clone(); },
+				[](int32 x, int32 y, int32 m, EImageFormat f, EInitializationType i) { return MakeShared<FImage>(x, y, m, f, i); },
+				[](TSharedPtr<FImage>& In) {In = nullptr; },
+				[](const FImage* In) { return In->Clone(); },
 				InFormatOverride
 			);
 		}
@@ -294,33 +282,33 @@ namespace mu
 		* \param onlyLOD If different than -1, only the specified lod level will be converted in the returned image.
 		* \return nullptr if the conversion failed, usually because not enough memory was allocated in the result. This is only checked for RLE compression.
 		*/
-		MUTABLERUNTIME_API Ptr<Image> ImagePixelFormat(int32 Quality, const Image* Base, EImageFormat TargetFormat, int32 OnlyLOD = -1);
+		MUTABLERUNTIME_API TSharedPtr<FImage> ImagePixelFormat(int32 Quality, const FImage* Base, EImageFormat TargetFormat, int32 OnlyLOD = -1);
 
 		/** Convert an image to another pixel format.
 		* \warning Not all format conversions are implemented.
 		* \param onlyLOD If different than -1, only the specified lod level will be converted in the returned image.
 		* \return false if the conversion failed, usually because not enough memory was allocated in the result. This is only checked for RLE compression.
 		*/
-		MUTABLERUNTIME_API void ImagePixelFormat(bool& bOutSuccess, int32 Quality, Image* Result, const Image* Base, int32 OnlyLOD = -1);
-		MUTABLERUNTIME_API void ImagePixelFormat(bool& bOutSuccess, int32 Quality, Image* Result, const Image* Base, int32 BeginResultLOD, int32 BeginBaseLOD, int32 NumLODs);
+		MUTABLERUNTIME_API void ImagePixelFormat(bool& bOutSuccess, int32 Quality, FImage* Result, const FImage* Base, int32 OnlyLOD = -1);
+		MUTABLERUNTIME_API void ImagePixelFormat(bool& bOutSuccess, int32 Quality, FImage* Result, const FImage* Base, int32 BeginResultLOD, int32 BeginBaseLOD, int32 NumLODs);
 
-		MUTABLERUNTIME_API Ptr<Image> ImageSwizzle(EImageFormat Format, const Ptr<const Image> Sources[], const uint8 Channels[]);
+		MUTABLERUNTIME_API TSharedPtr<FImage> ImageSwizzle(EImageFormat Format, const TSharedPtr<const FImage> Sources[], const uint8 Channels[]);
 
 		/** */
-		MUTABLERUNTIME_API Ptr<Image> ExtractMip(const Image* From, int32 Mip);
+		MUTABLERUNTIME_API TSharedPtr<FImage> ExtractMip(const FImage* From, int32 Mip);
 
 		/** Bilinear filter image resize. */
-		MUTABLERUNTIME_API void ImageResizeLinear(Image* Dest, int32 ImageCompressionQuality, const Image* Base);
+		MUTABLERUNTIME_API void ImageResizeLinear(FImage* Dest, int32 ImageCompressionQuality, const FImage* Base);
 
 		/** Fill the image with a plain colour. */
-		void FillColor(Image* Image, FVector4f Color);
+		void FillColor(FImage* Image, FVector4f Color);
 
 		/** Support struct to keep preallocated data required for some mipmap operations. */
 		struct FScratchImageMipmap
 		{
-			Ptr<Image> Uncompressed;
-			Ptr<Image> UncompressedMips;
-			Ptr<Image> CompressedMips;
+			TSharedPtr<FImage> Uncompressed;
+			TSharedPtr<FImage> UncompressedMips;
+			TSharedPtr<FImage> CompressedMips;
 		};
 
 		/** Generate the mipmaps for images.
@@ -328,32 +316,32 @@ namespace mu
 		* them in Dest (the full chain is spit in two images). Otherwise generate the mips missing
 		* from Base up to LevelCount and append them in Dest to the already generated Base's mips.
 		*/
-		void ImageMipmap(int32 CompressionQuality, Image* Dest, const Image* Base,
+		void ImageMipmap(int32 CompressionQuality, FImage* Dest, const FImage* Base,
 			int32 StartLevel, int32 LevelCount,
 			const FMipmapGenerationSettings&, bool bGenerateOnlyTail = false);
 
 		/** Mipmap separating the worst case treatment in 3 steps to manage allocations of temp data. */
-		void ImageMipmap_PrepareScratch(const Image* DestImage, int32 StartLevel, int32 LevelCount, FScratchImageMipmap&);
-		void ImageMipmap(FImageOperator::FScratchImageMipmap&, int32 CompressionQuality, Image* Dest, const Image* Base,
+		void ImageMipmap_PrepareScratch(const FImage* DestImage, int32 StartLevel, int32 LevelCount, FScratchImageMipmap&);
+		void ImageMipmap(FImageOperator::FScratchImageMipmap&, int32 CompressionQuality, FImage* Dest, const FImage* Base,
 			int32 LevelStart, int32 LevelCount,
 			const FMipmapGenerationSettings&, bool bGenerateOnlyTail = false);
 		void ImageMipmap_ReleaseScratch(FScratchImageMipmap&);
 
 		/** */
-		MUTABLERUNTIME_API bool ImageCrop(Image* Cropped, int32 CompressionQuality, const Image* Base, const box<FIntVector2>& Rect);
+		MUTABLERUNTIME_API bool ImageCrop(TSharedPtr<FImage>& OutCropped, int32 CompressionQuality, const FImage* Base, const box<FIntVector2>& Rect);
 
 		/** */
-		void ImageCompose(Image* Base, const Image* Block, const box<FIntVector2>& Rect);
+		void ImageCompose(FImage* Base, const FImage* Block, const box<FIntVector2>& Rect);
 
 	};
 
 	/** Update all the mipmaps in the image from the data in the base one.
 	* Only the mipmaps already existing in the image are updated.
 	*/
-	MUTABLERUNTIME_API void ImageMipmapInPlace(int32 CompressionQuality, Image* Base, const FMipmapGenerationSettings&);
+	MUTABLERUNTIME_API void ImageMipmapInPlace(int32 CompressionQuality, FImage* Base, const FMipmapGenerationSettings&);
 
 	/** */
-	MUTABLERUNTIME_API void ImageSwizzle(Image* Result, const Ptr<const Image> Sources[], const uint8 Channels[]);
+	MUTABLERUNTIME_API void ImageSwizzle(FImage* Result, const TSharedPtr<const FImage> Sources[], const uint8 Channels[]);
 
 	MUTABLERUNTIME_API inline EImageFormat GetUncompressedFormat(EImageFormat Format)
 	{
@@ -393,24 +381,39 @@ namespace mu
 		return Result;
 	}
 
-	//---------------------------------------------------------------------------------------------
-	//! Use with care.
-	//---------------------------------------------------------------------------------------------
+	/** */
 	template<class T>
-	Ptr<T> CloneOrTakeOver( const T* source )
+	TSharedPtr<T> CloneOrTakeOver(const TSharedPtr<T>& Source)
 	{
-		Ptr<T> result;
-		if (source->IsUnique())
+		TSharedPtr<T> Result;
+		if (Source.IsUnique())
 		{
-			result = const_cast<T*>(source);
+			Result = ConstCastSharedPtr<T>(Source);
 		}
 		else
 		{
-			result = source->Clone();
+			Result = Source->Clone();
 		}
 
-		return result;
+		return Result;
 	}
-	
+
+	/** */
+	template<class T>
+	TSharedPtr<T> CloneOrTakeOver(const TSharedPtr<const T>& Source)
+	{
+		TSharedPtr<T> Result;
+		if (Source.IsUnique())
+		{
+			Result = ConstCastSharedPtr<T>(Source);
+		}
+		else
+		{
+			Result = Source->Clone();
+		}
+
+		return Result;
+	}
+
 }
 
diff --git a/Engine/Plugins/Mutable/Source/MutableRuntime/Public/MuR/ImageDataStorage.h b/Engine/Plugins/Mutable/Source/MutableRuntime/Public/MuR/ImageDataStorage.h
--- a/Engine/Plugins/Mutable/Source/MutableRuntime/Public/MuR/ImageDataStorage.h
+++ b/Engine/Plugins/Mutable/Source/MutableRuntime/Public/MuR/ImageDataStorage.h
@@ -400,8 +400,8 @@ public:
 		return true;
 	}
 
-	void Serialise(OutputArchive& Arch) const;
-	void Unserialise(InputArchive& Arch);
+	void Serialise(FOutputArchive& Arch) const;
+	void Unserialise(FInputArchive& Arch);
 };
 
 }
diff --git a/Engine/Plugins/Mutable/Source/MutableRuntime/Public/MuR/ImageTypes.h b/Engine/Plugins/Mutable/Source/MutableRuntime/Public/MuR/ImageTypes.h
--- a/Engine/Plugins/Mutable/Source/MutableRuntime/Public/MuR/ImageTypes.h
+++ b/Engine/Plugins/Mutable/Source/MutableRuntime/Public/MuR/ImageTypes.h
@@ -226,8 +226,8 @@ namespace mu
 		EMipmapFilterType FilterType = EMipmapFilterType::SimpleAverage;
 		EAddressMode AddressMode = EAddressMode::None;
 
-		void Serialise(OutputArchive& Arch) const;
-		void Unserialise(InputArchive& Arch);
+		void Serialise(FOutputArchive& Arch) const;
+		void Unserialise(FInputArchive& Arch);
 	};
 }
 
diff --git a/Engine/Plugins/Mutable/Source/MutableRuntime/Public/MuR/Instance.h b/Engine/Plugins/Mutable/Source/MutableRuntime/Public/MuR/Instance.h
--- a/Engine/Plugins/Mutable/Source/MutableRuntime/Public/MuR/Instance.h
+++ b/Engine/Plugins/Mutable/Source/MutableRuntime/Public/MuR/Instance.h
@@ -11,18 +11,12 @@
 
 namespace mu
 {
-	class Instance;
 
-	typedef Ptr<Instance> InstancePtr;
-	typedef Ptr<const Instance> InstancePtrConst;
-
-
-    //! \brief A customised object created from a model and a set of parameter values.
-    //!
-    //! It corresponds to an "engine object" but the contents of its data depends on the Model, and
-    //! it may contain any number of LODs, components, surfaces, meshes and images, even none.
-	//! \ingroup runtime
-	class MUTABLERUNTIME_API Instance : public Resource
+    /** A customised object created from a model and a set of parameter values.
+    * It corresponds to an "engine object" but the contents of its data depends on the Model, and
+    * it may contain any number of LODs, components, surfaces, meshes and images, even none.
+	*/
+	class MUTABLERUNTIME_API FInstance : public FResource
 	{
 	public:
 
@@ -30,25 +24,28 @@ namespace mu
 		// Life cycle
 		//-----------------------------------------------------------------------------------------
 
-		Instance();
+		FInstance();
 
         //! Clone this instance
-        Ptr<Instance> Clone() const;
+		TSharedPtr<FInstance> Clone() const;
 
 		// Resource interface
 		int32 GetDataSize() const override;
 
+		/** */
+		~FInstance();
+
 		//-----------------------------------------------------------------------------------------
 		// Own interface
 		//-----------------------------------------------------------------------------------------
 
         //! Type for the instance unique identifiers.
-        typedef uint32 ID;
+        typedef uint32 FID;
 
         //! Get a unique identifier for this instance. It doesn't change during the entire
         //! lifecycle of each instance. This identifier can be used in the System methods to update
         //! or release the instance.
-        Instance::ID GetId() const;
+        FInstance::FID GetId() const;
 
 		//! Get the number of components of this instance.
 		int32 GetComponentCount() const;
@@ -193,7 +190,7 @@ namespace mu
 		//! \param Index Index of the ExtensionData to fetch
 		//! \param OutExtensionData Receives the ExtensionData
 		//! \param OutName Receives the name associated with the ExtensionData. Guaranteed to be a valid string of non-zero length.
-		void GetExtensionData(int32 Index, Ptr<const class ExtensionData>& OutExtensionData, FName& OutName) const;
+		void GetExtensionData(int32 Index, TSharedPtr<const FExtensionData>& OutExtensionData, FName& OutName) const;
 
         //-----------------------------------------------------------------------------------------
 		// Interface pattern
@@ -202,11 +199,6 @@ namespace mu
 
 		Private* GetPrivate() const;
 
-	protected:
-
-		//! Forbidden. Manage with the Ptr<> template.
-		~Instance();
-
 	private:
 
 		Private* m_pD;
diff --git a/Engine/Plugins/Mutable/Source/MutableRuntime/Public/MuR/Layout.h b/Engine/Plugins/Mutable/Source/MutableRuntime/Public/MuR/Layout.h
--- a/Engine/Plugins/Mutable/Source/MutableRuntime/Public/MuR/Layout.h
+++ b/Engine/Plugins/Mutable/Source/MutableRuntime/Public/MuR/Layout.h
@@ -81,13 +81,12 @@ namespace mu
 	};
 
 
-    //! \brief Image block layout class.
-    //!
-    //! It contains the information about what blocks are defined in a texture layout (texture
-    //! coordinates set from a mesh).
-    //! It is usually not necessary to use this objects, except for some advanced cases.
-	//! \ingroup runtime
-    class MUTABLERUNTIME_API Layout : public Resource
+    /** Image block layout class.
+    * It contains the information about what blocks are defined in a texture layout (texture
+    * coordinates set from a mesh).
+    * It is usually not necessary to use this objects, except for some advanced cases.
+	*/
+    class MUTABLERUNTIME_API FLayout : public FResource
 	{
 	public:
 
@@ -113,17 +112,15 @@ namespace mu
 		// Life cycle
 		//-----------------------------------------------------------------------------------------
 
-		Layout();
-
 		//! Deep clone this layout.
-		Ptr<Layout> Clone() const;
+		TSharedPtr<FLayout> Clone() const;
 
 		//! Serialisation
-		static void Serialise( const Layout* p, OutputArchive& arch );
-		static Ptr<Layout> StaticUnserialise( InputArchive& arch );
+		static void Serialise( const FLayout* p, FOutputArchive& arch );
+		static TSharedPtr<FLayout> StaticUnserialise( FInputArchive& arch );
 
 		//! Full compare
-		bool operator==( const Layout& other ) const;
+		bool operator==( const FLayout& other ) const;
 
 		// Resource interface
 		int32 GetDataSize() const override;
@@ -166,22 +163,18 @@ namespace mu
 		//! Set the texture layout packing strategy
 		EPackStrategy GetLayoutPackingStrategy() const;
 
-	protected:
-
-		//! Forbidden. Manage with the Ptr<> template.
-		~Layout() {}
 
 	public:
 
 
 		//!
-		void Serialise(OutputArchive& arch) const;
+		void Serialise(FOutputArchive& arch) const;
 
 		//!
-		void Unserialise(InputArchive& arch);
+		void Unserialise(FInputArchive& arch);
 
 		//!
-		bool IsSimilar(const Layout& o) const;
+		bool IsSimilar(const FLayout& o) const;
 
 		/** Find a block by id. This converts the "absolute" id to a relative index to the layout blocks. Return -1 if not found. */
 		int32 FindBlock(uint64 Id) const;
diff --git a/Engine/Plugins/Mutable/Source/MutableRuntime/Public/MuR/Mesh.h b/Engine/Plugins/Mutable/Source/MutableRuntime/Public/MuR/Mesh.h
--- a/Engine/Plugins/Mutable/Source/MutableRuntime/Public/MuR/Mesh.h
+++ b/Engine/Plugins/Mutable/Source/MutableRuntime/Public/MuR/Mesh.h
@@ -26,17 +26,12 @@ namespace mu
 {
 
 	// Forward references
-	class Layout;
-    class Skeleton;
-    class PhysicsBody;
-
-	class Mesh;
-    typedef Ptr<Mesh> MeshPtr;
-    typedef Ptr<const Mesh> MeshPtrConst;
+	class FLayout;
+    class FSkeleton;
+    class FPhysicsBody;
 
 	struct FSurfaceSubMesh
-	{
-		
+	{		
 		int32 VertexBegin = 0;
 		int32 VertexEnd = 0;
 		int32 IndexBegin = 0;
@@ -71,8 +66,8 @@ namespace mu
 				Lhs.SubMeshes == Rhs.SubMeshes;
 		}
 
-		inline void Serialise(OutputArchive& Arch) const;
-		inline void Unserialise(InputArchive& Arch);
+		inline void Serialise(FOutputArchive& Arch) const;
+		inline void Unserialise(FInputArchive& Arch);
 	};
 
 
@@ -166,13 +161,12 @@ namespace mu
 	};
 
 
-    //! \brief Mesh object containing any number of buffers with any number of channels.
-    //! The buffers can be per-index or per-vertex.
-    //! The mesh also includes layout information for every texture channel for internal usage, and
-    //! it can be ignored.
-    //! The meshes are always assumed to be triangle list primitives.
-    //! \ingroup runtime
-    class MUTABLERUNTIME_API Mesh : public Resource
+    /** Mesh object containing any number of buffers with any number of channels.
+	* The buffers can be per-index or per-vertex.
+    * The mesh also includes layout information for every texture channel for internal usage, and it can be ignored.
+    * The meshes are always assumed to be triangle list primitives.
+    */
+    class MUTABLERUNTIME_API FMesh : public FResource
     {
 	public:
 
@@ -185,20 +179,20 @@ namespace mu
         //-----------------------------------------------------------------------------------------
 				
 		/** Create a new empty mesh that repreents an external resource mesh. */
-		static Ptr<Mesh> CreateAsReference(uint32 ID, bool bForceLoad);
+		static TSharedPtr<FMesh> CreateAsReference(uint32 ID, bool bForceLoad);
 
         /** Deep clone this mesh. */
-        Ptr<Mesh> Clone() const;
+        TSharedPtr<FMesh> Clone() const;
 		
 		/** Clone with flags allowing to not include some parts in the cloned mesh */
-		Ptr<Mesh> Clone(EMeshCopyFlags Flags) const;
+		TSharedPtr<FMesh> Clone(EMeshCopyFlags Flags) const;
 
 		/** Copy form another mesh. */
-		void CopyFrom(const Mesh& From, EMeshCopyFlags Flags = EMeshCopyFlags::AllFlags);
+		void CopyFrom(const FMesh& From, EMeshCopyFlags Flags = EMeshCopyFlags::AllFlags);
 
         /** Serialisation */
-        static void Serialise(const Mesh* InMesh, OutputArchive& Arch );
-        static Ptr<Mesh> StaticUnserialise(InputArchive& Arch);
+        static void Serialise(const FMesh* InMesh, FOutputArchive& Arch );
+        static TSharedPtr<FMesh> StaticUnserialise(FInputArchive& Arch);
 
 		// Resource interface
 		int32 GetDataSize() const override;
@@ -269,34 +263,34 @@ namespace mu
         //! \{
 
         //!
-        void AddLayout( Ptr<const Layout> pLayout );
+        void AddLayout( TSharedPtr<const FLayout> );
 
         //!
         int32 GetLayoutCount() const;
 
         //!
-        const Layout* GetLayout( int32 i ) const;
+		TSharedPtr<const FLayout> GetLayout( int32 Index) const;
 
         //!
-        void SetLayout( int32 i, Ptr<const Layout> );
+        void SetLayout( int32 Index, TSharedPtr<const FLayout> );
         //! \}
 
         //! \name Skeleton information
         //! \{
 
-        void SetSkeleton( Ptr<const Skeleton> );
-        Ptr<const Skeleton> GetSkeleton() const;
+        void SetSkeleton(TSharedPtr<const FSkeleton> );
+		TSharedPtr<const FSkeleton> GetSkeleton() const;
 
         //! \}
 
         //! \name PhysicsBody information
         //! \{
 
-        void SetPhysicsBody( Ptr<const PhysicsBody> );
-        Ptr<const PhysicsBody> GetPhysicsBody() const;
+        void SetPhysicsBody(TSharedPtr<const FPhysicsBody> );
+		TSharedPtr<const FPhysicsBody> GetPhysicsBody() const;
 
-		int32 AddAdditionalPhysicsBody(Ptr<const PhysicsBody> Body);
-		Ptr<const PhysicsBody> GetAdditionalPhysicsBody(int32 I) const;
+		int32 AddAdditionalPhysicsBody(TSharedPtr<const FPhysicsBody>);
+		TSharedPtr<const FPhysicsBody> GetAdditionalPhysicsBody(int32 I) const;
 		//int32 GetAdditionalPhysicsBodyExternalId(int32 I) const;
 
         //! \}
@@ -305,16 +299,16 @@ namespace mu
         //! \{
 
         //!
-        void SetTagCount( int32 count );
+        void SetTagCount( int32 Count );
 
         //!
         int32 GetTagCount() const;
 
         //!
-        const FString& GetTag( int32 tagIndex ) const;
+        const FString& GetTag( int32 TagIndex ) const;
 
         //!
-        void SetTag( int32 tagIndex, const FString& Name );
+        void SetTag( int32 TagIndex, const FString& Name );
 
 		//!
 		void AddStreamedResource(uint64 ResourceId);
@@ -326,7 +320,7 @@ namespace mu
 		int32 FindBonePose(const FBoneName& BoneName) const;
 		
 		//!
-		void SetBonePoseCount(int32 count);
+		void SetBonePoseCount(int32 Count);
 
 		//!
 		int32 GetBonePoseCount() const;
@@ -367,12 +361,6 @@ namespace mu
 		 */
         uint32 GetId() const;
 
-
-    protected:
-
-        /** Forbidden. Manage with the Ptr<> template. */
-		~Mesh() {}
-
     
 	public:
 
@@ -413,17 +401,17 @@ namespace mu
 		 * This skeleton and physics body are not owned and may be used by other meshes, so it cannot be modified
 		 * once the mesh has been fully created.
 		 */
-		Ptr<const Skeleton> Skeleton;
-		Ptr<const PhysicsBody> PhysicsBody;
+		TSharedPtr<const FSkeleton> Skeleton;
+		TSharedPtr<const FPhysicsBody> PhysicsBody;
 
 		/** Additional physics bodies referenced by the mesh that don't merge. */
-		TArray<Ptr<const mu::PhysicsBody>> AdditionalPhysicsBodies;
+		TArray<TSharedPtr<const FPhysicsBody>> AdditionalPhysicsBodies;
 
 		/** 
 		 * Texture Layout blocks attached to this mesh. They are const because they could be shared with
 		 * other meshes, so they need to be cloned and replaced if a modification is needed.
 		 */
-		TArray<Ptr<const Layout>> Layouts;		
+		TArray<TSharedPtr<const FLayout>> Layouts;		
 
 		TArray<FString> Tags;
 
@@ -438,8 +426,8 @@ namespace mu
 			EBoneUsageFlags BoneUsageFlags = EBoneUsageFlags::None;
 			FTransform3f BoneTransform;
 
-			inline void Serialise(OutputArchive& arch) const;
-			inline void Unserialise(InputArchive& arch);
+			inline void Serialise(FOutputArchive& arch) const;
+			inline void Unserialise(FInputArchive& arch);
 
 			inline bool operator==(const FBonePose& Other) const
 			{
@@ -456,61 +444,61 @@ namespace mu
 		/** Array containing the bonemaps of all surfaces in the mesh. */
 		TArray<FBoneName> BoneMap;
 
-		inline void Serialise(OutputArchive& arch) const;
+		inline void Serialise(FOutputArchive& arch) const;
 
-		inline void Unserialise(InputArchive& arch);
+		inline void Unserialise(FInputArchive& arch);
 
 
-		inline bool operator==(const Mesh& o) const
+		inline bool operator==(const FMesh& Other) const
 		{
 			bool bEqual = true;
 
-			if (bEqual) bEqual = (ReferenceID == o.ReferenceID);
-			if (bEqual) bEqual = (MeshIDPrefix == o.MeshIDPrefix);
-			if (bEqual) bEqual = (IndexBuffers == o.IndexBuffers);
-			if (bEqual) bEqual = (VertexBuffers == o.VertexBuffers);
-			if (bEqual) bEqual = (Layouts.Num() == o.Layouts.Num());
-			if (bEqual) bEqual = (BonePoses.Num() == o.BonePoses.Num());
-			if (bEqual) bEqual = (BoneMap.Num() == o.BoneMap.Num());
-			if (bEqual && Skeleton != o.Skeleton)
+			if (bEqual) bEqual = (ReferenceID == Other.ReferenceID);
+			if (bEqual) bEqual = (MeshIDPrefix == Other.MeshIDPrefix);
+			if (bEqual) bEqual = (IndexBuffers == Other.IndexBuffers);
+			if (bEqual) bEqual = (VertexBuffers == Other.VertexBuffers);
+			if (bEqual) bEqual = (Layouts.Num() == Other.Layouts.Num());
+			if (bEqual) bEqual = (BonePoses.Num() == Other.BonePoses.Num());
+			if (bEqual) bEqual = (BoneMap.Num() == Other.BoneMap.Num());
+			if (bEqual && Skeleton != Other.Skeleton)
 			{
-				if (Skeleton && o.Skeleton)
+				if (Skeleton && Other.Skeleton)
 				{
-					bEqual = (*Skeleton == *o.Skeleton);
+					bEqual = (*Skeleton == *Other.Skeleton);
 				}
 				else
 				{
 					bEqual = false;
 				}
 			}
-			if (bEqual) bEqual = (StreamedResources == o.StreamedResources);
-			if (bEqual) bEqual = (Surfaces == o.Surfaces);
-			if (bEqual) bEqual = (Tags == o.Tags);
-			if (bEqual) bEqual = (SkeletonIDs == o.SkeletonIDs);
+			if (bEqual) bEqual = (StreamedResources == Other.StreamedResources);
+			if (bEqual) bEqual = (Surfaces == Other.Surfaces);
+			if (bEqual) bEqual = (Tags == Other.Tags);
+			if (bEqual) bEqual = (SkeletonIDs == Other.SkeletonIDs);
 
 			for (int32 i = 0; bEqual && i < Layouts.Num(); ++i)
 			{
-				bEqual &= (*Layouts[i]) == (*o.Layouts[i]);
+				bEqual &= (*Layouts[i]) == (*Other.Layouts[i]);
 			}
 
-			bEqual &= AdditionalBuffers.Num() == o.AdditionalBuffers.Num();
+			bEqual &= AdditionalBuffers.Num() == Other.AdditionalBuffers.Num();
 			for (int32 i = 0; bEqual && i < AdditionalBuffers.Num(); ++i)
 			{
-				bEqual &= AdditionalBuffers[i] == o.AdditionalBuffers[i];
+				bEqual &= AdditionalBuffers[i] == Other.AdditionalBuffers[i];
 			}
 
-			bEqual &= BonePoses.Num() == o.BonePoses.Num();
+			bEqual &= BonePoses.Num() == Other.BonePoses.Num();
 			for (int32 i = 0; bEqual && i < BonePoses.Num(); ++i)
 			{
-				bEqual &= BonePoses[i] == o.BonePoses[i];
+				bEqual &= BonePoses[i] == Other.BonePoses[i];
 			}
 
-			if (bEqual) bEqual = BoneMap == o.BoneMap;
+			if (bEqual) bEqual = BoneMap == Other.BoneMap;
 
-			bEqual &= AdditionalPhysicsBodies.Num() == o.AdditionalPhysicsBodies.Num();
+			bEqual &= AdditionalPhysicsBodies.Num() == Other.AdditionalPhysicsBodies.Num();
 			for (int32 i = 0; bEqual && i < AdditionalPhysicsBodies.Num(); ++i)
 			{
-				bEqual &= *AdditionalPhysicsBodies[i] == *o.AdditionalPhysicsBodies[i];
+				bEqual &= *AdditionalPhysicsBodies[i] == *Other.AdditionalPhysicsBodies[i];
 			}
 
 			return bEqual;
@@ -520,7 +508,7 @@ namespace mu
 		 * Compare the mesh with another one, but ignore internal data like generated vertex
 		 * indices.
 		 */
-		bool IsSimilar(const Mesh& Other, bool bCompareLayouts) const;
+		bool IsSimilar(const FMesh& Other, bool bCompareLayouts) const;
 
 
 		/**
@@ -538,16 +526,16 @@ namespace mu
 			bool DoMatch(int32 Vertex, int32 OtherVertex) const;
 		};
 
-		void GetVertexMap(const Mesh& Other, FVertexMatchMap& VertexMap, float Tolerance = 1e-3f) const;
+		void GetVertexMap(const FMesh& Other, FVertexMatchMap& VertexMap, float Tolerance = 1e-3f) const;
 
 		/** Compare the vertex attributes to check if they match. */
-		UE::Math::TIntVector3<uint32> GetFaceVertexIndices(int32 f) const;
+		UE::Math::TIntVector3<uint32> GetFaceVertexIndices(int32 FaceIndex) const;
 
 		/** 
 		 * Return true if the given mesh has the same vertex and index formats, and in the same
 		 * buffer structure.
 		 */
-		bool HasCompatibleFormat(const Mesh* Other) const;
+		bool HasCompatibleFormat(const FMesh* Other) const;
 
 		/** Update the flags identifying the mesh format as some of the optimised formats. */
 		void ResetStaticFormatFlags() const;
diff --git a/Engine/Plugins/Mutable/Source/MutableRuntime/Public/MuR/MeshBufferSet.h b/Engine/Plugins/Mutable/Source/MutableRuntime/Public/MuR/MeshBufferSet.h
--- a/Engine/Plugins/Mutable/Source/MutableRuntime/Public/MuR/MeshBufferSet.h
+++ b/Engine/Plugins/Mutable/Source/MutableRuntime/Public/MuR/MeshBufferSet.h
@@ -5,7 +5,6 @@
 #include "MuR/Serialisation.h"
 #include "Containers/Array.h"
 #include "HAL/PlatformMath.h"
-#include "MuR/MutableMemory.h"
 
 #include "MuR/MemoryTrackingAllocationPolicy.h"
 
@@ -181,8 +180,8 @@ namespace mu
 		TMemoryTrackedArray<uint8> Data;
 		uint32 ElementSize = 0;
 
-		void Serialise(mu::OutputArchive& Arch) const;
-		inline void Unserialise(mu::InputArchive& Arch);
+		void Serialise(mu::FOutputArchive& Arch) const;
+		inline void Unserialise(mu::FInputArchive& Arch);
 
 		inline bool operator==(const FMeshBuffer& Other) const
 		{
@@ -232,8 +231,8 @@ namespace mu
 		uint32 ElementCount = 0;
 		TArray<FMeshBuffer> Buffers;
 
-		void Serialise(OutputArchive& Arch) const;
-		void Unserialise(InputArchive& Arch);
+		void Serialise(FOutputArchive& Arch) const;
+		void Unserialise(FInputArchive& Arch);
 
 		inline bool operator==(const FMeshBufferSet& Other) const
 		{
diff --git a/Engine/Plugins/Mutable/Source/MutableRuntime/Public/MuR/Model.h b/Engine/Plugins/Mutable/Source/MutableRuntime/Public/MuR/Model.h
--- a/Engine/Plugins/Mutable/Source/MutableRuntime/Public/MuR/Model.h
+++ b/Engine/Plugins/Mutable/Source/MutableRuntime/Public/MuR/Model.h
@@ -11,23 +11,11 @@
 
 namespace mu
 {
-	class InputArchive;
-	class ModelWriter;
-	class OutputArchive;
-
-	typedef Ptr<Parameters> ParametersPtr;
-	typedef Ptr<const Parameters> ParametersPtrConst;
-
-    class ModelParametersGenerator;
-
-    typedef Ptr<ModelParametersGenerator> ModelParametersGeneratorPtr;
-    typedef Ptr<const ModelParametersGenerator> ModelParametersGeneratorPtrConst;
-
-    class System;
-
-    typedef Ptr<System> SystemPtr;
-    typedef Ptr<const System> SystemPtrConst;
-
+	class FInputArchive;
+	class FModelWriter;
+	class FOutputArchive;
+	class FSystem;
+    class FModelParametersGenerator;
 
 	/** Flags for the model data chunks. */
 	enum class ERomFlags : uint16
@@ -40,32 +28,31 @@ namespace mu
 	};
 
 
-    //! \brief A Model represents a customisable object with any number of parameters.
-    //!
-    //! When values are given to the parameters, specific Instances can be built, which hold the
-    //! built application-usable data.
-	//! \ingroup runtime
-    class MUTABLERUNTIME_API Model
+    /** A Model represents a customisable object with any number of parameters.
+    * When values are given to the parameters, specific Instances can be built, which hold the
+    * built application-usable data.
+	*/
+    class MUTABLERUNTIME_API FModel
 	{
 	public:
 
 		//-----------------------------------------------------------------------------------------
 		// Life cycle
 		//-----------------------------------------------------------------------------------------
-		Model();
+		FModel();
 
 		//! Don't call directly. Manage with a TSharedPtr.
-		~Model();
+		~FModel();
 
-		static void Serialise( const Model* p, OutputArchive& arch );
-		static TSharedPtr<Model> StaticUnserialise( InputArchive& arch );
+		static void Serialise( const FModel*, FOutputArchive& );
+		static TSharedPtr<FModel> StaticUnserialise( FInputArchive& );
 
 		/** Special serialise operation that serialises the data in separate "files". An object
          * with the ModelStreamer interface is responsible of storing this data and providing
 		 * the "file" concept. 
 		 * If bDropData is set to true, the rom data will be freed as it is serialized.
 		 */
-        static void Serialise( Model* p, ModelWriter& arch, bool bDropData=false );
+        static void Serialise( FModel*, FModelWriter&, bool bDropData=false );
 
 		//! Return true if the model has external data in other files. This kind of models will
 		//! require data streaming when used.
@@ -88,7 +75,7 @@ namespace mu
 		//! Create a set of new parameters of the model with the default values.
 		//! If old parameters are provided, they will be reused when possible instead of the
 		//! default values.
-        static Ptr<Parameters> NewParameters( TSharedPtr<const Model> Model, const Parameters* pOld = nullptr );
+        static TSharedPtr<FParameters> NewParameters( TSharedPtr<const FModel> Model, const FParameters* Old = nullptr );
 
 		/** Return true if the parameter is multi-dimensional */
 		bool IsParameterMultidimensional(int32 ParameterIndex) const;
@@ -103,11 +90,11 @@ namespace mu
 		int32 FindState( const FString& Name ) const;
 
 		//! Get the number of parameters available in a particular state.
-		int32 GetStateParameterCount( int32 stateIndex ) const;
+		int32 GetStateParameterCount( int32 StateIndex ) const;
 
 		//! Get the index of one of the parameters in the given state. The index refers to the
-		//! parameters in a Parameters object obtained from this model with NewParameters.
-		int32 GetStateParameterIndex( int32 stateIndex, int32 paramIndex ) const;
+		//! parameters in a FParameters object obtained from this model with NewParameters.
+		int32 GetStateParameterIndex( int32 StateIndex, int32 ParamIndex ) const;
 
         //! Free memory used by streaming data that may be loaded again when needed.
         void UnloadExternalData();
@@ -182,52 +169,5 @@ namespace mu
 
 	};
 
-
-    //! \brief This class enumerates all possible model instances and generates the parameters for
-    //! each of them.
-    //!
-    //! This class can be used to generate all the possible instances of a model, or to generate
-    //! random instances. It only takes into account discrete parameters (bools and ints) but not
-    //! continuous parameters like colours and floats, which may be set to random values.
-    //!
-    //! \ingroup runtime
-    class MUTABLERUNTIME_API ModelParametersGenerator : public RefCounted
-    {
-    public:
-
-        //!
-        ModelParametersGenerator( TSharedPtr<const Model>, System* );
-
-        //! Return the number of different possible instances that can be built from the model.
-        int64 GetInstanceCount();
-
-        //! Return the parameters of one of the possible instances of the model.
-        //! \param index is an index in the range of 0 to GetInstanceCount-1
-        //! \param randomGenerator if not null, this function will be used to generate random values
-        //! for the continuous parameters of the instance.
-        ParametersPtr GetInstance( int64 index, TFunction<float()> RandomGenerator);
-
-        //! Return the parameters of one of a random instance of the model.
-        //! \param randomGenerator This function will be used to generate random values
-        //! for the continuous parameters of the instance.
-        ParametersPtr GetRandomInstance(TFunctionRef<float()> RandomGenerator );
-
-        //-----------------------------------------------------------------------------------------
-        // Interface pattern
-        //-----------------------------------------------------------------------------------------
-        class Private;
-
-        Private* GetPrivate() const;
-
-    protected:
-
-        //! Forbidden. Manage with the Ptr<> template.
-        ~ModelParametersGenerator();
-
-    private:
-
-        Private* m_pD;
-    };
-
 }
 
diff --git a/Engine/Plugins/Mutable/Source/MutableRuntime/Public/MuR/MutableRuntimeModule.h b/Engine/Plugins/Mutable/Source/MutableRuntime/Public/MuR/MutableRuntimeModule.h
--- a/Engine/Plugins/Mutable/Source/MutableRuntime/Public/MuR/MutableRuntimeModule.h
+++ b/Engine/Plugins/Mutable/Source/MutableRuntime/Public/MuR/MutableRuntimeModule.h
@@ -6,17 +6,6 @@
 #include "Logging/LogMacros.h"
 
 
-// This will use unreal's malloc instead of system malloc for mutable memory operations
-#define USE_UNREAL_ALLOC_IN_MUTABLE				1
-
-// This will enable the internal memory manager for mutable operations, reducing the number of operations
-// but increasing the memory usage slightly. This only affects packaged builds, and it is not used in editor.
-#define USE_INTERNAL_MEMORY_MANAGER_IN_MUTABLE	1
-
-// This slows down mutable, but gathers some data about memory usage for debug
-#define USE_STAT_MUTABLE_MEMORY					0
-
-
 MUTABLERUNTIME_API DECLARE_LOG_CATEGORY_EXTERN(LogMutableCore, Log, All);
 
 
diff --git a/Engine/Plugins/Mutable/Source/MutableRuntime/Public/MuR/MutableString.h b/Engine/Plugins/Mutable/Source/MutableRuntime/Public/MuR/MutableString.h
--- a/Engine/Plugins/Mutable/Source/MutableRuntime/Public/MuR/MutableString.h
+++ b/Engine/Plugins/Mutable/Source/MutableRuntime/Public/MuR/MutableString.h
@@ -13,7 +13,7 @@ namespace mu
 {
 
     /** Container for text string data. */
-    class MUTABLERUNTIME_API String : public Resource
+    class MUTABLERUNTIME_API String : public FResource
 	{
 	public:
 
@@ -24,7 +24,7 @@ namespace mu
 		String(const FString&);
 
 		//! Deep clone this string.
-		Ptr<String> Clone() const;
+		TSharedPtr<String> Clone() const;
 
 		// Resource interface
 		int32 GetDataSize() const override;
@@ -36,11 +36,6 @@ namespace mu
 		//! Get the string data.
 		const FString& GetValue() const;
 
-	protected:
-
-		//! Forbidden. Manage with the Ptr<> template.
-		~String() {}
-
 	public:
 
 		//!
