[mutable] More enforcement of coding standards.

diff --git a/Engine/Plugins/Mutable/Source/MutableRuntime/Public/MuR/Parameters.h b/Engine/Plugins/Mutable/Source/MutableRuntime/Public/MuR/Parameters.h
--- a/Engine/Plugins/Mutable/Source/MutableRuntime/Public/MuR/Parameters.h
+++ b/Engine/Plugins/Mutable/Source/MutableRuntime/Public/MuR/Parameters.h
@@ -3,7 +3,6 @@
 #pragma once
 
 #include "HAL/Platform.h"
-#include "MuR/Ptr.h"
 #include "MuR/RefCounted.h"
 #include "MuR/Types.h"
 #include "Templates/SharedPointer.h"
@@ -13,20 +12,11 @@
 namespace mu
 {
 
-	//!
-	class InputArchive;
-	class OutputArchive;
-
-	class Parameters;
-    using ParametersPtr = Ptr<Parameters>;
-    using ParametersPtrConst = Ptr<const Parameters>;
-
-    class RangeIndex;
-
-	using RangeIndexPtr = Ptr<RangeIndex>;
-	using RangeIndexPtrConst = Ptr<const RangeIndex>;
-
-	class Model;
+	class FInputArchive;
+	class FOutputArchive;
+	class FParameters;
+    class FRangeIndex;
+	class FModel;
 
 	//! Default parameter values definition
 	static const FVector4f DefaultMutableColorValue(-UE_BIG_NUMBER);
@@ -41,7 +31,7 @@ namespace mu
         /** Boolean parameter type (true or false) */
         Bool,
 
-        /** Integer parameter type. It usually has a limited range of possible values that can be queried in the Parameters object. */
+        /** Integer parameter type. It usually has a limited range of possible values that can be queried in the FParameters object. */
         Int,
 
         /** Floating point value in the range of 0.0 to 1.0 */
@@ -85,16 +75,19 @@ namespace mu
     };
 
 
-    /** Type used to set or read multi - dimensional parameter values.
-     * If parameters have multiple values because of ranges, RangeIndex can be used to specify
+    /** Type used to set or read multi-dimensional parameter values.
+     * If parameters have multiple values because of ranges, FRangeIndex can be used to specify
      * what value is read or set in the methods of the Parameter class.
-     * RangeIndex objects can be reused in multiple calls for the same parameter in the Parameter
+     * FRangeIndex objects can be reused in multiple calls for the same parameter in the Parameter
      * class interface even after changing the position values.
 	 */
-	class MUTABLERUNTIME_API RangeIndex : public RefCounted
+	class MUTABLERUNTIME_API FRangeIndex
     {
     public:
 
+		/** Instances of this class must be obtained from a FParameters instance. */
+		FRangeIndex() = default;
+
         //-----------------------------------------------------------------------------------------
         // Own interface
         //-----------------------------------------------------------------------------------------
@@ -104,40 +97,32 @@ namespace mu
 
         //! Return the name of a range.
         //! \param index Index of the range from 0 to GetRangeCount()-1
-        const FString& GetRangeName( int32 index ) const;
+        const FString& GetRangeName( int32 Index ) const;
 
         //! Return the Guid of the parameter, resistant to parameter name changes.
         //! \param index Index of the parameter from 0 to GetRangeCount()-1
-        const FString& GetRangeUid( int32 index ) const;
+        const FString& GetRangeUid( int32 Index) const;
 
         //! Set the position in one of the dimensions of the index.
         //! \param index Index of the parameter from 0 to GetRangeCount()-1
-        void SetPosition( int32 index, int32 position );
+        void SetPosition( int32 Index, int32 Position );
 
         //! Return the position in one of the dimensions of the index.
         //! \param index Index of the parameter from 0 to GetRangeCount()-1
-        int32 GetPosition( int32 index ) const;
-
-
-        //-----------------------------------------------------------------------------------------
-        // Interface pattern
-        //-----------------------------------------------------------------------------------------
-        class Private;
-
-        Private* GetPrivate() const;
+        int32 GetPosition( int32 Index) const;
 
-    protected:
+    public:
 
-        //! Forbid creation. Instances of this class must be obtained from a Parameters instance
-        RangeIndex();
+        friend class FParameters;
 
-        //! Forbidden. Manage with the Ptr<> template.
-        ~RangeIndex() override;
+		//! Run-time data
+		TSharedPtr<const FParameters> Parameters;
 
-    private:
+		//! Index of the parameter for which we are a range index
+		int32 Parameter = -1;
 
-        friend class Parameters;
-        Private* m_pD;
+		//! Position in the several dimension of the range, as defined in Parameters
+		TArray<int32> Values;
 
     };
 
@@ -146,10 +131,10 @@ namespace mu
     * value, and additional resources associated to the parameter.
     *
     * \warning Every object of this class holds a reference to the model that it was created from.
-    * This implies that while any instance of a Parameters object is alive, its model will not
+    * This implies that while any instance of a FParameters object is alive, its model will not
     * be freed.
     */
-	class MUTABLERUNTIME_API Parameters : public RefCounted
+	class MUTABLERUNTIME_API FParameters : public TSharedFromThis<FParameters>
 	{
 	public:
 
@@ -157,12 +142,18 @@ namespace mu
 		// Life cycle
 		//-----------------------------------------------------------------------------------------
 
+		/** Instances of this class must be obtained from a Model instance. */
+		FParameters();
+
         //! Serialisation
-        static void Serialise( const Parameters* p, OutputArchive& arch );
-        static ParametersPtr StaticUnserialise( InputArchive& arch );
+        static void Serialise( const FParameters*, FOutputArchive& );
+        static TSharedPtr<FParameters> StaticUnserialise( FInputArchive& );
 
 		//! Deep clone this object.
-		ParametersPtr Clone() const;
+		TSharedPtr<FParameters> Clone() const;
+
+		/** */
+		~FParameters();
 
 		//-----------------------------------------------------------------------------------------
 		// Own interface
@@ -173,11 +164,11 @@ namespace mu
 
 		//! Return the name of the parameter.
 		//! \param index Index of the parameter from 0 to GetCount()-1
-		const FString& GetName( int32 index ) const;
+		const FString& GetName( int32 Index ) const;
 
 		//! Return the Guid of the parameter, resistant to parameter name changes.
 		//! \param index Index of the parameter from 0 to GetCount()-1
-		const FGuid& GetUid( int32 index ) const;
+		const FGuid& GetUid( int32 Index) const;
 
 		//! Find the parameter index by name.
 		//! It returns -1 if the parameter doesn't exist.
@@ -185,30 +176,30 @@ namespace mu
 
 		//! Return the type of the parameter.
 		//! \param index Index of the parameter from 0 to GetCount()-1
-		EParameterType GetType( int32 index ) const;
+		EParameterType GetType( int32 Index) const;
 
-        //! Create a new RangeIndex object to use to access a multi-dimensional parameter.
+        //! Create a new FRangeIndex object to use to access a multi-dimensional parameter.
         //! It will return nullptr if the parameter is not multidimensional.
-        Ptr<RangeIndex> NewRangeIndex( int32 paramIndex ) const;
+		TSharedPtr<FRangeIndex> NewRangeIndex( int32 ParamIndex) const;
 
         //! Return the number of values other than the default that have been set to a specific
         //! parameter.
-        int32 GetValueCount( int32 paramIndex ) const;
+        int32 GetValueCount( int32 ParamIndex ) const;
 
-        //! Return the RangeIndex of a value that has been set to a parameter.
+        //! Return the FRangeIndex of a value that has been set to a parameter.
         //! \param paramIndex Index of the parameter from 0 to GetCount()-1
         //! \param valueIndex Index of the value from 0 to GetValueCount()-1
-		Ptr<RangeIndex> GetValueIndex( int32 paramIndex, int32 valueIndex ) const;
+		TSharedPtr<FRangeIndex> GetValueIndex( int32 ParamIndex, int32 ValueIndex ) const;
 
         //! Remove all the multidimensional values for a parameter. The current non-dimensional
         //! value is kept.
-        void ClearAllValues( int32 paramIndex );
+        void ClearAllValues( int32 ParamIndex);
 
 		//! Return the value of a boolean parameter.
 		//! \pre The parameter specified by index is a T_BOOL.
 		//! \param index Index of the parameter from 0 to GetCount()-1
         //! \param pos Only for multidimensional parameters: relevant position to get in the ranges
-        bool GetBoolValue( int32 index, const Ptr<const RangeIndex>& pos=nullptr ) const;
+        bool GetBoolValue( int32 Index, const FRangeIndex* Pos =nullptr ) const;
 
 		//! Set the value of a parameter of type boolean.
 		//! \pre The parameter specified by index is a T_BOOL.
@@ -216,65 +207,65 @@ namespace mu
 		//! \param value new value of the parameter
         //! \param pos optional parameter to set a specific value for the given multidimensional
         //! position. If null, the value is set for all possible positions.
-        void SetBoolValue( int32 index, bool value, const Ptr<const RangeIndex>& pos=nullptr );
+        void SetBoolValue( int32 Index, bool bValue, const FRangeIndex* Pos =nullptr );
 
 		//! Get the number of possible values for this integer parameter.
 		//! If the number is zero, it means any integer value is accepted.
-		int32 GetIntPossibleValueCount( int32 paramIndex ) const;
+		int32 GetIntPossibleValueCount( int32 ParamIndex) const;
 
 		//! Get the value of one of the possible values for this integer.
 		//! The paramIndex is in the range of 0 to GetIntPossibleValueCount()-1
-		int32 GetIntPossibleValue( int32 paramIndex, int32 valueIndex ) const;
+		int32 GetIntPossibleValue( int32 ParamIndex, int32 ValueIndex) const;
 
         //! Get the name of one of the possible values for this integer.
         //! The paramIndex is in the range of 0 to GetIntPossibleValueCount()-1
-        const FString& GetIntPossibleValueName( int32 paramIndex, int32 valueIndex ) const;
+        const FString& GetIntPossibleValueName( int32 ParamIndex, int32 ValueIndex ) const;
 
         //! Get the index of the value of one of the possible values for this integer.
         //! The paramIndex is in the range of 0 to GetIntPossibleValueCount()-1
-		int32 GetIntValueIndex(int32 paramIndex, const FString& valueName) const;
+		int32 GetIntValueIndex(int32 ParamIndex, const FString& ValueIndex) const;
 
 		//! Get the index of the value of one of the possible values for this integer.
 		//! The paramIndex is in the range of 0 to GetIntPossibleValueCount()-1
-		int32 GetIntValueIndex(int32 paramIndex, int32 Value) const;
+		int32 GetIntValueIndex(int32 ParamIndex, int32 Value) const;
 
 		//! Return the value of a integer parameter.
 		//! \pre The parameter specified by index is a T_INT.
 		//! \param index Index of the parameter from 0 to GetCount()-1
         //! \param pos Only for multidimensional parameters: relevant position to get in the ranges
-        int32 GetIntValue( int32 index, const Ptr<const RangeIndex>& pos=nullptr ) const;
+        int32 GetIntValue( int32 Index, const FRangeIndex* Pos=nullptr ) const;
 
 		//! If the parameter is of the integer type, set its value.
 		//! \param index Index of the parameter from 0 to GetCount()-1
 		//! \param value new value of the parameter. It must be in the possible values for this
 		//! parameter (see GetIntPossibleValue), or the method will leave the value unchanged.
         //! \param pos Only for multidimensional parameters: relevant position to set in the ranges
-        void SetIntValue( int32 index, int32 value, const Ptr<const RangeIndex>& pos=nullptr );
+        void SetIntValue( int32 Index, int32 Value, const FRangeIndex* Pos=nullptr );
 
 		//! Return the value of a float parameter.
 		//! \pre The parameter specified by index is a T_FLOAT.
 		//! \param index Index of the parameter from 0 to GetCount()-1
         //! \param pos Only for multidimensional parameters: relevant position to get in the ranges
-        float GetFloatValue( int32 index, const Ptr<const RangeIndex>& pos=nullptr ) const;
+        float GetFloatValue( int32 Index, const FRangeIndex* Pos=nullptr ) const;
 
 		//! If the parameter is of the float type, set its value.
 		//! \param index Index of the parameter from 0 to GetCount()-1
 		//! \param value new value of the parameter
         //! \param pos Only for multidimensional parameters: relevant position to set in the ranges
-        void SetFloatValue( int32 index, float value, const Ptr<const RangeIndex>& pos=nullptr );
+        void SetFloatValue( int32 Index, float Value, const FRangeIndex* Pos=nullptr );
 
 		//! Return the value of a colour parameter.
 		//! \pre The parameter specified by index is a T_FLOAT.
         //! \param index Index of the parameter from 0 to GetCount()-1
         //! \param OutColor Values where every resulting colour channel will be stored
         //! \param pos Only for multidimensional parameters: relevant position to get in the ranges
-        void GetColourValue( int32 index, FVector4f& OutColor, const Ptr<const RangeIndex>& pos=nullptr ) const;
+        void GetColourValue( int32 Index, FVector4f& OutColor, const FRangeIndex* Pos=nullptr ) const;
 
 		//! If the parameter is of the colour type, set its value.
 		//! \param index Index of the parameter from 0 to GetCount()-1
         //! \param r,g,b new value of the parameter
         //! \param pos Only for multidimensional parameters: relevant position to set in the ranges
-        void SetColourValue( int32 index, FVector4f Color, const Ptr<const RangeIndex>& pos=nullptr );
+        void SetColourValue( int32 Index, FVector4f Color, const FRangeIndex* Pos=nullptr );
 
 		//! Return the value of a projector parameter, as a 4x4 matrix. The matrix is supposed to be
 		//! a linear transform in column-major.
@@ -294,7 +285,7 @@ namespace mu
 								FVector3f* OutUp,
 								FVector3f* OutScale,
                                 float* OutProjectionAngle,
-                                const Ptr<const RangeIndex>& RangePosition=nullptr ) const;
+								const FRangeIndex* RangePosition=nullptr ) const;
 
 		//! If the parameter is of the projector type, set its value.
 		//! \param ParameterIndex Index of the parameter from 0 to GetCount()-1
@@ -306,53 +297,53 @@ namespace mu
         //! projection area on the cylinder surface.
         //! \param RangePosition Only for multidimensional parameters: relevant position to set in the ranges
         void SetProjectorValue( int32 ParameterIndex,
-								const FVector3f& pos,
-								const FVector3f& dir,
-								const FVector3f& up,
-								const FVector3f& scale,
-                                float projectionAngle,
-                                const Ptr<const RangeIndex>& RangePosition=nullptr );
+								const FVector3f& Pos,
+								const FVector3f& Dir,
+								const FVector3f& Up,
+								const FVector3f& Scale,
+								float ProjectionAngle,
+								const FRangeIndex*  RangePosition=nullptr );
 
         //! Return the value of an image parameter.
         //! \pre The parameter specified by index is a T_IMAGE.
         //! \param index Index of the parameter from 0 to GetCount()-1
 		//! \param pos Only for multidimensional parameters: relevant position to set in the ranges
 	    //! \return The externalId specified when setting the image value (\see SetImageValue)
-		FName GetImageValue( int32 index, const Ptr<const RangeIndex>& pos = nullptr) const;
+		FName GetImageValue( int32 Index, const FRangeIndex* Pos = nullptr) const;
 
         //! If the parameter is of the image type, set its value.
         //! \param index Index of the parameter from 0 to GetCount()-1
         //! \param externalId Application-specific id used to identify this image during replication.
 		//! \param pos Only for multidimensional parameters: relevant position to set in the ranges
-		void SetImageValue( int32 index, FName externalId, const Ptr<const RangeIndex>& pos = nullptr);
+		void SetImageValue( int32 Index, FName externalId, const FRangeIndex* = nullptr);
 
         //! Return the value of a float parameter.
         //! \pre The parameter specified by index is a T_FLOAT.
         //! \param index Index of the parameter from 0 to GetCount()-1
         //! \param pos Only for multidimensional parameters: relevant position to get in the ranges
-        void GetStringValue( int32 index, FString& OutValue, const Ptr<const RangeIndex>& pos = nullptr) const;
+        void GetStringValue( int32 Index, FString& OutValue, const FRangeIndex* = nullptr) const;
 
         //! If the parameter is of the float type, set its value.
         //! \param index Index of the parameter from 0 to GetCount()-1
         //! \param value new value of the parameter
         //! \param pos Only for multidimensional parameters: relevant position to set in the ranges
-        void SetStringValue( int32 index, const FString& Value, const Ptr<const RangeIndex>& pos = nullptr );
+        void SetStringValue( int32 Index, const FString& Value, const FRangeIndex* = nullptr );
 
 		//! Return the value of a matrix parameter.
 		//! \pre The parameter specified by index is a T_MATRIX.
 		//! \param index Index of the parameter from 0 to GetCount()-1
 		//! \param pos Only for multidimensional parameters: relevant position to get in the ranges
-		void GetMatrixValue( int32 index, FMatrix44f& OutValue, const Ptr<const RangeIndex>& pos = nullptr) const;
+		void GetMatrixValue( int32 Index, FMatrix44f& OutValue, const FRangeIndex* = nullptr) const;
 
 		//! If the parameter is of the matrix type, set its value.
 		//! \param index Index of the parameter from 0 to GetCount()-1
 		//! \param value new value of the parameter
 		//! \param pos Only for multidimensional parameters: relevant position to set in the ranges
-		void SetMatrixValue( int32 index, const FMatrix44f& Value, const Ptr<const RangeIndex>& pos = nullptr );
+		void SetMatrixValue( int32 Index, const FMatrix44f& Value, const FRangeIndex* Pos = nullptr );
 
         //! Utility method to compare the values of a specific parameter with the values of another
-        //! Parameters object. It returns false if type or values are different.
-        bool HasSameValue( int32 thisParamIndex, const ParametersPtrConst& other, int32 otherParamIndex ) const;
+        //! FParameters object. It returns false if type or values are different.
+        bool HasSameValue( int32 ThisParamIndex, const TSharedPtr<const FParameters>& Other, int32 OtherParamIndex ) const;
 
 		//-----------------------------------------------------------------------------------------
 		// Interface pattern
@@ -361,17 +352,8 @@ namespace mu
 
 		Private* GetPrivate() const;
 
-	protected:
-
-		//! Forbid creation. Instances of this class must be obtained from a Model instance
-		Parameters();
-
-		//! Forbidden. Manage with the Ptr<> template.
-		~Parameters();
-
 	private:
 
-		friend class Model;
 		Private* m_pD;
 
 	};
diff --git a/Engine/Plugins/Mutable/Source/MutableRuntime/Public/MuR/PhysicsBody.h b/Engine/Plugins/Mutable/Source/MutableRuntime/Public/MuR/PhysicsBody.h
--- a/Engine/Plugins/Mutable/Source/MutableRuntime/Public/MuR/PhysicsBody.h
+++ b/Engine/Plugins/Mutable/Source/MutableRuntime/Public/MuR/PhysicsBody.h
@@ -31,10 +31,10 @@ namespace mu
 		}
 	
 		//!
-		inline void Serialise( OutputArchive& arch ) const;
+		inline void Serialise( FOutputArchive& arch ) const;
 
 		//!
-		inline void Unserialise( InputArchive& arch );
+		inline void Unserialise( FInputArchive& arch );
 	};
 
 	struct FSphereBody : FBodyShape
@@ -48,10 +48,10 @@ namespace mu
 		}
 		 
 		//!
-		inline void Serialise( OutputArchive& arch ) const;
+		inline void Serialise( FOutputArchive& arch ) const;
 
 		//!
-		inline void Unserialise( InputArchive& arch );
+		inline void Unserialise( FInputArchive& arch );
 	};
 	
 	struct FBoxBody : FBodyShape
@@ -66,10 +66,10 @@ namespace mu
 		}
 
 		//!
-		inline void Serialise( OutputArchive& arch ) const;
+		inline void Serialise( FOutputArchive& arch ) const;
 
 		//!
-		inline void Unserialise( InputArchive& arch );
+		inline void Unserialise( FInputArchive& arch );
 	};
 
 	struct FSphylBody : FBodyShape
@@ -89,10 +89,10 @@ namespace mu
 		}
 		 
 		//!
-		inline void Serialise( OutputArchive& arch ) const;
+		inline void Serialise( FOutputArchive& arch ) const;
 
 		//!
-		inline void Unserialise( InputArchive& arch );
+		inline void Unserialise( FInputArchive& arch );
 	};
 	
 	struct FTaperedCapsuleBody : FBodyShape
@@ -114,10 +114,10 @@ namespace mu
 		}
 		 
 		//!
-		inline void Serialise( OutputArchive& arch ) const;
+		inline void Serialise( FOutputArchive& arch ) const;
 
 		//!
-		inline void Unserialise( InputArchive& arch );
+		inline void Unserialise( FInputArchive& arch );
 	};
 	
 	struct FConvexBody : FBodyShape
@@ -136,10 +136,10 @@ namespace mu
 		}
 
 		//!
-		inline void Serialise( OutputArchive& arch ) const;
+		inline void Serialise( FOutputArchive& arch ) const;
 
 		//!
-		inline void Unserialise( InputArchive& arch );
+		inline void Unserialise( FInputArchive& arch );
 	};
 	
 	struct FPhysicsBodyAggregate
@@ -160,25 +160,21 @@ namespace mu
 		}
 
 		 //!
-		inline void Serialise( OutputArchive& arch ) const;
+		inline void Serialise( FOutputArchive& ) const;
 
 		//!
-		inline void Unserialise( InputArchive& arch );
+		inline void Unserialise( FInputArchive& );
 	};
 
-	class MUTABLERUNTIME_API PhysicsBody : public RefCounted
+	class MUTABLERUNTIME_API FPhysicsBody
 	{
 	public:
-
-		PhysicsBody() 
-		{
-		}
 		
-		mu::Ptr<PhysicsBody> Clone() const; 
+		TSharedPtr<FPhysicsBody> Clone() const;
 
 		//! Serialisation
-        static void Serialise( const PhysicsBody* p, OutputArchive& arch );
-		static mu::Ptr<PhysicsBody> StaticUnserialise( InputArchive& arch );
+        static void Serialise( const FPhysicsBody*, FOutputArchive& );
+		static TSharedPtr<FPhysicsBody> StaticUnserialise( FInputArchive& );
 
         //bool operator==( const PhysicsBody& Other ) const;
 
@@ -286,12 +282,6 @@ namespace mu
 		const FString& GetSphylName(int32 B, int32 I) const;
 		const FString& GetTaperedCapsuleName(int32 B, int32 I) const;
 
-	protected:
-		//! Forbidden. Manage with the Ptr<> template.
-		~PhysicsBody()
-		{
-		}
-
 	public:
 		int32 CustomId = -1;
 
@@ -304,16 +294,16 @@ namespace mu
 		bool bBodiesModified = false;
 
         //!
-        inline void Serialise( OutputArchive& arch ) const;
+        inline void Serialise( FOutputArchive& arch ) const;
 
 		//!
-        inline void Unserialise( InputArchive& arch );
+        inline void Unserialise( FInputArchive& arch );
 
 		//!
-		inline bool operator==(const PhysicsBody& Other) const
+		inline bool operator==(const FPhysicsBody& Other) const
         {
         	return CustomId        == Other.CustomId        &&
-				   BoneIds == Other.BoneIds					&&
+				   BoneIds		   == Other.BoneIds			&&
         		   Bodies          == Other.Bodies          &&
         	       BodiesCustomIds == Other.BodiesCustomIds &&
 				   bBodiesModified == Other.bBodiesModified;
diff --git a/Engine/Plugins/Mutable/Source/MutableRuntime/Public/MuR/RefCounted.h b/Engine/Plugins/Mutable/Source/MutableRuntime/Public/MuR/RefCounted.h
--- a/Engine/Plugins/Mutable/Source/MutableRuntime/Public/MuR/RefCounted.h
+++ b/Engine/Plugins/Mutable/Source/MutableRuntime/Public/MuR/RefCounted.h
@@ -4,7 +4,6 @@
 
 // HEADER_UNIT_SKIP - 
 
-#include "MuR/MutableMemory.h"
 #include "MuR/Ptr.h"
 
 #include "HAL/PlatformAtomics.h"
@@ -80,13 +79,16 @@ namespace mu
 		}
 	}
 
-	class MUTABLERUNTIME_API Resource : public RefCounted
+	class MUTABLERUNTIME_API FResource
 	{
 	public:
 
 		/** Return the size in bytes of all the LODs of the image. */
 		virtual int32 GetDataSize() const = 0;
 
+		/** */
+		virtual ~FResource() {}
+
 	};
 }
 
diff --git a/Engine/Plugins/Mutable/Source/MutableRuntime/Public/MuR/Serialisation.h b/Engine/Plugins/Mutable/Source/MutableRuntime/Public/MuR/Serialisation.h
--- a/Engine/Plugins/Mutable/Source/MutableRuntime/Public/MuR/Serialisation.h
+++ b/Engine/Plugins/Mutable/Source/MutableRuntime/Public/MuR/Serialisation.h
@@ -2,10 +2,7 @@
 
 #pragma once
 
-#include <string>
-
 #include "HAL/Platform.h"
-#include "MuR/MutableMemory.h"
 #include "MuR/Ptr.h"
 #include "MuR/RefCounted.h"
 #include "MuR/MutableMath.h"
@@ -17,64 +14,64 @@
 #include "Containers/Array.h"
 #include "Containers/Map.h"
 #include "Containers/ContainerAllocationPolicies.h"
+#include "Templates/SharedPointer.h"
 
 namespace mu
 {    
-    class Image;
-	class Model;
-
-#define MUTABLE_DEFINE_POD_SERIALISABLE(Type)					   \
-	void DLLEXPORT operator<<(OutputArchive& Arch, const Type& T); \
-	void DLLEXPORT operator>>(InputArchive& Arch, Type& T);		   \
-
-#define MUTABLE_DEFINE_POD_VECTOR_SERIALISABLE(Type)					\
-	template<typename Alloc>                                            \
-	void operator<<(OutputArchive& Arch, const TArray<Type, Alloc>& V);	\
-	                                                                    \
-	template<typename Alloc>                                            \
-	void operator>>(InputArchive& Arch, TArray<Type, Alloc>& V);        \
-
-#define MUTABLE_DEFINE_ENUM_SERIALISABLE(Type)					   \
-    void DLLEXPORT operator<<(OutputArchive& Arch, const Type& T); \
-    void DLLEXPORT operator>>(InputArchive& Arch, Type& T);		   \
+    class FImage;
+	class FModel;
+
+#define MUTABLE_DEFINE_POD_SERIALISABLE(Type)									\
+	void DLLEXPORT operator<<(mu::FOutputArchive& Arch, const Type& T);			\
+	void DLLEXPORT operator>>(mu::FInputArchive& Arch, Type& T);
+
+#define MUTABLE_DEFINE_POD_VECTOR_SERIALISABLE(Type)							\
+	template<typename Alloc>													\
+	void operator<<(mu::FOutputArchive& Arch, const TArray<Type, Alloc>& V);	\
+	template<typename Alloc>													\
+	void operator>>(mu::FInputArchive& Arch, TArray<Type, Alloc>& V);
+
+#define MUTABLE_DEFINE_ENUM_SERIALISABLE(Type)									\
+    void DLLEXPORT operator<<(mu::FOutputArchive& Arch, const Type& T);			\
+    void DLLEXPORT operator>>(mu::FInputArchive& Arch, Type& T);
 	
 	
     /** This repesents a handle to Mutable resource. */
     template<class R>
-    class ResourceProxy : public RefCounted
+    class TResourceProxy : public RefCounted
     {
     public:
 
-        virtual Ptr<const R> Get() = 0;
+        virtual TSharedPtr<const R> Get() = 0;
 
     };
 
 
 	/** Proxy implementation that always has the resource loaded in memory. */
 	template<class R>
-	class ResourceProxyMemory : public ResourceProxy<R>
+	class TResourceProxyMemory : public TResourceProxy<R>
 	{
 	private:
 
-		Ptr<const R> m_resource;
+		TSharedPtr<const R> Resource;
 
 	public:
 
-		ResourceProxyMemory(const Ptr<const R>& i)
+		TResourceProxyMemory(const TSharedPtr<const R>& InResource)
 		{
-			m_resource = i;
+			Resource = InResource;
 		}
 
-		// ResourceProxy interface
-		Ptr<const R> Get() override
+		// TResourceProxy interface
+		TSharedPtr<const R> Get() override
 		{
-			return m_resource;
+			return Resource;
 		}
 	};
 
 
 	/** */
-	class MUTABLERUNTIME_API ModelReader
+	class MUTABLERUNTIME_API FModelReader
     {
     public:
 
@@ -83,7 +80,7 @@ namespace mu
         //-----------------------------------------------------------------------------------------
 
         //! Ensure virtual destruction.
-        virtual ~ModelReader() {}
+        virtual ~FModelReader() {}
 
         //-----------------------------------------------------------------------------------------
         // Reading interface
@@ -104,7 +101,7 @@ namespace mu
 		//! \return a previously unused identifier, now used for this operation, that can be used in
 		//! calls to the other methods of this interface. If the return value is negative it indicates
 		//! an unrecoverable error.
-		virtual FOperationID BeginReadBlock(const mu::Model*, uint32 BlockKey, void* Buffer, uint64 size, TFunction<void(bool bSuccess)>* CompletionCallback = nullptr) = 0;
+		virtual FOperationID BeginReadBlock(const FModel*, uint32 BlockKey, void* Buffer, uint64 Size, TFunction<void(bool bSuccess)>* CompletionCallback = nullptr) = 0;
 
         //! Check if a data request operation has been completed.
         //! This is a weak check than *may* return true if the given operation has completed, but
@@ -124,7 +121,7 @@ namespace mu
 
 
 	/** */
-	class MUTABLERUNTIME_API ModelWriter
+	class MUTABLERUNTIME_API FModelWriter
 	{
 	public:
 
@@ -133,7 +130,7 @@ namespace mu
 		//-----------------------------------------------------------------------------------------
 
 		//! Ensure virtual destruction.
-		virtual ~ModelWriter() {}
+		virtual ~FModelWriter() {}
 
 		//-----------------------------------------------------------------------------------------
 		// Writing interface
@@ -156,7 +153,7 @@ namespace mu
 		//! There must be a file open with OpenWriteFile before calling this.
 		//! \param pBuffer pointer to the data that will be written to the file.
 		//! \param size size of the data to write on the file, in bytes.
-		virtual void Write(const void* pBuffer, uint64 size) = 0;
+		virtual void Write(const void* Buffer, uint64 Size) = 0;
 
 		//! \brief Close the file open for writing in a previous call to OpenWriteFile in this
 		//! object.
@@ -167,12 +164,12 @@ namespace mu
 
 
     /** Interface for any input stream to be use with InputArchives. */
-    class MUTABLERUNTIME_API InputStream
+    class MUTABLERUNTIME_API FInputStream
     {
     public:
 
 		/** Ensure virtual destruction. */
-		virtual ~InputStream() {}
+		virtual ~FInputStream() {}
 
         /** Read a byte buffer
          * \param pData destination buffer, must have at least size bytes allocated.
@@ -183,12 +180,12 @@ namespace mu
 
 
     /** Interface for any output stream to be used with OutputArchives. */
-    class MUTABLERUNTIME_API OutputStream
+    class MUTABLERUNTIME_API FOutputStream
     {
     public:
 
         /** Ensure virtual destruction. */
-        virtual ~OutputStream() {}
+        virtual ~FOutputStream() {}
 
         /** Write a byte buffer
          * \param pData source buffer where data will be read from.
@@ -200,42 +197,42 @@ namespace mu
 
 
     /** Archive containing data to be deserialised. */
-    class MUTABLERUNTIME_API InputArchive
+    class MUTABLERUNTIME_API FInputArchive
     {
     public:
 
         /** Construct form an input stream.The stream will not be owned by the archive and the
          * caller must make sure it is not modified or destroyed while serialisation is happening.
 		 */
-        InputArchive( InputStream* );
+        FInputArchive( FInputStream* );
 
 		/** Ensure virtual destruction. */
-		virtual ~InputArchive() {}
+		virtual ~FInputArchive() {}
 
         //
-        virtual Ptr<ResourceProxy<Image>> NewImageProxy();
+        virtual Ptr<TResourceProxy<FImage>> NewImageProxy();
 
 		/** Not owned. */
-		InputStream* Stream = nullptr;
+		FInputStream* Stream = nullptr;
 
 		/** Already read pointers. */
-		TArray< Ptr<RefCounted> > History;
+		TArray< TSharedPtr<void> > History;
 
     };
 
 
     /** Archive where data can be serialised to. */
-    class MUTABLERUNTIME_API OutputArchive
+    class MUTABLERUNTIME_API FOutputArchive
     {
     public:
 
         /** Construct form an output stream.The stream will not be owned by the archive and the
          * caller must make sure it is not modified or destroyed while serialisation is happening.
 		 */
-        OutputArchive( OutputStream* );
+        FOutputArchive( FOutputStream* );
 
 		/** Not owned. */
-		OutputStream* Stream;
+		FOutputStream* Stream = nullptr;
 
 		/** Already written pointers and their ids. */
 		TMap< const void*, int32 > History;
@@ -243,46 +240,17 @@ namespace mu
     };
 
 
-    /** */
-    class MUTABLERUNTIME_API InputArchiveWithProxies : public InputArchive
-    {
-    public:
-
-        class ProxyFactory
-        {
-        public:
-          virtual ~ProxyFactory() {}
-          virtual Ptr<ResourceProxy<Image>> NewImageProxy(InputArchive& arch) = 0;
-        };
-
-    public:
-
-        /** Construct form an input stream.The stream will not be owned by the archive and the
-         * caller must make sure it is not modified or destroyed while serialisation is happening.
-		 */
-        InputArchiveWithProxies( InputStream*, ProxyFactory*  );
-
-        // InputArchive interface
-        virtual Ptr<ResourceProxy<Image>> NewImageProxy() override;
-
-    private:
-
-		TArray< Ptr<ResourceProxy<Image>> > ProxyHistory;
-		ProxyFactory* Factory = nullptr;
-	};
-
-
     //!
-    class MUTABLERUNTIME_API OutputMemoryStream : public OutputStream
+    class MUTABLERUNTIME_API FOutputMemoryStream : public FOutputStream
     {
     public:
 
         /** Create the stream with an optional buffer size in bytes.
 		* The internal buffer will be enlarged as much as necessary.
 		*/
-        OutputMemoryStream( uint64 Reserve = 0 );
+        FOutputMemoryStream( uint64 Reserve = 0 );
 
-        // OutputStream interface
+        // FOutputStream interface
         virtual void Write( const void* Data, uint64 Size ) override;
 
         // Own interface
@@ -290,10 +258,10 @@ namespace mu
         /** Get the serialised data buffer pointer. This pointer invalidates after a Write
          * operation has been done, and you need to get it again.
 		 */
-        const void* GetBuffer() const;
+        const uint8* GetBuffer() const;
 
         /** Get the amount of data in the stream, in bytes. */
-        uint64 GetBufferSize() const;
+        int32 GetBufferSize() const;
 
 		/** Clear the internal buffer. */
 		void Reset();
@@ -306,12 +274,12 @@ namespace mu
 
 
 	/** This stream doesn't store any data: it just counts the amount of data serialised. */
-	class MUTABLERUNTIME_API OutputSizeStream : public OutputStream
+	class MUTABLERUNTIME_API FOutputSizeStream : public FOutputStream
 	{
 	public:
 
-		// OutputStream interface
-		void Write(const void* pData, uint64 size) override;
+		// FOutputStream interface
+		void Write(const void* Data, uint64 Size) override;
 
 		// Own interface
 
@@ -326,12 +294,12 @@ namespace mu
 
 
 	/** This stream doesn't store any data: it just calculates of a hash of the data as it receives it. */
-	class MUTABLERUNTIME_API OutputHashStream : public OutputStream
+	class MUTABLERUNTIME_API FOutputHashStream : public FOutputStream
 	{
 	public:
 
-		// OutputStream interface
-		void Write(const void* pData, uint64 size) override;
+		// FOutputStream interface
+		void Write(const void* Data, uint64 Size) override;
 
 		// Own interface
 
@@ -346,13 +314,13 @@ namespace mu
 
 
 	template<typename Type>
-	void operator<<(OutputArchive& Arch, const Type& Value)
+	void operator<<(FOutputArchive& Arch, const Type& Value)
 	{
         Value.Serialise(Arch);
 	}
 
 	template<typename Type>
-	void operator>>(InputArchive& Arch, Type& Value)
+	void operator>>(FInputArchive& Arch, Type& Value)
 	{
         Value.Unserialise(Arch);
 	}
@@ -404,25 +372,22 @@ namespace mu
 	MUTABLE_DEFINE_POD_VECTOR_SERIALISABLE(FVector4f);
 	
 	//---------------------------------------------------------------------------------------------
-	MUTABLERUNTIME_API void operator<<(OutputArchive& arch, const FString& t);
-	
-	MUTABLERUNTIME_API void operator>>(InputArchive& arch, FString& t);
+	MUTABLERUNTIME_API void operator<<(FOutputArchive&, const FString&);
+	MUTABLERUNTIME_API void operator>>(FInputArchive&, FString&);
 
 	//---------------------------------------------------------------------------------------------
-	MUTABLERUNTIME_API void operator<<(OutputArchive& arch, const FRichCurve& t);
-
-	MUTABLERUNTIME_API void operator>>(InputArchive& arch, FRichCurve& t);
+	MUTABLERUNTIME_API void operator<<(FOutputArchive&, const FRichCurve&);
+	MUTABLERUNTIME_API void operator>>(FInputArchive&, FRichCurve&);
 
 
 	//---------------------------------------------------------------------------------------------
-	MUTABLERUNTIME_API void operator<<(OutputArchive& arch, const FName& v);
-
-	MUTABLERUNTIME_API void operator>>(InputArchive& arch, FName& v);
+	MUTABLERUNTIME_API void operator<<(FOutputArchive&, const FName&);
+	MUTABLERUNTIME_API void operator>>(FInputArchive&, FName&);
 
 
 	//---------------------------------------------------------------------------------------------
 	template<typename T, typename Alloc> 
-	void operator<<(OutputArchive& Arch, const TArray<T, Alloc>& V)
+	void operator<<(FOutputArchive& Arch, const TArray<T, Alloc>& V)
 	{
 		const uint32 Num = (uint32)V.Num();
 		Arch << Num;
@@ -434,7 +399,7 @@ namespace mu
 	}
 
 	template<typename T, typename Alloc> 
-	void operator>>(InputArchive& Arch, TArray<T, Alloc>& V)
+	void operator>>(FInputArchive& Arch, TArray<T, Alloc>& V)
 	{
 		uint32 Num;
 		Arch >> Num;
@@ -446,12 +411,9 @@ namespace mu
 		}
 	}
 
-	
-	MUTABLERUNTIME_API void operator>> ( InputArchive& arch, std::string& v );
 
 	// Bool size is not a standard
-	MUTABLERUNTIME_API void operator<<(OutputArchive& Arch, const bool& T);
-
-	MUTABLERUNTIME_API void operator>>(InputArchive& Arch, bool& T);
+	MUTABLERUNTIME_API void operator<<(FOutputArchive&, const bool&);
+	MUTABLERUNTIME_API void operator>>(FInputArchive&, bool&);
 }
 
diff --git a/Engine/Plugins/Mutable/Source/MutableRuntime/Public/MuR/Settings.h b/Engine/Plugins/Mutable/Source/MutableRuntime/Public/MuR/Settings.h
--- a/Engine/Plugins/Mutable/Source/MutableRuntime/Public/MuR/Settings.h
+++ b/Engine/Plugins/Mutable/Source/MutableRuntime/Public/MuR/Settings.h
@@ -9,13 +9,10 @@
 namespace mu
 {
     /** Settings class that can be used to create a customised System instance. */
-    class MUTABLERUNTIME_API Settings : public RefCounted
+    class MUTABLERUNTIME_API FSettings
     {
     public:
 
-        /** Create new settings with the default values as explained in each method below. */
-        Settings();
-
         /** Record internal profiling data. Disabled by default. */
         void SetProfile( bool bEnabled );
 
@@ -40,17 +37,6 @@ namespace mu
         */
         void SetImageCompressionQuality( int32 Quality );
 
-        // Interface pattern
-        Settings( const Settings& ) = delete;
-        Settings( Settings&& ) = delete;
-        Settings& operator=(const Settings&) = delete;
-        Settings& operator=(Settings&&) = delete;
-
-    protected:
-
-        /** Forbidden.Manage with the Ptr<> template. */
-        ~Settings() override;
-
     public:
 
 		uint64 WorkingMemoryBytes = 0;
diff --git a/Engine/Plugins/Mutable/Source/MutableRuntime/Public/MuR/Skeleton.h b/Engine/Plugins/Mutable/Source/MutableRuntime/Public/MuR/Skeleton.h
--- a/Engine/Plugins/Mutable/Source/MutableRuntime/Public/MuR/Skeleton.h
+++ b/Engine/Plugins/Mutable/Source/MutableRuntime/Public/MuR/Skeleton.h
@@ -11,18 +11,9 @@
 #include "MuR/Ptr.h"
 #include "MuR/RefCounted.h"
 
-// Remove when removin deprecated data
-#include <string>
-
 
 namespace mu
 {	
-	// Forward references
-    class Skeleton;
-
-    typedef Ptr<Skeleton> SkeletonPtr;
-    typedef Ptr<const Skeleton> SkeletonPtrConst;
-
 	// Bone name identifier
 	struct FBoneName
 	{
@@ -32,14 +23,14 @@ namespace mu
 		// Hash built from the bone name (FString)
 		uint32 Id = 0;
 
-		inline void Serialise(OutputArchive& arch) const
+		inline void Serialise(FOutputArchive& Arch) const
 		{
-			arch << Id;
+			Arch << Id;
 		}
 
-		inline void Unserialise(InputArchive& arch)
+		inline void Unserialise(FInputArchive& Arch)
 		{
-			arch >> Id;
+			Arch >> Id;
 		}
 
 		//!
@@ -55,9 +46,8 @@ namespace mu
 	}
 
 
-    //! \brief Skeleton object.
-	//! \ingroup runtime
-    class MUTABLERUNTIME_API Skeleton : public RefCounted
+    /** Skeleton object. */
+	class MUTABLERUNTIME_API FSkeleton
 	{
 	public:
 
@@ -65,14 +55,12 @@ namespace mu
 		// Life cycle
 		//-----------------------------------------------------------------------------------------
 
-		Skeleton() {}
-
         //! Deep clone this skeleton.
-        Ptr<Skeleton> Clone() const;
+        TSharedPtr<FSkeleton> Clone() const;
 
 		//! Serialisation
-        static void Serialise( const Skeleton* p, OutputArchive& arch );
-        static Ptr<Skeleton> StaticUnserialise( InputArchive& arch );
+        static void Serialise( const FSkeleton*, FOutputArchive& );
+        static TSharedPtr<FSkeleton> StaticUnserialise( FInputArchive& );
 
 
 		//-----------------------------------------------------------------------------------------
@@ -81,15 +69,15 @@ namespace mu
 
         //! @return - Number of bones in the Skeleton
 		int32 GetBoneCount() const;
-        void SetBoneCount(int32 c);
+        void SetBoneCount(int32 Count);
 
 		//! @return - FName of the bone at 'Index'. Only valid in the editor
 		const FName GetDebugName(int32 Index) const;
 		void SetDebugName(const int32 Index, const FName BoneName);
 
         //! Get and set the parent bone of each bone. The parent can be -1 if the bone is a root.
-        int32 GetBoneParent(int32 boneIndex) const;
-        void SetBoneParent(int32 boneIndex, int32 parentBoneIndex);
+        int32 GetBoneParent(int32 BoneIndex) const;
+        void SetBoneParent(int32 BoneIndex, int32 ParentBoneIndex);
 
 		//! @return - BoneName of the Bone at 'Index'.
 		const FBoneName& GetBoneName(int32 Index) const;
@@ -99,17 +87,8 @@ namespace mu
 		int32 FindBone(const FBoneName& BoneName) const;
 		
 
-	protected:
-
-		//! Forbidden. Manage with the Ptr<> template.
-		~Skeleton() {}
-
 	public:
 
-		//! Deprecated
-		TArray<std::string> m_bones_DEPRECATED;
-		TArray<FTransform3f> m_boneTransforms_DEPRECATED;
-
 		//! DEBUG. FNames of the bones. Only valid in the editor. Do not serialize.
 		TArray<FName> DebugBoneNames;
 
@@ -121,16 +100,16 @@ namespace mu
 		TArray<int16> BoneParents;
 
 		//!
-		void Serialise(OutputArchive& arch) const;
+		void Serialise(FOutputArchive&) const;
 
 		//!
-		void Unserialise(InputArchive& arch);
+		void Unserialise(FInputArchive&);
 
 		//!
-		inline bool operator==(const Skeleton& o) const
+		inline bool operator==(const FSkeleton& Other) const
 		{
-			return BoneIds == o.BoneIds
-				&& BoneParents == o.BoneParents;
+			return BoneIds == Other.BoneIds
+				&& BoneParents == Other.BoneParents;
 		}
 	};
 
diff --git a/Engine/Plugins/Mutable/Source/MutableRuntime/Public/MuR/System.h b/Engine/Plugins/Mutable/Source/MutableRuntime/Public/MuR/System.h
--- a/Engine/Plugins/Mutable/Source/MutableRuntime/Public/MuR/System.h
+++ b/Engine/Plugins/Mutable/Source/MutableRuntime/Public/MuR/System.h
@@ -2,27 +2,25 @@
 
 #pragma once
 
-#include "HAL/Platform.h"
 #include "MuR/Image.h"
 #include "MuR/Instance.h"
-#include "MuR/MutableMemory.h"
-#include "MuR/Ptr.h"
 #include "MuR/RefCounted.h"
 #include "MuR/Settings.h"
 #include "MuR/Types.h"
 #include "Tasks/Task.h"
+#include "Templates/SharedPointer.h"
 #include "Templates/Tuple.h"
+#include "HAL/Platform.h"
+
+#include "System.generated.h"
 
 
-/** If set to 1, this enables some expensive Unreal Insights traces, but can lead to 5x slower mutable operation. 
+/** If set to 1, this enables some expensive Unreal Insights traces, but can lead to 5x slower mutable operation.
 * Other cheaper traces are enabled at all times.
 */
 #define UE_MUTABLE_ENABLE_SLOW_TRACES	0
 
 
-#include "System.generated.h"
-
-
 /** Despite being an UEnum, this is not always version-serialized (in MutableTools).
  * Beware of changing the enum options or order.
  */
@@ -39,21 +37,17 @@ enum class ETextureCompressionStrategy : uint8
 	NeverCompress
 };
 
+MUTABLE_DEFINE_ENUM_SERIALISABLE(ETextureCompressionStrategy);
+
+
 namespace mu
 {
 	// Forward references
-	class Model;
-	class ModelReader;
-	class Parameters;
-    class Mesh;
-	class ExtensionDataStreamer;
-
-    class System;
-    using SystemPtr=Ptr<System>;
-    using SystemPtrConst=Ptr<const System>;
-
-
-	MUTABLE_DEFINE_ENUM_SERIALISABLE(ETextureCompressionStrategy);
+	class FModel;
+	class FModelReader;
+	class FParameters;
+    class FMesh;
+	class FExtensionDataStreamer;
 
 
 	/** */
@@ -74,28 +68,27 @@ namespace mu
 
 
     /** Interface to request external images used as parameters. */
-    class MUTABLERUNTIME_API ExternalResourceProvider
+    class MUTABLERUNTIME_API FExternalResourceProvider
     {
     public:
 
         //! Ensure virtual destruction
-        virtual ~ExternalResourceProvider() = default;
+        virtual ~FExternalResourceProvider() = default;
 
         /** Returns the completion event and a cleanup function that must be called once event is completed. */
-		virtual TTuple<UE::Tasks::FTask, TFunction<void()>> GetImageAsync(FName Id, uint8 MipmapsToSkip, TFunction<void(Ptr<Image>)>& ResultCallback) = 0;
-		virtual TTuple<UE::Tasks::FTask, TFunction<void()>> GetReferencedImageAsync(const void* ModelPtr, int32 Id, uint8 MipmapsToSkip, TFunction<void(Ptr<Image>)>& ResultCallback) { check(false); return {}; }
+		virtual TTuple<UE::Tasks::FTask, TFunction<void()>> GetImageAsync(FName Id, uint8 MipmapsToSkip, TFunction<void(TSharedPtr<FImage>)>& ResultCallback) = 0;
+		virtual TTuple<UE::Tasks::FTask, TFunction<void()>> GetReferencedImageAsync(const void* ModelPtr, int32 Id, uint8 MipmapsToSkip, TFunction<void(TSharedPtr<FImage>)>& ResultCallback) { check(false); return {}; }
 
 		virtual mu::FImageDesc GetImageDesc(FName Id, uint8 MipmapsToSkip) = 0;
 
 		/** Returns the completion event and a cleanup function that must be called once event is completed. */
-		virtual TTuple<UE::Tasks::FTask, TFunction<void()>> GetMeshAsync(FName Id, TFunction<void(Ptr<Mesh>)>& ResultCallback) = 0;
-		virtual TTuple<UE::Tasks::FTask, TFunction<void()>> GetReferencedMeshAsync(const void* ModelPtr, int32 Id, TFunction<void(Ptr<Mesh>)>& ResultCallback) { check(false); return {}; }
+		virtual TTuple<UE::Tasks::FTask, TFunction<void()>> GetMeshAsync(FName Id, TFunction<void(TSharedPtr<FMesh>)>& ResultCallback) = 0;
+		virtual TTuple<UE::Tasks::FTask, TFunction<void()>> GetReferencedMeshAsync(const void* ModelPtr, int32 Id, TFunction<void(TSharedPtr<FMesh>)>& ResultCallback) { check(false); return {}; }
     };
 
 
-    //! \brief Main system class to load models and build instances.
-	//! \ingroup runtime
-	class MUTABLERUNTIME_API System : public RefCounted
+    /** Main system class to load models and build instances. */
+	class MUTABLERUNTIME_API FSystem
 	{
     public:
 
@@ -108,10 +101,10 @@ namespace mu
         //! \param Settings Optional class with the settings to use in this system. The default
         //! value configures a production-ready system.
 		//! \param DataStreamer Optional interface to allow the Model to stream in ExtensionData from disk
-        System( const Ptr<Settings>& Settings = nullptr, const TSharedPtr<ExtensionDataStreamer>& DataStreamer = nullptr );
+        FSystem( const FSettings& Settings = FSettings(), const TSharedPtr<FExtensionDataStreamer>& DataStreamer = nullptr);
 
         //! Set a new provider for model data. 
-        void SetStreamingInterface(const TSharedPtr<ModelReader>& );
+        void SetStreamingInterface(const TSharedPtr<FModelReader>& );
 
         /** Set the working memory limit, overrding any set in the settings when the system was created.
          * Refer to Settings::SetWorkingMemoryBudget for more information.
@@ -127,7 +120,7 @@ namespace mu
 		void SetGeneratedCacheSize(uint32 InCount);
 
         /** Set a new provider for external image data. This is only necessary if image parameters are used in the models. */
-        void SetExternalResourceProvider(const TSharedPtr<ExternalResourceProvider>&);
+        void SetExternalResourceProvider(const TSharedPtr<FExternalResourceProvider>&);
 
 		/** Set a function that will be used to convert image pixel formats instead of the internal conversion. 
 		* \warning The provided function can be called from any thread, and also concurrently.
@@ -142,7 +135,7 @@ namespace mu
         //! no longer needed.
         //! \param pModel Model to build an instance of
         //! \return An identifier that is always bigger than 0.
-        Instance::ID NewInstance(const TSharedPtr<const Model>& Model);
+        FInstance::FID NewInstance(const TSharedPtr<const FModel>& Model);
 
         //! \brief Update an instance with a new parameter set and/or state.
         //!
@@ -156,55 +149,55 @@ namespace mu
         //! \return the instance data with all the LOD, components, and ids to generate the meshes 
         //! and images. The returned Instance is only valid until the next call to EndUpdate with 
         //! the same instanceID parameter.
-        Ptr<const Instance> BeginUpdate(Instance::ID InstanceID, const Ptr<const Parameters>& Params, int32 StateIndex, uint32 LodMask);
+		TSharedPtr<const FInstance> BeginUpdate(FInstance::FID InstanceID, const TSharedPtr<const FParameters>& Params, int32 StateIndex, uint32 LodMask);
 
 		//! Only valid between BeginUpdate and EndUpdate
 		//! Calculate the description of an image, without generating it.
-		UE::Tasks::TTask<FImageDesc> GetImageDesc(Instance::ID InstanceID, FResourceID ImageId);
+		UE::Tasks::TTask<FImageDesc> GetImageDesc(FInstance::FID InstanceID, FResourceID ImageId);
 
 		//! Only valid between BeginUpdate and EndUpdate
 		//! \param MipsToSkip Number of mips to skip compared from the full image.
 		//! If 0, all mip levels will be generated. If more levels than possible to discard are specified, 
 		//! the image will still contain a minimum number of mips specified at model compile time.
-		UE::Tasks::TTask<Ptr<const Image>> GetImage(Instance::ID InstanceID, FResourceID ImageId, int32 MipsToSkip = 0, int32 LOD = 0);
+		UE::Tasks::TTask<TSharedPtr<const FImage>> GetImage(FInstance::FID InstanceID, FResourceID ImageId, int32 MipsToSkip = 0, int32 LOD = 0);
 
         //! Only valid between BeginUpdate and EndUpdate
-		UE::Tasks::TTask<Ptr<const Mesh>> GetMesh(Instance::ID InstanceID, FResourceID MeshId);
+		UE::Tasks::TTask<TSharedPtr<const FMesh>> GetMesh(FInstance::FID InstanceID, FResourceID MeshId);
 
 		//! Only valid between BeginUpdate and EndUpdate
 		//! Calculate the description of an image, without generating it.
-		void GetImageDescInline(Instance::ID InstanceID, FResourceID ImageId, FImageDesc& OutDesc);
+		void GetImageDescInline(FInstance::FID InstanceID, FResourceID ImageId, FImageDesc& OutDesc);
 
 		//! Only valid between BeginUpdate and EndUpdate
 		//! \param MipsToSkip Number of mips to skip compared from the full image.
 		//! If 0, all mip levels will be generated. If more levels than possible to discard are specified, 
 		//! the image will still contain a minimum number of mips specified at model compile time.
-		Ptr<const Image> GetImageInline(Instance::ID InstanceID, FResourceID ImageId, int32 MipsToSkip = 0, int32 LOD = 0);
+		TSharedPtr<const FImage> GetImageInline(FInstance::FID InstanceID, FResourceID ImageId, int32 MipsToSkip = 0, int32 LOD = 0);
 
         //! Only valid between BeginUpdate and EndUpdate
-		Ptr<const Mesh> GetMeshInline(Instance::ID InstanceID, FResourceID MeshId);
+		TSharedPtr<const FMesh> GetMeshInline(FInstance::FID InstanceID, FResourceID MeshId);
 
         //! Invalidate and free the last Instance data returned by a call to BeginUpdate with
         //! the same instance index. After a call to this method, that Instance cannot be used any
         //! more and its content is undefined.
         //! \param instance The index of the instance whose last data will be invalidated.
-        void EndUpdate(Instance::ID InstanceID);
+        void EndUpdate(FInstance::FID InstanceID);
 
         //! Completely destroy an instance. After a call to this method the given instance cannot be
         //! updated any more, and its resources may have been freed.
         //! \param instance The id of the instance to destroy.
-        void ReleaseInstance(Instance::ID InstanceID);
+        void ReleaseInstance(FInstance::FID InstanceID);
 
 		//! Calculate the relevancy of every parameter. Some parameters may be unused depending on
 		//! the values of other parameters. This method will set to true the flags for parameters
 		//! that are relevant, and to false otherwise. This is useful to hide irrelevant parameters
 		//! in dynamic user interfaces.
-        //! \param pModel The model used to create the Parameters instance.
+        //! \param pModel The model used to create the FParameters instance.
         //! \param pParameters Parameter set that we want to find the relevancy of.
         //! \param pFlags is a pointer to a preallocated array of booleans that contains at least
 		//! pParameters->GetCount() elements.
-        void GetParameterRelevancy( Instance::ID InstanceID,
-									const Ptr<const Parameters>& Parameters,
+        void GetParameterRelevancy( FInstance::FID InstanceID,
+									const TSharedPtr<const FParameters>& FParameters,
 									bool* Flags );
 
 		//-----------------------------------------------------------------------------------------
@@ -217,15 +210,12 @@ namespace mu
     public:
 
         // Prevent copy, move and assignment.
-        System( const System& ) = delete;
-        System& operator=( const System& ) = delete;
-        System( System&& ) = delete;
-        System& operator=( System&& ) = delete;
-
-	protected:
+        FSystem( const FSystem& ) = delete;
+		FSystem& operator=( const FSystem& ) = delete;
+        FSystem( FSystem&& ) = delete;
+        FSystem& operator=( FSystem&& ) = delete;
 
-		//! Forbidden. Manage with the Ptr<> template.
-        ~System() override;
+		~FSystem();
 
 	private:
 
diff --git a/Engine/Plugins/Mutable/Source/MutableRuntime/Public/MuR/Types.h b/Engine/Plugins/Mutable/Source/MutableRuntime/Public/MuR/Types.h
--- a/Engine/Plugins/Mutable/Source/MutableRuntime/Public/MuR/Types.h
+++ b/Engine/Plugins/Mutable/Source/MutableRuntime/Public/MuR/Types.h
@@ -33,8 +33,8 @@ enum class EClipVertexSelectionType : uint8
 
 namespace mu
 {
-	class OutputArchive;
-	class InputArchive;
+	class FOutputArchive;
+	class FInputArchive;
 
 	/** Handle used to identify separate resources generated by mutable: images and meshes included in instances. */
     typedef uint64 FResourceID;
diff --git a/Engine/Plugins/Mutable/Source/MutableTools/Internal/MuT/AST.h b/Engine/Plugins/Mutable/Source/MutableTools/Internal/MuT/AST.h
--- a/Engine/Plugins/Mutable/Source/MutableTools/Internal/MuT/AST.h
+++ b/Engine/Plugins/Mutable/Source/MutableTools/Internal/MuT/AST.h
@@ -2,14 +2,12 @@
 
 #pragma once
 
-#include "MuT/Platform.h"
 #include "MuT/Compiler.h"
 #include "MuT/NodeImage.h"
 #include "MuR/Image.h"
 #include "MuR/Mesh.h"
 #include "MuR/ModelPrivate.h"
 #include "MuR/MutableMath.h"
-#include "MuR/MutableMemory.h"
 #include "MuR/Operations.h"
 #include "MuR/ParametersPrivate.h"
 #include "MuR/Ptr.h"
@@ -27,10 +25,6 @@
 
 #include <array>
 #include <atomic>
-#include <unordered_map>
-#include <unordered_set>
-#include <set>
-
 
 namespace std
 {
@@ -253,7 +247,7 @@ namespace mu
 		bool bSeparateImageMips = true;
 
 		/** Structure used to speedup mesh constant comparison. */
-		struct FDeduplicationMeshFuncs : TDefaultMapHashableKeyFuncs<Ptr<const Mesh>, int32, false>
+		struct FDeduplicationMeshFuncs : TDefaultMapHashableKeyFuncs<TSharedPtr<const FMesh>, int32, false>
 		{
 			static FORCEINLINE bool Matches(KeyInitType A, KeyInitType B)
 			{
@@ -262,7 +256,7 @@ namespace mu
 
 			static FORCEINLINE uint32 GetKeyHash(KeyInitType Key)
 			{
-				const mu::Mesh* Data = Key.get();
+				const mu::FMesh* Data = Key.Get();
 				return HashCombineFast(
 					::GetTypeHash(Data->VertexBuffers.GetElementCount()),
 					::GetTypeHash(Data->IndexBuffers.GetElementCount())
@@ -270,10 +264,10 @@ namespace mu
 			}
 		};
 
-		TMap<Ptr<const Mesh>, int32, FDefaultSetAllocator, FDeduplicationMeshFuncs> MeshConstantMap;
+		TMap<TSharedPtr<const FMesh>, int32, FDefaultSetAllocator, FDeduplicationMeshFuncs> MeshConstantMap;
 
 		/** Structure used to speedup image mip comparison. */
-		struct FDeduplicationImageFuncs : TDefaultMapHashableKeyFuncs<Ptr<const Image>, int32, false>
+		struct FDeduplicationImageFuncs : TDefaultMapHashableKeyFuncs<TSharedPtr<const FImage>, int32, false>
 		{
 			static FORCEINLINE bool Matches(KeyInitType A, KeyInitType B)
 			{
@@ -282,7 +276,7 @@ namespace mu
 
 			static FORCEINLINE uint32 GetKeyHash(KeyInitType Key)
 			{
-				const mu::Image* Data = Key.get();
+				const mu::FImage* Data = Key.Get();
 				uint32 Hash = HashCombineFast(::GetTypeHash(Data->GetFormat()), GetTypeHash(Data->GetSize()));
 
 				TArrayView<const uint8> DataView = Data->DataStorage.GetLOD(0);
@@ -293,7 +287,7 @@ namespace mu
 			}
 		};
 
-		TMap<Ptr<const Image>, int32, FDefaultSetAllocator, FDeduplicationImageFuncs> ImageConstantMipMap;
+		TMap<TSharedPtr<const FImage>, int32, FDefaultSetAllocator, FDeduplicationImageFuncs> ImageConstantMipMap;
 
 		/** Image operation functions, so that they can be overriden. */
 		FImageOperator& ImageOperator;
@@ -343,8 +337,8 @@ namespace mu
 
 	protected:
 
-		const ASTOpImageCrop* m_root = nullptr;
-		Ptr<ASTOp> m_initialSource;
+		const ASTOpImageCrop* Root = nullptr;
+		Ptr<ASTOp> InitialSource;
 		//! For each operation we sink, the map from old instructions to new instructions.
 		TMap<FSinkerOldToNewKey, Ptr<ASTOp>> OldToNew;
 
@@ -361,8 +355,8 @@ namespace mu
 
 	protected:
 
-		const ASTOpImagePixelFormat* m_root = nullptr;
-		Ptr<ASTOp> m_initialSource;
+		const ASTOpImagePixelFormat* Root = nullptr;
+		Ptr<ASTOp> InitialSource;
 		TMap<FSinkerOldToNewKey, Ptr<ASTOp>> OldToNew;
 
 		Ptr<ASTOp> Visit(Ptr<ASTOp> at, const ASTOpImagePixelFormat* currentFormatOp);
@@ -395,8 +389,8 @@ namespace mu
 
 	protected:
 
-		const ASTOpMeshFormat* m_root = nullptr;
-		Ptr<ASTOp> m_initialSource;
+		const ASTOpMeshFormat* Root = nullptr;
+		Ptr<ASTOp> InitialSource;
 		TMap<FSinkerOldToNewKey, Ptr<ASTOp>> OldToNew;
 
 		Ptr<ASTOp> Visit(const Ptr<ASTOp>& Op, const ASTOpMeshFormat* CurrentSinkOp);
@@ -429,8 +423,8 @@ namespace mu
 
 	protected:
 
-		const ASTOpMeshExtractLayoutBlocks* m_root = nullptr;
-		Ptr<ASTOp> m_initialSource;
+		const ASTOpMeshExtractLayoutBlocks* Root = nullptr;
+		Ptr<ASTOp> InitialSource;
 		TMap<FSinkerOldToNewKey, Ptr<ASTOp>> OldToNew;
 
 		Ptr<ASTOp> Visit(const Ptr<ASTOp>& at, const ASTOpMeshExtractLayoutBlocks* currentFormatOp);
@@ -446,8 +440,8 @@ namespace mu
 
 	protected:
 
-		const ASTOpImageMipmap* m_root = nullptr;
-		Ptr<ASTOp> m_initialSource;
+		const ASTOpImageMipmap* Root = nullptr;
+		Ptr<ASTOp> InitialSource;
 
 		//! For each operation we sink, the map from old instructions to new instructions.
 		TMap<FSinkerOldToNewKey, Ptr<ASTOp>> OldToNew;
@@ -481,47 +475,47 @@ namespace mu
 
         // move constructor
         ASTChild(ASTChild&& rhs)
-             : m_parent(rhs.m_parent)
-             , m_child(rhs.m_child)
-             , m_parentIndexInChild(rhs.m_parentIndexInChild)
+             : Parent(rhs.Parent)
+             , Child(rhs.Child)
+             , ParentIndexInChild(rhs.ParentIndexInChild)
         {
-            rhs.m_parent=nullptr;
-            rhs.m_child.reset();
+            rhs.Parent=nullptr;
+            rhs.Child.reset();
         }
 
         // Move assignment
-        ASTChild& operator=( ASTChild&& rhs );
+        ASTChild& operator=( ASTChild&& );
 
-        ASTChild& operator=( const Ptr<ASTOp>& c );
+        ASTChild& operator=( const Ptr<ASTOp>& );
 
         inline explicit operator bool() const
         {
-            return m_child.get()!=nullptr;
+            return Child.get()!=nullptr;
         }
 
         inline Ptr<class ASTOp>& child()
         {
-            return m_child;
+            return Child;
         }
 
         inline const Ptr<class ASTOp>& child() const
         {
-            return m_child;
+            return Child;
         }
 
         inline const Ptr<class ASTOp>& operator->() const
         {
-            return m_child;
+            return Child;
         }
 
         inline bool operator==(const ASTChild& o) const
         {
-            return m_child==o.m_child;
+            return Child==o.Child;
         }
 
-        class ASTOp* m_parent;
-        Ptr<class ASTOp> m_child;
-        int32 m_parentIndexInChild = 0;
+        class ASTOp* Parent;
+        Ptr<class ASTOp> Child;
+        int32 ParentIndexInChild = 0;
 
     private:
 
@@ -555,7 +549,7 @@ namespace mu
         virtual ~ASTOp() {}
 
         //! Get the operation type
-        virtual OP_TYPE GetOpType() const = 0;
+        virtual EOpType GetOpType() const = 0;
 
         //! Validate that everything is fine with this tree
         virtual void Assert();
@@ -596,17 +590,17 @@ namespace mu
 
 		inline bool IsConstantOp() const
 		{
-			OP_TYPE Type = GetOpType();
-			return Type == OP_TYPE::BO_CONSTANT
-				|| Type == OP_TYPE::NU_CONSTANT
-				|| Type == OP_TYPE::SC_CONSTANT
-				|| Type == OP_TYPE::CO_CONSTANT
-				|| Type == OP_TYPE::IM_CONSTANT
-				|| Type == OP_TYPE::ME_CONSTANT
-				|| Type == OP_TYPE::LA_CONSTANT
-				|| Type == OP_TYPE::PR_CONSTANT
-				|| Type == OP_TYPE::ST_CONSTANT
-				|| Type == OP_TYPE::ED_CONSTANT
+			EOpType Type = GetOpType();
+			return Type == EOpType::BO_CONSTANT
+				|| Type == EOpType::NU_CONSTANT
+				|| Type == EOpType::SC_CONSTANT
+				|| Type == EOpType::CO_CONSTANT
+				|| Type == EOpType::IM_CONSTANT
+				|| Type == EOpType::ME_CONSTANT
+				|| Type == EOpType::LA_CONSTANT
+				|| Type == EOpType::PR_CONSTANT
+				|| Type == EOpType::ST_CONSTANT
+				|| Type == EOpType::ED_CONSTANT
 				;
 		}
 
@@ -722,6 +716,16 @@ namespace mu
 				TFunctionRef<bool(Ptr<ASTOp>&)> f
             );
 
+
+		template<typename STATE>
+		struct FKeyFuncs : BaseKeyFuncs<TPair<Ptr<ASTOp>, STATE>, Ptr<ASTOp>, false>
+		{
+			static Ptr<ASTOp> GetSetKey(TPair<Ptr<ASTOp>, STATE> Element) { return Element.Key; }
+			static bool Matches(const Ptr<ASTOp>& lhs, const Ptr<ASTOp>& rhs) { return lhs == rhs || *lhs == *rhs; }
+			static uint32 GetKeyHash(const Ptr<ASTOp>& Key) { return Key->Hash(); }
+		};
+
+
         //! Kind of top-down, but really not.
         //! This version is slighlty faster, but doesn't support recursive traversals so
         //! use it only in controlled cases.
@@ -734,8 +738,8 @@ namespace mu
         {
             if (!root) { return; }
 
-            TArray<TPair<Ptr<ASTOp>,STATE>> pending;
-			pending.Emplace( root, initialState );
+            TArray<TPair<Ptr<ASTOp>,STATE>> Pending;
+			Pending.Emplace( root, initialState );
 
 
             struct custom_partial_hash
@@ -745,29 +749,28 @@ namespace mu
                     return std::hash<const void*>()(k.first.get());
                 }
             };
+            TSet<TPair<Ptr<ASTOp>,STATE>, FKeyFuncs<STATE>> Traversed;
 
-            std::unordered_set<std::pair<Ptr<ASTOp>,STATE>, custom_partial_hash> traversed;
-
-            while (!pending.IsEmpty())
+            while (!Pending.IsEmpty())
             {
-				TPair<Ptr<ASTOp>,STATE> current = pending.Pop();
+				TPair<Ptr<ASTOp>,STATE> current = Pending.Pop();
 
                 // It could have been completed in another branch
-				auto iti = traversed.insert({ current.Key,current.Value });
-                if (iti.second)
+				bool bAlreadyAdded = false;
+				Traversed.Add({ current.Key,current.Value }, &bAlreadyAdded);
+                if (!bAlreadyAdded)
                 {
                     // Process. State in current.Value may change
-                    bool recurse = f(current.Key,current.Value,pending);
+                    bool bRecurse = f(current.Key,current.Value,Pending);
 
                     // Recurse children
-                    if (recurse)
+                    if (bRecurse)
                     {
                         current.Key->ForEachChild([&]( ASTChild& c )
                         {
-							std::pair<Ptr<ASTOp>, STATE> TraverseKey(c.m_child, current.Value);
-							if (c.m_child && !traversed.count(TraverseKey))
+							if (c.Child && !Traversed.Contains(c.Child))
                             {
-								pending.Emplace( c.m_child, current.Value );
+								Pending.Emplace( c.Child, current.Value );
                             }
                         });
                     }
@@ -936,80 +939,80 @@ namespace mu
     private:
 
         //! States found so far
-        TArray<STATE> m_states;
+        TArray<STATE> States;
 
-        //! Index of the current state, from the m_states TArray.
-        int m_currentState;
+        //! Index of the current state, from the States TArray.
+        int32 CurrentState;
 
-        struct PENDING
+        struct FPending
         {
-            PENDING()
+			FPending()
             {
-                stateIndex = 0;
+                StateIndex = 0;
             }
-            PENDING(Ptr<ASTOp> _at, int _stateIndex)
+			FPending(Ptr<ASTOp> InAt, int32 InStateIndex)
             {
-                at = _at;
-                stateIndex = _stateIndex;
+                At = InAt;
+                StateIndex = InStateIndex;
             }
 
-            Ptr<ASTOp> at;
-            int stateIndex;
+            Ptr<ASTOp> At;
+            int32 StateIndex;
         };
 
-        TArray<PENDING> m_pending;
+        TArray<FPending> Pending;
 
-        //! List of traversed nodes with the state in which they were traversed.
-        std::multimap<Ptr<ASTOp>,int> m_traversed;
+        //! List of Traversed nodes with the state in which they were Traversed.
+        TMultiMap<Ptr<ASTOp>,int32> Traversed;
 
     protected:
 
         //!
         const STATE& GetCurrentState() const
         {
-            return m_states[m_currentState];
+            return States[CurrentState];
         }
 
         //!
         STATE GetDefaultState() const
         {
-            return m_states[0];
+            return States[0];
         }
 
         //! For manual recursion that changes the state for a specific path.
-        void RecurseWithState(const Ptr<ASTOp>& at, const STATE& newState)
+        void RecurseWithState(const Ptr<ASTOp>& InAt, const STATE& NewState)
         {
-            if(at)
+            if(InAt)
             {
-                auto it = m_states.Find(newState);
-                if (it==INDEX_NONE)
+				int32 Index = States.Find(NewState);
+                if (Index==INDEX_NONE)
                 {
-                    m_states.Add(newState);
+                    States.Add(NewState);
                 }
-                int stateIndex = m_states.Find(newState);
+                int32 StateIndex = States.Find(NewState);
 
-                m_pending.Add( PENDING(at,stateIndex) );
+                Pending.Add( FPending(InAt,StateIndex) );
             }
         }
 
         //! For manual recursion that doesn't change the state for a specific path.
-        void RecurseWithCurrentState(const Ptr<ASTOp>& at)
+        void RecurseWithCurrentState(const Ptr<ASTOp>& InAt)
         {
-            if(at)
+            if(InAt)
             {
-                m_pending.Emplace( at, m_currentState );
+                Pending.Emplace(InAt, CurrentState );
             }
         }
 
         //! Can be called from visit to set the state to visit all children ops
-        void SetCurrentState(const STATE& newState)
+        void SetCurrentState(const STATE& NewState)
         {
-            auto it = m_states.Find(newState);
-            if (it==INDEX_NONE)
+			int32 Index = States.Find(NewState);
+            if (Index == INDEX_NONE)
             {
-                m_states.Add(newState);
+                States.Add(NewState);
             }
-            m_currentState = m_states.Find(newState);
+            CurrentState = States.Find(NewState);
         }
 
     public:
@@ -1018,53 +1021,54 @@ namespace mu
         virtual ~Visitor_TopDown_Unique_Const() {}
 
         //!
-        void Traverse(const ASTOpList& roots, const STATE& initialState)
+        void Traverse(const ASTOpList& Roots, const STATE& InitialState)
         {
-            m_pending.Empty();
-            m_states.Empty();
-            m_states.Add(initialState);
-            m_currentState = 0;
+            Pending.Empty();
+            States.Empty();
+            States.Add(InitialState);
+            CurrentState = 0;
 
-            for( const auto& r: roots)
+            for( const Ptr<ASTOp>& Root: Roots)
             {
-                if (r)
+                if (Root)
                 {
-                    m_pending.Add( PENDING(r,m_currentState) );
+                    Pending.Add(FPending(Root,CurrentState) );
                 }
             }
 
-            while (m_pending.Num())
+            while (Pending.Num())
             {
-                PENDING item = m_pending.Pop();
+                FPending Item = Pending.Pop();
 
-                auto thisNodeRange = m_traversed.equal_range(item.at);
-                bool visitedInThisState = false;
-                for (auto it=thisNodeRange.first; it!=thisNodeRange.second; ++it)
+				TArray<int32> Found;
+                Traversed.MultiFind(Item.At, Found);
+                bool bVisitedInThisState = false;
+                for (int32 Value: Found)
                 {
-                    if (it->second==item.stateIndex)
+                    if (Value==Item.StateIndex)
                     {
-                        visitedInThisState = true;
+                        bVisitedInThisState = true;
                         break;
                     }
                 }
 
                 // It could have been completed in another branch
-                if (!visitedInThisState)
+                if (!bVisitedInThisState)
                 {
-                    m_traversed.insert( std::make_pair(item.at,item.stateIndex) );
+					Traversed.Add( Item.At, Item.StateIndex);
 
                     // Process
-                    m_currentState = item.stateIndex;
-                    bool recurse = Visit(item.at);
+                    CurrentState = Item.StateIndex;
+                    bool bRecurse = Visit(Item.At);
 
                     // Recurse children
-                    if (recurse)
+                    if (bRecurse)
                     {
-                        item.at->ForEachChild([&]( ASTChild& c )
+                        Item.At->ForEachChild([&]( ASTChild& c )
                         {
                             if (c)
                             {
-                                m_pending.Add( PENDING(c.child(),m_currentState) );
+                                Pending.Add( FPending(c.child(), CurrentState) );
                             }
                         });
                     }
@@ -1097,32 +1101,32 @@ namespace mu
     private:
 
         //! Operations to be processed
-        TArray< TPair<bool,Ptr<ASTOp>> > m_pending;
+        TArray< TPair<bool,Ptr<ASTOp>> > Pending;
 
         //! Map for visited operations
-        std::unordered_map<Ptr<ASTOp>,Ptr<ASTOp>> m_oldToNew;
+        TMap<Ptr<ASTOp>,Ptr<ASTOp>> OldToNew;
 
         //!
-        inline Ptr<ASTOp> GetOldToNew( const Ptr<ASTOp>& old ) const
+        inline Ptr<ASTOp> GetOldToNew( const Ptr<ASTOp>& Old ) const
         {
             Ptr<ASTOp> n;
-            auto it = m_oldToNew.find(old);
-            if (it!=m_oldToNew.end())
+			const Ptr<ASTOp>* it = OldToNew.Find(Old);
+            if (it)
             {
-                n = it->second;
+                n = *it;
             }
 
-            it = m_oldToNew.find(n);
-            while (it!=m_oldToNew.end() && it->second.get()!=nullptr && it->second!=n)
+            it = OldToNew.Find(n);
+            while (it && it->get()!=nullptr && *it!=n)
             {
-                n = it->second;
-                it = m_oldToNew.find(n);
+                n = *it;
+                it = OldToNew.Find(n);
             }
 
             return n;
         }
 
-        //! Process all the pending operations and visit all children if necessary
+        //! Process all the Pending operations and visit all children if necessary
         void Process();
 
     };
@@ -1148,7 +1152,7 @@ namespace mu
 
         ~ASTOpFixed() override;
 
-		virtual OP_TYPE GetOpType() const override { return op.type; }
+		virtual EOpType GetOpType() const override { return op.type; }
 
 		virtual Ptr<ASTOp> Clone( MapChildFuncRef mapChild ) const override;
 		virtual void ForEachChild( const TFunctionRef<void(ASTChild&)> f ) override;
@@ -1245,17 +1249,17 @@ namespace mu
     //---------------------------------------------------------------------------------------------
     inline void ASTChild::AddParent()
     {
-        m_parentIndexInChild = m_child->m_parents.Num();
-        m_child->m_parents.Add(m_parent);
+        ParentIndexInChild = Child->m_parents.Num();
+        Child->m_parents.Add(Parent);
     }
 
     //---------------------------------------------------------------------------------------------
     inline void ASTChild::ClearParent()
     {
-        check( m_parentIndexInChild<m_child->m_parents.Num() );
+        check( ParentIndexInChild<Child->m_parents.Num() );
 		// Can't do this, because the indices are stored in children.
-		//m_child->m_parents.RemoveAtSwap(m_parentIndexInChild);
-		m_child->m_parents[m_parentIndexInChild] = nullptr;
+		//Child->m_parents.RemoveAtSwap(ParentIndexInChild);
+		Child->m_parents[ParentIndexInChild] = nullptr;
 	}
 
 
@@ -1281,29 +1285,34 @@ namespace mu
             }
         };
 
+		struct FKeyFuncs : BaseKeyFuncs<Ptr<ASTOp>, Ptr<ASTOp>, false>
+		{
+			static KeyInitType GetSetKey(ElementInitType Element) { return Element; }
+			static bool Matches(const Ptr<ASTOp>& lhs, const Ptr<ASTOp>& rhs) { return lhs==rhs || *lhs == *rhs; }
+			static uint32 GetKeyHash(const Ptr<ASTOp>& Key) { return Key->Hash(); }
+		};
+
         // Existing ops, per type
-        std::unordered_set<Ptr<ASTOp>,op_pointer_hash,op_pointer_equal> visited[(int)OP_TYPE::COUNT];
+        TSet<Ptr<ASTOp>, FKeyFuncs> Visited[(int32)EOpType::COUNT];
 
     public:
 
-        bool m_disabled = false;
+        bool bDisabled = false;
 
-        Ptr<ASTOp> Add( const Ptr<ASTOp>& n )
+        Ptr<ASTOp> Add( const Ptr<ASTOp>& Op )
         {
-            if (m_disabled) return n;
+			if (bDisabled)
+			{
+				return Op;
+			}
 
-            if (!n) return nullptr;
+			if (!Op)
+			{
+				return nullptr;
+			}
 
-            auto& container = visited[(int)n->GetOpType()];
-            auto it = container.insert(n);
-            if( !it.second )
-            {
-                return *it.first;
-            }
-            else
-            {
-                return n;
-            }
+			TSet<Ptr<ASTOp>, FKeyFuncs>& Container = Visited[(int32)Op->GetOpType()];
+            return Container.FindOrAdd(Op);
         }
     };
 
diff --git a/Engine/Plugins/Mutable/Source/MutableTools/Internal/MuT/ErrorLogPrivate.h b/Engine/Plugins/Mutable/Source/MutableTools/Internal/MuT/ErrorLogPrivate.h
--- a/Engine/Plugins/Mutable/Source/MutableTools/Internal/MuT/ErrorLogPrivate.h
+++ b/Engine/Plugins/Mutable/Source/MutableTools/Internal/MuT/ErrorLogPrivate.h
@@ -40,7 +40,7 @@ namespace mu
 	};
 
 
-    extern const TCHAR* GetOpName( OP_TYPE type );
+    extern const TCHAR* GetOpName( EOpType type );
 
 }
 
diff --git a/Engine/Plugins/Mutable/Source/MutableTools/Internal/MuT/NodeImageLayerColourPrivate.h b/Engine/Plugins/Mutable/Source/MutableTools/Internal/MuT/NodeImageLayerColourPrivate.h
--- a/Engine/Plugins/Mutable/Source/MutableTools/Internal/MuT/NodeImageLayerColourPrivate.h
+++ b/Engine/Plugins/Mutable/Source/MutableTools/Internal/MuT/NodeImageLayerColourPrivate.h
@@ -21,7 +21,7 @@ namespace mu
 		NodeImagePtr m_pBase;
 		NodeImagePtr m_pMask;
 		NodeColourPtr m_pColour;
-		EBlendType m_type;
+		EBlendType Type;
 	};
 
 }
diff --git a/Engine/Plugins/Mutable/Source/MutableTools/Internal/MuT/NodeImageLayerPrivate.h b/Engine/Plugins/Mutable/Source/MutableTools/Internal/MuT/NodeImageLayerPrivate.h
--- a/Engine/Plugins/Mutable/Source/MutableTools/Internal/MuT/NodeImageLayerPrivate.h
+++ b/Engine/Plugins/Mutable/Source/MutableTools/Internal/MuT/NodeImageLayerPrivate.h
@@ -19,7 +19,7 @@ namespace mu
 		NodeImagePtr m_pBase;
 		NodeImagePtr m_pMask;
 		NodeImagePtr m_pBlended;
-		EBlendType m_type;
+		EBlendType Type;
 	};
 
 }
diff --git a/Engine/Plugins/Mutable/Source/MutableTools/Internal/MuT/NodeImageSwitchPrivate.h b/Engine/Plugins/Mutable/Source/MutableTools/Internal/MuT/NodeImageSwitchPrivate.h
--- a/Engine/Plugins/Mutable/Source/MutableTools/Internal/MuT/NodeImageSwitchPrivate.h
+++ b/Engine/Plugins/Mutable/Source/MutableTools/Internal/MuT/NodeImageSwitchPrivate.h
@@ -18,7 +18,7 @@ namespace mu
 		static FNodeType s_type;
 
 		NodeScalarPtr m_pParameter;
-		TArray<NodeImagePtr> m_options;
+		TArray<NodeImagePtr> Options;
 	};
 
 
diff --git a/Engine/Plugins/Mutable/Source/MutableTools/Internal/MuT/NodeMeshConstantPrivate.h b/Engine/Plugins/Mutable/Source/MutableTools/Internal/MuT/NodeMeshConstantPrivate.h
--- a/Engine/Plugins/Mutable/Source/MutableTools/Internal/MuT/NodeMeshConstantPrivate.h
+++ b/Engine/Plugins/Mutable/Source/MutableTools/Internal/MuT/NodeMeshConstantPrivate.h
@@ -17,14 +17,14 @@ namespace mu
 
 		static FNodeType s_type;
 
-		Ptr<Mesh> Value;
+		TSharedPtr<FMesh> Value;
 
 		TArray<Ptr<NodeLayout>> Layouts;
 
 		struct FMorph
 		{
 			FString Name;
-			Ptr<Mesh> MorphedMesh;
+			TSharedPtr<FMesh> MorphedMesh;
 		};
 
 		TArray<FMorph> Morphs;
diff --git a/Engine/Plugins/Mutable/Source/MutableTools/Internal/MuT/NodeObjectGroupPrivate.h b/Engine/Plugins/Mutable/Source/MutableTools/Internal/MuT/NodeObjectGroupPrivate.h
--- a/Engine/Plugins/Mutable/Source/MutableTools/Internal/MuT/NodeObjectGroupPrivate.h
+++ b/Engine/Plugins/Mutable/Source/MutableTools/Internal/MuT/NodeObjectGroupPrivate.h
@@ -22,7 +22,7 @@ namespace mu
 		FString Name;
 		FString Uid;
 
-		CHILD_SELECTION m_type;
+		CHILD_SELECTION Type;
 
 		//! Set the child selection type
 		void SetSelectionType( CHILD_SELECTION );
diff --git a/Engine/Plugins/Mutable/Source/MutableTools/Internal/MuT/NodeProjectorPrivate.h b/Engine/Plugins/Mutable/Source/MutableTools/Internal/MuT/NodeProjectorPrivate.h
--- a/Engine/Plugins/Mutable/Source/MutableTools/Internal/MuT/NodeProjectorPrivate.h
+++ b/Engine/Plugins/Mutable/Source/MutableTools/Internal/MuT/NodeProjectorPrivate.h
@@ -25,7 +25,7 @@ namespace mu
 
 		static FNodeType s_type;
 
-        EProjectorType m_type = EProjectorType::Planar;
+        EProjectorType Type = EProjectorType::Planar;
 		FVector3f m_position;
 		FVector3f m_direction;
 		FVector3f m_up;
@@ -40,10 +40,10 @@ namespace mu
 
 		static FNodeType s_type;
 
-		FString m_name;
-		FString m_uid;
+		FString Name;
+		FString UID;
 
-		TArray<Ptr<NodeRange>> m_ranges;
+		TArray<Ptr<NodeRange>> Ranges;
 	};
 
 
diff --git a/Engine/Plugins/Mutable/Source/MutableTools/Internal/MuT/NodeScalarEnumParameterPrivate.h b/Engine/Plugins/Mutable/Source/MutableTools/Internal/MuT/NodeScalarEnumParameterPrivate.h
--- a/Engine/Plugins/Mutable/Source/MutableTools/Internal/MuT/NodeScalarEnumParameterPrivate.h
+++ b/Engine/Plugins/Mutable/Source/MutableTools/Internal/MuT/NodeScalarEnumParameterPrivate.h
@@ -19,9 +19,9 @@ namespace mu
 
 		static FNodeType s_type;
 
-		FString m_name;
-		FString m_uid;
-		int32 m_defaultValue = 0;
+		FString Name;
+		FString UID;
+		int32 DefaultValue = 0;
 
 		struct FOption
 		{
@@ -29,9 +29,9 @@ namespace mu
 			float value;
 		};
 
-		TArray<FOption> m_options;
+		TArray<FOption> Options;
 
-		TArray<Ptr<NodeRange>> m_ranges;
+		TArray<Ptr<NodeRange>> Ranges;
 	};
 
 }
diff --git a/Engine/Plugins/Mutable/Source/MutableTools/Internal/MuT/NodeScalarParameterPrivate.h b/Engine/Plugins/Mutable/Source/MutableTools/Internal/MuT/NodeScalarParameterPrivate.h
--- a/Engine/Plugins/Mutable/Source/MutableTools/Internal/MuT/NodeScalarParameterPrivate.h
+++ b/Engine/Plugins/Mutable/Source/MutableTools/Internal/MuT/NodeScalarParameterPrivate.h
@@ -20,11 +20,11 @@ namespace mu
 
 		static FNodeType s_type;
 
-		float m_defaultValue = 0.0f;
-		FString m_name;
-		FString m_uid;
+		float DefaultValue = 0.0f;
+		FString Name;
+		FString UID;
 
-		TArray<Ptr<NodeRange>> m_ranges;
+		TArray<Ptr<NodeRange>> Ranges;
 	};
 
 }
diff --git a/Engine/Plugins/Mutable/Source/MutableTools/Internal/MuT/TablePrivate.h b/Engine/Plugins/Mutable/Source/MutableTools/Internal/MuT/TablePrivate.h
--- a/Engine/Plugins/Mutable/Source/MutableTools/Internal/MuT/TablePrivate.h
+++ b/Engine/Plugins/Mutable/Source/MutableTools/Internal/MuT/TablePrivate.h
@@ -26,8 +26,8 @@ namespace mu
 		// TODO: Union
 		float Scalar;
 		FVector4f Color;
-		Ptr<ResourceProxy<Image>> ProxyImage;
-		Ptr<Mesh> Mesh;
+		Ptr<TResourceProxy<FImage>> ProxyImage;
+		TSharedPtr<FMesh> Mesh;
 		FString String;
 
 		const void* ErrorContext;
@@ -42,7 +42,7 @@ namespace mu
 
 
 	//!
-	class Table::Private
+	class FTable::Private
 	{
 	public:
 
