[mutable] More enforcement of coding standards.

diff --git a/Engine/Plugins/Mutable/Source/MutableTools/Private/MuT/AST.cpp b/Engine/Plugins/Mutable/Source/MutableTools/Private/MuT/AST.cpp
--- a/Engine/Plugins/Mutable/Source/MutableTools/Private/MuT/AST.cpp
+++ b/Engine/Plugins/Mutable/Source/MutableTools/Private/MuT/AST.cpp
@@ -13,7 +13,6 @@
 #include "MuT/ASTOpMeshRemoveMask.h"
 #include "MuT/ASTOpImageMultiLayer.h"
 #include "MuT/ASTOpParameter.h"
-#include "MuT/Platform.h"
 #include "MuT/ErrorLogPrivate.h"
 #include "Trace/Detail/Channel.h"
 
@@ -30,10 +29,10 @@ namespace mu
 //-------------------------------------------------------------------------------------------------
 //-------------------------------------------------------------------------------------------------
 ASTChild::ASTChild(ASTOp* p, const Ptr<ASTOp>& c)
-    : m_parent(p)
-    , m_child(c)
+    : Parent(p)
+    , Child(c)
 {
-    if (m_parent && m_child.get())
+    if (Parent && Child.get())
     {
         AddParent();
     }
@@ -46,7 +45,7 @@ ASTChild::ASTChild( const Ptr<ASTOp>& p, const Ptr<ASTOp>& c)
 
 ASTChild::~ASTChild()
 {
-    if (m_child && m_parent)
+    if (Child && Parent)
     {
         ClearParent();
     }
@@ -54,16 +53,16 @@ ASTChild::~ASTChild()
 
 ASTChild& ASTChild::operator=( const Ptr<ASTOp>& c )
 {
-    if (c!=m_child)
+    if (c!=Child)
     {
-        if (m_child && m_parent)
+        if (Child && Parent)
         {
             ClearParent();
         }
 
-        m_child = c;
+        Child = c;
 
-        if (m_child && m_parent)
+        if (Child && Parent)
         {
             AddParent();
         }
@@ -74,11 +73,11 @@ ASTChild& ASTChild::operator=( const Ptr<ASTOp>& c )
 
 ASTChild& ASTChild::operator=( ASTChild&& rhs )
 {
-    m_parent = rhs.m_parent;
-    m_parentIndexInChild = rhs.m_parentIndexInChild;
-    m_child = rhs.m_child;
-    rhs.m_parent=nullptr;
-    rhs.m_child.reset();
+    Parent = rhs.Parent;
+    ParentIndexInChild = rhs.ParentIndexInChild;
+    Child = rhs.Child;
+    rhs.Parent=nullptr;
+    rhs.Child.reset();
 
     return *this;
 }
@@ -145,7 +144,7 @@ void ASTOp::Assert()
             bool foundInParent = false;
             parent->ForEachChild([&](ASTChild&c)
             {
-                if (c && c.m_child.get()==this)
+                if (c && c.Child.get()==this)
                 {
                     foundInParent = true;
                 }
@@ -171,8 +170,8 @@ void ASTOp::Assert()
 //                }
 //            });
 //            check(found);
-            check( c.m_parentIndexInChild < size_t(c.child()->m_parents.Num()) );
-            check( c.child()->m_parents[c.m_parentIndexInChild]==this );
+            check( c.ParentIndexInChild < size_t(c.child()->m_parents.Num()) );
+            check( c.child()->m_parents[c.ParentIndexInChild]==this );
         }
      });
 }
@@ -218,7 +217,7 @@ mu::Ptr<ASTOp> ASTOp::DeepClone( const Ptr<ASTOp>& root )
 {
     MUTABLE_CPUPROFILER_SCOPE(AST_DeepClone);
 
-    std::unordered_map<Ptr<const ASTOp>,Ptr<ASTOp>> visited;
+    TMap<Ptr<const ASTOp>,Ptr<ASTOp>> Visited;
 
     MapChildFunc m = [&](const Ptr<ASTOp>&n)
     {
@@ -226,21 +225,21 @@ mu::Ptr<ASTOp> ASTOp::DeepClone( const Ptr<ASTOp>& root )
 		{
 			return Ptr<ASTOp>();
 		}
-		std::unordered_map<Ptr<const ASTOp>, Ptr<ASTOp>>::const_iterator it = visited.find(n);
-        check(it!=visited.end());
-        return it->second;
+		const Ptr<ASTOp>* it = Visited.Find(n);
+        check(it);
+        return *it;
     };
 
     Ptr<ASTOp> r = root;
-    Traverse_BottomUp_Unique( r, [&](Ptr<ASTOp> n)
+    Traverse_BottomUp_Unique( r, [&](Ptr<ASTOp> Op)
     {
-        Ptr<ASTOp> cloned = n->Clone( m );
-        visited[n] = cloned;
+        Ptr<ASTOp> Cloned = Op->Clone( m );
+        Visited.Add(Op, Cloned);
     });
 
-	std::unordered_map<Ptr<const ASTOp>, Ptr<ASTOp>>::const_iterator it = visited.find(r);
-    check(it!=visited.end());
-    return it->second;
+	const Ptr<ASTOp>* it = Visited.Find(r);
+    check(it);
+    return *it;
 }
 
 
@@ -267,7 +266,7 @@ void ASTOp::LogHistogram( ASTOpList& roots )
 {
     (void)roots;
 
- //   uint64 CountPerType[(int)OP_TYPE::COUNT];
+ //   uint64 CountPerType[(int)EOpType::COUNT];
  //   FMemory::Memzero(CountPerType,sizeof(CountPerType));
 
  //   size_t count=0;
@@ -280,16 +279,16 @@ void ASTOp::LogHistogram( ASTOpList& roots )
  //       return true;
  //   });
 
-	//TArray< TPair<uint64, OP_TYPE> > Sorted;
-	//Sorted.SetNum((int)OP_TYPE::COUNT);
- //   for (int i=0; i<(int)OP_TYPE::COUNT; ++i)
+	//TArray< TPair<uint64, EOpType> > Sorted;
+	//Sorted.SetNum((int)EOpType::COUNT);
+ //   for (int i=0; i<(int)EOpType::COUNT; ++i)
  //   {
- //       Sorted[i].Value = (OP_TYPE)i;
+ //       Sorted[i].Value = (EOpType)i;
  //       Sorted[i].Key = CountPerType[i];
  //   }
 
 
-	//Sorted.Sort( [](const TPair<uint64, OP_TYPE>& a, const TPair<uint64, OP_TYPE>& b)
+	//Sorted.Sort( [](const TPair<uint64, EOpType>& a, const TPair<uint64, EOpType>& b)
  //   {
  //       return a.Key >b.Key;
  //   });
@@ -306,7 +305,7 @@ void ASTOp::LogHistogram( ASTOpList& roots )
 	//Traverse_TopDown_Unique( roots,
 	//	[&](const Ptr<ASTOp>& n)
 	//	{
-	//		if (n->GetOpType() == OP_TYPE::IM_MULTILAYER)
+	//		if (n->GetOpType() == EOpType::IM_MULTILAYER)
 	//		{
 	//			Ptr<const ASTOpImageMultiLayer> Typed = static_cast<const ASTOpImageMultiLayer*>(n.get());
 	//			Ptr<const ASTOp> Base = Typed->base.child();
@@ -317,7 +316,7 @@ void ASTOp::LogHistogram( ASTOpList& roots )
 	//			UE_LOG(LogMutableCore, Log, TEXT("    base is type %s:"), s_opNames[(int)Base->GetOpType()]);
 	//			if (Constant)
 	//			{
-	//				int Format = int( ((mu::Image*)Constant->GetValue().get())->GetFormat() );
+	//				int Format = int( ((mu::FImage*)Constant->GetValue().get())->GetFormat() );
 	//				UE_LOG(LogMutableCore, Log, TEXT("    constant image format is %d:"), Format);
 	//			}
 	//			else
@@ -396,9 +395,9 @@ void ASTOp::Traverse_TopDown_Unique( const TArray<Ptr<ASTOp>>& roots, TFunctionR
             {
                 pCurrent->ForEachChild([&]( ASTChild& c )
                 {
-                    if (c && !traversed.Contains(c.m_child))
+                    if (c && !traversed.Contains(c.Child))
                     {
-                        pending.Enqueue( c.m_child );
+                        pending.Enqueue( c.Child );
                     }
                 });
             }
@@ -436,9 +435,9 @@ void ASTOp::Traverse_TopDown_Unique_Imprecise( const TArray<Ptr<ASTOp>>& roots,
             {
                 pCurrent->ForEachChild([&]( ASTChild& c )
                 {
-                    if (c && !traversed.Contains(c.m_child))
+                    if (c && !traversed.Contains(c.Child))
                     {
-                        pending.Enqueue( c.m_child );
+                        pending.Enqueue( c.Child );
                     }
                 });
             }
@@ -485,9 +484,9 @@ void ASTOp::Traverse_TopRandom_Unique_NonReentrant( const TArray<Ptr<ASTOp>>& ro
             {
                 pCurrent->ForEachChild([&]( ASTChild& c )
                 {
-                    if (c && c.m_child->m_traverseIndex!=traverseIndex)
+                    if (c && c.Child->m_traverseIndex!=traverseIndex)
                     {
-                        pending.Add( c.m_child );
+                        pending.Add( c.Child );
                     }
                 });
             }
@@ -500,26 +499,25 @@ void ASTOp::Traverse_TopRandom_Unique_NonReentrant( const TArray<Ptr<ASTOp>>& ro
 //-------------------------------------------------------------------------------------------------
 //-------------------------------------------------------------------------------------------------
 //-------------------------------------------------------------------------------------------------
-void Visitor_TopDown_Unique_Cloning::Traverse( Ptr<ASTOp>& root )
+void Visitor_TopDown_Unique_Cloning::Traverse( Ptr<ASTOp>& InOutRoot )
 {
     // Visit the given root
-    if (root)
+    if (InOutRoot)
     {
-		m_pending.Add({ false,root });
+		Pending.Add({ false,InOutRoot });
 
         Process();
 
-        root = GetOldToNew(root);
+		InOutRoot = GetOldToNew(InOutRoot);
     }
 }
 
 
-
 void Visitor_TopDown_Unique_Cloning::Process()
 {
-    while ( m_pending.Num() )
+    while ( Pending.Num() )
     {
-        TPair<bool,Ptr<ASTOp>> item = m_pending.Pop();
+        TPair<bool,Ptr<ASTOp>> item = Pending.Pop();
 
         Ptr<ASTOp> at = item.Value;
 
@@ -534,12 +532,12 @@ void Visitor_TopDown_Unique_Cloning::Process()
             bool childChanged = false;
             cop->ForEachChild( [&](ASTChild& ref)
             {
-				std::unordered_map<Ptr<ASTOp>, Ptr<ASTOp>>::iterator it = m_oldToNew.find(ref.m_child);
-                if ( ref && it!=m_oldToNew.end() && it->second.get()!=nullptr )
+				const Ptr<ASTOp>* it = OldToNew.Find(ref.Child);
+                if ( ref && it && it->get()!=nullptr )
                 {
-                    auto oldRef = ref.m_child;
-                    ref=GetOldToNew(ref.m_child);
-                    if (ref.m_child!=oldRef)
+                    auto oldRef = ref.Child;
+                    ref=GetOldToNew(ref.Child);
+                    if (ref.Child!=oldRef)
                     {
                         childChanged = true;
                     }
@@ -549,14 +547,14 @@ void Visitor_TopDown_Unique_Cloning::Process()
             // If any child changed, we need to replace this instruction
             if ( childChanged )
             {
-                m_oldToNew[at]=cop;
+                OldToNew.Add(at,cop);
             }
 
         }
         else
         {
-			std::unordered_map<Ptr<ASTOp>, Ptr<ASTOp>>::iterator it = m_oldToNew.find(at);
-            if (it==m_oldToNew.end())
+			const Ptr<ASTOp>* it = OldToNew.Find(at);
+            if (!it)
             {
 				Ptr<ASTOp> initialAt = at;
 
@@ -566,12 +564,12 @@ void Visitor_TopDown_Unique_Cloning::Process()
                     bool childChanged = false;
                     cop->ForEachChild( [&](ASTChild& ref)
                     {
-						std::unordered_map<Ptr<ASTOp>, Ptr<ASTOp>>::iterator ito = m_oldToNew.find(ref.m_child);
-                        if ( ref && ito!=m_oldToNew.end() && ito->second.get()!=nullptr )
+						const Ptr<ASTOp>* ito = OldToNew.Find(ref.Child);
+                        if ( ref && ito && ito->get()!=nullptr )
                         {
-							Ptr<ASTOp> oldRef = ref.m_child;
-                            ref=GetOldToNew(ref.m_child);
-                            if (ref.m_child!=oldRef)
+							Ptr<ASTOp> oldRef = ref.Child;
+                            ref=GetOldToNew(ref.Child);
+                            if (ref.Child!=oldRef)
                             {
                                 childChanged = true;
                             }
@@ -581,30 +579,25 @@ void Visitor_TopDown_Unique_Cloning::Process()
                     // If any child changed, we need to re-add this instruction
                     if ( childChanged )
                     {
-                        m_oldToNew[at]=cop;
-                        at = cop;
+						OldToNew.Add(at, cop);
+						at = cop;
                     }
                 }
 
-                //Ptr<ASTOp> test1 = at->Clone();
-                //check(*test1==*at);
-
                 bool processChildren = true;
 				Ptr<ASTOp> newAt = Visit( at, processChildren );
-                m_oldToNew[initialAt]=newAt;
-
-                //check(*test1==*at);
+				OldToNew.Add(initialAt,newAt);
 
                 // Proceed with children
                 if (processChildren)
                 {
-					m_pending.Add({ true, newAt });
+					Pending.Add({ true, newAt });
 
                     at->ForEachChild( [&](ASTChild& ref)
                     {
-                        if (ref && m_oldToNew.find(ref.m_child)==m_oldToNew.end())
+                        if (ref && !OldToNew.Contains(ref.Child))
                         {
-							m_pending.Add({ false,ref.m_child });
+							Pending.Add({ false,ref.Child });
                         }
                     });
                 }
@@ -657,9 +650,9 @@ void ASTOp::Traverse_BottomUp_Unique_NonReentrant
                 // ...after the children are processed
                 pCurrent->ForEachChild([&]( ASTChild& c )
                 {
-                    if (c && c.m_child->m_traverseIndex!=traverseIndex )
+                    if (c && c.Child->m_traverseIndex!=traverseIndex )
                     {
-						std::pair<Ptr<ASTOp>, int32> e = std::make_pair<>(c.m_child,0);
+						std::pair<Ptr<ASTOp>, int32> e = std::make_pair<>(c.Child,0);
                         pending.Add( e );
                     }
                 });
@@ -719,9 +712,9 @@ void ASTOp::Traverse_BottomUp_Unique_NonReentrant
                 // ...after the children are processed
                 pCurrent->ForEachChild([&]( ASTChild& c )
                 {
-                    if (c && accept(c.m_child.get()) && c.m_child->m_traverseIndex!=traverseIndex )
+                    if (c && accept(c.Child.get()) && c.Child->m_traverseIndex!=traverseIndex )
                     {
-                        pending.Add( std::make_pair<>(c.m_child,0) );
+                        pending.Add( std::make_pair<>(c.Child,0) );
                     }
                 });
             }
@@ -781,9 +774,9 @@ void ASTOp::Traverse_BottomUp_Unique
                 // ...after the children are processed
                 pCurrent->ForEachChild([&]( ASTChild& c )
                 {
-                    if (c && accept(c.m_child.get()) && !Traversed.Find(c.m_child) )
+                    if (c && accept(c.Child.get()) && !Traversed.Find(c.Child) )
                     {
-                        Pending.Add( std::make_pair<>(c.m_child,0) );
+                        Pending.Add( std::make_pair<>(c.Child,0) );
                     }
                 });
             }
@@ -848,7 +841,7 @@ void ASTOp::Replace( const Ptr<ASTOp>& node, const Ptr<ASTOp>& other )
         {
             p->ForEachChild( [=](ASTChild& c)
             {
-                if(c.m_child==node)
+                if(c.Child==node)
                 {
                     c = other;
                 }
@@ -931,27 +924,27 @@ void ASTOp::LinkRange(FProgram& program,
 	{
 		if (range.rangeSize->linkedRange < 0)
 		{
-			check(program.m_ranges.Num() < 255);
-			range.rangeSize->linkedRange = int8(program.m_ranges.Num());
+			check(program.Ranges.Num() < 255);
+			range.rangeSize->linkedRange = int8(program.Ranges.Num());
 			FRangeDesc rangeData;
-			rangeData.m_name = range.rangeName;
-			rangeData.m_uid = range.rangeUID;
+			rangeData.Name = range.rangeName;
+			rangeData.UID = range.rangeUID;
 
 			// Try to see if a parameter directly controls de size of the range. This is used
 			// to store hint data for instance generation in tools or randomizers that want to
 			// support multilayer, but it is not critical otherwise.
 			int32 EstimatedSizeParameter = -1;
-			if ( range.rangeSize->GetOpType() == OP_TYPE::SC_PARAMETER
+			if ( range.rangeSize->GetOpType() == EOpType::SC_PARAMETER
 				||
-				range.rangeSize->GetOpType() == OP_TYPE::NU_PARAMETER )
+				range.rangeSize->GetOpType() == EOpType::NU_PARAMETER )
 			{
 				const ASTOpParameter* ParamOp = static_cast<const ASTOpParameter*>(range.rangeSize.child().get());
 
 				EstimatedSizeParameter = ParamOp->LinkedParameterIndex;
 			}
-			rangeData.m_dimensionParameter = EstimatedSizeParameter;
+			rangeData.DimensionParameter = EstimatedSizeParameter;
 
-			program.m_ranges.Add(rangeData);
+			program.Ranges.Add(rangeData);
 		}
 		rangeSize = range.rangeSize->linkedAddress;
 		rangeId = uint16(range.rangeSize->linkedRange);
@@ -1003,7 +996,7 @@ void ASTOpFixed::Link( FProgram& program, FLinkerOptions* )
         // Fix the linked op and generate children
         ForEachReference(lop, [&](OP::ADDRESS* pChildOp)
         {
-            Ptr<ASTOp> pASTChild = children[*pChildOp].m_child;
+            Ptr<ASTOp> pASTChild = children[*pChildOp].Child;
             if (pASTChild)
             {
                 *pChildOp = pASTChild->linkedAddress;
@@ -1014,18 +1007,18 @@ void ASTOpFixed::Link( FProgram& program, FLinkerOptions* )
             }
         });
 
-        linkedAddress = (OP::ADDRESS)program.m_opAddress.Num();        
-        program.m_opAddress.Add((uint32)program.m_byteCode.Num());
+        linkedAddress = (OP::ADDRESS)program.OpAddress.Num();        
+        program.OpAddress.Add((uint32)program.ByteCode.Num());
 
-        AppendCode(program.m_byteCode,lop.type);
-        if (lop.type==OP_TYPE::ME_MERGE)
+        AppendCode(program.ByteCode,lop.type);
+        if (lop.type==EOpType::ME_MERGE)
         {
-            AppendCode(program.m_byteCode,lop.args.MeshMerge);
+            AppendCode(program.ByteCode,lop.args.MeshMerge);
         }
         else
         {
             // Generic encoding
-            AppendCode(program.m_byteCode,lop.args);
+            AppendCode(program.ByteCode,lop.args);
         }
     }
 }
@@ -1088,28 +1081,28 @@ FImageDesc ASTOpFixed::GetImageDesc( bool returnBestOption, FGetImageDescContext
         }
     }
 
-    OP_TYPE type = op.type;
+    EOpType type = op.type;
 
     switch ( type )
     {
 
-    case OP_TYPE::NONE:
+    case EOpType::NONE:
         break;
 
-    case OP_TYPE::IM_SATURATE:
+    case EOpType::IM_SATURATE:
         res = GetImageDesc( op.args.ImageSaturate.base, returnBestOption, context );
         break;
 
-    case OP_TYPE::IM_LUMINANCE:
+    case EOpType::IM_LUMINANCE:
         res = GetImageDesc( op.args.ImageLuminance.base,returnBestOption, context );
         res.m_format = EImageFormat::L_UByte;
         break;
 
-    case OP_TYPE::IM_INTERPOLATE:
+    case EOpType::IM_INTERPOLATE:
         res = GetImageDesc(op.args.ImageInterpolate.targets[0], returnBestOption, context );
         break;
 
-    case OP_TYPE::IM_PLAINCOLOUR:
+    case EOpType::IM_PLAINCOLOUR:
         res.m_format = EImageFormat(op.args.ImagePlainColour.format);
         res.m_size[0] = op.args.ImagePlainColour.size[0];
         res.m_size[1] = op.args.ImagePlainColour.size[1];
@@ -1117,7 +1110,7 @@ FImageDesc ASTOpFixed::GetImageDesc( bool returnBestOption, FGetImageDescContext
         check( res.m_format != EImageFormat::None );
         break;
 
-    case OP_TYPE::IM_RESIZE:
+    case EOpType::IM_RESIZE:
         res = GetImageDesc( op.args.ImageResize.source, returnBestOption, context );
 
         res.m_size = FImageSize
@@ -1127,13 +1120,13 @@ FImageDesc ASTOpFixed::GetImageDesc( bool returnBestOption, FGetImageDescContext
             );
         break;
 
-    case OP_TYPE::IM_RESIZEREL:
+    case EOpType::IM_RESIZEREL:
         res = GetImageDesc( op.args.ImageResizeRel.source, returnBestOption, context );
         res.m_size[0] = (uint16)( res.m_size[0]*op.args.ImageResizeRel.factor[0] );
         res.m_size[1] = (uint16)( res.m_size[1]*op.args.ImageResizeRel.factor[1] );
         break;
 
-    case OP_TYPE::IM_RESIZELIKE:
+    case EOpType::IM_RESIZELIKE:
         res = GetImageDesc( op.args.ImageResizeLike.source, returnBestOption, context );
         if (op.args.ImageResizeLike.sizeSource)
         {
@@ -1141,33 +1134,33 @@ FImageDesc ASTOpFixed::GetImageDesc( bool returnBestOption, FGetImageDescContext
         }
         break;
 
-    case OP_TYPE::IM_GRADIENT:
+    case EOpType::IM_GRADIENT:
         res.m_size[0] = op.args.ImageGradient.size[0];
         res.m_size[1] = op.args.ImageGradient.size[1];
         res.m_format = EImageFormat::RGB_UByte;
         break;
 
-    case OP_TYPE::IM_BLANKLAYOUT:
+    case EOpType::IM_BLANKLAYOUT:
         // TODO: We would need to process the layout to find the grid size, and then use
         // the block size with it.
         res.m_size = FImageSize(0,0);
         res.m_format = op.args.ImageBlankLayout.format;
         break;
 
-    case OP_TYPE::IM_BINARISE:
+    case EOpType::IM_BINARISE:
         res = GetImageDesc( op.args.ImageBinarise.base, returnBestOption, context );
         res.m_format = EImageFormat::L_UByte;
         break;
 
-    case OP_TYPE::IM_DISPLACE:
+    case EOpType::IM_DISPLACE:
         res = GetImageDesc( op.args.ImageDisplace.source, returnBestOption, context );
         break;
 
-	case OP_TYPE::IM_INVERT:
+	case EOpType::IM_INVERT:
 		res = GetImageDesc(op.args.ImageInvert.base, returnBestOption, context);
 		break;
 
-	case OP_TYPE::IM_COLOURMAP:
+	case EOpType::IM_COLOURMAP:
 		res = GetImageDesc( op.args.ImageColourMap.base, returnBestOption, context );
 		break;
 
@@ -1190,7 +1183,7 @@ void ASTOpFixed::GetLayoutBlockSize( int* pBlockX, int* pBlockY )
     switch ( op.type )
     {
     
-	case OP_TYPE::IM_RESIZEREL:
+	case EOpType::IM_RESIZEREL:
 	{
 		GetLayoutBlockSize(op.args.ImageResizeRel.source, pBlockX, pBlockY);
 		(*pBlockX) = int(float(*pBlockX) * op.args.ImageResizeRel.factor[0]);
@@ -1199,7 +1192,7 @@ void ASTOpFixed::GetLayoutBlockSize( int* pBlockX, int* pBlockY )
 		break;
 	}
 
-	case OP_TYPE::IM_RESIZE:
+	case EOpType::IM_RESIZE:
 	{
 		GetLayoutBlockSize(op.args.ImageResize.source, pBlockX, pBlockY);
 
@@ -1223,14 +1216,14 @@ void ASTOpFixed::GetLayoutBlockSize( int* pBlockX, int* pBlockY )
 		break;
 	}
 
-	case OP_TYPE::IM_BLANKLAYOUT:
+	case EOpType::IM_BLANKLAYOUT:
 	{
 		*pBlockX = int(op.args.ImageBlankLayout.blockSize[0]);
 		*pBlockY = int(op.args.ImageBlankLayout.blockSize[1]);
 		break;
 	}
 
-	case OP_TYPE::IM_PLAINCOLOUR:
+	case EOpType::IM_PLAINCOLOUR:
 	{
 		*pBlockX = 0;
 		*pBlockY = 0;
@@ -1265,7 +1258,7 @@ ASTOp::FBoolEvalResult ASTOpFixed::EvaluateBool( ASTOpList& facts, FEvaluateBool
     FBoolEvalResult result = BET_UNKNOWN;
     switch(op.type)
     {
-    case OP_TYPE::BO_NOT:
+    case EOpType::BO_NOT:
     {
         if (children[op.args.BoolNot.source])
         {
@@ -1282,7 +1275,7 @@ ASTOp::FBoolEvalResult ASTOpFixed::EvaluateBool( ASTOpList& facts, FEvaluateBool
         break;
     }
 
-    case OP_TYPE::BO_EQUAL_INT_CONST:
+    case EOpType::BO_EQUAL_INT_CONST:
     {
         int32 intValue = op.args.BoolEqualScalarConst.constant;
         const Ptr<ASTOp>& intExp = children[op.args.BoolEqualScalarConst.value].child();
@@ -1303,7 +1296,7 @@ ASTOp::FBoolEvalResult ASTOpFixed::EvaluateBool( ASTOpList& facts, FEvaluateBool
         break;
     }
 
-    case OP_TYPE::BO_AND:
+    case EOpType::BO_AND:
     {
         const Ptr<ASTOp>& a = children[op.args.BoolBinary.a].child();
         const Ptr<ASTOp>& b = children[op.args.BoolBinary.b].child();
@@ -1345,7 +1338,7 @@ ASTOp::FBoolEvalResult ASTOpFixed::EvaluateBool( ASTOpList& facts, FEvaluateBool
         break;
     }
 
-    case OP_TYPE::BO_OR:
+    case EOpType::BO_OR:
     {
         const Ptr<ASTOp>& a = children[op.args.BoolBinary.a].child();
         const Ptr<ASTOp>& b = children[op.args.BoolBinary.b].child();
@@ -1403,12 +1396,12 @@ int ASTOpFixed::EvaluateInt( ASTOpList& /*facts*/, bool &unknown ) const
 
     switch (GetOpType())
     {
-    case OP_TYPE::NU_CONSTANT:
+    case EOpType::NU_CONSTANT:
         unknown = false;
         return op.args.IntConstant.value;
         break;
 
-    case OP_TYPE::SC_CONSTANT:
+    case EOpType::SC_CONSTANT:
         unknown = false;
         return (int)op.args.ScalarConstant.value;
         break;
@@ -1437,7 +1430,7 @@ bool ASTOpFixed::IsImagePlainConstant(FVector4f& colour) const
     switch( op.type )
     {
 
-    case OP_TYPE::IM_BLANKLAYOUT:
+    case EOpType::IM_BLANKLAYOUT:
         res = true;
         colour[0] = 0;
         colour[1] = 0;
@@ -1445,19 +1438,19 @@ bool ASTOpFixed::IsImagePlainConstant(FVector4f& colour) const
         colour[3] = 0;
         break;
 
-    case OP_TYPE::IM_RESIZE:
+    case EOpType::IM_RESIZE:
         res = children[op.args.ImageResize.source]->IsImagePlainConstant( colour );
         break;
 
-    case OP_TYPE::IM_RESIZEREL:
+    case EOpType::IM_RESIZEREL:
         res = children[op.args.ImageResizeRel.source]->IsImagePlainConstant( colour );
         break;
 
-    case OP_TYPE::IM_RESIZELIKE:
+    case EOpType::IM_RESIZELIKE:
         res = children[op.args.ImageResizeLike.source]->IsImagePlainConstant(  colour );
         break;
 
-    case OP_TYPE::IM_PLAINCOLOUR:
+    case EOpType::IM_PLAINCOLOUR:
         res = children[op.args.ImagePlainColour.colour]->IsColourConstant( colour );
         break;
 
@@ -1478,7 +1471,7 @@ bool ASTOpFixed::IsColourConstant(FVector4f& colour) const
     bool res = false;
     switch( op.type )
     {
-    case OP_TYPE::CO_CONSTANT:
+    case EOpType::CO_CONSTANT:
         res = true;
         colour[0] = op.args.ColourConstant.value[0];
         colour[1] = op.args.ColourConstant.value[1];
@@ -1486,10 +1479,10 @@ bool ASTOpFixed::IsColourConstant(FVector4f& colour) const
         colour[3] = op.args.ColourConstant.value[3];
         break;
 
-    case OP_TYPE::CO_SAMPLEIMAGE:
-    case OP_TYPE::CO_SWIZZLE:
-    case OP_TYPE::CO_FROMSCALARS:
-    case OP_TYPE::CO_ARITHMETIC:
+    case EOpType::CO_SAMPLEIMAGE:
+    case EOpType::CO_SWIZZLE:
+    case EOpType::CO_FROMSCALARS:
+    case EOpType::CO_ARITHMETIC:
 
     default:
         // TODO: Improve this test with more operations
@@ -1507,10 +1500,10 @@ mu::Ptr<ImageSizeExpression> ASTOpFixed::GetImageSizeExpression() const
 {
     Ptr<ImageSizeExpression> pRes = new ImageSizeExpression;
 
-    OP_TYPE type = GetOpType();
+    EOpType type = GetOpType();
     switch ( type )
     {
-    case OP_TYPE::NONE:
+    case EOpType::NONE:
     {
         pRes->type = ImageSizeExpression::ISET_CONSTANT;
         pRes->size[0]=0;
@@ -1518,13 +1511,13 @@ mu::Ptr<ImageSizeExpression> ASTOpFixed::GetImageSizeExpression() const
         break;
     }
 
-    case OP_TYPE::IM_RESIZE:
+    case EOpType::IM_RESIZE:
         pRes->type = ImageSizeExpression::ISET_CONSTANT;
         pRes->size[0] = op.args.ImageResize.size[0];
         pRes->size[1] = op.args.ImageResize.size[1];
         break;
 
-    case OP_TYPE::IM_RESIZEREL:
+    case EOpType::IM_RESIZEREL:
         if ( children[op.args.ImageResizeRel.source] )
         {
             pRes = children[op.args.ImageResizeRel.source].child()->GetImageSizeExpression();
@@ -1542,75 +1535,75 @@ mu::Ptr<ImageSizeExpression> ASTOpFixed::GetImageSizeExpression() const
         }
         break;
 
-    case OP_TYPE::IM_RESIZELIKE:
+    case EOpType::IM_RESIZELIKE:
         if ( children[op.args.ImageResizeLike.sizeSource] )
         {
             pRes = children[op.args.ImageResizeLike.sizeSource].child()->GetImageSizeExpression();
         }
         break;
 
-    case OP_TYPE::IM_PLAINCOLOUR:
+    case EOpType::IM_PLAINCOLOUR:
         pRes->type = ImageSizeExpression::ISET_CONSTANT;
         pRes->size[0] = op.args.ImagePlainColour.size[0];
         pRes->size[1] = op.args.ImagePlainColour.size[1];
         break;
 
-    case OP_TYPE::IM_BLANKLAYOUT:
+    case EOpType::IM_BLANKLAYOUT:
         pRes->type = ImageSizeExpression::ISET_LAYOUTFACTOR;
         pRes->layout = children[op.args.ImageBlankLayout.layout].child();
         pRes->factor[0] = op.args.ImageBlankLayout.blockSize[0];
         pRes->factor[1] = op.args.ImageBlankLayout.blockSize[1];
         break;
 
-    case OP_TYPE::IM_INTERPOLATE:
+    case EOpType::IM_INTERPOLATE:
         if ( children[op.args.ImageInterpolate.targets[0]] )
         {
             pRes = children[op.args.ImageInterpolate.targets[0]].child()->GetImageSizeExpression();
         }
         break;
 
-    case OP_TYPE::IM_SATURATE:
+    case EOpType::IM_SATURATE:
         if ( children[op.args.ImageSaturate.base] )
         {
             pRes = children[op.args.ImageSaturate.base].child()->GetImageSizeExpression();
         }
         break;
 
-    case OP_TYPE::IM_LUMINANCE:
+    case EOpType::IM_LUMINANCE:
         if ( children[op.args.ImageLuminance.base] )
         {
             pRes = children[op.args.ImageLuminance.base].child()->GetImageSizeExpression();
         }
         break;
 
-    case OP_TYPE::IM_BINARISE:
+    case EOpType::IM_BINARISE:
         if ( children[op.args.ImageBinarise.base] )
         {
             pRes = children[op.args.ImageBinarise.base].child()->GetImageSizeExpression();
         }
         break;
 
-    case OP_TYPE::IM_COLOURMAP:
+    case EOpType::IM_COLOURMAP:
         if ( children[op.args.ImageColourMap.base] )
         {
             pRes = children[op.args.ImageColourMap.base].child()->GetImageSizeExpression();
         }
         break;
 
-    case OP_TYPE::IM_GRADIENT:
+    case EOpType::IM_GRADIENT:
         pRes->type = ImageSizeExpression::ISET_CONSTANT;
         pRes->size[0] = op.args.ImageGradient.size[0];
         pRes->size[1] = op.args.ImageGradient.size[1];
         break;
 
-    case OP_TYPE::IM_DISPLACE:
+    case EOpType::IM_DISPLACE:
         if ( children[op.args.ImageDisplace.source] )
         {
             pRes = children[op.args.ImageDisplace.source].child()->GetImageSizeExpression();
         }
         break;
 
-	case OP_TYPE::IM_INVERT:
+	case EOpType::IM_INVERT:
 		if (children[op.args.ImageInvert.base])
 		{
 			pRes = children[op.args.ImageInvert.base].child()->GetImageSizeExpression();
@@ -1648,26 +1641,26 @@ FSourceDataDescriptor ASTOpFixed::GetSourceDataDescriptor(FGetSourceDataDescript
 	}
 
 	// Calculate
-	OP_TYPE Type = op.type;
+	EOpType Type = op.type;
 
 	switch (Type)
 	{
 
-	case OP_TYPE::NONE:
-	case OP_TYPE::IM_PLAINCOLOUR:
-	case OP_TYPE::IM_GRADIENT:
-	case OP_TYPE::IM_BLANKLAYOUT:
+	case EOpType::NONE:
+	case EOpType::IM_PLAINCOLOUR:
+	case EOpType::IM_GRADIENT:
+	case EOpType::IM_BLANKLAYOUT:
 		break;
 
-	case OP_TYPE::IM_SATURATE:
+	case EOpType::IM_SATURATE:
 		Result = GetSourceDataDescriptor(op.args.ImageSaturate.base, Context);
 		break;
 
-	case OP_TYPE::IM_LUMINANCE:
+	case EOpType::IM_LUMINANCE:
 		Result = GetSourceDataDescriptor(op.args.ImageLuminance.base, Context);
 		break;
 
-	case OP_TYPE::IM_INTERPOLATE:
+	case EOpType::IM_INTERPOLATE:
 	{
 		for (int32 SourceIndex = 0; SourceIndex < MUTABLE_OP_MAX_INTERPOLATE_COUNT; ++SourceIndex)
 		{
@@ -1677,55 +1670,55 @@ FSourceDataDescriptor ASTOpFixed::GetSourceDataDescriptor(FGetSourceDataDescript
 		break;
 	}
 
-	case OP_TYPE::IM_RESIZE:
+	case EOpType::IM_RESIZE:
 		Result = GetSourceDataDescriptor(op.args.ImageResize.source, Context);
 		break;
 
-	case OP_TYPE::IM_RESIZEREL:
+	case EOpType::IM_RESIZEREL:
 		Result = GetSourceDataDescriptor(op.args.ImageResizeRel.source, Context);
 		break;
 
-	case OP_TYPE::IM_RESIZELIKE:
+	case EOpType::IM_RESIZELIKE:
 		Result = GetSourceDataDescriptor(op.args.ImageResizeLike.source, Context);
 		break;
 
-	case OP_TYPE::IM_BINARISE:
+	case EOpType::IM_BINARISE:
 		Result = GetSourceDataDescriptor(op.args.ImageBinarise.base, Context);
 		break;
 
-	case OP_TYPE::IM_DISPLACE:
+	case EOpType::IM_DISPLACE:
 		Result = GetSourceDataDescriptor(op.args.ImageDisplace.source, Context);
 		break;
 
-	case OP_TYPE::IM_INVERT:
+	case EOpType::IM_INVERT:
 		Result = GetSourceDataDescriptor(op.args.ImageInvert.base, Context);
 		break;
 
-	case OP_TYPE::IM_COLOURMAP:
+	case EOpType::IM_COLOURMAP:
 		Result = GetSourceDataDescriptor(op.args.ImageColourMap.base, Context);
 		break;
-	case OP_TYPE::ME_APPLYLAYOUT:
+	case EOpType::ME_APPLYLAYOUT:
 		Result = GetSourceDataDescriptor(op.args.MeshApplyLayout.mesh, Context);
 		break;	
-	case OP_TYPE::ME_MERGE:
+	case EOpType::ME_MERGE:
 		Result = GetSourceDataDescriptor(op.args.MeshMerge.base, Context);
 		break;	
-	case OP_TYPE::ME_INTERPOLATE:
+	case EOpType::ME_INTERPOLATE:
 		Result = GetSourceDataDescriptor(op.args.MeshInterpolate.base, Context);
 		break;	
-	case OP_TYPE::ME_MASKDIFF:
+	case EOpType::ME_MASKDIFF:
 		Result = GetSourceDataDescriptor(op.args.MeshMaskDiff.source, Context);
 		break;	
-	case OP_TYPE::ME_CLIPMORPHPLANE:
+	case EOpType::ME_CLIPMORPHPLANE:
 		Result = GetSourceDataDescriptor(op.args.MeshClipMorphPlane.source, Context);
 		break;	
-	case OP_TYPE::ME_CLIPWITHMESH:
+	case EOpType::ME_CLIPWITHMESH:
 		Result = GetSourceDataDescriptor(op.args.MeshClipWithMesh.source, Context);
 		break;	
-	case OP_TYPE::ME_SETSKELETON:
+	case EOpType::ME_SETSKELETON:
 		Result = GetSourceDataDescriptor(op.args.MeshSetSkeleton.source, Context);
 		break;	
-	case OP_TYPE::ME_PROJECT:
+	case EOpType::ME_PROJECT:
 		Result = GetSourceDataDescriptor(op.args.MeshProject.mesh, Context);
 		break;
 	default:
diff --git a/Engine/Plugins/Mutable/Source/MutableTools/Private/MuT/ASTOpAddExtensionData.cpp b/Engine/Plugins/Mutable/Source/MutableTools/Private/MuT/ASTOpAddExtensionData.cpp
--- a/Engine/Plugins/Mutable/Source/MutableTools/Private/MuT/ASTOpAddExtensionData.cpp
+++ b/Engine/Plugins/Mutable/Source/MutableTools/Private/MuT/ASTOpAddExtensionData.cpp
@@ -42,7 +42,7 @@ bool ASTOpAddExtensionData::IsEqual(const ASTOp& OtherUntyped) const
 
 uint64 ASTOpAddExtensionData::Hash() const
 {
-	uint64 Result = std::hash<uint64>()(uint64(OP_TYPE::IN_ADDEXTENSIONDATA));
+	uint64 Result = std::hash<uint64>()(uint64(EOpType::IN_ADDEXTENSIONDATA));
 	
 	hash_combine(Result, Instance.child().get());
 	hash_combine(Result, ExtensionData.child().get());
@@ -92,10 +92,10 @@ void ASTOpAddExtensionData::Link(FProgram& Program, FLinkerOptions*)
 	check(ExtensionDataName.Len() > 0);
 	Args.ExtensionDataName = Program.AddConstant(ExtensionDataName);
 
-	linkedAddress = (OP::ADDRESS)Program.m_opAddress.Num();
-	Program.m_opAddress.Add((uint32_t)Program.m_byteCode.Num());
-	AppendCode(Program.m_byteCode, OP_TYPE::IN_ADDEXTENSIONDATA);
-	AppendCode(Program.m_byteCode, Args);
+	linkedAddress = (OP::ADDRESS)Program.OpAddress.Num();
+	Program.OpAddress.Add((uint32_t)Program.ByteCode.Num());
+	AppendCode(Program.ByteCode, EOpType::IN_ADDEXTENSIONDATA);
+	AppendCode(Program.ByteCode, Args);
 }
 
 }
diff --git a/Engine/Plugins/Mutable/Source/MutableTools/Private/MuT/ASTOpAddExtensionData.h b/Engine/Plugins/Mutable/Source/MutableTools/Private/MuT/ASTOpAddExtensionData.h
--- a/Engine/Plugins/Mutable/Source/MutableTools/Private/MuT/ASTOpAddExtensionData.h
+++ b/Engine/Plugins/Mutable/Source/MutableTools/Private/MuT/ASTOpAddExtensionData.h
@@ -31,7 +31,7 @@ public:
 	ASTOpAddExtensionData(const ASTOpAddExtensionData&) = delete;
 	~ASTOpAddExtensionData();
 
-	OP_TYPE GetOpType() const override { return OP_TYPE::IN_ADDEXTENSIONDATA; }
+	EOpType GetOpType() const override { return EOpType::IN_ADDEXTENSIONDATA; }
 	uint64 Hash() const override;
 	void ForEachChild(const TFunctionRef<void(ASTChild&)> F) override;
 	bool IsEqual(const ASTOp& OtherUntyped) const override;
diff --git a/Engine/Plugins/Mutable/Source/MutableTools/Private/MuT/ASTOpAddLOD.cpp b/Engine/Plugins/Mutable/Source/MutableTools/Private/MuT/ASTOpAddLOD.cpp
--- a/Engine/Plugins/Mutable/Source/MutableTools/Private/MuT/ASTOpAddLOD.cpp
+++ b/Engine/Plugins/Mutable/Source/MutableTools/Private/MuT/ASTOpAddLOD.cpp
@@ -39,7 +39,7 @@ namespace mu
 
 	uint64 ASTOpAddLOD::Hash() const
 	{
-		uint64 res = std::hash<uint64>()(uint64(OP_TYPE::IN_ADDLOD));
+		uint64 res = std::hash<uint64>()(uint64(EOpType::IN_ADDLOD));
 		for (auto& c : lods)
 		{
 			hash_combine(res, c.child().get());
@@ -73,9 +73,9 @@ namespace mu
 		// Already linked?
 		if (!linkedAddress)
 		{
-			linkedAddress = (OP::ADDRESS)program.m_opAddress.Num();
-			program.m_opAddress.Add((uint32)program.m_byteCode.Num());
-			AppendCode(program.m_byteCode, OP_TYPE::IN_ADDLOD);
+			linkedAddress = (OP::ADDRESS)program.OpAddress.Num();
+			program.OpAddress.Add((uint32)program.ByteCode.Num());
+			AppendCode(program.ByteCode, EOpType::IN_ADDLOD);
 
 			// Calculate LODCount mimicking previous behaviour.
 			// It may be better to leave null LODs if they are null in the AST.
@@ -94,14 +94,14 @@ namespace mu
 				}
 			}
 
-			AppendCode(program.m_byteCode, LODCount);
+			AppendCode(program.ByteCode, LODCount);
 			for (const ASTChild& LOD : lods)
 			{
 				OP::ADDRESS LODAddress = 0;
 				if (LOD)
 				{
 					LODAddress = LOD->linkedAddress;
-					AppendCode(program.m_byteCode, LODAddress);
+					AppendCode(program.ByteCode, LODAddress);
 					if (LODCount == 255)
 					{
 						break;
diff --git a/Engine/Plugins/Mutable/Source/MutableTools/Private/MuT/ASTOpAddLOD.h b/Engine/Plugins/Mutable/Source/MutableTools/Private/MuT/ASTOpAddLOD.h
--- a/Engine/Plugins/Mutable/Source/MutableTools/Private/MuT/ASTOpAddLOD.h
+++ b/Engine/Plugins/Mutable/Source/MutableTools/Private/MuT/ASTOpAddLOD.h
@@ -26,7 +26,7 @@ struct FProgram;
 		ASTOpAddLOD(const ASTOpAddLOD&) = delete;
 		~ASTOpAddLOD();
 
-		OP_TYPE GetOpType() const override { return OP_TYPE::IN_ADDLOD; }
+		EOpType GetOpType() const override { return EOpType::IN_ADDLOD; }
 		uint64 Hash() const override;
 		void ForEachChild(const TFunctionRef<void(ASTChild&)>) override;
 		bool IsEqual(const ASTOp& otherUntyped) const override;
diff --git a/Engine/Plugins/Mutable/Source/MutableTools/Private/MuT/ASTOpConditional.cpp b/Engine/Plugins/Mutable/Source/MutableTools/Private/MuT/ASTOpConditional.cpp
--- a/Engine/Plugins/Mutable/Source/MutableTools/Private/MuT/ASTOpConditional.cpp
+++ b/Engine/Plugins/Mutable/Source/MutableTools/Private/MuT/ASTOpConditional.cpp
@@ -63,9 +63,9 @@ namespace mu
 	//-------------------------------------------------------------------------------------------------
 	uint64 ASTOpConditional::Hash() const
 	{
-		uint64 res = std::hash<void*>()(condition.m_child.get());
-		hash_combine(res, yes.m_child.get());
-		hash_combine(res, no.m_child.get());
+		uint64 res = std::hash<void*>()(condition.Child.get());
+		hash_combine(res, yes.Child.get());
+		hash_combine(res, no.Child.get());
 		return res;
 	}
 
@@ -75,14 +75,14 @@ namespace mu
 	{
 		switch (type)
 		{
-		case OP_TYPE::NU_CONDITIONAL:
-		case OP_TYPE::SC_CONDITIONAL:
-		case OP_TYPE::CO_CONDITIONAL:
-		case OP_TYPE::IM_CONDITIONAL:
-		case OP_TYPE::ME_CONDITIONAL:
-		case OP_TYPE::LA_CONDITIONAL:
-		case OP_TYPE::IN_CONDITIONAL:
-		case OP_TYPE::ED_CONDITIONAL:
+		case EOpType::NU_CONDITIONAL:
+		case EOpType::SC_CONDITIONAL:
+		case EOpType::CO_CONDITIONAL:
+		case EOpType::IM_CONDITIONAL:
+		case EOpType::ME_CONDITIONAL:
+		case EOpType::LA_CONDITIONAL:
+		case EOpType::IN_CONDITIONAL:
+		case EOpType::ED_CONDITIONAL:
 			break;
 		default:
 			// Unexpected type
@@ -116,11 +116,11 @@ namespace mu
 			if (yes) args.yes = yes->linkedAddress;
 			if (no) args.no = no->linkedAddress;
 
-			linkedAddress = (OP::ADDRESS)program.m_opAddress.Num();
+			linkedAddress = (OP::ADDRESS)program.OpAddress.Num();
 			//program.m_code.push_back(op);
-			program.m_opAddress.Add((uint32_t)program.m_byteCode.Num());
-			AppendCode(program.m_byteCode, type);
-			AppendCode(program.m_byteCode, args);
+			program.OpAddress.Add((uint32_t)program.ByteCode.Num());
+			AppendCode(program.ByteCode, type);
+			AppendCode(program.ByteCode, args);
 
 		}
 	}
@@ -147,7 +147,7 @@ namespace mu
 			}
 		}
 
-		if (type == OP_TYPE::IM_CONDITIONAL)
+		if (type == EOpType::IM_CONDITIONAL)
 		{
 			// In a conditional we cannot guarantee the size and format.
 			// We check both options, and if they are the same we return that.
@@ -196,7 +196,7 @@ namespace mu
 	//-------------------------------------------------------------------------------------------------
 	void ASTOpConditional::GetLayoutBlockSize(int32* pBlockX, int32* pBlockY)
 	{
-		if (type == OP_TYPE::IM_CONDITIONAL)
+		if (type == EOpType::IM_CONDITIONAL)
 		{
 			yes->GetLayoutBlockSize(pBlockX, pBlockY);
 
@@ -215,7 +215,7 @@ namespace mu
 	//-------------------------------------------------------------------------------------------------
 	void ASTOpConditional::GetBlockLayoutSize(uint64 BlockId, int32* pBlockX, int32* pBlockY, FBlockLayoutSizeCache* cache)
 	{
-		if (type == OP_TYPE::LA_CONDITIONAL)
+		if (type == EOpType::LA_CONDITIONAL)
 		{
 			yes->GetBlockLayoutSizeCached(BlockId, pBlockX, pBlockY, cache);
 
@@ -234,7 +234,7 @@ namespace mu
 	//-------------------------------------------------------------------------------------------------
 	bool ASTOpConditional::GetNonBlackRect(FImageRect& maskUsage) const
 	{
-		if (type == OP_TYPE::IM_CONDITIONAL)
+		if (type == EOpType::IM_CONDITIONAL)
 		{
 			FImageRect local;
 			bool localValid = false;
@@ -296,7 +296,7 @@ namespace mu
 		}
 
 		// Constant condition?
-		if (condition->GetOpType() == OP_TYPE::BO_CONSTANT)
+		if (condition->GetOpType() == EOpType::BO_CONSTANT)
 		{
 			const ASTOpConstantBool* typedCondition = static_cast<const ASTOpConstantBool*>(condition.child().get());
 			if (typedCondition->value)
diff --git a/Engine/Plugins/Mutable/Source/MutableTools/Private/MuT/ASTOpConditional.h b/Engine/Plugins/Mutable/Source/MutableTools/Private/MuT/ASTOpConditional.h
--- a/Engine/Plugins/Mutable/Source/MutableTools/Private/MuT/ASTOpConditional.h
+++ b/Engine/Plugins/Mutable/Source/MutableTools/Private/MuT/ASTOpConditional.h
@@ -21,7 +21,7 @@ struct FProgram;
 	public:
 
 		//! Type of switch
-		OP_TYPE type;
+		EOpType type;
 
 		//! Boolean expression
 		ASTChild condition;
@@ -36,7 +36,7 @@ struct FProgram;
 		ASTOpConditional(const ASTOpConditional&) = delete;
 		~ASTOpConditional() override;
 
-		virtual OP_TYPE GetOpType() const override { return type; }
+		virtual EOpType GetOpType() const override { return type; }
 		virtual bool IsEqual(const ASTOp& otherUntyped) const override;
 		virtual Ptr<ASTOp> Clone(MapChildFuncRef mapChild) const override;
 		virtual uint64 Hash() const override;
diff --git a/Engine/Plugins/Mutable/Source/MutableTools/Private/MuT/ASTOpConstantBool.cpp b/Engine/Plugins/Mutable/Source/MutableTools/Private/MuT/ASTOpConstantBool.cpp
--- a/Engine/Plugins/Mutable/Source/MutableTools/Private/MuT/ASTOpConstantBool.cpp
+++ b/Engine/Plugins/Mutable/Source/MutableTools/Private/MuT/ASTOpConstantBool.cpp
@@ -39,7 +39,7 @@ namespace mu
 
 	uint64 ASTOpConstantBool::Hash() const
 	{
-		uint64 res = std::hash<uint64>()(uint64(OP_TYPE::BO_CONSTANT));
+		uint64 res = std::hash<uint64>()(uint64(EOpType::BO_CONSTANT));
 		hash_combine(res, value);
 		return res;
 	}
@@ -61,11 +61,11 @@ namespace mu
 			memset(&args, 0, sizeof(args));
 			args.value = value;
 
-			linkedAddress = (OP::ADDRESS)program.m_opAddress.Num();
+			linkedAddress = (OP::ADDRESS)program.OpAddress.Num();
 			//        program.m_code.push_back(op);
-			program.m_opAddress.Add((uint32_t)program.m_byteCode.Num());
-			AppendCode(program.m_byteCode, OP_TYPE::BO_CONSTANT);
-			AppendCode(program.m_byteCode, args);
+			program.OpAddress.Add((uint32_t)program.ByteCode.Num());
+			AppendCode(program.ByteCode, EOpType::BO_CONSTANT);
+			AppendCode(program.ByteCode, args);
 		}
 	}
 
diff --git a/Engine/Plugins/Mutable/Source/MutableTools/Private/MuT/ASTOpConstantBool.h b/Engine/Plugins/Mutable/Source/MutableTools/Private/MuT/ASTOpConstantBool.h
--- a/Engine/Plugins/Mutable/Source/MutableTools/Private/MuT/ASTOpConstantBool.h
+++ b/Engine/Plugins/Mutable/Source/MutableTools/Private/MuT/ASTOpConstantBool.h
@@ -27,7 +27,7 @@ struct FProgram;
 
 		ASTOpConstantBool(bool value = true);
 
-		OP_TYPE GetOpType() const override { return OP_TYPE::BO_CONSTANT; }
+		EOpType GetOpType() const override { return EOpType::BO_CONSTANT; }
 		uint64 Hash() const override;
 		void ForEachChild(const TFunctionRef<void(ASTChild&)>) override;
 		bool IsEqual(const ASTOp& otherUntyped) const override;
diff --git a/Engine/Plugins/Mutable/Source/MutableTools/Private/MuT/ASTOpConstantExtensionData.cpp b/Engine/Plugins/Mutable/Source/MutableTools/Private/MuT/ASTOpConstantExtensionData.cpp
--- a/Engine/Plugins/Mutable/Source/MutableTools/Private/MuT/ASTOpConstantExtensionData.cpp
+++ b/Engine/Plugins/Mutable/Source/MutableTools/Private/MuT/ASTOpConstantExtensionData.cpp
@@ -37,7 +37,7 @@ mu::Ptr<ASTOp> ASTOpConstantExtensionData::Clone(MapChildFuncRef MapChild) const
 //-------------------------------------------------------------------------------------------------
 uint64 ASTOpConstantExtensionData::Hash() const
 {
-	return std::hash<uint64>()(Value.get() != nullptr ? Value->Hash() : 0ll);
+	return std::hash<uint64>()(Value != nullptr ? Value->Hash() : 0ll);
 }
 
 
@@ -51,10 +51,10 @@ void ASTOpConstantExtensionData::Link(FProgram& Program, FLinkerOptions* Options
 
 		Args.value = Program.AddConstant(Value);
 
-		linkedAddress = (OP::ADDRESS)Program.m_opAddress.Num();
-		Program.m_opAddress.Add((uint32_t)Program.m_byteCode.Num());
-		AppendCode(Program.m_byteCode, OP_TYPE::ED_CONSTANT);
-		AppendCode(Program.m_byteCode, Args);
+		linkedAddress = (OP::ADDRESS)Program.OpAddress.Num();
+		Program.OpAddress.Add((uint32_t)Program.ByteCode.Num());
+		AppendCode(Program.ByteCode, EOpType::ED_CONSTANT);
+		AppendCode(Program.ByteCode, Args);
 	}
 }
 
diff --git a/Engine/Plugins/Mutable/Source/MutableTools/Private/MuT/ASTOpConstantExtensionData.h b/Engine/Plugins/Mutable/Source/MutableTools/Private/MuT/ASTOpConstantExtensionData.h
--- a/Engine/Plugins/Mutable/Source/MutableTools/Private/MuT/ASTOpConstantExtensionData.h
+++ b/Engine/Plugins/Mutable/Source/MutableTools/Private/MuT/ASTOpConstantExtensionData.h
@@ -11,29 +11,27 @@
 
 namespace mu
 {
-struct FProgram;
 
+	struct FProgram;
 
-//---------------------------------------------------------------------------------------------
-//! A constant ExtensionData
-//---------------------------------------------------------------------------------------------
-class ASTOpConstantExtensionData final : public ASTOp
-{
-public:
+	/** A constant ExtensionData. */
+	class ASTOpConstantExtensionData final : public ASTOp
+	{
+	public:
 
-	//!
-	ExtensionDataPtrConst Value;
+		//!
+		TSharedPtr<const FExtensionData> Value;
 
-public:
+	public:
 
-	// ASTOp interface
-	OP_TYPE GetOpType() const override { return OP_TYPE::ED_CONSTANT; }
-	void ForEachChild(const TFunctionRef<void(ASTChild&)>) override {}
-	bool IsEqual(const ASTOp& OtherUntyped) const override;
-	Ptr<ASTOp> Clone(MapChildFuncRef MapChild) const override;
-	uint64 Hash() const override;
-	void Link(FProgram& Program, FLinkerOptions*) override;
-};
+		// ASTOp interface
+		EOpType GetOpType() const override { return EOpType::ED_CONSTANT; }
+		void ForEachChild(const TFunctionRef<void(ASTChild&)>) override {}
+		bool IsEqual(const ASTOp& OtherUntyped) const override;
+		Ptr<ASTOp> Clone(MapChildFuncRef MapChild) const override;
+		uint64 Hash() const override;
+		void Link(FProgram& Program, FLinkerOptions*) override;
+	};
 
 
 }
diff --git a/Engine/Plugins/Mutable/Source/MutableTools/Private/MuT/ASTOpConstantMatrix.cpp b/Engine/Plugins/Mutable/Source/MutableTools/Private/MuT/ASTOpConstantMatrix.cpp
--- a/Engine/Plugins/Mutable/Source/MutableTools/Private/MuT/ASTOpConstantMatrix.cpp
+++ b/Engine/Plugins/Mutable/Source/MutableTools/Private/MuT/ASTOpConstantMatrix.cpp
@@ -8,7 +8,7 @@ mu::ASTOpConstantMatrix::ASTOpConstantMatrix(FMatrix44f InitValue) : value(InitV
 
 uint64 mu::ASTOpConstantMatrix::Hash() const
 {
-	uint64 res = std::hash<uint64>()(uint64(OP_TYPE::MA_CONSTANT));
+	uint64 res = std::hash<uint64>()(uint64(EOpType::MA_CONSTANT));
 	hash_combine(res, value.ComputeHash());
 	return res;
 }
@@ -38,9 +38,9 @@ void mu::ASTOpConstantMatrix::Link(FProgram& program, FLinkerOptions* Options)
 		OP::MatrixConstantArgs args;
 		args.value = program.AddConstant(value);
 
-		linkedAddress = static_cast<OP::ADDRESS>(program.m_opAddress.Num());
-		program.m_opAddress.Add(static_cast<uint32_t>(program.m_byteCode.Num()));
-		AppendCode(program.m_byteCode, OP_TYPE::MA_CONSTANT);
-		AppendCode(program.m_byteCode, args);
+		linkedAddress = static_cast<OP::ADDRESS>(program.OpAddress.Num());
+		program.OpAddress.Add(static_cast<uint32_t>(program.ByteCode.Num()));
+		AppendCode(program.ByteCode, EOpType::MA_CONSTANT);
+		AppendCode(program.ByteCode, args);
 	}
 }
diff --git a/Engine/Plugins/Mutable/Source/MutableTools/Private/MuT/ASTOpConstantMatrix.h b/Engine/Plugins/Mutable/Source/MutableTools/Private/MuT/ASTOpConstantMatrix.h
--- a/Engine/Plugins/Mutable/Source/MutableTools/Private/MuT/ASTOpConstantMatrix.h
+++ b/Engine/Plugins/Mutable/Source/MutableTools/Private/MuT/ASTOpConstantMatrix.h
@@ -27,7 +27,7 @@ public:
 
 	ASTOpConstantMatrix(FMatrix44f value = FMatrix44f::Identity);
 
-	OP_TYPE GetOpType() const override { return OP_TYPE::MA_CONSTANT; }
+	EOpType GetOpType() const override { return EOpType::MA_CONSTANT; }
 	uint64 Hash() const override;
 	void ForEachChild(const TFunctionRef<void(ASTChild&)>) override {}
 	bool IsEqual(const ASTOp& otherUntyped) const override;
diff --git a/Engine/Plugins/Mutable/Source/MutableTools/Private/MuT/ASTOpConstantProjector.cpp b/Engine/Plugins/Mutable/Source/MutableTools/Private/MuT/ASTOpConstantProjector.cpp
--- a/Engine/Plugins/Mutable/Source/MutableTools/Private/MuT/ASTOpConstantProjector.cpp
+++ b/Engine/Plugins/Mutable/Source/MutableTools/Private/MuT/ASTOpConstantProjector.cpp
@@ -54,11 +54,11 @@ namespace mu
 			memset(&args, 0, sizeof(args));
 			args.value = program.AddConstant(value);
 
-			linkedAddress = (OP::ADDRESS)program.m_opAddress.Num();
+			linkedAddress = (OP::ADDRESS)program.OpAddress.Num();
 			//program.m_code.push_back(op);
-			program.m_opAddress.Add((uint32_t)program.m_byteCode.Num());
-			AppendCode(program.m_byteCode, OP_TYPE::PR_CONSTANT);
-			AppendCode(program.m_byteCode, args);
+			program.OpAddress.Add((uint32_t)program.ByteCode.Num());
+			AppendCode(program.ByteCode, EOpType::PR_CONSTANT);
+			AppendCode(program.ByteCode, args);
 		}
 	}
 
diff --git a/Engine/Plugins/Mutable/Source/MutableTools/Private/MuT/ASTOpConstantProjector.h b/Engine/Plugins/Mutable/Source/MutableTools/Private/MuT/ASTOpConstantProjector.h
--- a/Engine/Plugins/Mutable/Source/MutableTools/Private/MuT/ASTOpConstantProjector.h
+++ b/Engine/Plugins/Mutable/Source/MutableTools/Private/MuT/ASTOpConstantProjector.h
@@ -26,7 +26,7 @@ struct FProgram;
 
 	public:
 
-		OP_TYPE GetOpType() const override { return OP_TYPE::PR_CONSTANT; }
+		EOpType GetOpType() const override { return EOpType::PR_CONSTANT; }
 		uint64 Hash() const override;
 		void ForEachChild(const TFunctionRef<void(ASTChild&)>) override;
 		bool IsEqual(const ASTOp& otherUntyped) const override;
diff --git a/Engine/Plugins/Mutable/Source/MutableTools/Private/MuT/ASTOpConstantResource.cpp b/Engine/Plugins/Mutable/Source/MutableTools/Private/MuT/ASTOpConstantResource.cpp
--- a/Engine/Plugins/Mutable/Source/MutableTools/Private/MuT/ASTOpConstantResource.cpp
+++ b/Engine/Plugins/Mutable/Source/MutableTools/Private/MuT/ASTOpConstantResource.cpp
@@ -33,12 +33,12 @@ namespace mu
 	* - saved to a disk file compressed or uncompressed.
 	*/
 	template<class R>
-	class MUTABLETOOLS_API ResourceProxyTempFile : public ResourceProxy<R>
+	class MUTABLETOOLS_API ResourceProxyTempFile : public TResourceProxy<R>
 	{
 	private:
 
 		/** Actual resource to store. If the pointer is valid, it wasn't worth dumping to disk or compressing. */
-		Ptr<const R> Resource;
+		TSharedPtr<const R> Resource;
 
 		/** Temp filename used if it was necessary. */
 		FString FileName;
@@ -60,7 +60,7 @@ namespace mu
 
 	public:
 
-		ResourceProxyTempFile(const R* InResource, FProxyFileContext& InOptions)
+		ResourceProxyTempFile(TSharedPtr<const R> InResource, FProxyFileContext& InOptions)
 			: Options(InOptions)
 		{
 			if (!InResource)
@@ -70,13 +70,13 @@ namespace mu
 
 			IPlatformFile& PlatformFile = FPlatformFileManager::Get().GetPlatformFile();
 
-			OutputMemoryStream stream(128*1024);
-			OutputArchive arch(&stream);
-			R::Serialise(InResource, arch);
+			FOutputMemoryStream Stream(128*1024);
+			FOutputArchive Arch(&Stream);
+			R::Serialise(InResource.Get(), Arch);
 
-			UncompressedSize = stream.GetBufferSize();
+			UncompressedSize = Stream.GetBufferSize();
 
-			if (stream.GetBufferSize() <= Options.MinProxyFileSize)
+			if (Stream.GetBufferSize() <= Options.MinProxyFileSize)
 			{
 				// Not worth compressing or caching to disk
 				Resource = InResource;
@@ -88,13 +88,13 @@ namespace mu
 				constexpr bool bEnableCompression = true;
 				if (bEnableCompression)
 				{
-					int64 CompressedBufferSize = FOodleDataCompression::CompressedBufferSizeNeeded(stream.GetBufferSize());
-					CompressedBufferSize = FMath::Max(CompressedBufferSize, int64(stream.GetBufferSize() / 2));
+					int64 CompressedBufferSize = FOodleDataCompression::CompressedBufferSizeNeeded(Stream.GetBufferSize());
+					CompressedBufferSize = FMath::Max(CompressedBufferSize, int64(Stream.GetBufferSize() / 2));
 					CompressedBuffer.SetNumUninitialized(CompressedBufferSize);
 
 					CompressedSize = FOodleDataCompression::CompressParallel(
 						CompressedBuffer.GetData(), CompressedBufferSize,
-						stream.GetBuffer(), stream.GetBufferSize(),
+						Stream.GetBuffer(), Stream.GetBufferSize(),
 						FOodleDataCompression::ECompressor::Kraken,
 						FOodleDataCompression::ECompressionLevel::SuperFast,
 						true // CompressIndependentChunks
@@ -143,7 +143,7 @@ namespace mu
 					else
 					{
 						FileSize = UncompressedSize;
-						ResourceFile->Write((const uint8*)stream.GetBuffer(), FileSize);
+						ResourceFile->Write(Stream.GetBuffer(), FileSize);
 					}
 
 					CompressedBuffer.SetNum(0, EAllowShrinking::Yes);
@@ -169,11 +169,11 @@ namespace mu
 			}
 		}
 
-		Ptr<const R> Get() override
+		TSharedPtr<const R> Get() override
 		{
 			FScopeLock Lock(&Mutex);
 
-			Ptr<const R> Result;
+			TSharedPtr<const R> Result;
 			if (Resource)
 			{
 				// Cached as is
@@ -181,19 +181,19 @@ namespace mu
 			}
 			else if (!CompressedBuffer.Num() && !FileName.IsEmpty())
 			{
-				IFileHandle* resourceFile = FPlatformFileManager::Get().GetPlatformFile().OpenRead(*FileName);
-				check(resourceFile);
+				IFileHandle* ResourceFile = FPlatformFileManager::Get().GetPlatformFile().OpenRead(*FileName);
+				check(ResourceFile);
 
 				CompressedBuffer.SetNumUninitialized(FileSize);
-				resourceFile->Read(CompressedBuffer.GetData(), FileSize);
-				delete resourceFile;
+				ResourceFile->Read(CompressedBuffer.GetData(), FileSize);
+				delete ResourceFile;
 
 				bool bCompressed = FileSize != UncompressedSize;
 
 				if (!bCompressed)
 				{
-					InputMemoryStream stream(CompressedBuffer.GetData(), FileSize);
-					InputArchive arch(&stream);
+					FInputMemoryStream stream(CompressedBuffer.GetData(), FileSize);
+					FInputArchive arch(&stream);
 					Result = R::StaticUnserialise(arch);
 
 					CompressedBuffer.SetNum(0, EAllowShrinking::Yes);
@@ -216,8 +216,8 @@ namespace mu
 
 				if (bSuccess)
 				{
-					InputMemoryStream stream(UncompressedBuf.GetData(), UncompressedSize);
-					InputArchive arch(&stream);
+					FInputMemoryStream stream(UncompressedBuf.GetData(), UncompressedSize);
+					FInputArchive arch(&stream);
 					Result = R::StaticUnserialise(arch);
 				}
 
@@ -280,7 +280,7 @@ namespace mu
 	namespace
 	{
 		/** Adds a constant image data to a program and returns its constant index. */
-		int32 AddConstantImage(FProgram& Program, const Ptr<const Image>& pImage, FLinkerOptions& Options)
+		int32 AddConstantImage(FProgram& Program, const TSharedPtr<const FImage>& pImage, FLinkerOptions& Options)
 		{
 			MUTABLE_CPUPROFILER_SCOPE(AddConstantImage);
 
@@ -292,7 +292,7 @@ namespace mu
 			int32 FirstLODIndexIndex = Program.ConstantImageLODIndices.Num();
 
 			FImageOperator& ImOp = Options.ImageOperator;
-			Ptr<const Image> pMip;
+			TSharedPtr<const FImage> pMip;
 
 			if (!Options.bSeparateImageMips)
 			{
@@ -302,11 +302,11 @@ namespace mu
 			{
 				// We may want the full mipmaps for fragments of images, regardless of the resident mip size, for intermediate operations.
 				// \TODO: Calculate the mip ranges that makes sense to store.
-				int32 MaxMipmaps = Image::GetMipmapCount(pImage->GetSizeX(), pImage->GetSizeY());
+				int32 MaxMipmaps = FImage::GetMipmapCount(pImage->GetSizeX(), pImage->GetSizeY());
 				MipsToStore = MaxMipmaps;
 
 				// Some images cannot be resized or mipmaped
-				bool bCannotBeScaled = pImage->m_flags & Image::IF_CANNOT_BE_SCALED;
+				bool bCannotBeScaled = pImage->Flags & FImage::IF_CANNOT_BE_SCALED;
 				if (bCannotBeScaled)
 				{
 					// Store only the mips that we have already calculated. We assume we have calculated them correctly.
@@ -319,12 +319,12 @@ namespace mu
 				}
 				else
 				{
-					pMip = ImOp.ExtractMip(pImage.get(), 0);
+					pMip = ImOp.ExtractMip(pImage.Get(), 0);
 				}
 			}
 
 			// Temporary uncompressed version of the image, if we need to generate the mips and the source is compressed.
-			Ptr<const Image> UncompressedMip;
+			TSharedPtr<const FImage> UncompressedMip;
 			EImageFormat UncompressedFormat = GetUncompressedFormat( pMip->GetFormat() );
 
 			for (int32 Mip = 0; Mip < MipsToStore; ++Mip)
@@ -348,7 +348,7 @@ namespace mu
 
 				if (MipIndex<0)
 				{
-					MipIndex = Program.ConstantImageLODs.Add(TPair<int32, Ptr<const Image>>(-1, pMip));
+					MipIndex = Program.ConstantImageLODs.Add(TPair<int32, TSharedPtr<const FImage>>(-1, pMip));
 					Options.ImageConstantMipMap.Add(pMip, MipIndex);
 				}
 
@@ -357,11 +357,11 @@ namespace mu
 				// Generate next mip if necessary
 				if (Mip + 1 < MipsToStore)
 				{
-					Ptr<Image> NewMip;
+					TSharedPtr<FImage> NewMip;
 					if (Mip+1 < pImage->GetLODCount())
 					{
 						// Extract directly from source image
-						NewMip = ImOp.ExtractMip(pImage.get(), Mip + 1);
+						NewMip = ImOp.ExtractMip(pImage.Get(), Mip + 1);
 					}
 					else
 					{
@@ -372,15 +372,15 @@ namespace mu
 
 							if (!UncompressedMip)
 							{
-								UncompressedMip = ImOp.ImagePixelFormat(Quality, pMip.get(), UncompressedFormat);
+								UncompressedMip = ImOp.ImagePixelFormat(Quality, pMip.Get(), UncompressedFormat);
 							}
 
-							UncompressedMip = ImOp.ExtractMip(UncompressedMip.get(), 1);
-							NewMip = ImOp.ImagePixelFormat(Quality, UncompressedMip.get(), pMip->GetFormat());
+							UncompressedMip = ImOp.ExtractMip(UncompressedMip.Get(), 1);
+							NewMip = ImOp.ImagePixelFormat(Quality, UncompressedMip.Get(), pMip->GetFormat());
 						}
 						else
 						{
-							NewMip = ImOp.ExtractMip(pMip.get(), 1);
+							NewMip = ImOp.ExtractMip(pMip.Get(), 1);
 						}
 					}
 					check(NewMip);
@@ -407,34 +407,34 @@ namespace mu
 
 		if (!linkedAddress && !bLinkedAndNull)
 		{
-			if (Type == OP_TYPE::ME_CONSTANT)
+			if (Type == EOpType::ME_CONSTANT)
 			{
 				OP::MeshConstantArgs args;
 				FMemory::Memset(&args, 0, sizeof(args));
 
-				Ptr<Mesh> MeshData = static_cast<const Mesh*>(GetValue().get())->Clone();
+				TSharedPtr<FMesh> MeshData = static_cast<const FMesh*>(GetValue().Get())->Clone();
 				check(MeshData);
 
 				args.skeleton = -1;
-				if (Ptr<const Skeleton> pSkeleton = MeshData->GetSkeleton())
+				if (TSharedPtr<const FSkeleton> pSkeleton = MeshData->GetSkeleton())
 				{
-					args.skeleton = program.AddConstant(pSkeleton.get());
+					args.skeleton = program.AddConstant(pSkeleton);
 					MeshData->SetSkeleton(nullptr);
 				}
 
 				args.physicsBody = -1;
-				if (Ptr<const PhysicsBody> pPhysicsBody = MeshData->GetPhysicsBody())
+				if (TSharedPtr<const FPhysicsBody> pPhysicsBody = MeshData->GetPhysicsBody())
 				{
-					args.physicsBody = program.AddConstant(pPhysicsBody.get());
+					args.physicsBody = program.AddConstant(pPhysicsBody);
 					MeshData->SetPhysicsBody(nullptr);
 				}
 
 				// Use a map-based deduplication
-				mu::Ptr<const mu::Mesh> Key = MeshData;
+				TSharedPtr<const FMesh> Key = MeshData;
 				const int32* IndexPtr = Options->MeshConstantMap.Find(Key);
 				if (!IndexPtr)
 				{
-					args.value = program.AddConstant(MeshData.get());
+					args.value = program.AddConstant(MeshData);
 
 					int32 DataDescIndex = Options->AdditionalData.SourceMeshPerConstant.Add(SourceDataDescriptor);
 					check(DataDescIndex == args.value);
@@ -446,10 +446,10 @@ namespace mu
 					args.value = *IndexPtr;
 				}
 
-				linkedAddress = (OP::ADDRESS)program.m_opAddress.Num();
-				program.m_opAddress.Add((uint32_t)program.m_byteCode.Num());
-				AppendCode(program.m_byteCode, Type);
-				AppendCode(program.m_byteCode, args);
+				linkedAddress = (OP::ADDRESS)program.OpAddress.Num();
+				program.OpAddress.Add((uint32_t)program.ByteCode.Num());
+				AppendCode(program.ByteCode, Type);
+				AppendCode(program.ByteCode, args);
 			}
 			else
 			{
@@ -460,9 +460,9 @@ namespace mu
 
 				switch (Type)
 				{
-				case OP_TYPE::IM_CONSTANT:
+				case EOpType::IM_CONSTANT:
 				{
-					Ptr<const Image> pTyped = static_cast<const Image*>(GetValue().get());
+					TSharedPtr<const FImage> pTyped = StaticCastSharedPtr<const FImage>(GetValue());
 					check(pTyped);
 
 					if (pTyped->GetSizeX() * pTyped->GetSizeY() == 0)
@@ -480,9 +480,9 @@ namespace mu
 
 					break;
 				}
-				case OP_TYPE::LA_CONSTANT:
+				case EOpType::LA_CONSTANT:
 				{
-					Ptr<const Layout> pTyped = static_cast<const Layout*>(GetValue().get());
+					TSharedPtr<const FLayout> pTyped = StaticCastSharedPtr<const FLayout>(GetValue());
 					check(pTyped);
 					args.value = program.AddConstant(pTyped);
 					break;
@@ -493,10 +493,10 @@ namespace mu
 
 				if (bValidData)
 				{
-					linkedAddress = (OP::ADDRESS)program.m_opAddress.Num();
-					program.m_opAddress.Add((uint32)program.m_byteCode.Num());
-					AppendCode(program.m_byteCode, Type);
-					AppendCode(program.m_byteCode, args);
+					linkedAddress = (OP::ADDRESS)program.OpAddress.Num();
+					program.OpAddress.Add((uint32)program.ByteCode.Num());
+					AppendCode(program.ByteCode, Type);
+					AppendCode(program.ByteCode, args);
 				}
 				else
 				{
@@ -518,10 +518,10 @@ namespace mu
 	{
 		FImageDesc Result;
 
-		if (Type == OP_TYPE::IM_CONSTANT)
+		if (Type == EOpType::IM_CONSTANT)
 		{
 			// TODO: cache to avoid disk loading
-			Ptr<const Image> ConstImage = static_cast<const Image*>(GetValue().get());
+			TSharedPtr<const FImage> ConstImage = StaticCastSharedPtr<const FImage>(GetValue());
 			Result.m_format = ConstImage->GetFormat();
 			Result.m_lods = ConstImage->GetLODCount();
 			Result.m_size = ConstImage->GetSize();
@@ -539,9 +539,9 @@ namespace mu
 	{
 		switch (Type)
 		{
-		case OP_TYPE::LA_CONSTANT:
+		case EOpType::LA_CONSTANT:
 		{
-			Ptr<const Layout> pLayout = static_cast<const Layout*>(GetValue().get());
+			TSharedPtr<const FLayout> pLayout = StaticCastSharedPtr<const FLayout>(GetValue());
 			check(pLayout);
 
 			if (pLayout)
@@ -573,7 +573,7 @@ namespace mu
 		switch (Type)
 		{
 
-		case OP_TYPE::IM_CONSTANT:
+		case EOpType::IM_CONSTANT:
 		{
 			// We didn't find any layout.
 			*pBlockX = 0;
@@ -590,10 +590,10 @@ namespace mu
 	//-------------------------------------------------------------------------------------------------
 	bool ASTOpConstantResource::GetNonBlackRect(FImageRect& maskUsage) const
 	{
-		if (Type == OP_TYPE::IM_CONSTANT)
+		if (Type == EOpType::IM_CONSTANT)
 		{
 			// TODO: cache
-			Ptr<const Image> pMask = static_cast<const Image*>(GetValue().get());
+			TSharedPtr<const FImage> pMask = StaticCastSharedPtr<const FImage>(GetValue());
 			pMask->GetNonBlackRect(maskUsage);
 			return true;
 		}
@@ -609,17 +609,17 @@ namespace mu
 		switch (Type)
 		{
 
-		case OP_TYPE::IM_CONSTANT:
+		case EOpType::IM_CONSTANT:
 		{
-			Ptr<const Image> pImage = static_cast<const Image*>(GetValue().get());
+			const FImage* pImage = static_cast<const FImage*>(GetValue().Get());
 			if (pImage->GetSizeX() <= 0 || pImage->GetSizeY() <= 0)
 			{
 				res = true;
 				colour = FVector4f(0.0f,0.0f,0.0f,1.0f);
 			}
-			else if (pImage->m_flags & Image::IF_IS_PLAIN_COLOUR_VALID)
+			else if (pImage->Flags & FImage::IF_IS_PLAIN_COLOUR_VALID)
 			{
-				if (pImage->m_flags & Image::IF_IS_PLAIN_COLOUR)
+				if (pImage->Flags & FImage::IF_IS_PLAIN_COLOUR)
 				{
 					res = true;
 					colour = pImage->Sample(FVector2f(0, 0));
@@ -634,10 +634,10 @@ namespace mu
 				if (pImage->IsPlainColour(colour))
 				{
 					res = true;
-					pImage->m_flags |= Image::IF_IS_PLAIN_COLOUR;
+					pImage->Flags |= FImage::IF_IS_PLAIN_COLOUR;
 				}
 
-				pImage->m_flags |= Image::IF_IS_PLAIN_COLOUR_VALID;
+				pImage->Flags |= FImage::IF_IS_PLAIN_COLOUR_VALID;
 			}
 			break;
 		}
@@ -664,7 +664,7 @@ namespace mu
 
 
 	//-------------------------------------------------------------------------------------------------
-	mu::Ptr<const RefCounted> ASTOpConstantResource::GetValue() const
+	TSharedPtr<const FResource> ASTOpConstantResource::GetValue() const
 	{
 		if (LoadedValue)
 		{
@@ -675,11 +675,11 @@ namespace mu
 			switch (Type)
 			{
 
-			case OP_TYPE::IM_CONSTANT:
+			case EOpType::IM_CONSTANT:
 			{
-				Ptr<ResourceProxy<Image>> typedProxy = static_cast<ResourceProxy<Image>*>(Proxy.get());
-				Ptr<const Image> r = typedProxy->Get();
-				return r;
+				Ptr<TResourceProxy<FImage>> typedProxy = static_cast<TResourceProxy<FImage>*>(Proxy.get());
+				TSharedPtr<const FImage> Resource = typedProxy->Get();
+				return Resource;
 			}
 
 			default:
@@ -693,72 +693,72 @@ namespace mu
 
 
 	//-------------------------------------------------------------------------------------------------
-	void ASTOpConstantResource::SetValue(const Ptr<const RefCounted>& v, FProxyFileContext* DiskCacheContext)
+	void ASTOpConstantResource::SetValue(const TSharedPtr<const FResource>& Value, FProxyFileContext* DiskCacheContext)
 	{
 		MUTABLE_CPUPROFILER_SCOPE(ASTOpConstantResource_SetValue);
 
 		switch (Type)
 		{
-		case OP_TYPE::IM_CONSTANT:
+		case EOpType::IM_CONSTANT:
 		{
-			Ptr<const Image> r = static_cast<const Image*>(v.get());
+			TSharedPtr<const FImage> Resource = StaticCastSharedPtr<const FImage>(Value);
 
-			OutputHashStream stream;
+			FOutputHashStream stream;
 			{
 				MUTABLE_CPUPROFILER_SCOPE(Serialize);
-				OutputArchive arch(&stream);
-				Image::Serialise(r.get(), arch);
+				FOutputArchive arch(&stream);
+				FImage::Serialise(Resource.Get(), arch);
 			}
 
 			ValueHash = stream.GetHash();
 
 			if (DiskCacheContext)
 			{
-				Proxy = new ResourceProxyTempFile<Image>(r.get(), *DiskCacheContext);
+				Proxy = new ResourceProxyTempFile<FImage>(Resource, *DiskCacheContext);
 			}
 			else
 			{
-				LoadedValue = r;
+				LoadedValue = Resource;
 			}
 			break;
 		}
 
-		case OP_TYPE::ME_CONSTANT:
+		case EOpType::ME_CONSTANT:
 		{
-			Ptr<const Mesh> r = static_cast<const Mesh*>(v.get());
+			TSharedPtr<const FMesh> Resource = StaticCastSharedPtr<const FMesh>(Value);
 
-			OutputHashStream stream;
+			FOutputHashStream stream;
 			{
 				MUTABLE_CPUPROFILER_SCOPE(Serialize);
-				OutputArchive arch(&stream);
-				Mesh::Serialise(r.get(), arch);
+				FOutputArchive arch(&stream);
+				FMesh::Serialise(Resource.Get(), arch);
 			}
 
 			ValueHash = stream.GetHash();
 
-			LoadedValue = v;
+			LoadedValue = Resource;
 			break;
 		}
 
-		case OP_TYPE::LA_CONSTANT:
+		case EOpType::LA_CONSTANT:
 		{
-			Ptr<const Layout> r = static_cast<const Layout*>(v.get());
+			TSharedPtr<const FLayout> Resource = StaticCastSharedPtr<const FLayout>(Value);
 
-			OutputHashStream stream;
+			FOutputHashStream stream;
 			{
 				MUTABLE_CPUPROFILER_SCOPE(Serialize);
-				OutputArchive arch(&stream);
-				Layout::Serialise(r.get(), arch);
+				FOutputArchive arch(&stream);
+				FLayout::Serialise(Resource.Get(), arch);
 			}
 
 			ValueHash = stream.GetHash();
 
-			LoadedValue = v;
+			LoadedValue = Resource;
 			break;
 		}
 
 		default:
-			LoadedValue = v;
+			LoadedValue = Value;
 			break;
 		}
 	}
@@ -766,13 +766,13 @@ namespace mu
 
 	mu::Ptr<ImageSizeExpression> ASTOpConstantResource::GetImageSizeExpression() const
 	{
-		if (Type==OP_TYPE::IM_CONSTANT)
+		if (Type==EOpType::IM_CONSTANT)
 		{
-			Ptr<ImageSizeExpression> pRes = new ImageSizeExpression;
-			pRes->type = ImageSizeExpression::ISET_CONSTANT;
-			Ptr<const Image> pConst = static_cast<const Image*>(GetValue().get());
-			pRes->size = pConst->GetSize();
-			return pRes;
+			Ptr<ImageSizeExpression> Res = new ImageSizeExpression;
+			Res->type = ImageSizeExpression::ISET_CONSTANT;
+			TSharedPtr<const FImage> Const = StaticCastSharedPtr<const FImage>(GetValue());
+			Res->size = Const->GetSize();
+			return Res;
 		}
 
 		return nullptr;
diff --git a/Engine/Plugins/Mutable/Source/MutableTools/Private/MuT/ASTOpConstantResource.h b/Engine/Plugins/Mutable/Source/MutableTools/Private/MuT/ASTOpConstantResource.h
--- a/Engine/Plugins/Mutable/Source/MutableTools/Private/MuT/ASTOpConstantResource.h
+++ b/Engine/Plugins/Mutable/Source/MutableTools/Private/MuT/ASTOpConstantResource.h
@@ -21,7 +21,7 @@ namespace mu
 	private:
 
 		//!
-		Ptr<const RefCounted> LoadedValue;
+		TSharedPtr<const FResource> LoadedValue;
 		Ptr<RefCounted> Proxy;
 
 		//! Value hash
@@ -33,7 +33,7 @@ namespace mu
 	public:
 
 		//! Type of constant
-		OP_TYPE Type;
+		EOpType Type;
 
 		/** Source data descriptor. */
 		FSourceDataDescriptor SourceDataDescriptor;
@@ -48,16 +48,16 @@ namespace mu
 		uint64 GetValueHash() const;
 
 		//! Get a copy of the stored value
-		Ptr<const RefCounted> GetValue() const;
+		TSharedPtr<const FResource> GetValue() const;
 
 		/** Set the value to store in this op.
 		* If the DiskCacheContext is not null, the disk cache will be used.
 		*/
-		void SetValue(const Ptr<const RefCounted>& v, FProxyFileContext* DiskCacheContext);
+		void SetValue(const TSharedPtr<const FResource>& Value, FProxyFileContext* DiskCacheContext);
 
 
 		// ASTOp interface
-		virtual OP_TYPE GetOpType() const override { return Type; }
+		virtual EOpType GetOpType() const override { return Type; }
 		virtual void ForEachChild(const TFunctionRef<void(ASTChild&)>) override;
 		virtual bool IsEqual(const ASTOp& otherUntyped) const override;
 		virtual Ptr<ASTOp> Clone(MapChildFuncRef mapChild) const override;
diff --git a/Engine/Plugins/Mutable/Source/MutableTools/Private/MuT/ASTOpConstantString.cpp b/Engine/Plugins/Mutable/Source/MutableTools/Private/MuT/ASTOpConstantString.cpp
--- a/Engine/Plugins/Mutable/Source/MutableTools/Private/MuT/ASTOpConstantString.cpp
+++ b/Engine/Plugins/Mutable/Source/MutableTools/Private/MuT/ASTOpConstantString.cpp
@@ -54,10 +54,10 @@ namespace mu
 			FMemory::Memset(&args, 0, sizeof(args));
 			args.value = program.AddConstant(value);
 
-			linkedAddress = (OP::ADDRESS)program.m_opAddress.Num();
-			program.m_opAddress.Add((uint32_t)program.m_byteCode.Num());
-			AppendCode(program.m_byteCode, OP_TYPE::ST_CONSTANT);
-			AppendCode(program.m_byteCode, args);
+			linkedAddress = (OP::ADDRESS)program.OpAddress.Num();
+			program.OpAddress.Add((uint32_t)program.ByteCode.Num());
+			AppendCode(program.ByteCode, EOpType::ST_CONSTANT);
+			AppendCode(program.ByteCode, args);
 		}
 	}
 
diff --git a/Engine/Plugins/Mutable/Source/MutableTools/Private/MuT/ASTOpConstantString.h b/Engine/Plugins/Mutable/Source/MutableTools/Private/MuT/ASTOpConstantString.h
--- a/Engine/Plugins/Mutable/Source/MutableTools/Private/MuT/ASTOpConstantString.h
+++ b/Engine/Plugins/Mutable/Source/MutableTools/Private/MuT/ASTOpConstantString.h
@@ -23,7 +23,7 @@ struct FProgram;
 		FString value;
 
 	public:
-		OP_TYPE GetOpType() const override { return OP_TYPE::ST_CONSTANT; }
+		EOpType GetOpType() const override { return EOpType::ST_CONSTANT; }
 		uint64 Hash() const override;
 		void ForEachChild(const TFunctionRef<void(ASTChild&)>) override;
 		bool IsEqual(const ASTOp& otherUntyped) const override;
diff --git a/Engine/Plugins/Mutable/Source/MutableTools/Private/MuT/ASTOpImageCompose.cpp b/Engine/Plugins/Mutable/Source/MutableTools/Private/MuT/ASTOpImageCompose.cpp
--- a/Engine/Plugins/Mutable/Source/MutableTools/Private/MuT/ASTOpImageCompose.cpp
+++ b/Engine/Plugins/Mutable/Source/MutableTools/Private/MuT/ASTOpImageCompose.cpp
@@ -59,7 +59,7 @@ bool ASTOpImageCompose::IsEqual(const ASTOp& otherUntyped) const
 //-------------------------------------------------------------------------------------------------
 uint64 ASTOpImageCompose::Hash() const
 {
-	uint64 res = std::hash<OP_TYPE>()(OP_TYPE::IM_COMPOSE);
+	uint64 res = std::hash<EOpType>()(EOpType::IM_COMPOSE);
     hash_combine( res, Layout.child().get() );
 	hash_combine(res, Base.child().get());
 	hash_combine(res, BlockImage.child().get());
@@ -107,10 +107,10 @@ void ASTOpImageCompose::Link( FProgram& program, FLinkerOptions*)
 		if (Mask) args.mask = Mask->linkedAddress;
         args.BlockId = BlockId;
 
-        linkedAddress = (OP::ADDRESS)program.m_opAddress.Num();
-        program.m_opAddress.Add((uint32_t)program.m_byteCode.Num());
-        AppendCode(program.m_byteCode,OP_TYPE::IM_COMPOSE);
-        AppendCode(program.m_byteCode,args);
+        linkedAddress = (OP::ADDRESS)program.OpAddress.Num();
+        program.OpAddress.Add((uint32_t)program.ByteCode.Num());
+        AppendCode(program.ByteCode,EOpType::IM_COMPOSE);
+        AppendCode(program.ByteCode,args);
     }
 
 }
@@ -235,14 +235,14 @@ mu::Ptr<ASTOp> ASTOpImageCompose::OptimiseSemantic(const FModelOptimizationOptio
 	mu::Ptr<ASTOp> layoutAt = Layout.child();
 	if (layoutAt
 		&&
-		layoutAt->GetOpType() == OP_TYPE::LA_CONSTANT
+		layoutAt->GetOpType() == EOpType::LA_CONSTANT
 		&&
 		baseAt
 		&&
 		blockAt)
 	{
 		const ASTOpConstantResource* typedLayout = static_cast<const ASTOpConstantResource*>(layoutAt.get());
-		mu::Ptr<const mu::Layout> pLayout = static_cast<const mu::Layout*>(typedLayout->GetValue().get());
+		const mu::FLayout* pLayout = static_cast<const mu::FLayout*>(typedLayout->GetValue().Get());
 
 		// Constant single-block full layout?
 		if (pLayout->GetBlockCount() == 1
@@ -280,7 +280,7 @@ mu::Ptr<ASTOp> ASTOpImageCompose::OptimiseSemantic(const FModelOptimizationOptio
 				baseDesc.m_size[1] != 0)
 			{
 				mu::Ptr<ASTOpFixed> resize = new ASTOpFixed;
-				resize->op.type = OP_TYPE::IM_RESIZE;
+				resize->op.type = EOpType::IM_RESIZE;
 				resize->op.args.ImageResize.size[0] = baseDesc.m_size[0];
 				resize->op.args.ImageResize.size[1] = baseDesc.m_size[1];
 				resize->SetChild(resize->op.args.ImageResize.source, at);
diff --git a/Engine/Plugins/Mutable/Source/MutableTools/Private/MuT/ASTOpImageCompose.h b/Engine/Plugins/Mutable/Source/MutableTools/Private/MuT/ASTOpImageCompose.h
--- a/Engine/Plugins/Mutable/Source/MutableTools/Private/MuT/ASTOpImageCompose.h
+++ b/Engine/Plugins/Mutable/Source/MutableTools/Private/MuT/ASTOpImageCompose.h
@@ -29,7 +29,7 @@ namespace mu
 		ASTOpImageCompose(const ASTOpImageCompose&) = delete;
 		~ASTOpImageCompose();
 
-		virtual OP_TYPE GetOpType() const override { return OP_TYPE::IM_COMPOSE; }
+		virtual EOpType GetOpType() const override { return EOpType::IM_COMPOSE; }
 		virtual uint64 Hash() const override;
 		virtual bool IsEqual(const ASTOp& otherUntyped) const override;
 		virtual Ptr<ASTOp> Clone(MapChildFuncRef mapChild) const override;
diff --git a/Engine/Plugins/Mutable/Source/MutableTools/Private/MuT/ASTOpImageCrop.cpp b/Engine/Plugins/Mutable/Source/MutableTools/Private/MuT/ASTOpImageCrop.cpp
--- a/Engine/Plugins/Mutable/Source/MutableTools/Private/MuT/ASTOpImageCrop.cpp
+++ b/Engine/Plugins/Mutable/Source/MutableTools/Private/MuT/ASTOpImageCrop.cpp
@@ -46,7 +46,7 @@ namespace mu
 
 	uint64 ASTOpImageCrop::Hash() const
 	{
-		uint64 res = std::hash<OP_TYPE>()(OP_TYPE::IM_CROP);
+		uint64 res = std::hash<EOpType>()(EOpType::IM_CROP);
 		hash_combine(res, Source.child().get());
 		hash_combine(res, Min[0]);
 		hash_combine(res, Min[1]);
@@ -86,10 +86,10 @@ namespace mu
 			args.sizeX = Size[0];
 			args.sizeY = Size[1];
 
-			linkedAddress = (OP::ADDRESS)program.m_opAddress.Num();
-			program.m_opAddress.Add((uint32_t)program.m_byteCode.Num());
-			AppendCode(program.m_byteCode, OP_TYPE::IM_CROP);
-			AppendCode(program.m_byteCode, args);
+			linkedAddress = (OP::ADDRESS)program.OpAddress.Num();
+			program.OpAddress.Add((uint32_t)program.ByteCode.Num());
+			AppendCode(program.ByteCode, EOpType::IM_CROP);
+			AppendCode(program.ByteCode, args);
 		}
 
 	}
@@ -177,10 +177,10 @@ namespace mu
 		Ptr<ASTOp> sourceAt = Source.child();
 
 		// The instruction can be sunk
-		OP_TYPE sourceType = sourceAt->GetOpType();
+		EOpType sourceType = sourceAt->GetOpType();
 		switch (sourceType)
 		{
-		case OP_TYPE::IM_PLAINCOLOUR:
+		case EOpType::IM_PLAINCOLOUR:
 		{
 			Ptr<ASTOpFixed> NewOp = mu::Clone<ASTOpFixed>(sourceAt.get());
 			NewOp->op.args.ImagePlainColour.size[0] = Size[0];
@@ -208,7 +208,7 @@ namespace mu
 		switch (sourceAt->GetOpType())
 		{
 			// In case we have other operations with special optimisation rules.
-		case OP_TYPE::NONE:
+		case EOpType::NONE:
 			break;
 
 		default:
@@ -227,16 +227,16 @@ namespace mu
 
 	Ptr<ASTOp> Sink_ImageCropAST::Apply(const ASTOpImageCrop* InRoot)
 	{
-		check(InRoot->GetOpType() == OP_TYPE::IM_CROP);
+		check(InRoot->GetOpType() == EOpType::IM_CROP);
 
-		m_root = InRoot;
+		Root = InRoot;
 		OldToNew.Reset();
 
-		m_initialSource = InRoot->Source.child();
-		Ptr<ASTOp> newSource = Visit(m_initialSource, InRoot);
+		InitialSource = InRoot->Source.child();
+		Ptr<ASTOp> newSource = Visit(InitialSource, InRoot);
 
 		// If there is any change, it is the new root.
-		if (newSource != m_initialSource)
+		if (newSource != InitialSource)
 		{
 			return newSource;
 		}
@@ -261,7 +261,7 @@ namespace mu
 		switch (at->GetOpType())
 		{
 
-		case OP_TYPE::IM_CONDITIONAL:
+		case EOpType::IM_CONDITIONAL:
 		{
 			// We move down the two paths
 			Ptr<ASTOpConditional> newOp = mu::Clone<ASTOpConditional>(at);
@@ -271,7 +271,7 @@ namespace mu
 			break;
 		}
 
-		case OP_TYPE::IM_SWITCH:
+		case EOpType::IM_SWITCH:
 		{
 			// We move down all the paths
 			Ptr<ASTOpSwitch> newOp = mu::Clone<ASTOpSwitch>(at);
@@ -284,7 +284,7 @@ namespace mu
 			break;
 		}
 
-		case OP_TYPE::IM_PIXELFORMAT:
+		case EOpType::IM_PIXELFORMAT:
 		{
 			Ptr<ASTOpImagePixelFormat> nop = mu::Clone<ASTOpImagePixelFormat>(at);
 			nop->Source = Visit(nop->Source.child(), currentCropOp);
@@ -292,7 +292,7 @@ namespace mu
 			break;
 		}
 
-		case OP_TYPE::IM_PATCH:
+		case EOpType::IM_PATCH:
 		{
 			const ASTOpImagePatch* typedPatch = static_cast<const ASTOpImagePatch*>(at.get());
 
@@ -341,7 +341,7 @@ namespace mu
 			break;
 		}
 
-		case OP_TYPE::IM_CROP:
+		case EOpType::IM_CROP:
 		{
 			// We can combine the two crops into a possibly smaller crop
 			const ASTOpImageCrop* childCrop = static_cast<const ASTOpImageCrop*>(at.get());
@@ -395,7 +395,7 @@ namespace mu
 			break;
 		}
 
-		case OP_TYPE::IM_LAYER:
+		case EOpType::IM_LAYER:
 		{
 			// We move the op down the arguments
 			Ptr<ASTOpImageLayer> nop = mu::Clone<ASTOpImageLayer>(at);
@@ -413,7 +413,7 @@ namespace mu
 			break;
 		}
 
-		case OP_TYPE::IM_LAYERCOLOUR:
+		case EOpType::IM_LAYERCOLOUR:
 		{
 			// We move the op down the arguments
 			Ptr<ASTOpImageLayerColor> nop = mu::Clone<ASTOpImageLayerColor>(at);
@@ -428,7 +428,7 @@ namespace mu
 			break;
 		}
 
-		case OP_TYPE::IM_DISPLACE:
+		case EOpType::IM_DISPLACE:
 		{
 			// We move the op down the arguments
 			Ptr<ASTOpFixed> nop = mu::Clone<ASTOpFixed>(at);
@@ -443,7 +443,7 @@ namespace mu
 			break;
 		}
 
-		case OP_TYPE::IM_RASTERMESH:
+		case EOpType::IM_RASTERMESH:
 		{
 			// We add cropping data to the raster mesh if it doesn't have any
 			// \TODO: Is is possible to hit 2 crops on a raster mesh? Combine the crop.
@@ -470,7 +470,7 @@ namespace mu
 			break;
 		}
 
-		case OP_TYPE::IM_INTERPOLATE:
+		case EOpType::IM_INTERPOLATE:
 		{
 			// Move the op  down all the paths
 			auto newOp = mu::Clone<ASTOpFixed>(at);
@@ -491,10 +491,10 @@ namespace mu
 		}
 
 		// end on line, replace with crop
-		if (at == newAt && at != m_initialSource && !skipSinking)
+		if (at == newAt && at != InitialSource && !skipSinking)
 		{
 			Ptr<ASTOpImageCrop> newOp = mu::Clone<ASTOpImageCrop>(currentCropOp);
-			check(newOp->GetOpType() == OP_TYPE::IM_CROP);
+			check(newOp->GetOpType() == EOpType::IM_CROP);
 
 			newOp->Source = at;
 
diff --git a/Engine/Plugins/Mutable/Source/MutableTools/Private/MuT/ASTOpImageCrop.h b/Engine/Plugins/Mutable/Source/MutableTools/Private/MuT/ASTOpImageCrop.h
--- a/Engine/Plugins/Mutable/Source/MutableTools/Private/MuT/ASTOpImageCrop.h
+++ b/Engine/Plugins/Mutable/Source/MutableTools/Private/MuT/ASTOpImageCrop.h
@@ -28,7 +28,7 @@ namespace mu
 		ASTOpImageCrop(const ASTOpImageCrop&) = delete;
 		~ASTOpImageCrop();
 
-		virtual OP_TYPE GetOpType() const override { return OP_TYPE::IM_CROP; }
+		virtual EOpType GetOpType() const override { return EOpType::IM_CROP; }
 		virtual uint64 Hash() const override;
 		virtual bool IsEqual(const ASTOp& otherUntyped) const override;
 		virtual Ptr<ASTOp> Clone(MapChildFuncRef mapChild) const override;
diff --git a/Engine/Plugins/Mutable/Source/MutableTools/Private/MuT/ASTOpImageLayer.cpp b/Engine/Plugins/Mutable/Source/MutableTools/Private/MuT/ASTOpImageLayer.cpp
--- a/Engine/Plugins/Mutable/Source/MutableTools/Private/MuT/ASTOpImageLayer.cpp
+++ b/Engine/Plugins/Mutable/Source/MutableTools/Private/MuT/ASTOpImageLayer.cpp
@@ -57,7 +57,7 @@ namespace mu
 	//-------------------------------------------------------------------------------------------------
 	uint64 ASTOpImageLayer::Hash() const
 	{
-		uint64 res = std::hash<OP_TYPE>()(GetOpType());
+		uint64 res = std::hash<EOpType>()(GetOpType());
 		hash_combine(res, base.child().get());
 		hash_combine(res, blend.child().get());
 		hash_combine(res, mask.child().get());
@@ -107,10 +107,10 @@ namespace mu
 			if (blend) args.blended = blend->linkedAddress;
 			if (mask) args.mask = mask->linkedAddress;
 
-			linkedAddress = (OP::ADDRESS)program.m_opAddress.Num();
-			program.m_opAddress.Add((uint32)program.m_byteCode.Num());
-			AppendCode(program.m_byteCode, GetOpType());
-			AppendCode(program.m_byteCode, args);
+			linkedAddress = (OP::ADDRESS)program.OpAddress.Num();
+			program.OpAddress.Add((uint32)program.ByteCode.Num());
+			AppendCode(program.ByteCode, GetOpType());
+			AppendCode(program.ByteCode, args);
 		}
 	}
 
@@ -191,7 +191,7 @@ namespace mu
 		}
 
 		// Convert to image layer color if blend is plain
-		if (!at && blendAt && blendAt->GetOpType() == OP_TYPE::IM_PLAINCOLOUR)
+		if (!at && blendAt && blendAt->GetOpType() == EOpType::IM_PLAINCOLOUR)
 		{
 			// TODO: May some flags be supported?
 			if (Flags == 0)
@@ -260,7 +260,7 @@ namespace mu
 
 					switch (CurrentBlend->GetOpType())
 					{
-					case OP_TYPE::IM_LAYERCOLOUR:
+					case EOpType::IM_LAYERCOLOUR:
 					{
 						const ASTOpImageLayerColor* BlendLayer = static_cast<const ASTOpImageLayerColor*>(CurrentBlend.get());
 						if (BlendLayer->blendTypeAlpha == EBlendType::BT_NONE)
@@ -271,7 +271,7 @@ namespace mu
 						break;
 					}
 
-					case OP_TYPE::IM_LAYER:
+					case EOpType::IM_LAYER:
 					{
 						const ASTOpImageLayer* BlendLayer = static_cast<const ASTOpImageLayer*>(CurrentBlend.get());
 						if (BlendLayer->blendTypeAlpha == EBlendType::BT_NONE)
@@ -296,7 +296,7 @@ namespace mu
 				switch (CurrentMask->GetOpType())
 				{
 
-				case OP_TYPE::IM_DISPLACE:
+				case EOpType::IM_DISPLACE:
 				{
 					const ASTOpFixed* MaskDisplace = static_cast<const ASTOpFixed*>(CurrentMask.get());
 					const ASTOpFixed* BlendDisplace = static_cast<const ASTOpFixed*>(CurrentBlend.get());
@@ -313,7 +313,7 @@ namespace mu
 					break;
 				}
 
-				case OP_TYPE::IM_RASTERMESH:
+				case EOpType::IM_RASTERMESH:
 				{
 					const ASTOpImageRasterMesh* MaskRaster = static_cast<const ASTOpImageRasterMesh*>(CurrentMask.get());
 					const ASTOpImageRasterMesh* BlendRaster = static_cast<const ASTOpImageRasterMesh*>(CurrentBlend.get());
@@ -339,7 +339,7 @@ namespace mu
 					break;
 				}
 
-				case OP_TYPE::IM_RESIZE:
+				case EOpType::IM_RESIZE:
 				{
 					const ASTOpFixed* MaskResize = static_cast<const ASTOpFixed*>(CurrentMask.get());
 					const ASTOpFixed* BlendResize = static_cast<const ASTOpFixed*>(CurrentBlend.get());
@@ -362,7 +362,7 @@ namespace mu
 				}
 			}
 
-			if (CurrentMask && CurrentMask->GetOpType() == OP_TYPE::IM_SWIZZLE)
+			if (CurrentMask && CurrentMask->GetOpType() == EOpType::IM_SWIZZLE)
 			{
 				// End of the possible mask expression chain match should have a swizzle selecting the alpha.
 				const ASTOpImageSwizzle* MaskSwizzle = static_cast<const ASTOpImageSwizzle*>(CurrentMask.get());
@@ -390,7 +390,7 @@ namespace mu
 		if (!at)
 		{
 			// Is the base a swizzle expanding alpha from a texture?
-			if (baseAt->GetOpType() == OP_TYPE::IM_SWIZZLE)
+			if (baseAt->GetOpType() == EOpType::IM_SWIZZLE)
 			{
 				const ASTOpImageSwizzle* TypedBase = static_cast<const ASTOpImageSwizzle*>(baseAt.get());
 				bool bAreAllAlpha = true;
@@ -410,7 +410,7 @@ namespace mu
 			}
 
 			// Is the mask a swizzle expanding alpha from a texture?
-			if (!at && maskAt && maskAt->GetOpType() == OP_TYPE::IM_SWIZZLE)
+			if (!at && maskAt && maskAt->GetOpType() == EOpType::IM_SWIZZLE)
 			{
 				const ASTOpImageSwizzle* TypedBase = static_cast<const ASTOpImageSwizzle*>(maskAt.get());
 				bool bAreAllAlpha = true;
@@ -430,7 +430,7 @@ namespace mu
 			}
 
 			// Is the blend a swizzle selecting alpha?
-			if (Pass>0 && !at && blendAt && blendAt->GetOpType() == OP_TYPE::IM_SWIZZLE && Flags==0)
+			if (Pass>0 && !at && blendAt && blendAt->GetOpType() == EOpType::IM_SWIZZLE && Flags==0)
 			{
 				const ASTOpImageSwizzle* TypedBlend = static_cast<const ASTOpImageSwizzle*>(blendAt.get());
 				if (TypedBlend->Format==EImageFormat::L_UByte
@@ -546,12 +546,12 @@ namespace mu
 		}
 
 		// Promote conditions from the base
-		OP_TYPE baseType = baseAt->GetOpType();
+		EOpType baseType = baseAt->GetOpType();
 		switch (baseType)
 		{
 			// Seems to cause operation explosion in optimizer in bandit model.
 			// moved to generic sink in the default.
-//            case OP_TYPE::IM_CONDITIONAL:
+//            case EOpType::IM_CONDITIONAL:
 //            {
 //                m_modified = true;
 
@@ -569,7 +569,7 @@ namespace mu
 //                break;
 //            }
 
-		case OP_TYPE::IM_SWITCH:
+		case EOpType::IM_SWITCH:
 		{
 			// Warning:
 			// It seems to cause data explosion in optimizer in some models. Because
@@ -633,9 +633,9 @@ namespace mu
 //                OP templateOp = program.m_code[at];
 
 //                CloneNeutralPartialTreeVisitor neutral;
-//                OP_TYPE supportedOps[] = {
-//                    OP_TYPE::IM_CONDITIONAL,
-//                    OP_TYPE::NONE
+//                EOpType supportedOps[] = {
+//                    EOpType::IM_CONDITIONAL,
+//                    EOpType::NONE
 //                };
 //                auto newRoot = neutral.Apply
 //                         (
@@ -658,11 +658,11 @@ namespace mu
 		// If we failed to optimize so far, see if it is worth optimizing the blended branch only.
 		if (!at && blendAt)
 		{
-			OP_TYPE BlendType = blendAt->GetOpType();
+			EOpType BlendType = blendAt->GetOpType();
 			switch (BlendType)
 			{
 
-			case OP_TYPE::IM_SWITCH:
+			case EOpType::IM_SWITCH:
 			{
 				const ASTOpSwitch* BlendSwitch = static_cast<const ASTOpSwitch*>(blendAt.get());
 
@@ -673,7 +673,7 @@ namespace mu
 					if (BlendSwitch->cases[v].branch)
 					{
 						// \TODO: Use the smarter query function to detect plain images?
-						if (BlendSwitch->cases[v].branch->GetOpType() == OP_TYPE::IM_PLAINCOLOUR)
+						if (BlendSwitch->cases[v].branch->GetOpType() == EOpType::IM_PLAINCOLOUR)
 						{
 							bWorthSinking = true;
 							break;
@@ -685,7 +685,7 @@ namespace mu
 				{
 					bool bMaskIsCompatibleSwitch = false;
 					const ASTOpSwitch* MaskSwitch = nullptr;
-					if (maskAt && maskAt->GetOpType()== OP_TYPE::IM_SWITCH)
+					if (maskAt && maskAt->GetOpType()== EOpType::IM_SWITCH)
 					{
 						MaskSwitch = static_cast<const ASTOpSwitch*>(maskAt.get());
 						bMaskIsCompatibleSwitch = MaskSwitch->IsCompatibleWith(BlendSwitch);
@@ -733,11 +733,11 @@ namespace mu
 		// If we failed to optimize so far, see if it is worth optimizing the mask branch only.
 		if (!at && maskAt)
 		{
-			OP_TYPE MaskType = maskAt->GetOpType();
+			EOpType MaskType = maskAt->GetOpType();
 			switch (MaskType)
 			{
 
-			case OP_TYPE::IM_SWITCH:
+			case EOpType::IM_SWITCH:
 			{
 				const ASTOpSwitch* MaskSwitch = static_cast<const ASTOpSwitch*>(maskAt.get());
 
@@ -748,7 +748,7 @@ namespace mu
 					if (MaskSwitch->cases[v].branch)
 					{
 						// \TODO: Use the smarter query function to detect plain images?
-						if (MaskSwitch->cases[v].branch->GetOpType() == OP_TYPE::IM_PLAINCOLOUR)
+						if (MaskSwitch->cases[v].branch->GetOpType() == EOpType::IM_PLAINCOLOUR)
 						{
 							bWorthSinking = true;
 							break;
diff --git a/Engine/Plugins/Mutable/Source/MutableTools/Private/MuT/ASTOpImageLayer.h b/Engine/Plugins/Mutable/Source/MutableTools/Private/MuT/ASTOpImageLayer.h
--- a/Engine/Plugins/Mutable/Source/MutableTools/Private/MuT/ASTOpImageLayer.h
+++ b/Engine/Plugins/Mutable/Source/MutableTools/Private/MuT/ASTOpImageLayer.h
@@ -43,7 +43,7 @@ namespace mu
 		ASTOpImageLayer(const ASTOpImageLayer&) = delete;
 		~ASTOpImageLayer();
 
-		virtual OP_TYPE GetOpType() const override { return OP_TYPE::IM_LAYER; }
+		virtual EOpType GetOpType() const override { return EOpType::IM_LAYER; }
 		virtual uint64 Hash() const override;
 		virtual bool IsEqual(const ASTOp& otherUntyped) const override;
 		virtual Ptr<ASTOp> Clone(MapChildFuncRef mapChild) const override;
diff --git a/Engine/Plugins/Mutable/Source/MutableTools/Private/MuT/ASTOpImageLayerColor.cpp b/Engine/Plugins/Mutable/Source/MutableTools/Private/MuT/ASTOpImageLayerColor.cpp
--- a/Engine/Plugins/Mutable/Source/MutableTools/Private/MuT/ASTOpImageLayerColor.cpp
+++ b/Engine/Plugins/Mutable/Source/MutableTools/Private/MuT/ASTOpImageLayerColor.cpp
@@ -52,7 +52,7 @@ namespace mu
 	//-------------------------------------------------------------------------------------------------
 	uint64 ASTOpImageLayerColor::Hash() const
 	{
-		uint64 res = std::hash<OP_TYPE>()(GetOpType());
+		uint64 res = std::hash<EOpType>()(GetOpType());
 		hash_combine(res, base.child().get());
 		hash_combine(res, color.child().get());
 		hash_combine(res, mask.child().get());
@@ -103,10 +103,10 @@ namespace mu
 			if (color) args.colour = color->linkedAddress;
 			if (mask) args.mask = mask->linkedAddress;
 
-			linkedAddress = (OP::ADDRESS)program.m_opAddress.Num();
-			program.m_opAddress.Add((uint32)program.m_byteCode.Num());
-			AppendCode(program.m_byteCode, GetOpType());
-			AppendCode(program.m_byteCode, args);
+			linkedAddress = (OP::ADDRESS)program.OpAddress.Num();
+			program.OpAddress.Add((uint32)program.ByteCode.Num());
+			AppendCode(program.ByteCode, GetOpType());
+			AppendCode(program.ByteCode, args);
 		}
 	}
 
@@ -205,7 +205,7 @@ namespace mu
 			bool bAlphaUnchanged = blendTypeAlpha == EBlendType::BT_NONE;
 
 			FVector4f ColorConst(0,0,0,1);
-			if (!color || color.child()->GetOpType()==OP_TYPE::CO_CONSTANT)
+			if (!color || color.child()->GetOpType()==EOpType::CO_CONSTANT)
 			{
 				if (color)
 				{
@@ -277,7 +277,7 @@ namespace mu
 		if (!at)
 		{
 			// Is the base a swizzle getting expanding alpha from the same texture?
-			if (base.child()->GetOpType() == OP_TYPE::IM_SWIZZLE)
+			if (base.child()->GetOpType() == EOpType::IM_SWIZZLE)
 			{
 				const ASTOpImageSwizzle* TypedBase = static_cast<const ASTOpImageSwizzle*>(base.child().get());
 				bool bAreAllSameAlpha = true;
@@ -328,12 +328,12 @@ namespace mu
 		auto maskAt = mask.child();
 
 		// Promote conditions from the base
-		OP_TYPE baseType = baseAt->GetOpType();
+		EOpType baseType = baseAt->GetOpType();
 		switch (baseType)
 		{
 			// Seems to cause operation explosion in optimizer in bandit model.
 			// moved to generic sink in the default.
-//            case OP_TYPE::IM_CONDITIONAL:
+//            case EOpType::IM_CONDITIONAL:
 //            {
 //                m_modified = true;
 
@@ -351,7 +351,7 @@ namespace mu
 //                break;
 //            }
 
-		case OP_TYPE::IM_SWITCH:
+		case EOpType::IM_SWITCH:
 		{
 			// Warning:
 			// It seems to cause data explosion in optimizer in some models. Because
@@ -390,7 +390,7 @@ namespace mu
 			break;
 		}
 
-		case OP_TYPE::IM_DISPLACE:
+		case EOpType::IM_DISPLACE:
 		{
 			// Mask not supported yet. If there is a mask it wouldn't be correct to sink
 			// unless the mask was a similar displace.
@@ -410,7 +410,7 @@ namespace mu
 			break;
 		}
 
-		case OP_TYPE::IM_RASTERMESH:
+		case EOpType::IM_RASTERMESH:
 		{
 			// Mask not supported yet. If there is a mask it wouldn't be correct to sink.				
 			if (maskAt)
diff --git a/Engine/Plugins/Mutable/Source/MutableTools/Private/MuT/ASTOpImageLayerColor.h b/Engine/Plugins/Mutable/Source/MutableTools/Private/MuT/ASTOpImageLayerColor.h
--- a/Engine/Plugins/Mutable/Source/MutableTools/Private/MuT/ASTOpImageLayerColor.h
+++ b/Engine/Plugins/Mutable/Source/MutableTools/Private/MuT/ASTOpImageLayerColor.h
@@ -43,7 +43,7 @@ namespace mu
 		ASTOpImageLayerColor(const ASTOpImageLayerColor&) = delete;
 		~ASTOpImageLayerColor();
 
-		virtual OP_TYPE GetOpType() const override { return OP_TYPE::IM_LAYERCOLOUR; }
+		virtual EOpType GetOpType() const override { return EOpType::IM_LAYERCOLOUR; }
 		virtual uint64 Hash() const override;
 		virtual bool IsEqual(const ASTOp& otherUntyped) const override;
 		virtual Ptr<ASTOp> Clone(MapChildFuncRef mapChild) const override;
diff --git a/Engine/Plugins/Mutable/Source/MutableTools/Private/MuT/ASTOpImageMakeGrowMap.cpp b/Engine/Plugins/Mutable/Source/MutableTools/Private/MuT/ASTOpImageMakeGrowMap.cpp
--- a/Engine/Plugins/Mutable/Source/MutableTools/Private/MuT/ASTOpImageMakeGrowMap.cpp
+++ b/Engine/Plugins/Mutable/Source/MutableTools/Private/MuT/ASTOpImageMakeGrowMap.cpp
@@ -80,11 +80,11 @@ namespace mu
 			args.border = Border;
 			if (Mask) args.mask = Mask->linkedAddress;
 
-			linkedAddress = (OP::ADDRESS)program.m_opAddress.Num();
+			linkedAddress = (OP::ADDRESS)program.OpAddress.Num();
 			//program.m_code.push_back(op);
-			program.m_opAddress.Add((uint32_t)program.m_byteCode.Num());
-			AppendCode(program.m_byteCode, GetOpType());
-			AppendCode(program.m_byteCode, args);
+			program.OpAddress.Add((uint32_t)program.ByteCode.Num());
+			AppendCode(program.ByteCode, GetOpType());
+			AppendCode(program.ByteCode, args);
 		}
 	}
 
@@ -96,7 +96,7 @@ namespace mu
 		switch (Mask.child()->GetOpType())
 		{
 
-		case OP_TYPE::IM_CONDITIONAL:
+		case EOpType::IM_CONDITIONAL:
 		{
 			// We move the format down the two paths
 			Ptr<ASTOpConditional> nop = mu::Clone<ASTOpConditional>(Mask.child());
@@ -113,7 +113,7 @@ namespace mu
 			break;
 		}
 
-		case OP_TYPE::IM_SWITCH:
+		case EOpType::IM_SWITCH:
 		{
 			// Move the format down all the paths
 			Ptr<ASTOpSwitch> nop = mu::Clone<ASTOpSwitch>(Mask.child());
diff --git a/Engine/Plugins/Mutable/Source/MutableTools/Private/MuT/ASTOpImageMakeGrowMap.h b/Engine/Plugins/Mutable/Source/MutableTools/Private/MuT/ASTOpImageMakeGrowMap.h
--- a/Engine/Plugins/Mutable/Source/MutableTools/Private/MuT/ASTOpImageMakeGrowMap.h
+++ b/Engine/Plugins/Mutable/Source/MutableTools/Private/MuT/ASTOpImageMakeGrowMap.h
@@ -27,7 +27,7 @@ namespace mu
 		ASTOpImageMakeGrowMap(const ASTOpImageMakeGrowMap&) = delete;
 		~ASTOpImageMakeGrowMap();
 
-		virtual OP_TYPE GetOpType() const override { return OP_TYPE::IM_MAKEGROWMAP; }
+		virtual EOpType GetOpType() const override { return EOpType::IM_MAKEGROWMAP; }
 		virtual uint64 Hash() const override;
 		virtual bool IsEqual(const ASTOp& otherUntyped) const override;
 		virtual Ptr<ASTOp> Clone(MapChildFuncRef mapChild) const override;
