[mutable] More enforcement of coding standards.

diff --git a/Engine/Plugins/Mutable/Source/MutableTools/Private/MuT/ASTOpImageMipmap.cpp b/Engine/Plugins/Mutable/Source/MutableTools/Private/MuT/ASTOpImageMipmap.cpp
--- a/Engine/Plugins/Mutable/Source/MutableTools/Private/MuT/ASTOpImageMipmap.cpp
+++ b/Engine/Plugins/Mutable/Source/MutableTools/Private/MuT/ASTOpImageMipmap.cpp
@@ -95,16 +95,16 @@ namespace mu
 			args.FilterType = FilterType;
 			if (Source) args.source = Source->linkedAddress;
 
-			linkedAddress = (OP::ADDRESS)program.m_opAddress.Num();
+			linkedAddress = (OP::ADDRESS)program.OpAddress.Num();
 			//program.m_code.push_back(op);
-			program.m_opAddress.Add((uint32_t)program.m_byteCode.Num());
-			AppendCode(program.m_byteCode, GetOpType());
-			AppendCode(program.m_byteCode, args);
+			program.OpAddress.Add((uint32_t)program.ByteCode.Num());
+			AppendCode(program.ByteCode, GetOpType());
+			AppendCode(program.ByteCode, args);
 		}
 	}
 
 
-	mu::Ptr<ASTOp> ASTOpImageMipmap::OptimiseSemantic(const FModelOptimizationOptions& options, int32 Pass) const
+	mu::Ptr<ASTOp> ASTOpImageMipmap::OptimiseSemantic(const FModelOptimizationOptions&, int32 Pass) const
 	{
 		mu::Ptr<ASTOp> at;
 
@@ -125,7 +125,7 @@ namespace mu
 	}
 
 
-	mu::Ptr<ASTOp> ASTOpImageMipmap::OptimiseSink(const FModelOptimizationOptions& options, FOptimizeSinkContext& context) const
+	mu::Ptr<ASTOp> ASTOpImageMipmap::OptimiseSink(const FModelOptimizationOptions&, FOptimizeSinkContext& context) const
 	{
 		mu::Ptr<ASTOp> at;
 
@@ -133,7 +133,7 @@ namespace mu
 		switch (sourceAt->GetOpType())
 		{
 
-		case OP_TYPE::IM_BLANKLAYOUT:
+		case EOpType::IM_BLANKLAYOUT:
 		{
 			// Set the mipmap generation on the blank layout operation
 			Ptr<ASTOpFixed> mop = mu::Clone<ASTOpFixed>(sourceAt);
@@ -143,7 +143,7 @@ namespace mu
 			break;
 		}
 
-		case OP_TYPE::IM_PLAINCOLOUR:
+		case EOpType::IM_PLAINCOLOUR:
 		{
 			// Set the mipmap generation on the plaincolour operation
 			Ptr<ASTOpFixed> mop = mu::Clone<ASTOpFixed>(sourceAt);
@@ -152,7 +152,7 @@ namespace mu
 			break;
 		}
 
-		case OP_TYPE::IM_PIXELFORMAT:
+		case EOpType::IM_PIXELFORMAT:
 		{
 			// Swap unless the mipmap operation builds only the tail or is compressed.
 			// Otherwise, we could fall in a loop of swapping mipmaps and pixelformats.
@@ -286,50 +286,50 @@ namespace mu
 	//---------------------------------------------------------------------------------------------
 	mu::Ptr<ASTOp> Sink_ImageMipmapAST::Apply(const ASTOpImageMipmap* InRoot)
 	{
-		m_root = InRoot;
+		Root = InRoot;
 		OldToNew.Empty();
 
-		if (m_root->bOnlyTail)
+		if (Root->bOnlyTail)
 		{
 			return nullptr;
 		}
 
 		mu::Ptr<ASTOp> newSource;
-		m_initialSource = m_root->Source.child();
+		InitialSource = Root->Source.child();
 
 		// Before sinking, see if it is worth splitting into miptail and mip.
-		if (!m_root->bPreventSplitTail
+		if (!Root->bPreventSplitTail
 			// \TODO: Review this: It seems it fails for a BC1 8x4 texture having blockSize=4, and level=0
 			&&
 			// No need to split if there is no difference between the levels considering blocks
-			m_root->Levels != m_root->BlockLevels 
+			Root->Levels != Root->BlockLevels 
 			)
 		{
 			// the block mipmaps can be done before composition is done.
-			mu::Ptr<ASTOpImageMipmap> newMip = mu::Clone<ASTOpImageMipmap>(m_root);
-			newMip->Levels = m_root->BlockLevels;
-			newMip->BlockLevels = m_root->BlockLevels;
+			mu::Ptr<ASTOpImageMipmap> newMip = mu::Clone<ASTOpImageMipmap>(Root);
+			newMip->Levels = Root->BlockLevels;
+			newMip->BlockLevels = Root->BlockLevels;
 			newMip->bOnlyTail = false;
 
 			// the smallest mipmaps after the composition is done.
-			mu::Ptr<ASTOpImageMipmap> topMipOp = mu::Clone<ASTOpImageMipmap>(m_root);
+			mu::Ptr<ASTOpImageMipmap> topMipOp = mu::Clone<ASTOpImageMipmap>(Root);
 			topMipOp->bOnlyTail = true;
 
 			// Proceed
-			topMipOp->Source = Visit(m_initialSource, newMip.get());
+			topMipOp->Source = Visit(InitialSource, newMip.get());
 
 			newSource = topMipOp;
 		}
 		else
 		{
 			// Proceed
-			newSource = Visit(m_initialSource, m_root);
+			newSource = Visit(InitialSource, Root);
 		}
 
-		m_root = nullptr;
+		Root = nullptr;
 
 		// If there is any change, it is the new root.
-		if (newSource != m_initialSource)
+		if (newSource != InitialSource)
 		{
 			return newSource;
 		}
@@ -356,7 +356,7 @@ namespace mu
 		switch (at->GetOpType())
 		{
 
-		case OP_TYPE::IM_CONDITIONAL:
+		case EOpType::IM_CONDITIONAL:
 		{
 			// We move the op down the two paths
 			Ptr<ASTOpConditional> newOp = mu::Clone<ASTOpConditional>(at);
@@ -366,7 +366,7 @@ namespace mu
 			break;
 		}
 
-		case OP_TYPE::IM_SWITCH:
+		case EOpType::IM_SWITCH:
 		{
 			// We move the op down all the paths
 			Ptr<ASTOpSwitch> newOp = mu::Clone<ASTOpSwitch>(at);
@@ -379,7 +379,7 @@ namespace mu
 			break;
 		}
 
-		case OP_TYPE::IM_COMPOSE:
+		case EOpType::IM_COMPOSE:
 		{
 			const ASTOpImageCompose* typedAt = static_cast<const ASTOpImageCompose*>(at.get());
 			if (!currentMipmapOp->bOnlyTail
@@ -403,7 +403,7 @@ namespace mu
 			break;
 		}
 
-		case OP_TYPE::IM_PATCH:
+		case EOpType::IM_PATCH:
 		{
 			if (!currentMipmapOp->bOnlyTail)
 			{
@@ -479,10 +479,10 @@ namespace mu
 		}
 
 		// end on line, replace with mipmap
-		if (at == newAt && at != m_initialSource)
+		if (at == newAt && at != InitialSource)
 		{
 			mu::Ptr<ASTOpImageMipmap> newOp = mu::Clone<ASTOpImageMipmap>(currentMipmapOp);
-			check(newOp->GetOpType() == OP_TYPE::IM_MIPMAP);
+			check(newOp->GetOpType() == EOpType::IM_MIPMAP);
 
 			newOp->Source = at;
 
diff --git a/Engine/Plugins/Mutable/Source/MutableTools/Private/MuT/ASTOpImageMipmap.h b/Engine/Plugins/Mutable/Source/MutableTools/Private/MuT/ASTOpImageMipmap.h
--- a/Engine/Plugins/Mutable/Source/MutableTools/Private/MuT/ASTOpImageMipmap.h
+++ b/Engine/Plugins/Mutable/Source/MutableTools/Private/MuT/ASTOpImageMipmap.h
@@ -41,7 +41,7 @@ namespace mu
 		ASTOpImageMipmap(const ASTOpImageMipmap&) = delete;
 		~ASTOpImageMipmap();
 
-		virtual OP_TYPE GetOpType() const override { return OP_TYPE::IM_MIPMAP; }
+		virtual EOpType GetOpType() const override { return EOpType::IM_MIPMAP; }
 		virtual uint64 Hash() const override;
 		virtual bool IsEqual(const ASTOp& otherUntyped) const override;
 		virtual Ptr<ASTOp> Clone(MapChildFuncRef mapChild) const override;
diff --git a/Engine/Plugins/Mutable/Source/MutableTools/Private/MuT/ASTOpImageMultiLayer.cpp b/Engine/Plugins/Mutable/Source/MutableTools/Private/MuT/ASTOpImageMultiLayer.cpp
--- a/Engine/Plugins/Mutable/Source/MutableTools/Private/MuT/ASTOpImageMultiLayer.cpp
+++ b/Engine/Plugins/Mutable/Source/MutableTools/Private/MuT/ASTOpImageMultiLayer.cpp
@@ -52,7 +52,7 @@ namespace mu
 	//-------------------------------------------------------------------------------------------------
 	uint64 ASTOpImageMultiLayer::Hash() const
 	{
-		uint64 res = std::hash<OP_TYPE>()(GetOpType());
+		uint64 res = std::hash<EOpType>()(GetOpType());
 		hash_combine(res, base.child().get());
 		hash_combine(res, blend.child().get());
 		hash_combine(res, mask.child().get());
@@ -110,10 +110,10 @@ namespace mu
 				LinkRange(program, range, args.rangeSize, args.rangeId);
 			}
 
-			linkedAddress = (OP::ADDRESS)program.m_opAddress.Num();
-			program.m_opAddress.Add((uint32)program.m_byteCode.Num());
-			AppendCode(program.m_byteCode, GetOpType());
-			AppendCode(program.m_byteCode, args);
+			linkedAddress = (OP::ADDRESS)program.OpAddress.Num();
+			program.OpAddress.Add((uint32)program.ByteCode.Num());
+			AppendCode(program.ByteCode, GetOpType());
+			AppendCode(program.ByteCode, args);
 		}
 	}
 
diff --git a/Engine/Plugins/Mutable/Source/MutableTools/Private/MuT/ASTOpImageMultiLayer.h b/Engine/Plugins/Mutable/Source/MutableTools/Private/MuT/ASTOpImageMultiLayer.h
--- a/Engine/Plugins/Mutable/Source/MutableTools/Private/MuT/ASTOpImageMultiLayer.h
+++ b/Engine/Plugins/Mutable/Source/MutableTools/Private/MuT/ASTOpImageMultiLayer.h
@@ -44,7 +44,7 @@ namespace mu
 		ASTOpImageMultiLayer(const ASTOpImageMultiLayer&) = delete;
 		~ASTOpImageMultiLayer();
 
-		OP_TYPE GetOpType() const override { return OP_TYPE::IM_MULTILAYER; }
+		EOpType GetOpType() const override { return EOpType::IM_MULTILAYER; }
 		uint64 Hash() const override;
 		bool IsEqual(const ASTOp& otherUntyped) const override;
 		Ptr<ASTOp> Clone(MapChildFuncRef mapChild) const override;
diff --git a/Engine/Plugins/Mutable/Source/MutableTools/Private/MuT/ASTOpImageNormalComposite.cpp b/Engine/Plugins/Mutable/Source/MutableTools/Private/MuT/ASTOpImageNormalComposite.cpp
--- a/Engine/Plugins/Mutable/Source/MutableTools/Private/MuT/ASTOpImageNormalComposite.cpp
+++ b/Engine/Plugins/Mutable/Source/MutableTools/Private/MuT/ASTOpImageNormalComposite.cpp
@@ -94,10 +94,10 @@ namespace mu
 			args.power = Power;
 			args.mode = Mode;
 
-			linkedAddress = (OP::ADDRESS)program.m_opAddress.Num();
-			program.m_opAddress.Add((uint32_t)program.m_byteCode.Num());
-			AppendCode(program.m_byteCode, OP_TYPE::IM_NORMALCOMPOSITE);
-			AppendCode(program.m_byteCode, args);
+			linkedAddress = (OP::ADDRESS)program.OpAddress.Num();
+			program.OpAddress.Add((uint32_t)program.ByteCode.Num());
+			AppendCode(program.ByteCode, EOpType::IM_NORMALCOMPOSITE);
+			AppendCode(program.ByteCode, args);
 		}
 
 	}
diff --git a/Engine/Plugins/Mutable/Source/MutableTools/Private/MuT/ASTOpImageNormalComposite.h b/Engine/Plugins/Mutable/Source/MutableTools/Private/MuT/ASTOpImageNormalComposite.h
--- a/Engine/Plugins/Mutable/Source/MutableTools/Private/MuT/ASTOpImageNormalComposite.h
+++ b/Engine/Plugins/Mutable/Source/MutableTools/Private/MuT/ASTOpImageNormalComposite.h
@@ -27,7 +27,7 @@ struct FProgram;
 		ASTOpImageNormalComposite(const ASTOpImageNormalComposite&) = delete;
 		~ASTOpImageNormalComposite();
 
-		OP_TYPE GetOpType() const override { return OP_TYPE::IM_NORMALCOMPOSITE; }
+		EOpType GetOpType() const override { return EOpType::IM_NORMALCOMPOSITE; }
 		uint64 Hash() const override;
 		bool IsEqual(const ASTOp& otherUntyped) const override;
 		Ptr<ASTOp> Clone(MapChildFuncRef mapChild) const override;
diff --git a/Engine/Plugins/Mutable/Source/MutableTools/Private/MuT/ASTOpImagePatch.cpp b/Engine/Plugins/Mutable/Source/MutableTools/Private/MuT/ASTOpImagePatch.cpp
--- a/Engine/Plugins/Mutable/Source/MutableTools/Private/MuT/ASTOpImagePatch.cpp
+++ b/Engine/Plugins/Mutable/Source/MutableTools/Private/MuT/ASTOpImagePatch.cpp
@@ -43,7 +43,7 @@ namespace mu
 	//-------------------------------------------------------------------------------------------------
 	uint64 ASTOpImagePatch::Hash() const
 	{
-		uint64 res = std::hash<OP_TYPE>()(OP_TYPE::IM_PATCH);
+		uint64 res = std::hash<EOpType>()(EOpType::IM_PATCH);
 		hash_combine(res, base.child().get());
 		hash_combine(res, patch.child().get());
 		return res;
@@ -83,10 +83,10 @@ namespace mu
 			args.minX = location[0];
 			args.minY = location[1];
 
-			linkedAddress = (OP::ADDRESS)program.m_opAddress.Num();
-			program.m_opAddress.Add((uint32_t)program.m_byteCode.Num());
-			AppendCode(program.m_byteCode, OP_TYPE::IM_PATCH);
-			AppendCode(program.m_byteCode, args);
+			linkedAddress = (OP::ADDRESS)program.OpAddress.Num();
+			program.OpAddress.Add((uint32_t)program.ByteCode.Num());
+			AppendCode(program.ByteCode, EOpType::IM_PATCH);
+			AppendCode(program.ByteCode, args);
 		}
 
 	}
diff --git a/Engine/Plugins/Mutable/Source/MutableTools/Private/MuT/ASTOpImagePatch.h b/Engine/Plugins/Mutable/Source/MutableTools/Private/MuT/ASTOpImagePatch.h
--- a/Engine/Plugins/Mutable/Source/MutableTools/Private/MuT/ASTOpImagePatch.h
+++ b/Engine/Plugins/Mutable/Source/MutableTools/Private/MuT/ASTOpImagePatch.h
@@ -28,7 +28,7 @@ struct FProgram;
 		ASTOpImagePatch(const ASTOpImagePatch&) = delete;
 		~ASTOpImagePatch();
 
-		virtual OP_TYPE GetOpType() const override { return OP_TYPE::IM_PATCH; }
+		virtual EOpType GetOpType() const override { return EOpType::IM_PATCH; }
 		virtual uint64 Hash() const override;
 		virtual bool IsEqual(const ASTOp& otherUntyped) const override;
 		virtual Ptr<ASTOp> Clone(MapChildFuncRef mapChild) const override;
diff --git a/Engine/Plugins/Mutable/Source/MutableTools/Private/MuT/ASTOpImagePixelFormat.cpp b/Engine/Plugins/Mutable/Source/MutableTools/Private/MuT/ASTOpImagePixelFormat.cpp
--- a/Engine/Plugins/Mutable/Source/MutableTools/Private/MuT/ASTOpImagePixelFormat.cpp
+++ b/Engine/Plugins/Mutable/Source/MutableTools/Private/MuT/ASTOpImagePixelFormat.cpp
@@ -87,11 +87,11 @@ namespace mu
 			args.formatIfAlpha = FormatIfAlpha;
 			if (Source) args.source = Source->linkedAddress;
 
-			linkedAddress = (OP::ADDRESS)program.m_opAddress.Num();
+			linkedAddress = (OP::ADDRESS)program.OpAddress.Num();
 			//program.m_code.push_back(op);
-			program.m_opAddress.Add((uint32_t)program.m_byteCode.Num());
-			AppendCode(program.m_byteCode, OP_TYPE::IM_PIXELFORMAT);
-			AppendCode(program.m_byteCode, args);
+			program.OpAddress.Add((uint32_t)program.ByteCode.Num());
+			AppendCode(program.ByteCode, EOpType::IM_PIXELFORMAT);
+			AppendCode(program.ByteCode, args);
 		}
 
 	}
@@ -126,10 +126,10 @@ namespace mu
 		//bool isBlockFormat = GetImageFormatData( format ).PixelsPerBlockX!=0;
 
 		// The instruction can be sunk
-		OP_TYPE sourceType = sourceAt ? sourceAt->GetOpType() : OP_TYPE::NONE;
+		EOpType sourceType = sourceAt ? sourceAt->GetOpType() : EOpType::NONE;
 		switch (sourceType)
 		{
-		case OP_TYPE::IM_PIXELFORMAT:
+		case EOpType::IM_PIXELFORMAT:
 		{
 			// Keep only the top pixel format
 			const ASTOpImagePixelFormat* typedSource = static_cast<const ASTOpImagePixelFormat*>(sourceAt.get());
@@ -139,7 +139,7 @@ namespace mu
 			break;
 		}
 
-		case OP_TYPE::IM_DISPLACE:
+		case EOpType::IM_DISPLACE:
 		{
 			// This op doesn't support compressed formats
 			if (!bIsCompressedFormat)
@@ -155,7 +155,7 @@ namespace mu
 			break;
 		}
 
-		case OP_TYPE::IM_INVERT:
+		case EOpType::IM_INVERT:
 		{
 			// This op doesn't support compressed formats
 			if (!bIsCompressedFormat)
@@ -171,7 +171,7 @@ namespace mu
 			break;
 		}
 
-		case OP_TYPE::IM_RASTERMESH:
+		case EOpType::IM_RASTERMESH:
 		{
 			// This op doesn't support compressed formats
 			if (!bIsCompressedFormat
@@ -189,7 +189,7 @@ namespace mu
 			break;
 		}
 
-		case OP_TYPE::IM_BLANKLAYOUT:
+		case EOpType::IM_BLANKLAYOUT:
 		{
 			// Just make sure the layout format is the right one and forget the op
 			auto nop = mu::Clone<ASTOpFixed>(sourceAt);
@@ -207,7 +207,7 @@ namespace mu
 			break;
 		}
 
-		case OP_TYPE::IM_PLAINCOLOUR:
+		case EOpType::IM_PLAINCOLOUR:
 		{
 			// Just make sure the format is the right one and forget the op
 			auto nop = mu::Clone<ASTOpFixed>(sourceAt);
@@ -341,18 +341,18 @@ namespace mu
 	//---------------------------------------------------------------------------------------------
 	mu::Ptr<ASTOp> Sink_ImagePixelFormatAST::Apply(const ASTOpImagePixelFormat* root)
 	{
-		m_root = root;
+		Root = root;
 		OldToNew.Reset();
 
-		check(root->GetOpType() == OP_TYPE::IM_PIXELFORMAT);
+		check(root->GetOpType() == EOpType::IM_PIXELFORMAT);
 
-		m_initialSource = m_root->Source.child();
-		mu::Ptr<ASTOp> newSource = Visit(m_initialSource, m_root);
+		InitialSource = Root->Source.child();
+		mu::Ptr<ASTOp> newSource = Visit(InitialSource, Root);
 
-		m_root = nullptr;
+		Root = nullptr;
 
 		// If there is any change, it is the new root.
-		if (newSource != m_initialSource)
+		if (newSource != InitialSource)
 		{
 			return newSource;
 		}
@@ -381,7 +381,7 @@ namespace mu
 		switch (at->GetOpType())
 		{
 
-		case OP_TYPE::IM_CONDITIONAL:
+		case EOpType::IM_CONDITIONAL:
 		{
 			// We move the op down the two paths
 			Ptr<ASTOpConditional> newOp = mu::Clone<ASTOpConditional>(at);
@@ -391,7 +391,7 @@ namespace mu
 			break;
 		}
 
-		case OP_TYPE::IM_SWITCH:
+		case EOpType::IM_SWITCH:
 		{
 			// We move the op down all the paths
 			Ptr<ASTOpSwitch> newOp = mu::Clone<ASTOpSwitch>(at);
@@ -404,7 +404,7 @@ namespace mu
 			break;
 		}
 
-		case OP_TYPE::IM_COMPOSE:
+		case EOpType::IM_COMPOSE:
 		{
 			if (isBlockFormat)
 			{
@@ -428,7 +428,7 @@ namespace mu
 			break;
 		}
 
-		case OP_TYPE::IM_PATCH:
+		case EOpType::IM_PATCH:
 		{
 			if (isBlockFormat)
 			{
@@ -444,7 +444,7 @@ namespace mu
 			break;
 		}
 
-		case OP_TYPE::IM_MIPMAP:
+		case EOpType::IM_MIPMAP:
 		{
 			const ASTOpImageMipmap* typedSource = static_cast<const ASTOpImageMipmap*>(at.get());
 
@@ -462,7 +462,7 @@ namespace mu
 			break;
 		}
 
-		case OP_TYPE::IM_INTERPOLATE:
+		case EOpType::IM_INTERPOLATE:
 		{
 			// This op doesn't support compressed formats
 			if (!bIsCompressedFormat)
@@ -482,7 +482,7 @@ namespace mu
 			break;
 		}
 
-		case OP_TYPE::IM_LAYER:
+		case EOpType::IM_LAYER:
 		{
 			if (GetOpDesc(at->GetOpType()).supportedBasePixelFormats[(size_t)format])
 			{
@@ -497,7 +497,7 @@ namespace mu
 			break;
 		}
 
-		case OP_TYPE::IM_LAYERCOLOUR:
+		case EOpType::IM_LAYERCOLOUR:
 		{
 			if (GetOpDesc(at->GetOpType()).supportedBasePixelFormats[(size_t)format])
 			{
@@ -517,10 +517,10 @@ namespace mu
 		}
 
 		// end on tree branch, replace with format
-		if (at == newAt && at != m_initialSource)
+		if (at == newAt && at != InitialSource)
 		{
 			mu::Ptr<ASTOpImagePixelFormat> newOp = mu::Clone<ASTOpImagePixelFormat>(currentFormatOp);
-			check(newOp->GetOpType() == OP_TYPE::IM_PIXELFORMAT);
+			check(newOp->GetOpType() == EOpType::IM_PIXELFORMAT);
 
 			newOp->Source = at;
 
diff --git a/Engine/Plugins/Mutable/Source/MutableTools/Private/MuT/ASTOpImagePixelFormat.h b/Engine/Plugins/Mutable/Source/MutableTools/Private/MuT/ASTOpImagePixelFormat.h
--- a/Engine/Plugins/Mutable/Source/MutableTools/Private/MuT/ASTOpImagePixelFormat.h
+++ b/Engine/Plugins/Mutable/Source/MutableTools/Private/MuT/ASTOpImagePixelFormat.h
@@ -27,7 +27,7 @@ namespace mu
 		ASTOpImagePixelFormat(const ASTOpImagePixelFormat&) = delete;
 		~ASTOpImagePixelFormat();
 
-		virtual OP_TYPE GetOpType() const override { return OP_TYPE::IM_PIXELFORMAT; }
+		virtual EOpType GetOpType() const override { return EOpType::IM_PIXELFORMAT; }
 		virtual uint64 Hash() const override;
 		virtual bool IsEqual(const ASTOp& otherUntyped) const override;
 		virtual Ptr<ASTOp> Clone(MapChildFuncRef mapChild) const override;
diff --git a/Engine/Plugins/Mutable/Source/MutableTools/Private/MuT/ASTOpImageRasterMesh.cpp b/Engine/Plugins/Mutable/Source/MutableTools/Private/MuT/ASTOpImageRasterMesh.cpp
--- a/Engine/Plugins/Mutable/Source/MutableTools/Private/MuT/ASTOpImageRasterMesh.cpp
+++ b/Engine/Plugins/Mutable/Source/MutableTools/Private/MuT/ASTOpImageRasterMesh.cpp
@@ -81,7 +81,7 @@ namespace mu
 	//-------------------------------------------------------------------------------------------------
 	uint64 ASTOpImageRasterMesh::Hash() const
 	{
-		uint64 res = std::hash<OP_TYPE>()(GetOpType());
+		uint64 res = std::hash<EOpType>()(GetOpType());
 		hash_combine(res, mesh.child().get());
 		hash_combine(res, image.child().get());
 		hash_combine(res, angleFadeProperties.child().get());
@@ -160,10 +160,10 @@ namespace mu
 			if (mask) args.mask = mask->linkedAddress;
 			if (projector) args.projector = projector->linkedAddress;
 
-			linkedAddress = (OP::ADDRESS)program.m_opAddress.Num();
-			program.m_opAddress.Add((uint32)program.m_byteCode.Num());
-			AppendCode(program.m_byteCode, GetOpType());
-			AppendCode(program.m_byteCode, args);
+			linkedAddress = (OP::ADDRESS)program.OpAddress.Num();
+			program.OpAddress.Add((uint32)program.ByteCode.Num());
+			AppendCode(program.ByteCode, GetOpType());
+			AppendCode(program.ByteCode, args);
 		}
 	}
 
@@ -239,11 +239,11 @@ namespace mu
 		Ptr<ASTOp> sourceAt = mesh.child();
 		Ptr<ASTOp> imageAt = image.child();
 
-		OP_TYPE sourceType = sourceAt->GetOpType();
+		EOpType sourceType = sourceAt->GetOpType();
 		switch (sourceType)
 		{
 
-		case OP_TYPE::ME_PROJECT:
+		case EOpType::ME_PROJECT:
 		{
 			// If we are rastering just the UV layout (to create a mask) we don't care about
 			// mesh project operations, which modify only the positions.
@@ -260,7 +260,7 @@ namespace mu
 			break;
 		}
 
-		case OP_TYPE::ME_INTERPOLATE:
+		case EOpType::ME_INTERPOLATE:
 		{
 			// TODO: should be sink only if no imageAt?
 			auto typedSource = static_cast<const ASTOpFixed*>(sourceAt.get());
@@ -270,7 +270,7 @@ namespace mu
 			break;
 		}
 
-		case OP_TYPE::ME_MORPH:
+		case EOpType::ME_MORPH:
 		{
 			// TODO: should be sink only if no imageAt?
 			const ASTOpMeshMorph* typedSource = static_cast<const ASTOpMeshMorph*>(sourceAt.get());
@@ -280,7 +280,7 @@ namespace mu
 			break;
 		}
 
-		case OP_TYPE::ME_ADDTAGS:
+		case EOpType::ME_ADDTAGS:
 		{
 			// Ignore tags
 			const ASTOpMeshAddTags* typedSource = static_cast<const ASTOpMeshAddTags*>(sourceAt.get());
@@ -290,10 +290,10 @@ namespace mu
 			break;
 		}
 
-		case OP_TYPE::ME_CONDITIONAL:
+		case EOpType::ME_CONDITIONAL:
 		{
 			auto nop = mu::Clone<ASTOpConditional>(sourceAt.get());
-			nop->type = OP_TYPE::IM_CONDITIONAL;
+			nop->type = EOpType::IM_CONDITIONAL;
 
 			Ptr<ASTOpImageRasterMesh> aOp = mu::Clone<ASTOpImageRasterMesh>(this);
 			aOp->mesh = nop->yes.child();
@@ -307,11 +307,11 @@ namespace mu
 			break;
 		}
 
-		case OP_TYPE::ME_SWITCH:
+		case EOpType::ME_SWITCH:
 		{
 			// Make an image for every path
 			auto nop = mu::Clone<ASTOpSwitch>(sourceAt.get());
-			nop->type = OP_TYPE::IM_SWITCH;
+			nop->type = EOpType::IM_SWITCH;
 
 			if (nop->def)
 			{
@@ -342,15 +342,15 @@ namespace mu
 		// If we didn't optimize the mesh child, try to optimize the image child.
 		if (OriginalAt == at && imageAt)
 		{
-			OP_TYPE imageType = imageAt->GetOpType();
+			EOpType imageType = imageAt->GetOpType();
 			switch (imageType)
 			{
 
 				// TODO: Implement for image conditionals.
-				//case OP_TYPE::ME_CONDITIONAL:
+				//case EOpType::ME_CONDITIONAL:
 				//{
 				//	auto nop = mu::Clone<ASTOpConditional>(sourceAt.get());
-				//	nop->type = OP_TYPE::IM_CONDITIONAL;
+				//	nop->type = EOpType::IM_CONDITIONAL;
 
 				//	Ptr<ASTOpFixed> aOp = mu::Clone<ASTOpFixed>(this);
 				//	aOp->SetChild(aOp->op.args.ImageRasterMesh.mesh, nop->yes);
@@ -364,7 +364,7 @@ namespace mu
 				//	break;
 				//}
 
-			case OP_TYPE::IM_SWITCH:
+			case EOpType::IM_SWITCH:
 			{
 				// TODO: Do this only if the projector is constant?
 
diff --git a/Engine/Plugins/Mutable/Source/MutableTools/Private/MuT/ASTOpImageRasterMesh.h b/Engine/Plugins/Mutable/Source/MutableTools/Private/MuT/ASTOpImageRasterMesh.h
--- a/Engine/Plugins/Mutable/Source/MutableTools/Private/MuT/ASTOpImageRasterMesh.h
+++ b/Engine/Plugins/Mutable/Source/MutableTools/Private/MuT/ASTOpImageRasterMesh.h
@@ -50,7 +50,7 @@ namespace mu
 		ASTOpImageRasterMesh(const ASTOpImageRasterMesh&) = delete;
 		~ASTOpImageRasterMesh();
 
-		virtual OP_TYPE GetOpType() const override { return OP_TYPE::IM_RASTERMESH; }
+		virtual EOpType GetOpType() const override { return EOpType::IM_RASTERMESH; }
 		virtual uint64 Hash() const override;
 		virtual bool IsEqual(const ASTOp& otherUntyped) const override;
 		virtual Ptr<ASTOp> Clone(MapChildFuncRef mapChild) const override;
diff --git a/Engine/Plugins/Mutable/Source/MutableTools/Private/MuT/ASTOpImageSwizzle.cpp b/Engine/Plugins/Mutable/Source/MutableTools/Private/MuT/ASTOpImageSwizzle.cpp
--- a/Engine/Plugins/Mutable/Source/MutableTools/Private/MuT/ASTOpImageSwizzle.cpp
+++ b/Engine/Plugins/Mutable/Source/MutableTools/Private/MuT/ASTOpImageSwizzle.cpp
@@ -109,10 +109,10 @@ namespace mu
 				args.sourceChannels[i] = SourceChannels[i];
 			}			
 
-			linkedAddress = (OP::ADDRESS)program.m_opAddress.Num();
-			program.m_opAddress.Add((uint32)program.m_byteCode.Num());
-			AppendCode(program.m_byteCode, GetOpType());
-			AppendCode(program.m_byteCode, args);
+			linkedAddress = (OP::ADDRESS)program.OpAddress.Num();
+			program.OpAddress.Add((uint32)program.ByteCode.Num());
+			AppendCode(program.ByteCode, GetOpType());
+			AppendCode(program.ByteCode, args);
 		}
 	}
 
@@ -132,7 +132,7 @@ namespace mu
 			switch (candidate->GetOpType())
 			{
 			// Swizzle
-			case OP_TYPE::IM_SWIZZLE:
+			case EOpType::IM_SWIZZLE:
 			{
 				if (!sat)
 				{
@@ -148,7 +148,7 @@ namespace mu
 			}
 
 			// Format
-			case OP_TYPE::IM_PIXELFORMAT:
+			case EOpType::IM_PIXELFORMAT:
 			{
 				// We can remove the format if its source is already an uncompressed format
 				ASTOpImagePixelFormat* typedCandidate = static_cast<ASTOpImagePixelFormat*>(candidate.get());
@@ -186,7 +186,7 @@ namespace mu
 		//---------------------------------------------------------------------------------------------
 		void ReplaceAllSources(Ptr<ASTOpImageSwizzle>& op, Ptr<ASTOp>& value)
 		{
-			check(op->GetOpType() == OP_TYPE::IM_SWIZZLE);
+			check(op->GetOpType() == EOpType::IM_SWIZZLE);
 			for (int c = 0; c < MUTABLE_OP_MAX_SWIZZLE_CHANNELS; ++c)
 			{
 				if (op->Sources[c])
@@ -256,7 +256,7 @@ namespace mu
 			}
 		}
 
-		OP_TYPE sourceType = channelSourceAt->GetOpType();
+		EOpType sourceType = channelSourceAt->GetOpType();
 
 		if (bAllChannelsAreTheSame)
 		{
@@ -275,9 +275,9 @@ namespace mu
 				&&
 				Sources[0] == Sources[1] && Sources[0] == Sources[2]
 				&&
-				Sources[0] && Sources[0]->GetOpType() == OP_TYPE::IM_MULTILAYER
+				Sources[0] && Sources[0]->GetOpType() == EOpType::IM_MULTILAYER
 				&&
-				Sources[3] && Sources[3]->GetOpType() == OP_TYPE::IM_MULTILAYER
+				Sources[3] && Sources[3]->GetOpType() == EOpType::IM_MULTILAYER
 				&&
 				SourceChannels[0] == 0 && SourceChannels[1] == 1 && SourceChannels[2] == 2 && SourceChannels[3] == 0
 				)
@@ -334,9 +334,9 @@ namespace mu
 				&&
 				Sources[0] == Sources[1] && (Sources[0] == Sources[2] || !Sources[2])
 				&&
-				Sources[0] && Sources[0]->GetOpType() == OP_TYPE::IM_LAYER
+				Sources[0] && Sources[0]->GetOpType() == EOpType::IM_LAYER
 				&&
-				Sources[3] && Sources[3]->GetOpType() == OP_TYPE::IM_LAYER
+				Sources[3] && Sources[3]->GetOpType() == EOpType::IM_LAYER
 				&&
 				SourceChannels[0] == 0 && SourceChannels[1] == 1 && (SourceChannels[2] == 2 || !Sources[2] ) && SourceChannels[3] == 0
 				)
@@ -386,7 +386,7 @@ namespace mu
 
 
 			// If the channels are compatible switches, we can still sink the swizzle.
-			if (!at && sourceType == OP_TYPE::IM_SWITCH)
+			if (!at && sourceType == EOpType::IM_SWITCH)
 			{
 				const ASTOpSwitch* FirstSwitch = static_cast<const ASTOpSwitch*>(Sources[0].child().get());
 				check(FirstSwitch);
@@ -447,7 +447,7 @@ namespace mu
 			}
 
 			// Swizzle down compatible displaces.
-			if (!at && sourceType == OP_TYPE::IM_DISPLACE)
+			if (!at && sourceType == EOpType::IM_DISPLACE)
 			{
 				const ASTOpFixed* FirstDisplace = static_cast<const ASTOpFixed*>(Sources[0].child().get());
 				check(FirstDisplace);
@@ -490,7 +490,7 @@ namespace mu
 			}
 
 			// Swizzle down compatible raster meshes.
-			if (!at && sourceType == OP_TYPE::IM_RASTERMESH)
+			if (!at && sourceType == EOpType::IM_RASTERMESH)
 			{
 				const ASTOpImageRasterMesh* FirstRasterMesh = static_cast<const ASTOpImageRasterMesh*>(Sources[0].child().get());
 				check(FirstRasterMesh);
@@ -564,7 +564,7 @@ namespace mu
 			}
 
 			// Swizzle down compatible image transforms.
-			if (!at && sourceType == OP_TYPE::IM_TRANSFORM)
+			if (!at && sourceType == EOpType::IM_TRANSFORM)
 			{
 				const ASTOpImageTransform* FirstTransform = static_cast<const ASTOpImageTransform*>(Sources[0].child().get());
 				check(FirstTransform);
@@ -619,7 +619,7 @@ namespace mu
 			}
 
 			// Swizzle down compatible resizes.
-			//if (!at && sourceType == OP_TYPE::IM_RESIZE)
+			//if (!at && sourceType == EOpType::IM_RESIZE)
 			//{
 			//	const ASTOpFixed* FirstResize = static_cast<const ASTOpFixed*>(Sources[0].child().get());
 			//	check(FirstResize);
@@ -666,7 +666,7 @@ namespace mu
 			//}
 
 			// Swizzle down compatible pixelformats.
-			if (!at && sourceType == OP_TYPE::IM_PIXELFORMAT && bSameChannelOrder)
+			if (!at && sourceType == EOpType::IM_PIXELFORMAT && bSameChannelOrder)
 			{
 				const ASTOpImagePixelFormat* FirstFormat = static_cast<const ASTOpImagePixelFormat*>(Sources[0].child().get());
 				check(FirstFormat);
@@ -697,11 +697,11 @@ namespace mu
 			}
 
 			// Swizzle down plaincolours.
-			if (!at && sourceType == OP_TYPE::IM_PLAINCOLOUR)
+			if (!at && sourceType == EOpType::IM_PLAINCOLOUR)
 			{
 				Ptr<ASTOpFixed> NewPlain = mu::Clone<ASTOpFixed>(channelSourceAt);
 				Ptr<ASTOpFixed> NewSwizzle = new ASTOpFixed;
-				NewSwizzle->op.type = OP_TYPE::CO_SWIZZLE;
+				NewSwizzle->op.type = EOpType::CO_SWIZZLE;
 				for (int c = 0; c < MUTABLE_OP_MAX_SWIZZLE_CHANNELS; ++c)
 				{
 					if (Sources[c])
@@ -726,7 +726,7 @@ namespace mu
 		//	&&
 		//	Sources[0] && Sources[0]==Sources[1] && Sources[0]==Sources[2]
 		//	&&
-		//	Sources[3] && Sources[3]->GetOpType()==OP_TYPE::IM_LAYERCOLOUR
+		//	Sources[3] && Sources[3]->GetOpType()==EOpType::IM_LAYERCOLOUR
 		//	)
 		//{
 		//	// Move the swizzle down all the paths
@@ -752,7 +752,7 @@ namespace mu
 		//	&&
 		//	Sources[0] && Sources[0] == Sources[1] && Sources[0] == Sources[2]
 		//	&&
-		//	Sources[3] && Sources[3]->GetOpType() == OP_TYPE::IM_LAYER)
+		//	Sources[3] && Sources[3]->GetOpType() == EOpType::IM_LAYER)
 		//{
 		//	// Move the swizzle down all the paths
 		//	Ptr<ASTOpImageLayer> NewLayer = mu::Clone<ASTOpImageLayer>(Sources[3].child());
@@ -778,7 +778,7 @@ namespace mu
 			&& 
 			(!Sources[2] || Sources[0] == Sources[2])
 			&&
-			Sources[0]->GetOpType() == OP_TYPE::IM_LAYERCOLOUR
+			Sources[0]->GetOpType() == EOpType::IM_LAYERCOLOUR
 			&& 
 			!(Sources[3]==Sources[0]) )
 		{
@@ -805,7 +805,7 @@ namespace mu
 
 			for (int32 Channel = 0; Channel < MUTABLE_OP_MAX_SWIZZLE_CHANNELS; ++Channel)
 			{
-				if (Sources[Channel] && SourceChannels[Channel]==3 && Sources[Channel]->GetOpType() == OP_TYPE::IM_SATURATE)
+				if (Sources[Channel] && SourceChannels[Channel]==3 && Sources[Channel]->GetOpType() == EOpType::IM_SATURATE)
 				{
 					// Remove the saturate for this channel
 					if (!NewSwizzle)
@@ -828,7 +828,7 @@ namespace mu
 			&&
 			Sources[0]
 			&&
-			Sources[0]->GetOpType() == OP_TYPE::IM_SATURATE 
+			Sources[0]->GetOpType() == EOpType::IM_SATURATE 
 			&&
 			(Sources[0] == Sources[1]) && (Sources[0] == Sources[2])
 			&&
@@ -881,7 +881,7 @@ namespace mu
 			&&
 			SourceChannels[0] == 0 && SourceChannels[1] == 1 && SourceChannels[2] == 2
 			&&
-			Sources[3] && Sources[3]->GetOpType() == OP_TYPE::IM_LAYER
+			Sources[3] && Sources[3]->GetOpType() == EOpType::IM_LAYER
 			)
 		{
 			const ASTOpImageLayer* OldLayer = static_cast<const ASTOpImageLayer*>(Sources[3].child().get());
@@ -897,7 +897,7 @@ namespace mu
 						bUpdated = false;
 						switch (Op->GetOpType())
 						{
-						case OP_TYPE::IM_PIXELFORMAT:
+						case EOpType::IM_PIXELFORMAT:
 						{
 							const ASTOpImagePixelFormat* Typed = static_cast<const ASTOpImagePixelFormat*>(Op.get());
 							Op = Typed->Source.child();
@@ -953,7 +953,7 @@ namespace mu
 		// \TODO: Make more general.
 		if (!at
 			&&
-			Sources[3] && Sources[3]->GetOpType() == OP_TYPE::IM_LAYER
+			Sources[3] && Sources[3]->GetOpType() == EOpType::IM_LAYER
 			)
 		{
 			ASTOpImageLayer* OldLayer = static_cast<ASTOpImageLayer*>(Sources[3].child().get());
@@ -963,10 +963,10 @@ namespace mu
 				&&
 				SourceChannels[3] == 0
 				&&
-				OldLayer->blend->GetOpType()==OP_TYPE::IM_PIXELFORMAT )
+				OldLayer->blend->GetOpType()==EOpType::IM_PIXELFORMAT )
 			{
 				const ASTOpImagePixelFormat* OldFormat = static_cast<const ASTOpImagePixelFormat*>(OldLayer->blend.child().get());
-				if (OldFormat->Source->GetOpType() == OP_TYPE::IM_SWIZZLE
+				if (OldFormat->Source->GetOpType() == EOpType::IM_SWIZZLE
 					&&
 					(
 						OldFormat->Format == EImageFormat::RGB_UByte
@@ -1132,7 +1132,7 @@ namespace mu
 		Root = InRoot;
 		OldToNew.Reset();
 
-		check(Root->GetOpType() == OP_TYPE::IM_SWIZZLE);
+		check(Root->GetOpType() == EOpType::IM_SWIZZLE);
 
 		// This sinker only works assuming all swizzle channels come from the same image operation.
 		bool bAllChannelsAreTheSame = true;
@@ -1189,7 +1189,7 @@ namespace mu
 		switch (at->GetOpType())
 		{
 
-		case OP_TYPE::IM_CONDITIONAL:
+		case EOpType::IM_CONDITIONAL:
 		{
 			// We move the op down the two paths
 			auto newOp = mu::Clone<ASTOpConditional>(at);
@@ -1199,7 +1199,7 @@ namespace mu
 			break;
 		}
 
-		case OP_TYPE::IM_SWITCH:
+		case EOpType::IM_SWITCH:
 		{
 			// We move the op down all the paths
 			Ptr<ASTOpSwitch> newOp = mu::Clone<ASTOpSwitch>(at);
@@ -1212,7 +1212,7 @@ namespace mu
 			break;
 		}
 
-		case OP_TYPE::IM_COMPOSE:
+		case EOpType::IM_COMPOSE:
 		{
 			Ptr<ASTOpImageCompose> newOp = mu::Clone<ASTOpImageCompose>(at);
 			newOp->Base = Visit(newOp->Base.child(), CurrentSwizzleOp);
@@ -1221,7 +1221,7 @@ namespace mu
 			break;
 		}
 
-		case OP_TYPE::IM_PATCH:
+		case EOpType::IM_PATCH:
 		{
 			Ptr<ASTOpImagePatch> newOp = mu::Clone<ASTOpImagePatch>(at);
 			newOp->base = Visit(newOp->base.child(), CurrentSwizzleOp);
@@ -1230,7 +1230,7 @@ namespace mu
 			break;
 		}
 
-		case OP_TYPE::IM_MIPMAP:
+		case EOpType::IM_MIPMAP:
 		{
 			Ptr<ASTOpImageMipmap> newOp = mu::Clone<ASTOpImageMipmap>(at);
 			newOp->Source = Visit(newOp->Source.child(), CurrentSwizzleOp);
@@ -1238,7 +1238,7 @@ namespace mu
 			break;
 		}
 
-		case OP_TYPE::IM_INTERPOLATE:
+		case EOpType::IM_INTERPOLATE:
 		{
 			Ptr<ASTOpFixed> newOp = mu::Clone<ASTOpFixed>(at);
 
@@ -1253,7 +1253,7 @@ namespace mu
 			break;
 		}
 
-		case OP_TYPE::IM_LAYER:
+		case EOpType::IM_LAYER:
 		{
 			Ptr<ASTOpImageLayer> nop = mu::Clone<ASTOpImageLayer>(at);
 			nop->base = Visit(nop->base.child(), CurrentSwizzleOp);
@@ -1262,14 +1262,14 @@ namespace mu
 			break;
 		}
 
-		case OP_TYPE::IM_LAYERCOLOUR:
+		case EOpType::IM_LAYERCOLOUR:
 		{
 			Ptr<ASTOpImageLayerColor> nop = mu::Clone<ASTOpImageLayerColor>(at);
 			nop->base = Visit(nop->base.child(), CurrentSwizzleOp);
 
 			// We need to swizzle the colour too
 			Ptr<ASTOpFixed> NewColorOp = new ASTOpFixed;
-			NewColorOp->op.type = OP_TYPE::CO_SWIZZLE;
+			NewColorOp->op.type = EOpType::CO_SWIZZLE;
 			for (int i = 0; i < MUTABLE_OP_MAX_SWIZZLE_CHANNELS; ++i)
 			{
 				NewColorOp->SetChild(NewColorOp->op.args.ColourSwizzle.sources[i], nop->color);
@@ -1281,7 +1281,7 @@ namespace mu
 			break;
 		}
 
-		case OP_TYPE::IM_DISPLACE:
+		case EOpType::IM_DISPLACE:
 		{
 			Ptr<ASTOpFixed> NewOp = mu::Clone<ASTOpFixed>(at);
 
@@ -1293,7 +1293,7 @@ namespace mu
 			break;
 		}
 
-		case OP_TYPE::IM_INVERT:
+		case EOpType::IM_INVERT:
 		{
 			Ptr<ASTOpFixed> NewOp = mu::Clone<ASTOpFixed>(at);
 
@@ -1305,7 +1305,7 @@ namespace mu
 			break;
 		}
 
-		case OP_TYPE::IM_RASTERMESH:
+		case EOpType::IM_RASTERMESH:
 		{
 			Ptr<ASTOpImageRasterMesh> NewOp = mu::Clone<ASTOpImageRasterMesh>(at);
 
@@ -1328,7 +1328,7 @@ namespace mu
 			break;
 		}
 
-		case OP_TYPE::IM_TRANSFORM:
+		case EOpType::IM_TRANSFORM:
 		{
 			Ptr<ASTOpImageTransform> nop = mu::Clone<ASTOpImageTransform>(at);
 			nop->Base = Visit(nop->Base.child(), CurrentSwizzleOp);
@@ -1336,7 +1336,7 @@ namespace mu
 			break;
 		}
 
-		case OP_TYPE::IM_PIXELFORMAT:
+		case EOpType::IM_PIXELFORMAT:
 		{
 			// If we are not changing channel order, just remove the swizzle and adjust the format.
 			bool bSameChannelOrder = true;
@@ -1358,7 +1358,7 @@ namespace mu
 			break;
 		}
 
-		case OP_TYPE::IM_BLANKLAYOUT:
+		case EOpType::IM_BLANKLAYOUT:
 		{
 			Ptr<ASTOpFixed> NewOp = mu::Clone<ASTOpFixed>(at);
 
@@ -1378,7 +1378,7 @@ namespace mu
 		if (at == newAt && at != InitialSource)
 		{
 			mu::Ptr<ASTOpImageSwizzle> NewOp = mu::Clone<ASTOpImageSwizzle>(CurrentSwizzleOp);
-			check(NewOp->GetOpType() == OP_TYPE::IM_SWIZZLE);
+			check(NewOp->GetOpType() == EOpType::IM_SWIZZLE);
 
 			ReplaceAllSources(NewOp, at);
 
diff --git a/Engine/Plugins/Mutable/Source/MutableTools/Private/MuT/ASTOpImageSwizzle.h b/Engine/Plugins/Mutable/Source/MutableTools/Private/MuT/ASTOpImageSwizzle.h
--- a/Engine/Plugins/Mutable/Source/MutableTools/Private/MuT/ASTOpImageSwizzle.h
+++ b/Engine/Plugins/Mutable/Source/MutableTools/Private/MuT/ASTOpImageSwizzle.h
@@ -29,7 +29,7 @@ namespace mu
 		ASTOpImageSwizzle(const ASTOpImageSwizzle&) = delete;
 		~ASTOpImageSwizzle();
 
-		virtual OP_TYPE GetOpType() const override { return OP_TYPE::IM_SWIZZLE; }
+		virtual EOpType GetOpType() const override { return EOpType::IM_SWIZZLE; }
 		virtual uint64 Hash() const override;
 		virtual bool IsEqual(const ASTOp& otherUntyped) const override;
 		virtual Ptr<ASTOp> Clone(MapChildFuncRef mapChild) const override;
diff --git a/Engine/Plugins/Mutable/Source/MutableTools/Private/MuT/ASTOpImageTransform.cpp b/Engine/Plugins/Mutable/Source/MutableTools/Private/MuT/ASTOpImageTransform.cpp
--- a/Engine/Plugins/Mutable/Source/MutableTools/Private/MuT/ASTOpImageTransform.cpp
+++ b/Engine/Plugins/Mutable/Source/MutableTools/Private/MuT/ASTOpImageTransform.cpp
@@ -59,7 +59,7 @@ namespace mu
 	//-------------------------------------------------------------------------------------------------
 	uint64 ASTOpImageTransform::Hash() const
 	{
-		uint64 Res = std::hash<OP_TYPE>()(OP_TYPE::IM_TRANSFORM);
+		uint64 Res = std::hash<EOpType>()(EOpType::IM_TRANSFORM);
 		hash_combine(Res, Base.child().get());
 		hash_combine(Res, OffsetX.child().get());
 		hash_combine(Res, OffsetY.child().get());
@@ -130,10 +130,10 @@ namespace mu
 			Args.SourceSizeX = SourceSizeX;
 			Args.SourceSizeY = SourceSizeY;
 
-			linkedAddress = (OP::ADDRESS)Program.m_opAddress.Num();
-			Program.m_opAddress.Add((uint32_t)Program.m_byteCode.Num());
-			AppendCode(Program.m_byteCode, OP_TYPE::IM_TRANSFORM);
-			AppendCode(Program.m_byteCode, Args);
+			linkedAddress = (OP::ADDRESS)Program.OpAddress.Num();
+			Program.OpAddress.Add((uint32_t)Program.ByteCode.Num());
+			AppendCode(Program.ByteCode, EOpType::IM_TRANSFORM);
+			AppendCode(Program.ByteCode, Args);
 		}
 	}
 
diff --git a/Engine/Plugins/Mutable/Source/MutableTools/Private/MuT/ASTOpImageTransform.h b/Engine/Plugins/Mutable/Source/MutableTools/Private/MuT/ASTOpImageTransform.h
--- a/Engine/Plugins/Mutable/Source/MutableTools/Private/MuT/ASTOpImageTransform.h
+++ b/Engine/Plugins/Mutable/Source/MutableTools/Private/MuT/ASTOpImageTransform.h
@@ -43,7 +43,7 @@ struct FProgram;
 		ASTOpImageTransform(const ASTOpImageTransform&) = delete;
 		~ASTOpImageTransform();
 
-		virtual OP_TYPE GetOpType() const override { return OP_TYPE::IM_TRANSFORM; }
+		virtual EOpType GetOpType() const override { return EOpType::IM_TRANSFORM; }
 		virtual uint64 Hash() const override;
 		virtual bool IsEqual(const ASTOp& OtherUntyped) const override;
 		virtual Ptr<ASTOp> Clone(MapChildFuncRef MapChild) const override;
diff --git a/Engine/Plugins/Mutable/Source/MutableTools/Private/MuT/ASTOpInstanceAdd.cpp b/Engine/Plugins/Mutable/Source/MutableTools/Private/MuT/ASTOpInstanceAdd.cpp
--- a/Engine/Plugins/Mutable/Source/MutableTools/Private/MuT/ASTOpInstanceAdd.cpp
+++ b/Engine/Plugins/Mutable/Source/MutableTools/Private/MuT/ASTOpInstanceAdd.cpp
@@ -72,14 +72,14 @@ namespace mu
 	{
 		switch (type)
 		{
-		case OP_TYPE::IN_ADDMESH:
-		case OP_TYPE::IN_ADDIMAGE:
-		case OP_TYPE::IN_ADDVECTOR:
-		case OP_TYPE::IN_ADDSCALAR:
-		case OP_TYPE::IN_ADDSTRING:
-		case OP_TYPE::IN_ADDSURFACE:
-		case OP_TYPE::IN_ADDCOMPONENT:
-		case OP_TYPE::IN_ADDLOD:
+		case EOpType::IN_ADDMESH:
+		case EOpType::IN_ADDIMAGE:
+		case EOpType::IN_ADDVECTOR:
+		case EOpType::IN_ADDSCALAR:
+		case EOpType::IN_ADDSTRING:
+		case EOpType::IN_ADDSURFACE:
+		case EOpType::IN_ADDCOMPONENT:
+		case EOpType::IN_ADDLOD:
 			break;
 		default:
 			// Unexpected type
@@ -113,8 +113,8 @@ namespace mu
 			if (instance) args.instance = instance->linkedAddress;
 			if (value) args.value = value->linkedAddress;
 
-			if (type == OP_TYPE::IN_ADDIMAGE ||
-				type == OP_TYPE::IN_ADDMESH)
+			if (type == EOpType::IN_ADDIMAGE ||
+				type == EOpType::IN_ADDMESH)
 			{
 				// Find out relevant parameters. \todo: this may be optimised by reusing partial
 				// values in a LINK_CONTEXT or similar
@@ -122,12 +122,12 @@ namespace mu
 				visitor.Run(value.child());
 
 				TArray<uint16> params;
-				for (const FString& paramName : visitor.m_params)
+				for (const FString& paramName : visitor.Parameters)
 				{
-					for (int32 i = 0; i < program.m_parameters.Num(); ++i)
+					for (int32 i = 0; i < program.Parameters.Num(); ++i)
 					{
-						const auto& param = program.m_parameters[i];
-						if (param.m_name == paramName)
+						const auto& param = program.Parameters[i];
+						if (param.Name == paramName)
 						{
 							params.Add(uint16(i));
 							break;
@@ -137,7 +137,7 @@ namespace mu
 
 				params.Sort();
 
-				auto it = program.m_parameterLists.Find(params);
+				auto it = program.ParameterLists.Find(params);
 
 				if (it != INDEX_NONE)
 				{
@@ -145,15 +145,15 @@ namespace mu
 				}
 				else
 				{
-					args.relevantParametersListIndex = uint32_t(program.m_parameterLists.Num());
-					program.m_parameterLists.Add(params);
+					args.relevantParametersListIndex = uint32_t(program.ParameterLists.Num());
+					program.ParameterLists.Add(params);
 				}
 			}
 
-			linkedAddress = (OP::ADDRESS)program.m_opAddress.Num();
-			program.m_opAddress.Add((uint32_t)program.m_byteCode.Num());
-			AppendCode(program.m_byteCode, type);
-			AppendCode(program.m_byteCode, args);
+			linkedAddress = (OP::ADDRESS)program.OpAddress.Num();
+			program.OpAddress.Add((uint32_t)program.ByteCode.Num());
+			AppendCode(program.ByteCode, type);
+			AppendCode(program.ByteCode, args);
 		}
 
 	}
@@ -166,7 +166,7 @@ namespace mu
 		switch (type)
 		{
 
-		case OP_TYPE::IN_ADDMESH:
+		case EOpType::IN_ADDMESH:
 		{
 			Ptr<ASTOp> ValueAt = value.child();
 			if (!ValueAt)
@@ -174,21 +174,21 @@ namespace mu
 				break;
 			}
 
-			OP_TYPE ValueType = ValueAt->GetOpType();
+			EOpType ValueType = ValueAt->GetOpType();
 			switch (ValueType)
 			{
 
 			// We want to move the conditional up the op graph because this way the mesh root operation
 			// can easily match other mesh root operations in the program. This is important because this operation
 			// address is used for caching, and avoiding duplicate mesh work with multi-components.
-			case OP_TYPE::ME_CONDITIONAL:
+			case EOpType::ME_CONDITIONAL:
 			{
 				const ASTOpConditional* TypedValue = static_cast<const ASTOpConditional*>(ValueAt.get());
 
 				bool bCanSink = !instance.child();
 
 				const ASTOpConditional* TypedSource = nullptr;
-				if (instance.child() && instance.child()->GetOpType() == OP_TYPE::IN_CONDITIONAL)
+				if (instance.child() && instance.child()->GetOpType() == EOpType::IN_CONDITIONAL)
 				{
 					TypedSource = static_cast<const ASTOpConditional*>(instance.child().get());
 					if (TypedSource->condition == TypedValue->condition)
@@ -200,7 +200,7 @@ namespace mu
 				if (bCanSink)
 				{
 					Ptr<ASTOpConditional> NewOp = mu::Clone<ASTOpConditional>(ValueAt);
-					NewOp->type = OP_TYPE::IN_CONDITIONAL;
+					NewOp->type = EOpType::IN_CONDITIONAL;
 
 					Ptr<ASTOpInstanceAdd> TrueOp = mu::Clone<ASTOpInstanceAdd>(this);
 					TrueOp->instance = TypedSource ? TypedSource->yes.child() : nullptr; 
@@ -219,7 +219,7 @@ namespace mu
 
 			// I don't think we have switches at this point but if we did, it would be interesting to optimize
 			// for the same reasons 
-			case OP_TYPE::ME_SWITCH:
+			case EOpType::ME_SWITCH:
 			{
 				// TODO
 				break;
diff --git a/Engine/Plugins/Mutable/Source/MutableTools/Private/MuT/ASTOpInstanceAdd.h b/Engine/Plugins/Mutable/Source/MutableTools/Private/MuT/ASTOpInstanceAdd.h
--- a/Engine/Plugins/Mutable/Source/MutableTools/Private/MuT/ASTOpInstanceAdd.h
+++ b/Engine/Plugins/Mutable/Source/MutableTools/Private/MuT/ASTOpInstanceAdd.h
@@ -21,7 +21,7 @@ struct FProgram;
 	public:
 
 		//! Type of switch
-		OP_TYPE type = OP_TYPE::NONE;
+		EOpType type = EOpType::NONE;
 
 		ASTChild instance;
 		ASTChild value;
@@ -36,7 +36,7 @@ struct FProgram;
 		ASTOpInstanceAdd(const ASTOpInstanceAdd&) = delete;
 		~ASTOpInstanceAdd();
 
-		virtual OP_TYPE GetOpType() const override { return type; }
+		virtual EOpType GetOpType() const override { return type; }
 
 		virtual void ForEachChild(const TFunctionRef<void(ASTChild&)>) override;
 		virtual bool IsEqual(const ASTOp& otherUntyped) const override;
diff --git a/Engine/Plugins/Mutable/Source/MutableTools/Private/MuT/ASTOpLayoutFromMesh.cpp b/Engine/Plugins/Mutable/Source/MutableTools/Private/MuT/ASTOpLayoutFromMesh.cpp
--- a/Engine/Plugins/Mutable/Source/MutableTools/Private/MuT/ASTOpLayoutFromMesh.cpp
+++ b/Engine/Plugins/Mutable/Source/MutableTools/Private/MuT/ASTOpLayoutFromMesh.cpp
@@ -87,11 +87,11 @@ namespace mu
 			args.LayoutIndex = LayoutIndex;
 			if (Mesh) args.Mesh = Mesh->linkedAddress;
 
-			linkedAddress = (OP::ADDRESS)program.m_opAddress.Num();
+			linkedAddress = (OP::ADDRESS)program.OpAddress.Num();
 			//program.m_code.push_back(op);
-			program.m_opAddress.Add((uint32_t)program.m_byteCode.Num());
-			AppendCode(program.m_byteCode, GetOpType());
-			AppendCode(program.m_byteCode, args);
+			program.OpAddress.Add((uint32_t)program.ByteCode.Num());
+			AppendCode(program.ByteCode, GetOpType());
+			AppendCode(program.ByteCode, args);
 		}
 
 	}
@@ -154,7 +154,7 @@ namespace mu
 				switch (at->GetOpType())
 				{
 
-				case OP_TYPE::ME_MORPH:
+				case EOpType::ME_MORPH:
 				{
 					// Sink, ignoring the op
 					const ASTOpMeshMorph* Typed = static_cast<const ASTOpMeshMorph*>(at.get());
@@ -162,7 +162,7 @@ namespace mu
 					break;
 				}
 
-				case OP_TYPE::ME_FORMAT:
+				case EOpType::ME_FORMAT:
 				{
 					// Sink, ignoring the op
 					const ASTOpMeshFormat* Typed = static_cast<const ASTOpMeshFormat*>(at.get());
@@ -170,7 +170,7 @@ namespace mu
 					break;
 				}
 
-				case OP_TYPE::ME_APPLYSHAPE:
+				case EOpType::ME_APPLYSHAPE:
 				{
 					// Sink, ignoring the op
 					const ASTOpMeshApplyShape* Typed = static_cast<const ASTOpMeshApplyShape*>(at.get());
@@ -178,7 +178,7 @@ namespace mu
 					break;
 				}
 
-				case OP_TYPE::ME_BINDSHAPE:
+				case EOpType::ME_BINDSHAPE:
 				{
 					// Sink, ignoring the op
 					const ASTOpMeshBindShape* Typed = static_cast<const ASTOpMeshBindShape*>(at.get());
@@ -186,7 +186,7 @@ namespace mu
 					break;
 				}
 
-				case OP_TYPE::ME_ADDTAGS:
+				case EOpType::ME_ADDTAGS:
 				{
 					// Sink, ignoring the op
 					const ASTOpMeshAddTags* Typed = static_cast<const ASTOpMeshAddTags*>(at.get());
@@ -194,7 +194,7 @@ namespace mu
 					break;
 				}
 
-				case OP_TYPE::ME_INTERPOLATE:
+				case EOpType::ME_INTERPOLATE:
 				{
 					// Sink, ignoring the op
 					const ASTOpFixed* Typed = static_cast<const ASTOpFixed*>(at.get());
@@ -202,20 +202,20 @@ namespace mu
 					break;
 				}
 
-				case OP_TYPE::ME_CONDITIONAL:
+				case EOpType::ME_CONDITIONAL:
 				{
 					Ptr<ASTOpConditional> NewConditional = mu::Clone<ASTOpConditional>(at);
-					NewConditional->type = OP_TYPE::LA_CONDITIONAL;
+					NewConditional->type = EOpType::LA_CONDITIONAL;
 					NewConditional->yes = Visit(NewConditional->yes.child(), CurrentSinkingOp);
 					NewConditional->no = Visit(NewConditional->no.child(), CurrentSinkingOp);
 					newAt = NewConditional;
 					break;
 				}
 
-				case OP_TYPE::ME_SWITCH:
+				case EOpType::ME_SWITCH:
 				{
 					Ptr<ASTOpSwitch> NewOp = mu::Clone<ASTOpSwitch>(at);
-					NewOp->type = OP_TYPE::LA_SWITCH;
+					NewOp->type = EOpType::LA_SWITCH;
 					NewOp->def = Visit(NewOp->def.child(), CurrentSinkingOp);
 					for (ASTOpSwitch::FCase& c : NewOp->cases)
 					{
@@ -225,7 +225,7 @@ namespace mu
 					break;
 				}
 
-				case OP_TYPE::ME_MERGE:
+				case EOpType::ME_MERGE:
 				{
 					const ASTOpFixed* Typed = static_cast<const ASTOpFixed*>(at.get());
 
diff --git a/Engine/Plugins/Mutable/Source/MutableTools/Private/MuT/ASTOpLayoutFromMesh.h b/Engine/Plugins/Mutable/Source/MutableTools/Private/MuT/ASTOpLayoutFromMesh.h
--- a/Engine/Plugins/Mutable/Source/MutableTools/Private/MuT/ASTOpLayoutFromMesh.h
+++ b/Engine/Plugins/Mutable/Source/MutableTools/Private/MuT/ASTOpLayoutFromMesh.h
@@ -25,7 +25,7 @@ namespace mu
 		ASTOpLayoutFromMesh(const ASTOpLayoutFromMesh&) = delete;
 		~ASTOpLayoutFromMesh();
 
-		OP_TYPE GetOpType() const override { return OP_TYPE::LA_FROMMESH; }
+		EOpType GetOpType() const override { return EOpType::LA_FROMMESH; }
 		uint64 Hash() const override;
 		bool IsEqual(const ASTOp& otherUntyped) const override;
 		Ptr<ASTOp> Clone(MapChildFuncRef mapChild) const override;
diff --git a/Engine/Plugins/Mutable/Source/MutableTools/Private/MuT/ASTOpLayoutMerge.cpp b/Engine/Plugins/Mutable/Source/MutableTools/Private/MuT/ASTOpLayoutMerge.cpp
--- a/Engine/Plugins/Mutable/Source/MutableTools/Private/MuT/ASTOpLayoutMerge.cpp
+++ b/Engine/Plugins/Mutable/Source/MutableTools/Private/MuT/ASTOpLayoutMerge.cpp
@@ -84,10 +84,10 @@ namespace mu
 			if (Base) args.Base = Base->linkedAddress;
 			if (Added) args.Added = Added->linkedAddress;
 
-			linkedAddress = (OP::ADDRESS)program.m_opAddress.Num();
-			program.m_opAddress.Add((uint32_t)program.m_byteCode.Num());
-			AppendCode(program.m_byteCode, GetOpType());
-			AppendCode(program.m_byteCode, args);
+			linkedAddress = (OP::ADDRESS)program.OpAddress.Num();
+			program.OpAddress.Add((uint32_t)program.ByteCode.Num());
+			AppendCode(program.ByteCode, GetOpType());
+			AppendCode(program.ByteCode, args);
 		}
 
 	}
diff --git a/Engine/Plugins/Mutable/Source/MutableTools/Private/MuT/ASTOpLayoutMerge.h b/Engine/Plugins/Mutable/Source/MutableTools/Private/MuT/ASTOpLayoutMerge.h
--- a/Engine/Plugins/Mutable/Source/MutableTools/Private/MuT/ASTOpLayoutMerge.h
+++ b/Engine/Plugins/Mutable/Source/MutableTools/Private/MuT/ASTOpLayoutMerge.h
@@ -25,7 +25,7 @@ namespace mu
 		ASTOpLayoutMerge(const ASTOpLayoutMerge&) = delete;
 		~ASTOpLayoutMerge();
 
-		OP_TYPE GetOpType() const override { return OP_TYPE::LA_MERGE; }
+		EOpType GetOpType() const override { return EOpType::LA_MERGE; }
 		uint64 Hash() const override;
 		bool IsEqual(const ASTOp& otherUntyped) const override;
 		Ptr<ASTOp> Clone(MapChildFuncRef mapChild) const override;
diff --git a/Engine/Plugins/Mutable/Source/MutableTools/Private/MuT/ASTOpLayoutPack.cpp b/Engine/Plugins/Mutable/Source/MutableTools/Private/MuT/ASTOpLayoutPack.cpp
--- a/Engine/Plugins/Mutable/Source/MutableTools/Private/MuT/ASTOpLayoutPack.cpp
+++ b/Engine/Plugins/Mutable/Source/MutableTools/Private/MuT/ASTOpLayoutPack.cpp
@@ -79,10 +79,10 @@ namespace mu
 
 			if (Source) args.Source = Source->linkedAddress;
 
-			linkedAddress = (OP::ADDRESS)program.m_opAddress.Num();
-			program.m_opAddress.Add((uint32_t)program.m_byteCode.Num());
-			AppendCode(program.m_byteCode, GetOpType());
-			AppendCode(program.m_byteCode, args);
+			linkedAddress = (OP::ADDRESS)program.OpAddress.Num();
+			program.OpAddress.Add((uint32_t)program.ByteCode.Num());
+			AppendCode(program.ByteCode, GetOpType());
+			AppendCode(program.ByteCode, args);
 		}
 
 	}
diff --git a/Engine/Plugins/Mutable/Source/MutableTools/Private/MuT/ASTOpLayoutPack.h b/Engine/Plugins/Mutable/Source/MutableTools/Private/MuT/ASTOpLayoutPack.h
--- a/Engine/Plugins/Mutable/Source/MutableTools/Private/MuT/ASTOpLayoutPack.h
+++ b/Engine/Plugins/Mutable/Source/MutableTools/Private/MuT/ASTOpLayoutPack.h
@@ -24,7 +24,7 @@ namespace mu
 		ASTOpLayoutPack(const ASTOpLayoutPack&) = delete;
 		~ASTOpLayoutPack();
 
-		OP_TYPE GetOpType() const override { return OP_TYPE::LA_PACK; }
+		EOpType GetOpType() const override { return EOpType::LA_PACK; }
 		uint64 Hash() const override;
 		bool IsEqual(const ASTOp& otherUntyped) const override;
 		Ptr<ASTOp> Clone(MapChildFuncRef mapChild) const override;
diff --git a/Engine/Plugins/Mutable/Source/MutableTools/Private/MuT/ASTOpLayoutRemoveBlocks.cpp b/Engine/Plugins/Mutable/Source/MutableTools/Private/MuT/ASTOpLayoutRemoveBlocks.cpp
--- a/Engine/Plugins/Mutable/Source/MutableTools/Private/MuT/ASTOpLayoutRemoveBlocks.cpp
+++ b/Engine/Plugins/Mutable/Source/MutableTools/Private/MuT/ASTOpLayoutRemoveBlocks.cpp
@@ -84,10 +84,10 @@ namespace mu
 			if (Source) args.Source = Source->linkedAddress;
 			if (ReferenceLayout) args.ReferenceLayout = ReferenceLayout->linkedAddress;
 
-			linkedAddress = (OP::ADDRESS)program.m_opAddress.Num();
-			program.m_opAddress.Add((uint32_t)program.m_byteCode.Num());
-			AppendCode(program.m_byteCode, OP_TYPE::LA_REMOVEBLOCKS);
-			AppendCode(program.m_byteCode, args);
+			linkedAddress = (OP::ADDRESS)program.OpAddress.Num();
+			program.OpAddress.Add((uint32_t)program.ByteCode.Num());
+			AppendCode(program.ByteCode, EOpType::LA_REMOVEBLOCKS);
+			AppendCode(program.ByteCode, args);
 		}
 
 	}
diff --git a/Engine/Plugins/Mutable/Source/MutableTools/Private/MuT/ASTOpLayoutRemoveBlocks.h b/Engine/Plugins/Mutable/Source/MutableTools/Private/MuT/ASTOpLayoutRemoveBlocks.h
--- a/Engine/Plugins/Mutable/Source/MutableTools/Private/MuT/ASTOpLayoutRemoveBlocks.h
+++ b/Engine/Plugins/Mutable/Source/MutableTools/Private/MuT/ASTOpLayoutRemoveBlocks.h
@@ -25,7 +25,7 @@ namespace mu
 		ASTOpLayoutRemoveBlocks(const ASTOpLayoutRemoveBlocks&) = delete;
 		~ASTOpLayoutRemoveBlocks();
 
-		OP_TYPE GetOpType() const override { return OP_TYPE::LA_REMOVEBLOCKS; }
+		EOpType GetOpType() const override { return EOpType::LA_REMOVEBLOCKS; }
 		uint64 Hash() const override;
 		bool IsEqual(const ASTOp& otherUntyped) const override;
 		Ptr<ASTOp> Clone(MapChildFuncRef mapChild) const override;
diff --git a/Engine/Plugins/Mutable/Source/MutableTools/Private/MuT/ASTOpMeshAddTags.cpp b/Engine/Plugins/Mutable/Source/MutableTools/Private/MuT/ASTOpMeshAddTags.cpp
--- a/Engine/Plugins/Mutable/Source/MutableTools/Private/MuT/ASTOpMeshAddTags.cpp
+++ b/Engine/Plugins/Mutable/Source/MutableTools/Private/MuT/ASTOpMeshAddTags.cpp
@@ -63,17 +63,17 @@ namespace mu
 		// Already linked?
 		if (!linkedAddress)
 		{
-			linkedAddress = (OP::ADDRESS)Program.m_opAddress.Num();
+			linkedAddress = (OP::ADDRESS)Program.OpAddress.Num();
 
-			Program.m_opAddress.Add((uint32)Program.m_byteCode.Num());
-			AppendCode(Program.m_byteCode, OP_TYPE::ME_ADDTAGS);
+			Program.OpAddress.Add((uint32)Program.ByteCode.Num());
+			AppendCode(Program.ByteCode, EOpType::ME_ADDTAGS);
 			OP::ADDRESS SourceAt = Source ? Source->linkedAddress : 0;
-			AppendCode(Program.m_byteCode, SourceAt);
-			AppendCode(Program.m_byteCode, (uint16)Tags.Num());
+			AppendCode(Program.ByteCode, SourceAt);
+			AppendCode(Program.ByteCode, (uint16)Tags.Num());
 			for (const FString& Tag : Tags)
 			{
 				OP::ADDRESS TagConstantAddress = Program.AddConstant(Tag);
-				AppendCode(Program.m_byteCode, TagConstantAddress);
+				AppendCode(Program.ByteCode, TagConstantAddress);
 			}
 		}
 	}
diff --git a/Engine/Plugins/Mutable/Source/MutableTools/Private/MuT/ASTOpMeshAddTags.h b/Engine/Plugins/Mutable/Source/MutableTools/Private/MuT/ASTOpMeshAddTags.h
--- a/Engine/Plugins/Mutable/Source/MutableTools/Private/MuT/ASTOpMeshAddTags.h
+++ b/Engine/Plugins/Mutable/Source/MutableTools/Private/MuT/ASTOpMeshAddTags.h
@@ -30,7 +30,7 @@ namespace mu
 		virtual ~ASTOpMeshAddTags() override;
 
 		// ASTOp interface
-		virtual OP_TYPE GetOpType() const override { return OP_TYPE::ME_ADDTAGS; }
+		virtual EOpType GetOpType() const override { return EOpType::ME_ADDTAGS; }
 		virtual uint64 Hash() const override;
 		virtual void ForEachChild(const TFunctionRef<void(ASTChild&)>) override;
 		virtual bool IsEqual(const ASTOp& otherUntyped) const override;
diff --git a/Engine/Plugins/Mutable/Source/MutableTools/Private/MuT/ASTOpMeshApplyPose.cpp b/Engine/Plugins/Mutable/Source/MutableTools/Private/MuT/ASTOpMeshApplyPose.cpp
--- a/Engine/Plugins/Mutable/Source/MutableTools/Private/MuT/ASTOpMeshApplyPose.cpp
+++ b/Engine/Plugins/Mutable/Source/MutableTools/Private/MuT/ASTOpMeshApplyPose.cpp
@@ -74,11 +74,11 @@ namespace mu
 			if (base) args.base = base->linkedAddress;
 			if (pose) args.pose = pose->linkedAddress;
 
-			linkedAddress = (OP::ADDRESS)program.m_opAddress.Num();
+			linkedAddress = (OP::ADDRESS)program.OpAddress.Num();
 			//program.m_code.push_back(op);
-			program.m_opAddress.Add((uint32_t)program.m_byteCode.Num());
-			AppendCode(program.m_byteCode, OP_TYPE::ME_APPLYPOSE);
-			AppendCode(program.m_byteCode, args);
+			program.OpAddress.Add((uint32_t)program.ByteCode.Num());
+			AppendCode(program.ByteCode, EOpType::ME_APPLYPOSE);
+			AppendCode(program.ByteCode, args);
 		}
 
 	}
@@ -94,11 +94,11 @@ namespace mu
 			return nullptr;
 		}
 
-		OP_TYPE MeshType = MeshAt->GetOpType();
+		EOpType MeshType = MeshAt->GetOpType();
 
 		switch (MeshType)
 		{
-			case OP_TYPE::ME_ADDTAGS:
+			case EOpType::ME_ADDTAGS:
 			{
 				Ptr<ASTOpMeshAddTags> New = mu::Clone<ASTOpMeshAddTags>(MeshAt);
 				if (New->Source)
diff --git a/Engine/Plugins/Mutable/Source/MutableTools/Private/MuT/ASTOpMeshApplyPose.h b/Engine/Plugins/Mutable/Source/MutableTools/Private/MuT/ASTOpMeshApplyPose.h
--- a/Engine/Plugins/Mutable/Source/MutableTools/Private/MuT/ASTOpMeshApplyPose.h
+++ b/Engine/Plugins/Mutable/Source/MutableTools/Private/MuT/ASTOpMeshApplyPose.h
@@ -29,7 +29,7 @@ struct FProgram;
 		ASTOpMeshApplyPose(const ASTOpMeshApplyPose&) = delete;
 		~ASTOpMeshApplyPose();
 
-		OP_TYPE GetOpType() const override { return OP_TYPE::ME_APPLYPOSE; }
+		EOpType GetOpType() const override { return EOpType::ME_APPLYPOSE; }
 		uint64 Hash() const override;
 		bool IsEqual(const ASTOp& otherUntyped) const override;
 		Ptr<ASTOp> Clone(MapChildFuncRef mapChild) const override;
diff --git a/Engine/Plugins/Mutable/Source/MutableTools/Private/MuT/ASTOpMeshApplyShape.cpp b/Engine/Plugins/Mutable/Source/MutableTools/Private/MuT/ASTOpMeshApplyShape.cpp
--- a/Engine/Plugins/Mutable/Source/MutableTools/Private/MuT/ASTOpMeshApplyShape.cpp
+++ b/Engine/Plugins/Mutable/Source/MutableTools/Private/MuT/ASTOpMeshApplyShape.cpp
@@ -106,10 +106,10 @@ namespace mu
 			Args.mesh = Mesh ? Mesh->linkedAddress : 0;
 			Args.shape = Shape ? Shape->linkedAddress : 0;
 
-			linkedAddress = (OP::ADDRESS)program.m_opAddress.Num();
-			program.m_opAddress.Add((uint32_t)program.m_byteCode.Num());
-			AppendCode(program.m_byteCode, OP_TYPE::ME_APPLYSHAPE);
-			AppendCode(program.m_byteCode, Args);
+			linkedAddress = (OP::ADDRESS)program.OpAddress.Num();
+			program.OpAddress.Add((uint32_t)program.ByteCode.Num());
+			AppendCode(program.ByteCode, EOpType::ME_APPLYSHAPE);
+			AppendCode(program.ByteCode, Args);
 		}
 
 	}
diff --git a/Engine/Plugins/Mutable/Source/MutableTools/Private/MuT/ASTOpMeshApplyShape.h b/Engine/Plugins/Mutable/Source/MutableTools/Private/MuT/ASTOpMeshApplyShape.h
--- a/Engine/Plugins/Mutable/Source/MutableTools/Private/MuT/ASTOpMeshApplyShape.h
+++ b/Engine/Plugins/Mutable/Source/MutableTools/Private/MuT/ASTOpMeshApplyShape.h
@@ -33,7 +33,7 @@ struct FProgram;
 		ASTOpMeshApplyShape(const ASTOpMeshApplyShape&) = delete;
 		~ASTOpMeshApplyShape();
 
-		OP_TYPE GetOpType() const override { return OP_TYPE::ME_APPLYSHAPE; }
+		EOpType GetOpType() const override { return EOpType::ME_APPLYSHAPE; }
 		uint64 Hash() const override;
 		bool IsEqual(const ASTOp& otherUntyped) const override;
 		Ptr<ASTOp> Clone(MapChildFuncRef mapChild) const override;
diff --git a/Engine/Plugins/Mutable/Source/MutableTools/Private/MuT/ASTOpMeshBindShape.cpp b/Engine/Plugins/Mutable/Source/MutableTools/Private/MuT/ASTOpMeshBindShape.cpp
--- a/Engine/Plugins/Mutable/Source/MutableTools/Private/MuT/ASTOpMeshBindShape.cpp
+++ b/Engine/Plugins/Mutable/Source/MutableTools/Private/MuT/ASTOpMeshBindShape.cpp
@@ -177,21 +177,21 @@ namespace mu
 				Args.shape = Shape->linkedAddress;
 			}
 
-			linkedAddress = (OP::ADDRESS)program.m_opAddress.Num();
-			program.m_opAddress.Add((uint32_t)program.m_byteCode.Num());
-			AppendCode(program.m_byteCode, OP_TYPE::ME_BINDSHAPE);
-			AppendCode(program.m_byteCode, Args);
+			linkedAddress = (OP::ADDRESS)program.OpAddress.Num();
+			program.OpAddress.Add((uint32_t)program.ByteCode.Num());
+			AppendCode(program.ByteCode, EOpType::ME_BINDSHAPE);
+			AppendCode(program.ByteCode, Args);
 
-			AppendCode(program.m_byteCode, (int32)BonesToDeform.Num());
+			AppendCode(program.ByteCode, (int32)BonesToDeform.Num());
 			for (const FBoneName& S : BonesToDeform)
 			{
-				AppendCode(program.m_byteCode, S);
+				AppendCode(program.ByteCode, S);
 			}
 
-			AppendCode(program.m_byteCode, (int32)PhysicsToDeform.Num());
+			AppendCode(program.ByteCode, (int32)PhysicsToDeform.Num());
 			for (const FBoneName& S : PhysicsToDeform)
 			{
-				AppendCode(program.m_byteCode, S);
+				AppendCode(program.ByteCode, S);
 			}
 		}
 	}
@@ -214,8 +214,8 @@ namespace mu
 			return nullptr;
 		}
 
-		OP_TYPE MeshType = MeshAt->GetOpType();
-		OP_TYPE ShapeType = ShapeAt->GetOpType();
+		EOpType MeshType = MeshAt->GetOpType();
+		EOpType ShapeType = ShapeAt->GetOpType();
 
 		// See if both mesh and shape have an operation that can be optimized in a combined way
 		if (MeshType == ShapeType)
@@ -223,7 +223,7 @@ namespace mu
 			switch (MeshType)
 			{
 
-			case OP_TYPE::ME_SWITCH:
+			case EOpType::ME_SWITCH:
 			{
 				// If the switch variable and structure is the same
 				const ASTOpSwitch* MeshSwitch = static_cast<const ASTOpSwitch*>(MeshAt.get());
@@ -261,7 +261,7 @@ namespace mu
 			}
 
 
-			case OP_TYPE::ME_CONDITIONAL:
+			case EOpType::ME_CONDITIONAL:
 			{
 				const ASTOpConditional* MeshConditional = static_cast<const ASTOpConditional*>(MeshAt.get());
 				const ASTOpConditional* ShapeConditional = static_cast<const ASTOpConditional*>(ShapeAt.get());
@@ -308,7 +308,7 @@ namespace mu
 			switch (MeshType)
 			{
 
-			case OP_TYPE::ME_SWITCH:
+			case EOpType::ME_SWITCH:
 			{
 				// Move the operation down all the paths
 				Ptr<ASTOpSwitch> NewSwitch = mu::Clone<ASTOpSwitch>(MeshAt);
@@ -334,7 +334,7 @@ namespace mu
 				break;
 			}
 
-			case OP_TYPE::ME_CONDITIONAL:
+			case EOpType::ME_CONDITIONAL:
 			{
 				// Move the operation down all the paths
 				Ptr<ASTOpConditional> NewConditional = mu::Clone<ASTOpConditional>(MeshAt);
@@ -357,7 +357,7 @@ namespace mu
 				break;
 			}
 
-			case OP_TYPE::ME_REMOVEMASK:
+			case EOpType::ME_REMOVEMASK:
 			{
 				// We bind something that could have a part removed: we can reorder to bind the entire mesh
 				// and apply remove later at runtime.
@@ -374,7 +374,7 @@ namespace mu
 				break;
 			}
 
-			case OP_TYPE::ME_ADDTAGS:
+			case EOpType::ME_ADDTAGS:
 			{
 				Ptr<ASTOpMeshAddTags> New = mu::Clone<ASTOpMeshAddTags>(MeshAt);
 				if (New->Source)
@@ -401,7 +401,7 @@ namespace mu
 			switch (ShapeType)
 			{
 
-			case OP_TYPE::ME_SWITCH:
+			case EOpType::ME_SWITCH:
 			{
 				// Move the operation down all the paths
 				Ptr<ASTOpSwitch> NewSwitch = mu::Clone<ASTOpSwitch>(ShapeAt);
@@ -427,7 +427,7 @@ namespace mu
 				break;
 			}
 
-			case OP_TYPE::ME_CONDITIONAL:
+			case EOpType::ME_CONDITIONAL:
 			{
 				// Move the operation down all the paths
 				Ptr<ASTOpConditional> NewConditional = mu::Clone<ASTOpConditional>(ShapeAt);
@@ -450,7 +450,7 @@ namespace mu
 				break;
 			}
 
-			case OP_TYPE::ME_ADDTAGS:
+			case EOpType::ME_ADDTAGS:
 			{
 				// Ignore the tags in the shape
 				Ptr<ASTOpMeshBindShape> NewBind = mu::Clone<ASTOpMeshBindShape>(this);
diff --git a/Engine/Plugins/Mutable/Source/MutableTools/Private/MuT/ASTOpMeshBindShape.h b/Engine/Plugins/Mutable/Source/MutableTools/Private/MuT/ASTOpMeshBindShape.h
--- a/Engine/Plugins/Mutable/Source/MutableTools/Private/MuT/ASTOpMeshBindShape.h
+++ b/Engine/Plugins/Mutable/Source/MutableTools/Private/MuT/ASTOpMeshBindShape.h
@@ -45,7 +45,7 @@ namespace mu
 		ASTOpMeshBindShape(const ASTOpMeshBindShape&) = delete;
 		~ASTOpMeshBindShape();
 
-		OP_TYPE GetOpType() const override { return OP_TYPE::ME_BINDSHAPE; }
+		EOpType GetOpType() const override { return EOpType::ME_BINDSHAPE; }
 		uint64 Hash() const override;
 		bool IsEqual(const ASTOp& otherUntyped) const override;
 		Ptr<ASTOp> Clone(MapChildFuncRef mapChild) const override;
diff --git a/Engine/Plugins/Mutable/Source/MutableTools/Private/MuT/ASTOpMeshClipDeform.cpp b/Engine/Plugins/Mutable/Source/MutableTools/Private/MuT/ASTOpMeshClipDeform.cpp
--- a/Engine/Plugins/Mutable/Source/MutableTools/Private/MuT/ASTOpMeshClipDeform.cpp
+++ b/Engine/Plugins/Mutable/Source/MutableTools/Private/MuT/ASTOpMeshClipDeform.cpp
@@ -83,10 +83,10 @@ namespace mu
 				args.clipShape = ClipShape->linkedAddress;
 			}
 
-			linkedAddress = (OP::ADDRESS)program.m_opAddress.Num();
-			program.m_opAddress.Add((uint32_t)program.m_byteCode.Num());
-			AppendCode(program.m_byteCode, OP_TYPE::ME_CLIPDEFORM);
-			AppendCode(program.m_byteCode, args);
+			linkedAddress = (OP::ADDRESS)program.OpAddress.Num();
+			program.OpAddress.Add((uint32_t)program.ByteCode.Num());
+			AppendCode(program.ByteCode, EOpType::ME_CLIPDEFORM);
+			AppendCode(program.ByteCode, args);
 		}
 
 	}
diff --git a/Engine/Plugins/Mutable/Source/MutableTools/Private/MuT/ASTOpMeshClipDeform.h b/Engine/Plugins/Mutable/Source/MutableTools/Private/MuT/ASTOpMeshClipDeform.h
--- a/Engine/Plugins/Mutable/Source/MutableTools/Private/MuT/ASTOpMeshClipDeform.h
+++ b/Engine/Plugins/Mutable/Source/MutableTools/Private/MuT/ASTOpMeshClipDeform.h
@@ -32,7 +32,7 @@ struct FProgram;
 		ASTOpMeshClipDeform(const ASTOpMeshClipDeform&) = delete;
 		~ASTOpMeshClipDeform();
 
-		OP_TYPE GetOpType() const override { return OP_TYPE::ME_CLIPDEFORM; }
+		EOpType GetOpType() const override { return EOpType::ME_CLIPDEFORM; }
 		uint64 Hash() const override;
 		bool IsEqual(const ASTOp& otherUntyped) const override;
 		Ptr<ASTOp> Clone(MapChildFuncRef mapChild) const override;
diff --git a/Engine/Plugins/Mutable/Source/MutableTools/Private/MuT/ASTOpMeshClipMorphPlane.cpp b/Engine/Plugins/Mutable/Source/MutableTools/Private/MuT/ASTOpMeshClipMorphPlane.cpp
--- a/Engine/Plugins/Mutable/Source/MutableTools/Private/MuT/ASTOpMeshClipMorphPlane.cpp
+++ b/Engine/Plugins/Mutable/Source/MutableTools/Private/MuT/ASTOpMeshClipMorphPlane.cpp
@@ -88,8 +88,8 @@ namespace mu
 
 			if (source) args.source = source->linkedAddress;
 
-			args.morphShape = (OP::ADDRESS)program.m_constantShapes.Num();
-			program.m_constantShapes.Add(morphShape);
+			args.morphShape = (OP::ADDRESS)program.ConstantShapes.Num();
+			program.ConstantShapes.Add(morphShape);
 
 			args.FaceCullStrategy = FaceCullStrategy;
 
@@ -108,11 +108,11 @@ namespace mu
 			args.factor = factor;
 			args.maxBoneRadius = vertexSelectionBoneMaxRadius;
 
-			linkedAddress = (OP::ADDRESS)program.m_opAddress.Num();
+			linkedAddress = (OP::ADDRESS)program.OpAddress.Num();
 			//program.m_code.push_back(op);
-			program.m_opAddress.Add((uint32_t)program.m_byteCode.Num());
-			AppendCode(program.m_byteCode, OP_TYPE::ME_CLIPMORPHPLANE);
-			AppendCode(program.m_byteCode, args);
+			program.OpAddress.Add((uint32_t)program.ByteCode.Num());
+			AppendCode(program.ByteCode, EOpType::ME_CLIPMORPHPLANE);
+			AppendCode(program.ByteCode, args);
 		}
 
 	}
@@ -127,13 +127,13 @@ namespace mu
 			return nullptr;
 		}
 
-		OP_TYPE SourceType = source.child()->GetOpType();
+		EOpType SourceType = source.child()->GetOpType();
 
 		// Optimize only the mesh parameter
 		switch (SourceType)
 		{
 
-		case OP_TYPE::ME_ADDTAGS:
+		case EOpType::ME_ADDTAGS:
 		{
 			Ptr<ASTOpMeshAddTags> NewAddTags = mu::Clone<ASTOpMeshAddTags>(source.child());
 
diff --git a/Engine/Plugins/Mutable/Source/MutableTools/Private/MuT/ASTOpMeshClipMorphPlane.h b/Engine/Plugins/Mutable/Source/MutableTools/Private/MuT/ASTOpMeshClipMorphPlane.h
--- a/Engine/Plugins/Mutable/Source/MutableTools/Private/MuT/ASTOpMeshClipMorphPlane.h
+++ b/Engine/Plugins/Mutable/Source/MutableTools/Private/MuT/ASTOpMeshClipMorphPlane.h
@@ -37,7 +37,7 @@ struct FProgram;
 		ASTOpMeshClipMorphPlane(const ASTOpMeshClipMorphPlane&) = delete;
 		virtual ~ASTOpMeshClipMorphPlane();
 
-		virtual OP_TYPE GetOpType() const override { return OP_TYPE::ME_CLIPMORPHPLANE; }
+		virtual EOpType GetOpType() const override { return EOpType::ME_CLIPMORPHPLANE; }
 		virtual uint64 Hash() const override;
 		virtual bool IsEqual(const ASTOp& otherUntyped) const override;
 		virtual Ptr<ASTOp> Clone(MapChildFuncRef mapChild) const override;
diff --git a/Engine/Plugins/Mutable/Source/MutableTools/Private/MuT/ASTOpMeshDifference.cpp b/Engine/Plugins/Mutable/Source/MutableTools/Private/MuT/ASTOpMeshDifference.cpp
--- a/Engine/Plugins/Mutable/Source/MutableTools/Private/MuT/ASTOpMeshDifference.cpp
+++ b/Engine/Plugins/Mutable/Source/MutableTools/Private/MuT/ASTOpMeshDifference.cpp
@@ -80,24 +80,24 @@ namespace mu
 		// Already linked?
 		if (!linkedAddress)
 		{
-			linkedAddress = (OP::ADDRESS)program.m_opAddress.Num();
+			linkedAddress = (OP::ADDRESS)program.OpAddress.Num();
 
-			program.m_opAddress.Add((uint32_t)program.m_byteCode.Num());
-			AppendCode(program.m_byteCode, OP_TYPE::ME_DIFFERENCE);
+			program.OpAddress.Add((uint32_t)program.ByteCode.Num());
+			AppendCode(program.ByteCode, EOpType::ME_DIFFERENCE);
 
 			OP::ADDRESS BaseAt = Base ? Base->linkedAddress : 0;
-			AppendCode(program.m_byteCode, BaseAt);
+			AppendCode(program.ByteCode, BaseAt);
 
 			OP::ADDRESS TargetAt = Target ? Target->linkedAddress : 0;
-			AppendCode(program.m_byteCode, TargetAt);
+			AppendCode(program.ByteCode, TargetAt);
 
-			AppendCode(program.m_byteCode, (uint8)bIgnoreTextureCoords);
+			AppendCode(program.ByteCode, (uint8)bIgnoreTextureCoords);
 
-			AppendCode(program.m_byteCode, (uint8)Channels.Num());
+			AppendCode(program.ByteCode, (uint8)Channels.Num());
 			for (const FChannel& b : Channels)
 			{
-				AppendCode(program.m_byteCode, b.Semantic);
-				AppendCode(program.m_byteCode, b.SemanticIndex);
+				AppendCode(program.ByteCode, b.Semantic);
+				AppendCode(program.ByteCode, b.SemanticIndex);
 			}
 		}
 	}
@@ -120,8 +120,8 @@ namespace mu
 			return nullptr;
 		}
 
-		OP_TYPE BaseType = BaseAt->GetOpType();
-		OP_TYPE TargetType = TargetAt->GetOpType();
+		EOpType BaseType = BaseAt->GetOpType();
+		EOpType TargetType = TargetAt->GetOpType();
 
 		// See if both base and target have an operation that can be optimized in a combined way
 		if (BaseType == TargetType)
@@ -129,7 +129,7 @@ namespace mu
 			switch (BaseType)
 			{
 
-			case OP_TYPE::ME_SWITCH:
+			case EOpType::ME_SWITCH:
 			{
 				// If the switch variable and structure is the same
 				const ASTOpSwitch* BaseSwitch = static_cast<const ASTOpSwitch*>(BaseAt.get());
@@ -167,7 +167,7 @@ namespace mu
 			}
 
 
-			case OP_TYPE::ME_CONDITIONAL:
+			case EOpType::ME_CONDITIONAL:
 			{
 				const ASTOpConditional* BaseConditional = static_cast<const ASTOpConditional*>(BaseAt.get());
 				const ASTOpConditional* TargetConditional = static_cast<const ASTOpConditional*>(TargetAt.get());
@@ -214,7 +214,7 @@ namespace mu
 			switch (BaseType)
 			{
 
-			case OP_TYPE::ME_SWITCH:
+			case EOpType::ME_SWITCH:
 			{
 				// Move the operation down all the paths
 				Ptr<ASTOpSwitch> NewSwitch = mu::Clone<ASTOpSwitch>(BaseAt);
@@ -240,7 +240,7 @@ namespace mu
 				break;
 			}
 
-			case OP_TYPE::ME_CONDITIONAL:
+			case EOpType::ME_CONDITIONAL:
 			{
 				// Move the operation down all the paths
 				Ptr<ASTOpConditional> NewConditional = mu::Clone<ASTOpConditional>(BaseAt);
@@ -263,7 +263,7 @@ namespace mu
 				break;
 			}
 
-			case OP_TYPE::ME_ADDTAGS:
+			case EOpType::ME_ADDTAGS:
 			{
 				Ptr<ASTOpMeshAddTags> NewAdd = mu::Clone<ASTOpMeshAddTags>(BaseAt);
 
@@ -291,7 +291,7 @@ namespace mu
 			switch (TargetType)
 			{
 
-			case OP_TYPE::ME_SWITCH:
+			case EOpType::ME_SWITCH:
 			{
 				// Move the operation down all the paths
 				Ptr<ASTOpSwitch> NewSwitch = mu::Clone<ASTOpSwitch>(TargetAt);
@@ -317,7 +317,7 @@ namespace mu
 				break;
 			}
 
-			case OP_TYPE::ME_CONDITIONAL:
+			case EOpType::ME_CONDITIONAL:
 			{
 				// Move the operation down all the paths
 				Ptr<ASTOpConditional> NewConditional = mu::Clone<ASTOpConditional>(TargetAt);
@@ -340,7 +340,7 @@ namespace mu
 				break;
 			}
 
-			case OP_TYPE::ME_ADDTAGS:
+			case EOpType::ME_ADDTAGS:
 			{
 				// Ignore tags in this branch
 				const ASTOpMeshAddTags* Add = static_cast<const ASTOpMeshAddTags*>(TargetAt.get());
diff --git a/Engine/Plugins/Mutable/Source/MutableTools/Private/MuT/ASTOpMeshDifference.h b/Engine/Plugins/Mutable/Source/MutableTools/Private/MuT/ASTOpMeshDifference.h
--- a/Engine/Plugins/Mutable/Source/MutableTools/Private/MuT/ASTOpMeshDifference.h
+++ b/Engine/Plugins/Mutable/Source/MutableTools/Private/MuT/ASTOpMeshDifference.h
@@ -46,7 +46,7 @@ namespace mu
 		~ASTOpMeshDifference() override;
 
 		// ASTOp interface
-		OP_TYPE GetOpType() const override { return OP_TYPE::ME_DIFFERENCE; }
+		EOpType GetOpType() const override { return EOpType::ME_DIFFERENCE; }
 		uint64 Hash() const override;
 		void ForEachChild(const TFunctionRef<void(ASTChild&)>) override;
 		bool IsEqual(const ASTOp& otherUntyped) const override;
diff --git a/Engine/Plugins/Mutable/Source/MutableTools/Private/MuT/ASTOpMeshExtractLayoutBlocks.cpp b/Engine/Plugins/Mutable/Source/MutableTools/Private/MuT/ASTOpMeshExtractLayoutBlocks.cpp
--- a/Engine/Plugins/Mutable/Source/MutableTools/Private/MuT/ASTOpMeshExtractLayoutBlocks.cpp
+++ b/Engine/Plugins/Mutable/Source/MutableTools/Private/MuT/ASTOpMeshExtractLayoutBlocks.cpp
@@ -79,18 +79,18 @@ namespace mu
 		// Already linked?
 		if (!linkedAddress)
 		{
-			linkedAddress = (OP::ADDRESS)program.m_opAddress.Num();
+			linkedAddress = (OP::ADDRESS)program.OpAddress.Num();
 
-			program.m_opAddress.Add((uint32)program.m_byteCode.Num());
-			AppendCode(program.m_byteCode, OP_TYPE::ME_EXTRACTLAYOUTBLOCK);
+			program.OpAddress.Add((uint32)program.ByteCode.Num());
+			AppendCode(program.ByteCode, EOpType::ME_EXTRACTLAYOUTBLOCK);
 			OP::ADDRESS sourceAt = Source ? Source->linkedAddress : 0;
-			AppendCode(program.m_byteCode, sourceAt);
-			AppendCode(program.m_byteCode, (uint16)LayoutIndex);
-			AppendCode(program.m_byteCode, (uint16)Blocks.Num());
+			AppendCode(program.ByteCode, sourceAt);
+			AppendCode(program.ByteCode, (uint16)LayoutIndex);
+			AppendCode(program.ByteCode, (uint16)Blocks.Num());
 
 			for (uint64 Id : Blocks)
 			{
-				AppendCode(program.m_byteCode, Id);
+				AppendCode(program.ByteCode, Id);
 			}
 		}
 	}
@@ -117,15 +117,15 @@ namespace mu
 	//---------------------------------------------------------------------------------------------
 	mu::Ptr<ASTOp> Sink_MeshExtractLayoutBlocksAST::Apply(const ASTOpMeshExtractLayoutBlocks* root)
 	{
-		m_root = root;
+		Root = root;
 
 		OldToNew.Reset();
 
-		m_initialSource = m_root->Source.child();
-		mu::Ptr<ASTOp> newSource = Visit(m_initialSource, m_root);
+		InitialSource = Root->Source.child();
+		mu::Ptr<ASTOp> newSource = Visit(InitialSource, Root);
 
 		// If there is any change, it is the new root.
-		if (newSource != m_initialSource)
+		if (newSource != InitialSource)
 		{
 			return newSource;
 		}
@@ -148,7 +148,7 @@ namespace mu
 		switch (at->GetOpType())
 		{
 
-		case OP_TYPE::ME_APPLYLAYOUT:
+		case EOpType::ME_APPLYLAYOUT:
 		{
 			Ptr<ASTOpFixed> newOp = mu::Clone<ASTOpFixed>(at);
 			newOp->SetChild(newOp->op.args.MeshApplyLayout.mesh, Visit(newOp->children[newOp->op.args.MeshApplyLayout.mesh].child(), currentSinkOp));
@@ -156,7 +156,7 @@ namespace mu
 			break;
 		}
 
-		case OP_TYPE::ME_SETSKELETON:
+		case EOpType::ME_SETSKELETON:
 		{
 			Ptr<ASTOpFixed> newOp = mu::Clone<ASTOpFixed>(at);
 			newOp->SetChild(newOp->op.args.MeshSetSkeleton.source, Visit(newOp->children[newOp->op.args.MeshSetSkeleton.source].child(), currentSinkOp));
@@ -164,7 +164,7 @@ namespace mu
 			break;
 		}
 
-		case OP_TYPE::ME_ADDTAGS:
+		case EOpType::ME_ADDTAGS:
 		{
 			Ptr<ASTOpMeshAddTags> newOp = mu::Clone<ASTOpMeshAddTags>(at);
 			newOp->Source = Visit(newOp->Source.child(), currentSinkOp);
@@ -172,7 +172,7 @@ namespace mu
 			break;
 		}
 
-		case OP_TYPE::ME_CLIPMORPHPLANE:
+		case EOpType::ME_CLIPMORPHPLANE:
 		{
 			Ptr<ASTOpMeshClipMorphPlane> newOp = mu::Clone<ASTOpMeshClipMorphPlane>(at);
 			newOp->source = Visit(newOp->source.child(), currentSinkOp);
@@ -180,7 +180,7 @@ namespace mu
 			break;
 		}
 
-		case OP_TYPE::ME_MORPH:
+		case EOpType::ME_MORPH:
 		{
 			Ptr<ASTOpMeshMorph> NewOp = mu::Clone<ASTOpMeshMorph>(at);
 			NewOp->Base = Visit(NewOp->Base.child(), currentSinkOp);
@@ -189,7 +189,7 @@ namespace mu
 			break;
 		}
 
-		case OP_TYPE::ME_MERGE:
+		case EOpType::ME_MERGE:
 		{
 			Ptr<ASTOpFixed> newOp = mu::Clone<ASTOpFixed>(at);
 			newOp->SetChild(newOp->op.args.MeshMerge.base, Visit(newOp->children[newOp->op.args.MeshMerge.base].child(), currentSinkOp));
@@ -198,7 +198,7 @@ namespace mu
 			break;
 		}
 
-		case OP_TYPE::ME_APPLYPOSE:
+		case EOpType::ME_APPLYPOSE:
 		{
 			Ptr<ASTOpMeshApplyPose> NewOp = mu::Clone<ASTOpMeshApplyPose>(at);
 			NewOp->base = Visit(NewOp->base.child(), currentSinkOp);
@@ -206,7 +206,7 @@ namespace mu
 			break;
 		}
 
-		case OP_TYPE::ME_INTERPOLATE:
+		case EOpType::ME_INTERPOLATE:
 		{
 			Ptr<ASTOpFixed> newOp = mu::Clone<ASTOpFixed>(at);
 			newOp->SetChild(newOp->op.args.MeshInterpolate.base, Visit(newOp->children[newOp->op.args.MeshInterpolate.base].child(), currentSinkOp));
@@ -223,7 +223,7 @@ namespace mu
 			break;
 		}
 
-		case OP_TYPE::ME_REMOVEMASK:
+		case EOpType::ME_REMOVEMASK:
 		{
 			// TODO: Make mask smaller?
 			Ptr<ASTOpMeshRemoveMask> newOp = mu::Clone<ASTOpMeshRemoveMask>(at);
@@ -232,7 +232,7 @@ namespace mu
 			break;
 		}
 
-		case OP_TYPE::ME_CONDITIONAL:
+		case EOpType::ME_CONDITIONAL:
 		{
 			Ptr<ASTOpConditional> newOp = mu::Clone<ASTOpConditional>(at);
 			newOp->yes = Visit(newOp->yes.child(), currentSinkOp);
@@ -241,7 +241,7 @@ namespace mu
 			break;
 		}
 
-		case OP_TYPE::ME_SWITCH:
+		case EOpType::ME_SWITCH:
 		{
 			Ptr<ASTOpSwitch> newOp = mu::Clone<ASTOpSwitch>(at);
 			newOp->def = Visit(newOp->def.child(), currentSinkOp);
@@ -255,7 +255,7 @@ namespace mu
 
 		// If we reach here it means the operation type has not bee optimized.
 		default:
-			if (at != m_initialSource)
+			if (at != InitialSource)
 			{
 				mu::Ptr<ASTOpMeshExtractLayoutBlocks> newOp = mu::Clone<ASTOpMeshExtractLayoutBlocks>(currentSinkOp);
 				newOp->Source = at;
diff --git a/Engine/Plugins/Mutable/Source/MutableTools/Private/MuT/ASTOpMeshExtractLayoutBlocks.h b/Engine/Plugins/Mutable/Source/MutableTools/Private/MuT/ASTOpMeshExtractLayoutBlocks.h
--- a/Engine/Plugins/Mutable/Source/MutableTools/Private/MuT/ASTOpMeshExtractLayoutBlocks.h
+++ b/Engine/Plugins/Mutable/Source/MutableTools/Private/MuT/ASTOpMeshExtractLayoutBlocks.h
@@ -36,7 +36,7 @@ namespace mu
 		~ASTOpMeshExtractLayoutBlocks() override;
 
 		// ASTOp interface
-		OP_TYPE GetOpType() const override { return OP_TYPE::ME_EXTRACTLAYOUTBLOCK; }
+		EOpType GetOpType() const override { return EOpType::ME_EXTRACTLAYOUTBLOCK; }
 		uint64 Hash() const override;
 		void ForEachChild(const TFunctionRef<void(ASTChild&)>) override;
 		bool IsEqual(const ASTOp& otherUntyped) const override;
diff --git a/Engine/Plugins/Mutable/Source/MutableTools/Private/MuT/ASTOpMeshFormat.cpp b/Engine/Plugins/Mutable/Source/MutableTools/Private/MuT/ASTOpMeshFormat.cpp
--- a/Engine/Plugins/Mutable/Source/MutableTools/Private/MuT/ASTOpMeshFormat.cpp
+++ b/Engine/Plugins/Mutable/Source/MutableTools/Private/MuT/ASTOpMeshFormat.cpp
@@ -96,10 +96,10 @@ void ASTOpMeshFormat::Link( FProgram& program, FLinkerOptions* )
 		if (Source) args.source = Source->linkedAddress;
 		if (Format) args.format = Format->linkedAddress;
 
-        linkedAddress = (OP::ADDRESS)program.m_opAddress.Num();
-        program.m_opAddress.Add((uint32)program.m_byteCode.Num());
-        AppendCode(program.m_byteCode,OP_TYPE::ME_FORMAT);
-        AppendCode(program.m_byteCode,args);
+        linkedAddress = (OP::ADDRESS)program.OpAddress.Num();
+        program.OpAddress.Add((uint32)program.ByteCode.Num());
+        AppendCode(program.ByteCode,EOpType::ME_FORMAT);
+        AppendCode(program.ByteCode,args);
     }
 
 }
@@ -128,15 +128,15 @@ FSourceDataDescriptor ASTOpMeshFormat::GetSourceDataDescriptor(FGetSourceDataDes
 //---------------------------------------------------------------------------------------------
 mu::Ptr<ASTOp> Sink_MeshFormatAST::Apply(const ASTOpMeshFormat* root)
 {
-	m_root = root;
+	Root = root;
 
 	OldToNew.Reset();
 
-	m_initialSource = m_root->Source.child();
-	mu::Ptr<ASTOp> newSource = Visit(m_initialSource, m_root);
+	InitialSource = Root->Source.child();
+	mu::Ptr<ASTOp> newSource = Visit(InitialSource, Root);
 
 	// If there is any change, it is the new root.
-	if (newSource != m_initialSource)
+	if (newSource != InitialSource)
 	{
 		return newSource;
 	}
@@ -147,16 +147,16 @@ mu::Ptr<ASTOp> Sink_MeshFormatAST::Apply(const ASTOpMeshFormat* root)
 
 namespace
 {
-	mu::Ptr<const Mesh> FindBaseMeshConstant(mu::Ptr<ASTOp> at)
+	TSharedPtr<const FMesh> FindBaseMeshConstant(mu::Ptr<ASTOp> at)
 	{
-		mu::Ptr<const Mesh> res;
+		TSharedPtr<const FMesh> res;
 
 		switch (at->GetOpType())
 		{
-		case OP_TYPE::ME_CONSTANT:
+		case EOpType::ME_CONSTANT:
 		{
 			const ASTOpConstantResource* typed = static_cast<const ASTOpConstantResource*>(at.get());
-			res = static_cast<const Mesh*>(typed->GetValue().get());
+			res = StaticCastSharedPtr<const FMesh>(typed->GetValue());
 			break;
 		}
 
@@ -171,7 +171,7 @@ namespace
 
 
 	// Make a mesh format suitable to morph a particular other format.
-	Ptr<Mesh> MakeMorphTargetFormat(Ptr<const Mesh> pTargetFormat)
+	TSharedPtr<FMesh> MakeMorphTargetFormat(TSharedPtr<const FMesh> pTargetFormat)
 	{
 		MUTABLE_CPUPROFILER_SCOPE(MakeMorphTargetFormat);
 
@@ -211,7 +211,7 @@ namespace
 		}
 
 
-		MeshPtr pTargetMorphFormat = new Mesh;
+		TSharedPtr<FMesh> pTargetMorphFormat = MakeShared<FMesh>();
 		pTargetMorphFormat->GetVertexBuffers().SetBufferCount(1);
 
 		pTargetMorphFormat->GetVertexBuffers().SetBuffer(0, offset,
@@ -225,7 +225,7 @@ namespace
 		return pTargetMorphFormat;
 	}
 
-    Ptr<const Mesh> EnsureFormatHasSkinningBuffers(Ptr<const Mesh>& FormatMesh)
+    TSharedPtr<const FMesh> EnsureFormatHasSkinningBuffers(TSharedPtr<const FMesh>& FormatMesh)
     {
         const FMeshBufferSet& FormatMeshVertexBuffers = FormatMesh->GetVertexBuffers();
     
@@ -242,7 +242,7 @@ namespace
             return FormatMesh;
         }
 
-        Ptr<Mesh> NewMesh = FormatMesh->Clone();
+		TSharedPtr<FMesh> NewMesh = FormatMesh->Clone();
         FMeshBufferSet& MeshBuffers = NewMesh->GetVertexBuffers();
         
         FMeshBuffer& Buffer = MeshBuffers.Buffers.AddDefaulted_GetRef();
@@ -287,7 +287,7 @@ mu::Ptr<ASTOp> Sink_MeshFormatAST::Visit(const mu::Ptr<ASTOp>& at, const ASTOpMe
 	switch (at->GetOpType())
 	{
 
-	case OP_TYPE::ME_APPLYLAYOUT:
+	case EOpType::ME_APPLYLAYOUT:
 	{
 		auto newOp = mu::Clone<ASTOpFixed>(at);
 		newOp->SetChild(newOp->op.args.MeshApplyLayout.mesh, Visit(newOp->children[newOp->op.args.MeshApplyLayout.mesh].child(), currentFormatOp));
@@ -295,7 +295,7 @@ mu::Ptr<ASTOp> Sink_MeshFormatAST::Visit(const mu::Ptr<ASTOp>& at, const ASTOpMe
 		break;
 	}
 
-	case OP_TYPE::ME_SETSKELETON:
+	case EOpType::ME_SETSKELETON:
 	{
 		auto newOp = mu::Clone<ASTOpFixed>(at);
 		newOp->SetChild(newOp->op.args.MeshSetSkeleton.source, Visit(newOp->children[newOp->op.args.MeshSetSkeleton.source].child(), currentFormatOp));
@@ -303,7 +303,7 @@ mu::Ptr<ASTOp> Sink_MeshFormatAST::Visit(const mu::Ptr<ASTOp>& at, const ASTOpMe
 		break;
 	}
 
-	case OP_TYPE::ME_ADDTAGS:
+	case EOpType::ME_ADDTAGS:
 	{
 		Ptr<ASTOpMeshAddTags> newOp = mu::Clone<ASTOpMeshAddTags>(at);
 		newOp->Source = Visit(newOp->Source.child(), currentFormatOp);
@@ -311,7 +311,7 @@ mu::Ptr<ASTOp> Sink_MeshFormatAST::Visit(const mu::Ptr<ASTOp>& at, const ASTOpMe
 		break;
 	}
 
-	case OP_TYPE::ME_CLIPMORPHPLANE:
+	case EOpType::ME_CLIPMORPHPLANE:
 	{
 		Ptr<ASTOpMeshClipMorphPlane> newOp = mu::Clone<ASTOpMeshClipMorphPlane>(at);
 		newOp->source = Visit(newOp->source.child(), currentFormatOp);
@@ -319,7 +319,7 @@ mu::Ptr<ASTOp> Sink_MeshFormatAST::Visit(const mu::Ptr<ASTOp>& at, const ASTOpMe
 		break;
 	}
 
-	case OP_TYPE::ME_TRANSFORMWITHMESH:
+	case EOpType::ME_TRANSFORMWITHMESH:
 	{
 		Ptr<ASTOpMeshTransformWithBoundingMesh> NewOp = mu::Clone<ASTOpMeshTransformWithBoundingMesh>(at);
 		NewOp->source = Visit(NewOp->source.child(), currentFormatOp);
@@ -331,18 +331,18 @@ mu::Ptr<ASTOp> Sink_MeshFormatAST::Visit(const mu::Ptr<ASTOp>& at, const ASTOpMe
 		break;
 	}
 
-	case OP_TYPE::ME_MORPH:
+	case EOpType::ME_MORPH:
 	{
 		// Move the format down the base of the morph
 		Ptr<ASTOpMeshMorph> NewOp = mu::Clone<ASTOpMeshMorph>(at);
 		NewOp->Base = Visit(NewOp->Base.child(), currentFormatOp);
 
 		// Reformat the morph targets to match the new format.
-		Ptr<const Mesh> pTargetFormat = FindBaseMeshConstant(currentFormatOp->Format.child());
-		Ptr<const Mesh> pTargetMorphFormat = MakeMorphTargetFormat(pTargetFormat);
+		TSharedPtr<const FMesh> pTargetFormat = FindBaseMeshConstant(currentFormatOp->Format.child());
+		TSharedPtr<const FMesh> pTargetMorphFormat = MakeMorphTargetFormat(pTargetFormat);
 
 		mu::Ptr<ASTOpConstantResource> NewFormatConstant = new ASTOpConstantResource();
-		NewFormatConstant->Type = OP_TYPE::ME_CONSTANT;
+		NewFormatConstant->Type = EOpType::ME_CONSTANT;
 		NewFormatConstant->SetValue(pTargetMorphFormat, nullptr);
 		NewFormatConstant->SourceDataDescriptor = at->GetSourceDataDescriptor();
 
@@ -359,7 +359,7 @@ mu::Ptr<ASTOp> Sink_MeshFormatAST::Visit(const mu::Ptr<ASTOp>& at, const ASTOpMe
 		break;
 	}
 
-	case OP_TYPE::ME_MERGE:
+	case EOpType::ME_MERGE:
 	{
 		Ptr<ASTOpFixed> newOp = mu::Clone<ASTOpFixed>(at);
 		newOp->SetChild(newOp->op.args.MeshMerge.base, Visit(newOp->children[newOp->op.args.MeshMerge.base].child(), currentFormatOp));
@@ -368,22 +368,22 @@ mu::Ptr<ASTOp> Sink_MeshFormatAST::Visit(const mu::Ptr<ASTOp>& at, const ASTOpMe
 		break;
 	}
 
-	case OP_TYPE::ME_APPLYPOSE:
+	case EOpType::ME_APPLYPOSE:
 	{
-		Ptr<const Mesh> TargetFormat = FindBaseMeshConstant(currentFormatOp->Format.child());
+		TSharedPtr<const FMesh> TargetFormat = FindBaseMeshConstant(currentFormatOp->Format.child());
         TargetFormat = EnsureFormatHasSkinningBuffers(TargetFormat);
 
 		Ptr<ASTOpMeshApplyPose> NewOp = mu::Clone<ASTOpMeshApplyPose>(at);
         mu::Ptr<ASTOpMeshFormat> NewFormat = mu::Clone<ASTOpMeshFormat>(currentFormatOp);
         
 		mu::Ptr<ASTOpConstantResource> NewFormatConstant = new ASTOpConstantResource();
-		NewFormatConstant->Type = OP_TYPE::ME_CONSTANT; 
+		NewFormatConstant->Type = EOpType::ME_CONSTANT;
 		NewFormatConstant->SetValue(TargetFormat, nullptr);
 		NewFormatConstant->SourceDataDescriptor = at->GetSourceDataDescriptor();
 	
         NewFormat->Flags = NewFormat->Flags | OP::MeshFormatArgs::OptimizeBuffers;
 
-        //TODO: Optimize, in case no skinning data is found in the format mesh a generic buffer that can represent
+        // TODO: Optimize, in case no skinning data is found in the format mesh a generic buffer that can represent
         // all possible skinning formats is added. This is not optimal, we may want to add a flag to the format op
         // to indicate it should copy the skinning from the base mesh.
         NewFormat->Format = NewFormatConstant;
@@ -394,18 +394,18 @@ mu::Ptr<ASTOp> Sink_MeshFormatAST::Visit(const mu::Ptr<ASTOp>& at, const ASTOpMe
 		break;
 	}
 
-	case OP_TYPE::ME_INTERPOLATE:
+	case EOpType::ME_INTERPOLATE:
 	{
 		// Move the format down the base of the morph
 		Ptr<ASTOpFixed> newOp = mu::Clone<ASTOpFixed>(at);
 		newOp->SetChild(newOp->op.args.MeshInterpolate.base, Visit(newOp->children[newOp->op.args.MeshInterpolate.base].child(), currentFormatOp));
 
 		// Reformat the morph targets to match the new format.
-		Ptr<const Mesh> pTargetFormat = FindBaseMeshConstant(currentFormatOp->Format.child());
-		Ptr<const Mesh> pTargetMorphFormat = MakeMorphTargetFormat(pTargetFormat);
+		TSharedPtr<const FMesh> pTargetFormat = FindBaseMeshConstant(currentFormatOp->Format.child());
+		TSharedPtr<const FMesh> pTargetMorphFormat = MakeMorphTargetFormat(pTargetFormat);
 
 		mu::Ptr<ASTOpConstantResource> TargetMorphFormatOp = new ASTOpConstantResource();
-		TargetMorphFormatOp->Type = OP_TYPE::ME_CONSTANT;
+		TargetMorphFormatOp->Type = EOpType::ME_CONSTANT;
 		TargetMorphFormatOp->SetValue(pTargetMorphFormat, nullptr );
 		TargetMorphFormatOp->SourceDataDescriptor = at->GetSourceDataDescriptor();
 
@@ -425,7 +425,7 @@ mu::Ptr<ASTOp> Sink_MeshFormatAST::Visit(const mu::Ptr<ASTOp>& at, const ASTOpMe
 		break;
 	}
 
-	case OP_TYPE::ME_REMOVEMASK:
+	case EOpType::ME_REMOVEMASK:
 	{
 		auto newOp = mu::Clone<ASTOpMeshRemoveMask>(at);
 		newOp->source = Visit(newOp->source.child(), currentFormatOp);
@@ -433,7 +433,7 @@ mu::Ptr<ASTOp> Sink_MeshFormatAST::Visit(const mu::Ptr<ASTOp>& at, const ASTOpMe
 		break;
 	}
 
-	case OP_TYPE::ME_CONDITIONAL:
+	case EOpType::ME_CONDITIONAL:
 	{
 		auto newOp = mu::Clone<ASTOpConditional>(at);
 		newOp->yes = Visit(newOp->yes.child(), currentFormatOp);
@@ -442,7 +442,7 @@ mu::Ptr<ASTOp> Sink_MeshFormatAST::Visit(const mu::Ptr<ASTOp>& at, const ASTOpMe
 		break;
 	}
 
-	case OP_TYPE::ME_SWITCH:
+	case EOpType::ME_SWITCH:
 	{
 		auto newOp = mu::Clone<ASTOpSwitch>(at);
 		newOp->def = Visit(newOp->def.child(), currentFormatOp);
@@ -454,18 +454,18 @@ mu::Ptr<ASTOp> Sink_MeshFormatAST::Visit(const mu::Ptr<ASTOp>& at, const ASTOpMe
 		break;
 	}
 
-	case OP_TYPE::ME_FORMAT:
+	case EOpType::ME_FORMAT:
 		// TODO: The child format can be removed. 
 		// Unless channels are removed and re-added, which would change their content?
 		break;
 
 
 	// This operation should not be optimized.
-	case OP_TYPE::ME_DIFFERENCE:
+	case EOpType::ME_DIFFERENCE:
 
 	// If we reach here it means the operation type has not bee optimized.
 	default:
-		if (at != m_initialSource)
+		if (at != InitialSource)
 		{
 			mu::Ptr<ASTOpMeshFormat> newOp = mu::Clone<ASTOpMeshFormat>(currentFormatOp);
 			newOp->Source = at;
diff --git a/Engine/Plugins/Mutable/Source/MutableTools/Private/MuT/ASTOpMeshFormat.h b/Engine/Plugins/Mutable/Source/MutableTools/Private/MuT/ASTOpMeshFormat.h
--- a/Engine/Plugins/Mutable/Source/MutableTools/Private/MuT/ASTOpMeshFormat.h
+++ b/Engine/Plugins/Mutable/Source/MutableTools/Private/MuT/ASTOpMeshFormat.h
@@ -27,7 +27,7 @@ namespace mu
 		ASTOpMeshFormat(const ASTOpMeshFormat&) = delete;
 		~ASTOpMeshFormat();
 
-		OP_TYPE GetOpType() const override { return OP_TYPE::ME_FORMAT; }
+		EOpType GetOpType() const override { return EOpType::ME_FORMAT; }
 		uint64 Hash() const override;
 		bool IsEqual(const ASTOp& otherUntyped) const override;
 		Ptr<ASTOp> Clone(MapChildFuncRef mapChild) const override;
diff --git a/Engine/Plugins/Mutable/Source/MutableTools/Private/MuT/ASTOpMeshGeometryOperation.cpp b/Engine/Plugins/Mutable/Source/MutableTools/Private/MuT/ASTOpMeshGeometryOperation.cpp
--- a/Engine/Plugins/Mutable/Source/MutableTools/Private/MuT/ASTOpMeshGeometryOperation.cpp
+++ b/Engine/Plugins/Mutable/Source/MutableTools/Private/MuT/ASTOpMeshGeometryOperation.cpp
@@ -81,11 +81,11 @@ namespace mu
 			if (scalarA) args.scalarA = scalarA->linkedAddress;
 			if (scalarB) args.scalarB = scalarB->linkedAddress;
 
-			linkedAddress = (OP::ADDRESS)program.m_opAddress.Num();
+			linkedAddress = (OP::ADDRESS)program.OpAddress.Num();
 			//program.m_code.push_back(op);
-			program.m_opAddress.Add((uint32_t)program.m_byteCode.Num());
-			AppendCode(program.m_byteCode, OP_TYPE::ME_GEOMETRYOPERATION);
-			AppendCode(program.m_byteCode, args);
+			program.OpAddress.Add((uint32_t)program.ByteCode.Num());
+			AppendCode(program.ByteCode, EOpType::ME_GEOMETRYOPERATION);
+			AppendCode(program.ByteCode, args);
 		}
 
 	}
diff --git a/Engine/Plugins/Mutable/Source/MutableTools/Private/MuT/ASTOpMeshGeometryOperation.h b/Engine/Plugins/Mutable/Source/MutableTools/Private/MuT/ASTOpMeshGeometryOperation.h
--- a/Engine/Plugins/Mutable/Source/MutableTools/Private/MuT/ASTOpMeshGeometryOperation.h
+++ b/Engine/Plugins/Mutable/Source/MutableTools/Private/MuT/ASTOpMeshGeometryOperation.h
@@ -31,7 +31,7 @@ namespace mu
 		ASTOpMeshGeometryOperation(const ASTOpMeshGeometryOperation&) = delete;
 		~ASTOpMeshGeometryOperation();
 
-		OP_TYPE GetOpType() const override { return OP_TYPE::ME_GEOMETRYOPERATION; }
+		EOpType GetOpType() const override { return EOpType::ME_GEOMETRYOPERATION; }
 		uint64 Hash() const override;
 		bool IsEqual(const ASTOp& otherUntyped) const override;
 		Ptr<ASTOp> Clone(MapChildFuncRef mapChild) const override;
diff --git a/Engine/Plugins/Mutable/Source/MutableTools/Private/MuT/ASTOpMeshMaskClipMesh.cpp b/Engine/Plugins/Mutable/Source/MutableTools/Private/MuT/ASTOpMeshMaskClipMesh.cpp
--- a/Engine/Plugins/Mutable/Source/MutableTools/Private/MuT/ASTOpMeshMaskClipMesh.cpp
+++ b/Engine/Plugins/Mutable/Source/MutableTools/Private/MuT/ASTOpMeshMaskClipMesh.cpp
@@ -85,11 +85,11 @@ namespace mu
 			if (source) args.source = source->linkedAddress;
 			if (clip) args.clip = clip->linkedAddress;
 
-			linkedAddress = (OP::ADDRESS)program.m_opAddress.Num();
+			linkedAddress = (OP::ADDRESS)program.OpAddress.Num();
 			//program.m_code.push_back(op);
-			program.m_opAddress.Add((uint32_t)program.m_byteCode.Num());
-			AppendCode(program.m_byteCode, OP_TYPE::ME_MASKCLIPMESH);
-			AppendCode(program.m_byteCode, args);
+			program.OpAddress.Add((uint32_t)program.ByteCode.Num());
+			AppendCode(program.ByteCode, EOpType::ME_MASKCLIPMESH);
+			AppendCode(program.ByteCode, args);
 		}
 	}
 
@@ -107,14 +107,14 @@ namespace mu
 
 			mu::Ptr<ASTOp> Apply(const ASTOpMeshMaskClipMesh* root)
 			{
-				m_root = root;
+				Root = root;
 				m_oldToNew.Empty();
 
-				m_initialSource = root->source.child();
-				mu::Ptr<ASTOp> NewSource = Visit(m_initialSource);
+				InitialSource = root->source.child();
+				mu::Ptr<ASTOp> NewSource = Visit(InitialSource);
 
 				// If there is any change, it is the new root.
-				if (NewSource != m_initialSource)
+				if (NewSource != InitialSource)
 				{
 					return NewSource;
 				}
@@ -124,8 +124,8 @@ namespace mu
 
 		protected:
 
-			const ASTOpMeshMaskClipMesh* m_root;
-			mu::Ptr<ASTOp> m_initialSource;
+			const ASTOpMeshMaskClipMesh* Root;
+			mu::Ptr<ASTOp> InitialSource;
 			TMap<mu::Ptr<ASTOp>, mu::Ptr<ASTOp>> m_oldToNew;
 			TArray<mu::Ptr<ASTOp>> m_newOps;
 
@@ -153,12 +153,12 @@ namespace mu
 				// This cannot be sunk since the result is different. Since the clipping is now correctly
 				// generated at the end of the chain when really necessary, this wrong optimisation is no 
 				// longer needed.
-				//case OP_TYPE::ME_MORPH:
+				//case EOpType::ME_MORPH:
 		        //{
 		        //    break;
 		        //}
 
-				case OP_TYPE::ME_REMOVEMASK:
+				case EOpType::ME_REMOVEMASK:
 				{
 					// Remove this op:
 					// This may lead to the mask being bigger than needed since it will include
@@ -170,7 +170,7 @@ namespace mu
 					break;
 				}
 
-				case OP_TYPE::ME_ADDTAGS:
+				case EOpType::ME_ADDTAGS:
 				{
 					Ptr<ASTOpMeshAddTags> newOp = mu::Clone<ASTOpMeshAddTags>(at);
 					newOp->Source = Visit(newOp->Source.child());
@@ -178,7 +178,7 @@ namespace mu
 					break;
 				}
 
-				case OP_TYPE::ME_CONDITIONAL:
+				case EOpType::ME_CONDITIONAL:
 				{
 					// We move the mask creation down the two paths
 					// It always needs to be a clone because otherwise we could be modifying an
@@ -191,7 +191,7 @@ namespace mu
 					break;
 				}
 
-				case OP_TYPE::ME_SWITCH:
+				case EOpType::ME_SWITCH:
 				{
 					// We move the mask creation down all the paths
 					Ptr<ASTOpSwitch> newOp = mu::Clone<ASTOpSwitch>(at);
@@ -207,7 +207,7 @@ namespace mu
 				// This cannot be sunk since the result is different. Since the clipping is now correctly
 				// generated at the end of the chain when really necessary, this wrong optimisation is no 
 				// longer needed.
-				//case OP_TYPE::ME_CLIPMORPHPLANE:
+				//case EOpType::ME_CLIPMORPHPLANE:
 	   //         {
 	   //             // We move the mask creation down the source
 	   //             auto typedAt = dynamic_cast<const ASTOpMeshClipMorphPlane*>(at.get());
@@ -218,9 +218,9 @@ namespace mu
 				default:
 				{
 					//
-					if (at != m_initialSource)
+					if (at != InitialSource)
 					{
-						Ptr<ASTOpMeshMaskClipMesh> newOp = mu::Clone<ASTOpMeshMaskClipMesh>(m_root);
+						Ptr<ASTOpMeshMaskClipMesh> newOp = mu::Clone<ASTOpMeshMaskClipMesh>(Root);
 						newOp->source = at;
 						newAt = newOp;
 					}
@@ -244,7 +244,7 @@ namespace mu
 
 			mu::Ptr<ASTOp> Apply(const ASTOpMeshMaskClipMesh* root)
 			{
-				m_root = root;
+				Root = root;
 				m_oldToNew.Empty();
 
 				m_initialClip = root->clip.child();
@@ -261,7 +261,7 @@ namespace mu
 
 		protected:
 
-			const ASTOpMeshMaskClipMesh* m_root;
+			const ASTOpMeshMaskClipMesh* Root;
 			mu::Ptr<ASTOp> m_initialClip;
 			TMap<mu::Ptr<ASTOp>, mu::Ptr<ASTOp>> m_oldToNew;
 			TArray<mu::Ptr<ASTOp>> m_newOps;
@@ -287,7 +287,7 @@ namespace mu
 				switch (at->GetOpType())
 				{
 
-				case OP_TYPE::ME_CONDITIONAL:
+				case EOpType::ME_CONDITIONAL:
 				{
 					// We move the mask creation down the two paths
 					// It always needs to be a clone because otherwise we could be modifying an
@@ -300,7 +300,7 @@ namespace mu
 					break;
 				}
 
-				case OP_TYPE::ME_SWITCH:
+				case EOpType::ME_SWITCH:
 				{
 					// We move the mask creation down all the paths
 					Ptr<ASTOpSwitch> newOp = mu::Clone<ASTOpSwitch>(at);
@@ -313,7 +313,7 @@ namespace mu
 					break;
 				}
 
-				case OP_TYPE::ME_ADDTAGS:
+				case EOpType::ME_ADDTAGS:
 				{
 					// Ignore tags in the clip mesh branch
 					const ASTOpMeshAddTags* AddOp = static_cast<const ASTOpMeshAddTags*>(at.get());
@@ -326,7 +326,7 @@ namespace mu
 					//
 					if (at != m_initialClip)
 					{
-						Ptr<ASTOpMeshMaskClipMesh> newOp = mu::Clone<ASTOpMeshMaskClipMesh>(m_root);
+						Ptr<ASTOpMeshMaskClipMesh> newOp = mu::Clone<ASTOpMeshMaskClipMesh>(Root);
 						newOp->clip = at;
 						newAt = newOp;
 					}
diff --git a/Engine/Plugins/Mutable/Source/MutableTools/Private/MuT/ASTOpMeshMaskClipMesh.h b/Engine/Plugins/Mutable/Source/MutableTools/Private/MuT/ASTOpMeshMaskClipMesh.h
--- a/Engine/Plugins/Mutable/Source/MutableTools/Private/MuT/ASTOpMeshMaskClipMesh.h
+++ b/Engine/Plugins/Mutable/Source/MutableTools/Private/MuT/ASTOpMeshMaskClipMesh.h
@@ -29,7 +29,7 @@ struct FProgram;
 		ASTOpMeshMaskClipMesh(const ASTOpMeshMaskClipMesh&) = delete;
 		~ASTOpMeshMaskClipMesh();
 
-		OP_TYPE GetOpType() const override { return OP_TYPE::ME_MASKCLIPMESH; }
+		EOpType GetOpType() const override { return EOpType::ME_MASKCLIPMESH; }
 		uint64 Hash() const override;
 		bool IsEqual(const ASTOp& otherUntyped) const override;
 		Ptr<ASTOp> Clone(MapChildFuncRef mapChild) const override;
