[mutable] More enforcement of coding standards.

diff --git a/Engine/Plugins/Mutable/Source/MutableTools/Private/MuT/ASTOpMeshMaskClipUVMask.cpp b/Engine/Plugins/Mutable/Source/MutableTools/Private/MuT/ASTOpMeshMaskClipUVMask.cpp
--- a/Engine/Plugins/Mutable/Source/MutableTools/Private/MuT/ASTOpMeshMaskClipUVMask.cpp
+++ b/Engine/Plugins/Mutable/Source/MutableTools/Private/MuT/ASTOpMeshMaskClipUVMask.cpp
@@ -101,10 +101,10 @@ namespace mu
 			if (MaskLayout) args.MaskLayout = MaskLayout->linkedAddress;
 			args.LayoutIndex = LayoutIndex;
 
-			linkedAddress = (OP::ADDRESS)program.m_opAddress.Num();
-			program.m_opAddress.Add((uint32)program.m_byteCode.Num());
-			AppendCode(program.m_byteCode, OP_TYPE::ME_MASKCLIPUVMASK);
-			AppendCode(program.m_byteCode, args);
+			linkedAddress = (OP::ADDRESS)program.OpAddress.Num();
+			program.OpAddress.Add((uint32)program.ByteCode.Num());
+			AppendCode(program.ByteCode, EOpType::ME_MASKCLIPUVMASK);
+			AppendCode(program.ByteCode, args);
 		}
 	}
 
@@ -122,14 +122,14 @@ namespace mu
 
 			mu::Ptr<ASTOp> Apply(const ASTOpMeshMaskClipUVMask* root)
 			{
-				m_root = root;
+				Root = root;
 				m_oldToNew.Empty();
 
-				m_initialSource = root->Source.child();
-				mu::Ptr<ASTOp> NewSource = Visit(m_initialSource);
+				InitialSource = root->Source.child();
+				mu::Ptr<ASTOp> NewSource = Visit(InitialSource);
 
 				// If there is any change, it is the new root.
-				if (NewSource != m_initialSource)
+				if (NewSource != InitialSource)
 				{
 					return NewSource;
 				}
@@ -139,8 +139,8 @@ namespace mu
 
 		protected:
 
-			const ASTOpMeshMaskClipUVMask* m_root;
-			mu::Ptr<ASTOp> m_initialSource;
+			const ASTOpMeshMaskClipUVMask* Root;
+			mu::Ptr<ASTOp> InitialSource;
 			TMap<mu::Ptr<ASTOp>, mu::Ptr<ASTOp>> m_oldToNew;
 			TArray<mu::Ptr<ASTOp>> m_newOps;
 
@@ -168,12 +168,12 @@ namespace mu
 				// This cannot be sunk since the result is different. Since the clipping is now correctly
 				// generated at the end of the chain when really necessary, this wrong optimisation is no 
 				// longer needed.
-				//case OP_TYPE::ME_MORPH:
+				//case EOpType::ME_MORPH:
 		        //{
 		        //    break;
 		        //}
 
-				case OP_TYPE::ME_REMOVEMASK:
+				case EOpType::ME_REMOVEMASK:
 				{
 					// Remove this op:
 					// This may lead to the mask being bigger than needed since it will include
@@ -185,7 +185,7 @@ namespace mu
 					break;
 				}
 
-				case OP_TYPE::ME_ADDTAGS:
+				case EOpType::ME_ADDTAGS:
 				{
 					Ptr<ASTOpMeshAddTags> newOp = mu::Clone<ASTOpMeshAddTags>(at);
 					newOp->Source = Visit(newOp->Source.child());
@@ -193,7 +193,7 @@ namespace mu
 					break;
 				}
 
-				case OP_TYPE::ME_CONDITIONAL:
+				case EOpType::ME_CONDITIONAL:
 				{
 					// We move the mask creation down the two paths
 					// It always needs to be a clone because otherwise we could be modifying an
@@ -206,7 +206,7 @@ namespace mu
 					break;
 				}
 
-				case OP_TYPE::ME_SWITCH:
+				case EOpType::ME_SWITCH:
 				{
 					// We move the mask creation down all the paths
 					Ptr<ASTOpSwitch> newOp = mu::Clone<ASTOpSwitch>(at);
@@ -222,7 +222,7 @@ namespace mu
 				// This cannot be sunk since the result is different. Since the clipping is now correctly
 				// generated at the end of the chain when really necessary, this wrong optimisation is no 
 				// longer needed.
-				//case OP_TYPE::ME_CLIPMORPHPLANE:
+				//case EOpType::ME_CLIPMORPHPLANE:
 	   //         {
 	   //             // We move the mask creation down the Source
 	   //             auto typedAt = dynamic_cast<const ASTOpMeshClipMorphPlane*>(at.get());
@@ -233,9 +233,9 @@ namespace mu
 				default:
 				{
 					//
-					if (at != m_initialSource)
+					if (at != InitialSource)
 					{
-						Ptr<ASTOpMeshMaskClipUVMask> newOp = mu::Clone<ASTOpMeshMaskClipUVMask>(m_root);
+						Ptr<ASTOpMeshMaskClipUVMask> newOp = mu::Clone<ASTOpMeshMaskClipUVMask>(Root);
 						newOp->Source = at;
 						newAt = newOp;
 					}
@@ -259,7 +259,7 @@ namespace mu
 
 			mu::Ptr<ASTOp> Apply(const ASTOpMeshMaskClipUVMask* root)
 			{
-				m_root = root;
+				Root = root;
 				m_oldToNew.Empty();
 
 				m_initialClip = root->MaskImage.child();
@@ -276,7 +276,7 @@ namespace mu
 
 		protected:
 
-			const ASTOpMeshMaskClipUVMask* m_root;
+			const ASTOpMeshMaskClipUVMask* Root;
 			mu::Ptr<ASTOp> m_initialClip;
 			TMap<mu::Ptr<ASTOp>, mu::Ptr<ASTOp>> m_oldToNew;
 			TArray<mu::Ptr<ASTOp>> m_newOps;
@@ -302,7 +302,7 @@ namespace mu
 				switch (at->GetOpType())
 				{
 
-				case OP_TYPE::ME_CONDITIONAL:
+				case EOpType::ME_CONDITIONAL:
 				{
 					// We move the mask creation down the two paths
 					// It always needs to be a clone because otherwise we could be modifying an
@@ -315,7 +315,7 @@ namespace mu
 					break;
 				}
 
-				case OP_TYPE::ME_SWITCH:
+				case EOpType::ME_SWITCH:
 				{
 					// We move the mask creation down all the paths
 					Ptr<ASTOpSwitch> newOp = mu::Clone<ASTOpSwitch>(at);
@@ -334,7 +334,7 @@ namespace mu
 					//
 					if (at != m_initialClip)
 					{
-						Ptr<ASTOpMeshMaskClipUVMask> newOp = mu::Clone<ASTOpMeshMaskClipUVMask>(m_root);
+						Ptr<ASTOpMeshMaskClipUVMask> newOp = mu::Clone<ASTOpMeshMaskClipUVMask>(Root);
 						newOp->MaskImage = at;
 						newAt = newOp;
 					}
diff --git a/Engine/Plugins/Mutable/Source/MutableTools/Private/MuT/ASTOpMeshMaskClipUVMask.h b/Engine/Plugins/Mutable/Source/MutableTools/Private/MuT/ASTOpMeshMaskClipUVMask.h
--- a/Engine/Plugins/Mutable/Source/MutableTools/Private/MuT/ASTOpMeshMaskClipUVMask.h
+++ b/Engine/Plugins/Mutable/Source/MutableTools/Private/MuT/ASTOpMeshMaskClipUVMask.h
@@ -29,7 +29,7 @@ namespace mu
 		ASTOpMeshMaskClipUVMask(const ASTOpMeshMaskClipUVMask&) = delete;
 		~ASTOpMeshMaskClipUVMask();
 
-		OP_TYPE GetOpType() const override { return OP_TYPE::ME_MASKCLIPUVMASK; }
+		EOpType GetOpType() const override { return EOpType::ME_MASKCLIPUVMASK; }
 		uint64 Hash() const override;
 		bool IsEqual(const ASTOp& otherUntyped) const override;
 		Ptr<ASTOp> Clone(MapChildFuncRef mapChild) const override;
diff --git a/Engine/Plugins/Mutable/Source/MutableTools/Private/MuT/ASTOpMeshMorph.cpp b/Engine/Plugins/Mutable/Source/MutableTools/Private/MuT/ASTOpMeshMorph.cpp
--- a/Engine/Plugins/Mutable/Source/MutableTools/Private/MuT/ASTOpMeshMorph.cpp
+++ b/Engine/Plugins/Mutable/Source/MutableTools/Private/MuT/ASTOpMeshMorph.cpp
@@ -75,19 +75,19 @@ namespace mu
 		// Already linked?
 		if (!linkedAddress)
 		{
-			linkedAddress = (OP::ADDRESS)program.m_opAddress.Num();
+			linkedAddress = (OP::ADDRESS)program.OpAddress.Num();
 
-			program.m_opAddress.Add((uint32_t)program.m_byteCode.Num());
-			AppendCode(program.m_byteCode, OP_TYPE::ME_MORPH);
+			program.OpAddress.Add((uint32_t)program.ByteCode.Num());
+			AppendCode(program.ByteCode, EOpType::ME_MORPH);
 
 			OP::ADDRESS FactorAt = Factor ? Factor->linkedAddress : 0;
-			AppendCode(program.m_byteCode, FactorAt);
+			AppendCode(program.ByteCode, FactorAt);
 
 			OP::ADDRESS BaseAt = Base ? Base->linkedAddress : 0;
-			AppendCode(program.m_byteCode, BaseAt);
+			AppendCode(program.ByteCode, BaseAt);
 
 			OP::ADDRESS TargetAt = Target ? Target->linkedAddress : 0;
-			AppendCode(program.m_byteCode, TargetAt);
+			AppendCode(program.ByteCode, TargetAt);
 		}
 	}
 
@@ -103,11 +103,11 @@ namespace mu
 		}
 
 		// Base optimizations
-		OP_TYPE BaseType = Base.child()->GetOpType();
+		EOpType BaseType = Base.child()->GetOpType();
 		switch (BaseType)
 		{
 
-		case OP_TYPE::ME_ADDTAGS:
+		case EOpType::ME_ADDTAGS:
 		{
 			// Add the base tags after morphing
 			Ptr<ASTOpMeshAddTags> NewAddTags = mu::Clone<ASTOpMeshAddTags>(Base.child());
@@ -132,11 +132,11 @@ namespace mu
 		if (!NewOp && Target)
 		{
 			// Target optimizations
-			OP_TYPE MorphType = Target.child()->GetOpType();
+			EOpType MorphType = Target.child()->GetOpType();
 			switch (MorphType)
 			{
 
-			case OP_TYPE::ME_ADDTAGS:
+			case EOpType::ME_ADDTAGS:
 			{
 				// Ignore the morph target tags
 				const ASTOpMeshAddTags* AddTags = static_cast<const ASTOpMeshAddTags*>(Target.child().get());
diff --git a/Engine/Plugins/Mutable/Source/MutableTools/Private/MuT/ASTOpMeshMorph.h b/Engine/Plugins/Mutable/Source/MutableTools/Private/MuT/ASTOpMeshMorph.h
--- a/Engine/Plugins/Mutable/Source/MutableTools/Private/MuT/ASTOpMeshMorph.h
+++ b/Engine/Plugins/Mutable/Source/MutableTools/Private/MuT/ASTOpMeshMorph.h
@@ -35,7 +35,7 @@ namespace mu
 		~ASTOpMeshMorph() override;
 
 		// ASTOp interface
-		OP_TYPE GetOpType() const override { return OP_TYPE::ME_MORPH; }
+		EOpType GetOpType() const override { return EOpType::ME_MORPH; }
 		uint64 Hash() const override;
 		void ForEachChild(const TFunctionRef<void(ASTChild&)>) override;
 		bool IsEqual(const ASTOp& otherUntyped) const override;
diff --git a/Engine/Plugins/Mutable/Source/MutableTools/Private/MuT/ASTOpMeshMorphReshape.cpp b/Engine/Plugins/Mutable/Source/MutableTools/Private/MuT/ASTOpMeshMorphReshape.cpp
--- a/Engine/Plugins/Mutable/Source/MutableTools/Private/MuT/ASTOpMeshMorphReshape.cpp
+++ b/Engine/Plugins/Mutable/Source/MutableTools/Private/MuT/ASTOpMeshMorphReshape.cpp
@@ -81,10 +81,10 @@ namespace mu
 				Args.Reshape = Reshape->linkedAddress;
 			}
 
-			linkedAddress = (OP::ADDRESS)program.m_opAddress.Num();
-			program.m_opAddress.Add((uint32_t)program.m_byteCode.Num());
-			AppendCode(program.m_byteCode, OP_TYPE::ME_MORPHRESHAPE);
-			AppendCode(program.m_byteCode, Args);
+			linkedAddress = (OP::ADDRESS)program.OpAddress.Num();
+			program.OpAddress.Add((uint32_t)program.ByteCode.Num());
+			AppendCode(program.ByteCode, EOpType::ME_MORPHRESHAPE);
+			AppendCode(program.ByteCode, Args);
 		}
 
 	}
diff --git a/Engine/Plugins/Mutable/Source/MutableTools/Private/MuT/ASTOpMeshMorphReshape.h b/Engine/Plugins/Mutable/Source/MutableTools/Private/MuT/ASTOpMeshMorphReshape.h
--- a/Engine/Plugins/Mutable/Source/MutableTools/Private/MuT/ASTOpMeshMorphReshape.h
+++ b/Engine/Plugins/Mutable/Source/MutableTools/Private/MuT/ASTOpMeshMorphReshape.h
@@ -28,7 +28,7 @@ struct FProgram;
 		ASTOpMeshMorphReshape(const ASTOpMeshMorphReshape&) = delete;
 		~ASTOpMeshMorphReshape();
 
-		OP_TYPE GetOpType() const override { return OP_TYPE::ME_MORPHRESHAPE; }
+		EOpType GetOpType() const override { return EOpType::ME_MORPHRESHAPE; }
 		uint64 Hash() const override;
 		bool IsEqual(const ASTOp& otherUntyped) const override;
 		Ptr<ASTOp> Clone(MapChildFuncRef mapChild) const override;
diff --git a/Engine/Plugins/Mutable/Source/MutableTools/Private/MuT/ASTOpMeshOptimizeSkinning.cpp b/Engine/Plugins/Mutable/Source/MutableTools/Private/MuT/ASTOpMeshOptimizeSkinning.cpp
--- a/Engine/Plugins/Mutable/Source/MutableTools/Private/MuT/ASTOpMeshOptimizeSkinning.cpp
+++ b/Engine/Plugins/Mutable/Source/MutableTools/Private/MuT/ASTOpMeshOptimizeSkinning.cpp
@@ -72,10 +72,10 @@ namespace mu
 				Args.source = Source->linkedAddress;
 			}
 
-			linkedAddress = OP::ADDRESS(Program.m_opAddress.Num());
-			Program.m_opAddress.Add((uint32)Program.m_byteCode.Num());
-			AppendCode(Program.m_byteCode, OP_TYPE::ME_OPTIMIZESKINNING);
-			AppendCode(Program.m_byteCode, Args);
+			linkedAddress = OP::ADDRESS(Program.OpAddress.Num());
+			Program.OpAddress.Add((uint32)Program.ByteCode.Num());
+			AppendCode(Program.ByteCode, EOpType::ME_OPTIMIZESKINNING);
+			AppendCode(Program.ByteCode, Args);
 		}
 	}
 
@@ -142,7 +142,7 @@ namespace mu
 		switch (at->GetOpType())
 		{
 
-		case OP_TYPE::ME_APPLYLAYOUT:
+		case EOpType::ME_APPLYLAYOUT:
 		{
 			Ptr<ASTOpFixed> NewOp = mu::Clone<ASTOpFixed>(at);
 			NewOp->SetChild(NewOp->op.args.MeshApplyLayout.mesh, Visit(NewOp->children[NewOp->op.args.MeshApplyLayout.mesh].child(), CurrentSinkOp));
@@ -150,7 +150,7 @@ namespace mu
 			break;
 		}
 
-		case OP_TYPE::ME_ADDTAGS:
+		case EOpType::ME_ADDTAGS:
 		{
 			Ptr<ASTOpMeshAddTags> NewOp = mu::Clone<ASTOpMeshAddTags>(at);
 			NewOp->Source = Visit(NewOp->Source.child(), CurrentSinkOp);
@@ -158,7 +158,7 @@ namespace mu
 			break;
 		}
 
-		case OP_TYPE::ME_CONDITIONAL:
+		case EOpType::ME_CONDITIONAL:
 		{
 			Ptr<ASTOpConditional> NewOp = mu::Clone<ASTOpConditional>(at);
 			NewOp->yes = Visit(NewOp->yes.child(), CurrentSinkOp);
@@ -167,7 +167,7 @@ namespace mu
 			break;
 		}
 
-		case OP_TYPE::ME_SWITCH:
+		case EOpType::ME_SWITCH:
 		{
 			Ptr<ASTOpSwitch> NewOp = mu::Clone<ASTOpSwitch>(at);
 			NewOp->def = Visit(NewOp->def.child(), CurrentSinkOp);
diff --git a/Engine/Plugins/Mutable/Source/MutableTools/Private/MuT/ASTOpMeshOptimizeSkinning.h b/Engine/Plugins/Mutable/Source/MutableTools/Private/MuT/ASTOpMeshOptimizeSkinning.h
--- a/Engine/Plugins/Mutable/Source/MutableTools/Private/MuT/ASTOpMeshOptimizeSkinning.h
+++ b/Engine/Plugins/Mutable/Source/MutableTools/Private/MuT/ASTOpMeshOptimizeSkinning.h
@@ -27,7 +27,7 @@ namespace mu
 		~ASTOpMeshOptimizeSkinning() override;
 
 		// ASTOp interface
-		virtual OP_TYPE GetOpType() const override { return OP_TYPE::ME_OPTIMIZESKINNING; }
+		virtual EOpType GetOpType() const override { return EOpType::ME_OPTIMIZESKINNING; }
 		virtual uint64 Hash() const override;
 		virtual void ForEachChild(const TFunctionRef<void(ASTChild&)>) override;
 		virtual bool IsEqual(const ASTOp&) const override;
diff --git a/Engine/Plugins/Mutable/Source/MutableTools/Private/MuT/ASTOpMeshRemoveMask.cpp b/Engine/Plugins/Mutable/Source/MutableTools/Private/MuT/ASTOpMeshRemoveMask.cpp
--- a/Engine/Plugins/Mutable/Source/MutableTools/Private/MuT/ASTOpMeshRemoveMask.cpp
+++ b/Engine/Plugins/Mutable/Source/MutableTools/Private/MuT/ASTOpMeshRemoveMask.cpp
@@ -93,24 +93,24 @@ namespace mu
 		// Already linked?
 		if (!linkedAddress)
 		{
-			linkedAddress = (OP::ADDRESS)program.m_opAddress.Num();
+			linkedAddress = (OP::ADDRESS)program.OpAddress.Num();
 
-			program.m_opAddress.Add((uint32_t)program.m_byteCode.Num());
-			AppendCode(program.m_byteCode, OP_TYPE::ME_REMOVEMASK);
+			program.OpAddress.Add((uint32_t)program.ByteCode.Num());
+			AppendCode(program.ByteCode, EOpType::ME_REMOVEMASK);
 
 			OP::ADDRESS sourceAt = source ? source->linkedAddress : 0;
-			AppendCode(program.m_byteCode, sourceAt);
+			AppendCode(program.ByteCode, sourceAt);
 
-			AppendCode(program.m_byteCode, FaceCullStrategy);
+			AppendCode(program.ByteCode, FaceCullStrategy);
 
-			AppendCode(program.m_byteCode, (uint16)removes.Num());
+			AppendCode(program.ByteCode, (uint16)removes.Num());
 			for (const TPair<ASTChild, ASTChild>& b : removes)
 			{
 				OP::ADDRESS condition = b.Key ? b.Key->linkedAddress : 0;
-				AppendCode(program.m_byteCode, condition);
+				AppendCode(program.ByteCode, condition);
 
 				OP::ADDRESS remove = b.Value ? b.Value->linkedAddress : 0;
-				AppendCode(program.m_byteCode, remove);
+				AppendCode(program.ByteCode, remove);
 			}
 		}
 	}
@@ -129,14 +129,14 @@ namespace mu
 
 			mu::Ptr<ASTOp> Apply(const ASTOpMeshRemoveMask* root)
 			{
-				m_root = root;
+				Root = root;
 				m_oldToNew.Empty();
 
-				m_initialSource = root->source.child();
-				mu::Ptr<ASTOp> newSource = Visit(m_initialSource);
+				InitialSource = root->source.child();
+				mu::Ptr<ASTOp> newSource = Visit(InitialSource);
 
 				// If there is any change, it is the new root.
-				if (newSource != m_initialSource)
+				if (newSource != InitialSource)
 				{
 					return newSource;
 				}
@@ -146,8 +146,8 @@ namespace mu
 
 		protected:
 
-			const ASTOpMeshRemoveMask* m_root;
-			mu::Ptr<ASTOp> m_initialSource;
+			const ASTOpMeshRemoveMask* Root;
+			mu::Ptr<ASTOp> InitialSource;
 			TMap<mu::Ptr<ASTOp>, mu::Ptr<ASTOp>> m_oldToNew;
 			TArray<mu::Ptr<ASTOp>> m_newOps;
 
@@ -172,7 +172,7 @@ namespace mu
 				switch (at->GetOpType())
 				{
 
-				case OP_TYPE::ME_MORPH:
+				case EOpType::ME_MORPH:
 				{
 					mu::Ptr<ASTOpMeshMorph> newOp = mu::Clone<ASTOpMeshMorph>(at);
 					newOp->Base = Visit(newOp->Base.child());
@@ -180,7 +180,7 @@ namespace mu
 					break;
 				}
 
-				case OP_TYPE::ME_ADDTAGS:
+				case EOpType::ME_ADDTAGS:
 				{
 					mu::Ptr<ASTOpMeshAddTags> newOp = mu::Clone<ASTOpMeshAddTags>(at);
 					newOp->Source = Visit(newOp->Source.child());
@@ -189,7 +189,7 @@ namespace mu
 				}
 
 				// disabled to avoid code explosion (or bug?) TODO
-//            case OP_TYPE::ME_CONDITIONAL:
+//            case EOpType::ME_CONDITIONAL:
 //            {
 //                Ptr<ASTOpConditional> newOp = mu::Clone<ASTOpConditional>(at);
 //                newOp->yes = Visit(newOp->yes.child());
@@ -198,7 +198,7 @@ namespace mu
 //                break;
 //            }
 
-//            case OP_TYPE::ME_SWITCH:
+//            case EOpType::ME_SWITCH:
 //            {
 //                auto newOp = mu::Clone<ASTOpSwitch>(at);
 //                newOp->def = Visit(newOp->def.child());
@@ -213,9 +213,9 @@ namespace mu
 				default:
 				{
 					//
-					if (at != m_initialSource)
+					if (at != InitialSource)
 					{
-						Ptr<ASTOpMeshRemoveMask> newOp = mu::Clone<ASTOpMeshRemoveMask>(m_root);
+						Ptr<ASTOpMeshRemoveMask> newOp = mu::Clone<ASTOpMeshRemoveMask>(Root);
 						newOp->source = at;
 						newAt = newOp;
 					}
@@ -250,10 +250,10 @@ namespace mu
 					continue;
 				}
 
-				OP_TYPE RemoveType = removes[RemoveIndex].Value->GetOpType();
+				EOpType RemoveType = removes[RemoveIndex].Value->GetOpType();
 				switch (RemoveType)
 				{
-				case OP_TYPE::ME_ADDTAGS:
+				case EOpType::ME_ADDTAGS:
 				{
 					// It can be ignored.
 					if (!NewOp)
diff --git a/Engine/Plugins/Mutable/Source/MutableTools/Private/MuT/ASTOpMeshRemoveMask.h b/Engine/Plugins/Mutable/Source/MutableTools/Private/MuT/ASTOpMeshRemoveMask.h
--- a/Engine/Plugins/Mutable/Source/MutableTools/Private/MuT/ASTOpMeshRemoveMask.h
+++ b/Engine/Plugins/Mutable/Source/MutableTools/Private/MuT/ASTOpMeshRemoveMask.h
@@ -33,7 +33,7 @@ namespace mu
 		~ASTOpMeshRemoveMask() override;
 
 		// ASTOp interface
-		OP_TYPE GetOpType() const override { return OP_TYPE::ME_REMOVEMASK; }
+		EOpType GetOpType() const override { return EOpType::ME_REMOVEMASK; }
 		uint64 Hash() const override;
 		void ForEachChild(const TFunctionRef<void(ASTChild&)>) override;
 		bool IsEqual(const ASTOp& otherUntyped) const override;
diff --git a/Engine/Plugins/Mutable/Source/MutableTools/Private/MuT/ASTOpMeshTransform.cpp b/Engine/Plugins/Mutable/Source/MutableTools/Private/MuT/ASTOpMeshTransform.cpp
--- a/Engine/Plugins/Mutable/Source/MutableTools/Private/MuT/ASTOpMeshTransform.cpp
+++ b/Engine/Plugins/Mutable/Source/MutableTools/Private/MuT/ASTOpMeshTransform.cpp
@@ -38,7 +38,7 @@ namespace mu
 
 	uint64 ASTOpMeshTransform::Hash() const
 	{
-		uint64 res = std::hash<OP_TYPE>()(OP_TYPE::ME_TRANSFORM);
+		uint64 res = std::hash<EOpType>()(EOpType::ME_TRANSFORM);
 		hash_combine(res, source.child().get());
 		return res;
 	}
@@ -71,11 +71,11 @@ namespace mu
 
 			args.matrix = program.AddConstant(matrix);
 
-			linkedAddress = (OP::ADDRESS)program.m_opAddress.Num();
+			linkedAddress = (OP::ADDRESS)program.OpAddress.Num();
 			//program.m_code.push_back(op);
-			program.m_opAddress.Add((uint32_t)program.m_byteCode.Num());
-			AppendCode(program.m_byteCode, OP_TYPE::ME_TRANSFORM);
-			AppendCode(program.m_byteCode, args);
+			program.OpAddress.Add((uint32_t)program.ByteCode.Num());
+			AppendCode(program.ByteCode, EOpType::ME_TRANSFORM);
+			AppendCode(program.ByteCode, args);
 		}
 
 	}
diff --git a/Engine/Plugins/Mutable/Source/MutableTools/Private/MuT/ASTOpMeshTransform.h b/Engine/Plugins/Mutable/Source/MutableTools/Private/MuT/ASTOpMeshTransform.h
--- a/Engine/Plugins/Mutable/Source/MutableTools/Private/MuT/ASTOpMeshTransform.h
+++ b/Engine/Plugins/Mutable/Source/MutableTools/Private/MuT/ASTOpMeshTransform.h
@@ -32,7 +32,7 @@ struct FProgram;
 		~ASTOpMeshTransform();
 
 		// ASTOp interface
-		OP_TYPE GetOpType() const override { return OP_TYPE::ME_TRANSFORM; }
+		EOpType GetOpType() const override { return EOpType::ME_TRANSFORM; }
 		uint64 Hash() const override;
 		void ForEachChild(const TFunctionRef<void(ASTChild&)>) override;
 		bool IsEqual(const ASTOp& otherUntyped) const override;
diff --git a/Engine/Plugins/Mutable/Source/MutableTools/Private/MuT/ASTOpMeshTransformWithBoundingMesh.cpp b/Engine/Plugins/Mutable/Source/MutableTools/Private/MuT/ASTOpMeshTransformWithBoundingMesh.cpp
--- a/Engine/Plugins/Mutable/Source/MutableTools/Private/MuT/ASTOpMeshTransformWithBoundingMesh.cpp
+++ b/Engine/Plugins/Mutable/Source/MutableTools/Private/MuT/ASTOpMeshTransformWithBoundingMesh.cpp
@@ -20,7 +20,7 @@ mu::ASTOpMeshTransformWithBoundingMesh::~ASTOpMeshTransformWithBoundingMesh()
 
 uint64 mu::ASTOpMeshTransformWithBoundingMesh::Hash() const
 {
-	uint64 res = std::hash<OP_TYPE>()(GetOpType());
+	uint64 res = std::hash<EOpType>()(GetOpType());
 	hash_combine(res, source.child());
 	hash_combine(res, boundingMesh.child());
 	hash_combine(res, matrix.child());
@@ -71,10 +71,10 @@ void mu::ASTOpMeshTransformWithBoundingMesh::Link(FProgram& program, FLinkerOpti
 		{
 			args.matrix = matrix->linkedAddress;
 		}
-		linkedAddress = static_cast<OP::ADDRESS>(program.m_opAddress.Num());
-		program.m_opAddress.Add(program.m_byteCode.Num());
-		AppendCode(program.m_byteCode, OP_TYPE::ME_TRANSFORMWITHMESH);
-		AppendCode(program.m_byteCode, args);
+		linkedAddress = static_cast<OP::ADDRESS>(program.OpAddress.Num());
+		program.OpAddress.Add(program.ByteCode.Num());
+		AppendCode(program.ByteCode, EOpType::ME_TRANSFORMWITHMESH);
+		AppendCode(program.ByteCode, args);
 	}
 }
 
diff --git a/Engine/Plugins/Mutable/Source/MutableTools/Private/MuT/ASTOpMeshTransformWithBoundingMesh.h b/Engine/Plugins/Mutable/Source/MutableTools/Private/MuT/ASTOpMeshTransformWithBoundingMesh.h
--- a/Engine/Plugins/Mutable/Source/MutableTools/Private/MuT/ASTOpMeshTransformWithBoundingMesh.h
+++ b/Engine/Plugins/Mutable/Source/MutableTools/Private/MuT/ASTOpMeshTransformWithBoundingMesh.h
@@ -30,7 +30,7 @@ public:
 	ASTOpMeshTransformWithBoundingMesh(const ASTOpMeshTransformWithBoundingMesh&) = delete;
 	~ASTOpMeshTransformWithBoundingMesh() override;
 
-	OP_TYPE GetOpType() const override { return OP_TYPE::ME_TRANSFORMWITHMESH; }
+	EOpType GetOpType() const override { return EOpType::ME_TRANSFORMWITHMESH; }
 	uint64 Hash() const override;
 	bool IsEqual(const ASTOp& otherUntyped) const override;
 	Ptr<ASTOp> Clone(MapChildFuncRef mapChild) const override;
diff --git a/Engine/Plugins/Mutable/Source/MutableTools/Private/MuT/ASTOpParameter.cpp b/Engine/Plugins/Mutable/Source/MutableTools/Private/MuT/ASTOpParameter.cpp
--- a/Engine/Plugins/Mutable/Source/MutableTools/Private/MuT/ASTOpParameter.cpp
+++ b/Engine/Plugins/Mutable/Source/MutableTools/Private/MuT/ASTOpParameter.cpp
@@ -59,8 +59,8 @@ namespace mu
 	uint64 ASTOpParameter::Hash() const
 	{
 		uint64 res = std::hash<uint64>()(uint64(type));
-		hash_combine(res, parameter.m_type);
-		hash_combine(res, parameter.m_name.Len());
+		hash_combine(res, parameter.Type);
+		hash_combine(res, parameter.Name.Len());
 		return res;
 	}
 
@@ -79,7 +79,7 @@ namespace mu
 			OP::ParameterArgs args;
 			FMemory::Memzero(&args, sizeof(args));
 
-			LinkedParameterIndex = Program.m_parameters.Find(parameter);
+			LinkedParameterIndex = Program.Parameters.Find(parameter);
 
 			// If this fails, it means an ASTOpParameter was created at code generation time, but not registered into the
 			// parameters map in the CodeGenerator_FirstPass.
@@ -92,14 +92,14 @@ namespace mu
 				OP::ADDRESS sizeAt = 0;
 				uint16 rangeId = 0;
 				LinkRange(Program, d, sizeAt, rangeId);
-				Program.m_parameters[LinkedParameterIndex].m_ranges.Add(rangeId);
+				Program.Parameters[LinkedParameterIndex].Ranges.Add(rangeId);
 			}
 
-			linkedAddress = (OP::ADDRESS)Program.m_opAddress.Num();
+			linkedAddress = (OP::ADDRESS)Program.OpAddress.Num();
 
-			Program.m_opAddress.Add((uint32)Program.m_byteCode.Num());
-			AppendCode(Program.m_byteCode, type);
-			AppendCode(Program.m_byteCode, args);
+			Program.OpAddress.Add((uint32)Program.ByteCode.Num());
+			AppendCode(Program.ByteCode, type);
+			AppendCode(Program.ByteCode, args);
 		}
 	}
 
@@ -111,7 +111,7 @@ namespace mu
 		// Check the facts, in case we have the value for our parameter.
 		for (const auto& f : facts)
 		{
-			if (f->GetOpType() == OP_TYPE::BO_EQUAL_INT_CONST)
+			if (f->GetOpType() == EOpType::BO_EQUAL_INT_CONST)
 			{
 				const auto typedFact = static_cast<const ASTOpFixed*>(f.get());
 				auto value = typedFact->children[typedFact->op.args.BoolEqualScalarConst.value].child();
@@ -135,14 +135,14 @@ namespace mu
 	ASTOp::FBoolEvalResult ASTOpParameter::EvaluateBool(ASTOpList& /*facts*/,
 		FEvaluateBoolCache*) const
 	{
-		check(type == OP_TYPE::BO_PARAMETER);
+		check(type == EOpType::BO_PARAMETER);
 		return BET_UNKNOWN;
 	}
 
 
 	FImageDesc ASTOpParameter::GetImageDesc(bool, FGetImageDescContext*) const
 	{
-		check(type == OP_TYPE::IM_PARAMETER);
+		check(type == EOpType::IM_PARAMETER);
 		return FImageDesc();
 	}
 
diff --git a/Engine/Plugins/Mutable/Source/MutableTools/Private/MuT/ASTOpParameter.h b/Engine/Plugins/Mutable/Source/MutableTools/Private/MuT/ASTOpParameter.h
--- a/Engine/Plugins/Mutable/Source/MutableTools/Private/MuT/ASTOpParameter.h
+++ b/Engine/Plugins/Mutable/Source/MutableTools/Private/MuT/ASTOpParameter.h
@@ -23,7 +23,7 @@ namespace mu
 	public:
 
 		//! Type of parameter
-		OP_TYPE type;
+		EOpType type;
 
 		//!
 		FParameterDesc parameter;
@@ -38,7 +38,7 @@ namespace mu
 
 		~ASTOpParameter() override;
 
-		OP_TYPE GetOpType() const override { return type; }
+		EOpType GetOpType() const override { return type; }
 		uint64 Hash() const override;
 		void ForEachChild(const TFunctionRef<void(ASTChild&)>) override;
 		bool IsEqual(const ASTOp& otherUntyped) const override;
diff --git a/Engine/Plugins/Mutable/Source/MutableTools/Private/MuT/ASTOpReferenceResource.cpp b/Engine/Plugins/Mutable/Source/MutableTools/Private/MuT/ASTOpReferenceResource.cpp
--- a/Engine/Plugins/Mutable/Source/MutableTools/Private/MuT/ASTOpReferenceResource.cpp
+++ b/Engine/Plugins/Mutable/Source/MutableTools/Private/MuT/ASTOpReferenceResource.cpp
@@ -61,10 +61,10 @@ namespace mu
 			Args.ForceLoad = bForceLoad ? 1 : 0;
 			Args.ImageDesc = ImageDesc;
 
-			linkedAddress = (OP::ADDRESS)program.m_opAddress.Num();
-			program.m_opAddress.Add((uint32)program.m_byteCode.Num());
-			AppendCode(program.m_byteCode, type);
-			AppendCode(program.m_byteCode, Args);
+			linkedAddress = (OP::ADDRESS)program.OpAddress.Num();
+			program.OpAddress.Add((uint32)program.ByteCode.Num());
+			AppendCode(program.ByteCode, type);
+			AppendCode(program.ByteCode, Args);
 		}
 	}
 
@@ -82,7 +82,7 @@ namespace mu
 		switch (type)
 		{
 
-		case OP_TYPE::IM_REFERENCE:
+		case EOpType::IM_REFERENCE:
 		{
 			*pBlockX = 0;
 			*pBlockY = 0;
@@ -101,7 +101,7 @@ namespace mu
 		switch (type)
 		{
 
-		case OP_TYPE::IM_REFERENCE:
+		case EOpType::IM_REFERENCE:
 		{
 			// We didn't find any layout.
 			*pBlockX = 0;
diff --git a/Engine/Plugins/Mutable/Source/MutableTools/Private/MuT/ASTOpReferenceResource.h b/Engine/Plugins/Mutable/Source/MutableTools/Private/MuT/ASTOpReferenceResource.h
--- a/Engine/Plugins/Mutable/Source/MutableTools/Private/MuT/ASTOpReferenceResource.h
+++ b/Engine/Plugins/Mutable/Source/MutableTools/Private/MuT/ASTOpReferenceResource.h
@@ -24,7 +24,7 @@ namespace mu
 	public:
 
 		//! Type of switch
-		OP_TYPE type = OP_TYPE::NONE;
+		EOpType type = EOpType::NONE;
 		
 		/** */
 		bool bForceLoad = false;
@@ -43,7 +43,7 @@ namespace mu
 		// Own interface
 
 		// ASTOp interface
-		virtual OP_TYPE GetOpType() const override { return type; }
+		virtual EOpType GetOpType() const override { return type; }
 		virtual void ForEachChild(const TFunctionRef<void(ASTChild&)>) override;
 		virtual bool IsEqual(const ASTOp& otherUntyped) const override;
 		virtual Ptr<ASTOp> Clone(MapChildFuncRef mapChild) const override;
diff --git a/Engine/Plugins/Mutable/Source/MutableTools/Private/MuT/ASTOpScalarCurve.cpp b/Engine/Plugins/Mutable/Source/MutableTools/Private/MuT/ASTOpScalarCurve.cpp
--- a/Engine/Plugins/Mutable/Source/MutableTools/Private/MuT/ASTOpScalarCurve.cpp
+++ b/Engine/Plugins/Mutable/Source/MutableTools/Private/MuT/ASTOpScalarCurve.cpp
@@ -33,7 +33,7 @@ namespace mu
 
 	uint64 ASTOpScalarCurve::Hash() const
 	{
-		uint64 res = std::hash<uint64>()(size_t(OP_TYPE::SC_CURVE));
+		uint64 res = std::hash<uint64>()(size_t(EOpType::SC_CURVE));
 		hash_combine(res, Curve.Keys.Num());
 		return res;
 	}
@@ -68,11 +68,11 @@ namespace mu
 			args.time = time ? time->linkedAddress : 0;
 			args.curve = program.AddConstant(Curve);
 
-			linkedAddress = (OP::ADDRESS)program.m_opAddress.Num();
+			linkedAddress = (OP::ADDRESS)program.OpAddress.Num();
 			//program.m_code.push_back(op);
-			program.m_opAddress.Add((uint32_t)program.m_byteCode.Num());
-			AppendCode(program.m_byteCode, OP_TYPE::SC_CURVE);
-			AppendCode(program.m_byteCode, args);
+			program.OpAddress.Add((uint32_t)program.ByteCode.Num());
+			AppendCode(program.ByteCode, EOpType::SC_CURVE);
+			AppendCode(program.ByteCode, args);
 		}
 
 	}
diff --git a/Engine/Plugins/Mutable/Source/MutableTools/Private/MuT/ASTOpScalarCurve.h b/Engine/Plugins/Mutable/Source/MutableTools/Private/MuT/ASTOpScalarCurve.h
--- a/Engine/Plugins/Mutable/Source/MutableTools/Private/MuT/ASTOpScalarCurve.h
+++ b/Engine/Plugins/Mutable/Source/MutableTools/Private/MuT/ASTOpScalarCurve.h
@@ -33,7 +33,7 @@ struct FProgram;
 		ASTOpScalarCurve(const ASTOpScalarCurve&) = delete;
 		~ASTOpScalarCurve() override;
 
-		OP_TYPE GetOpType() const override { return OP_TYPE::SC_CURVE; }
+		EOpType GetOpType() const override { return EOpType::SC_CURVE; }
 		uint64 Hash() const override;
 		void ForEachChild(const TFunctionRef<void(ASTChild&)>) override;
 		bool IsEqual(const ASTOp& otherUntyped) const override;
diff --git a/Engine/Plugins/Mutable/Source/MutableTools/Private/MuT/ASTOpSwitch.cpp b/Engine/Plugins/Mutable/Source/MutableTools/Private/MuT/ASTOpSwitch.cpp
--- a/Engine/Plugins/Mutable/Source/MutableTools/Private/MuT/ASTOpSwitch.cpp
+++ b/Engine/Plugins/Mutable/Source/MutableTools/Private/MuT/ASTOpSwitch.cpp
@@ -65,14 +65,14 @@ namespace mu
 	{
 		switch (type)
 		{
-		case OP_TYPE::NU_SWITCH:
-		case OP_TYPE::SC_SWITCH:
-		case OP_TYPE::CO_SWITCH:
-		case OP_TYPE::IM_SWITCH:
-		case OP_TYPE::ME_SWITCH:
-		case OP_TYPE::LA_SWITCH:
-		case OP_TYPE::IN_SWITCH:
-		case OP_TYPE::ED_SWITCH:
+		case EOpType::NU_SWITCH:
+		case EOpType::SC_SWITCH:
+		case EOpType::CO_SWITCH:
+		case EOpType::IM_SWITCH:
+		case EOpType::ME_SWITCH:
+		case EOpType::LA_SWITCH:
+		case EOpType::IN_SWITCH:
+		case EOpType::ED_SWITCH:
 			break;
 		default:
 			// Unexpected type
@@ -171,22 +171,22 @@ namespace mu
 		// Already linked?
 		if (!linkedAddress)
 		{
-			linkedAddress = (OP::ADDRESS)program.m_opAddress.Num();
-			program.m_opAddress.Add((uint32_t)program.m_byteCode.Num());
+			linkedAddress = (OP::ADDRESS)program.OpAddress.Num();
+			program.OpAddress.Add((uint32_t)program.ByteCode.Num());
 
 			OP::ADDRESS VarAddress = variable ? variable->linkedAddress : 0;
 			OP::ADDRESS DefAddress = def ? def->linkedAddress : 0;
 
-			AppendCode(program.m_byteCode, type);
-			AppendCode(program.m_byteCode, VarAddress);
-			AppendCode(program.m_byteCode, DefAddress);
-			AppendCode(program.m_byteCode, (uint32_t)cases.Num());
+			AppendCode(program.ByteCode, type);
+			AppendCode(program.ByteCode, VarAddress);
+			AppendCode(program.ByteCode, DefAddress);
+			AppendCode(program.ByteCode, (uint32_t)cases.Num());
 
 			for (const FCase& Case : cases)
 			{
 				OP::ADDRESS CaseBranchAddress = Case.branch ? Case.branch->linkedAddress : 0;
-				AppendCode(program.m_byteCode, Case.condition);
-				AppendCode(program.m_byteCode, CaseBranchAddress);
+				AppendCode(program.ByteCode, Case.condition);
+				AppendCode(program.ByteCode, CaseBranchAddress);
 			}
 		}
 	}
@@ -294,7 +294,7 @@ namespace mu
 	{
 		switch (type)
 		{
-		case OP_TYPE::LA_SWITCH:
+		case EOpType::LA_SWITCH:
 		{
 			Ptr<ASTOp> child = GetFirstValidValue();
 			if (!child)
@@ -326,7 +326,7 @@ namespace mu
 	{
 		switch (type)
 		{
-		case OP_TYPE::IM_SWITCH:
+		case EOpType::IM_SWITCH:
 		{
 			Ptr<ASTOp> child = GetFirstValidValue();
 			if (!child)
@@ -354,7 +354,7 @@ namespace mu
 	//-------------------------------------------------------------------------------------------------
 	bool ASTOpSwitch::GetNonBlackRect(FImageRect& maskUsage) const
 	{
-		if (type == OP_TYPE::IM_SWITCH)
+		if (type == EOpType::IM_SWITCH)
 		{
 			FImageRect local;
 			bool localValid = false;
@@ -416,7 +416,7 @@ namespace mu
 	mu::Ptr<ASTOp> ASTOpSwitch::OptimiseSemantic(const FModelOptimizationOptions&, int32 Pass) const
 	{
 		// Constant condition?
-		if (variable->GetOpType() == OP_TYPE::NU_CONSTANT)
+		if (variable->GetOpType() == EOpType::NU_CONSTANT)
 		{
 			Ptr<ASTOp> Branch = def.child();
 
@@ -436,23 +436,23 @@ namespace mu
 			return Branch;
 		}
 
-		else if (variable->GetOpType() == OP_TYPE::NU_PARAMETER)
+		else if (variable->GetOpType() == EOpType::NU_PARAMETER)
 		{
 			// If all the branches for the possible values are the same op remove the instruction
 			const ASTOpParameter* ParamOp = static_cast<const ASTOpParameter*>(variable.child().get());
 			check(ParamOp);
-			check(!ParamOp->parameter.m_possibleValues.IsEmpty());
+			check(!ParamOp->parameter.PossibleValues.IsEmpty());
 
 			bool bFirstValue = true;
 			bool bAllSame = true;
 			Ptr<ASTOp> SameBranch = nullptr;
-			for (const FParameterDesc::FIntValueDesc& Value : ParamOp->parameter.m_possibleValues)
+			for (const FParameterDesc::FIntValueDesc& Value : ParamOp->parameter.PossibleValues)
 			{
 				// Look for the switch branch it would take
 				Ptr<ASTOp> Branch = def.child();
 				for (const FCase& Case : cases)
 				{
-					if (Case.condition == Value.m_value)
+					if (Case.condition == Value.Value)
 					{
 						Branch = Case.branch.child();
 						break;
diff --git a/Engine/Plugins/Mutable/Source/MutableTools/Private/MuT/ASTOpSwitch.h b/Engine/Plugins/Mutable/Source/MutableTools/Private/MuT/ASTOpSwitch.h
--- a/Engine/Plugins/Mutable/Source/MutableTools/Private/MuT/ASTOpSwitch.h
+++ b/Engine/Plugins/Mutable/Source/MutableTools/Private/MuT/ASTOpSwitch.h
@@ -23,7 +23,7 @@ template <class SCALAR> class vec4;
 	public:
 
 		//! Type of switch
-		OP_TYPE type;
+		EOpType type;
 
 		//! Variable whose value will be used to choose the switch branch.
 		ASTChild variable;
@@ -59,7 +59,7 @@ template <class SCALAR> class vec4;
 		~ASTOpSwitch() override;
 
 		// ASTOp interface
-		virtual OP_TYPE GetOpType() const override { return type; }
+		virtual EOpType GetOpType() const override { return type; }
 		virtual uint64 Hash() const override;
 		virtual void ForEachChild(const TFunctionRef<void(ASTChild&)> f) override;
 		virtual bool IsEqual(const ASTOp& otherUntyped) const override;
diff --git a/Engine/Plugins/Mutable/Source/MutableTools/Private/MuT/CodeGenerator.cpp b/Engine/Plugins/Mutable/Source/MutableTools/Private/MuT/CodeGenerator.cpp
--- a/Engine/Plugins/Mutable/Source/MutableTools/Private/MuT/CodeGenerator.cpp
+++ b/Engine/Plugins/Mutable/Source/MutableTools/Private/MuT/CodeGenerator.cpp
@@ -313,15 +313,15 @@ namespace mu
 		Ptr<ASTOp> result;
 
         FParameterDesc param;
-        param.m_name = CacheKey.ParameterName;
+        param.Name = CacheKey.ParameterName;
 
-        if (param.m_name.Len() == 0)
+        if (param.Name.Len() == 0)
         {
-            param.m_name = CacheKey.Table->GetName();
+            param.Name = CacheKey.Table->GetName();
         }
         
-		param.m_type = EParameterType::Int;
-        param.m_defaultValue.Set<FParamIntType>(0);
+		param.Type = EParameterType::Int;
+        param.DefaultValue.Set<FParamIntType>(0);
 
 		// Add the possible values
 		{
@@ -340,10 +340,10 @@ namespace mu
 			if (bAddNoneOption)
 			{
 				FParameterDesc::FIntValueDesc nullValue;
-				nullValue.m_value = -1;
-				nullValue.m_name = "None";
-				param.m_possibleValues.Add(nullValue);
-				param.m_defaultValue.Set<FParamIntType>(nullValue.m_value);
+				nullValue.Value = -1;
+				nullValue.Name = "None";
+				param.PossibleValues.Add(nullValue);
+				param.DefaultValue.Set<FParamIntType>(nullValue.Value);
 			}
 
 			// Add every row
@@ -353,31 +353,31 @@ namespace mu
 			for (int32 RowIndex = 0; RowIndex < RowCount; ++RowIndex)
 			{
 				FParameterDesc::FIntValueDesc value;
-				value.m_value = RowIndex;
+				value.Value = RowIndex;
 
 				if (nameCol > -1)
 				{
-					value.m_name = CacheKey.Table->GetPrivate()->Rows[RowIndex].Values[nameCol].String;
+					value.Name = CacheKey.Table->GetPrivate()->Rows[RowIndex].Values[nameCol].String;
 				}
 
-				param.m_possibleValues.Add(value);
+				param.PossibleValues.Add(value);
 
 				// Set the first row as the default one (if there is none option)
 				if (RowIndex == 0 && !bAddNoneOption)
 				{
-					param.m_defaultValue.Set<FParamIntType>(value.m_value);
+					param.DefaultValue.Set<FParamIntType>(value.Value);
 				}
 				
 				// Set the selected row as default (if exists)
-				if (value.m_name == DefaultRowName)
+				if (value.Name == DefaultRowName)
 				{
-					param.m_defaultValue.Set<FParamIntType>(value.m_value);
+					param.DefaultValue.Set<FParamIntType>(value.Value);
 				}
             }
         }
 
 		Ptr<ASTOpParameter> op = new ASTOpParameter();
-		op->type = OP_TYPE::NU_PARAMETER;
+		op->type = EOpType::NU_PARAMETER;
 		op->parameter = param;
 
 		FirstPass.ParameterNodes.Add(InNode, op);
@@ -386,16 +386,16 @@ namespace mu
 	}
 
 
-	Ptr<const Layout> CodeGenerator::GenerateLayout(Ptr<const NodeLayout> SourceLayout, uint32 MeshIDPrefix)
+	TSharedPtr<const FLayout> CodeGenerator::GenerateLayout(Ptr<const NodeLayout> SourceLayout, uint32 MeshIDPrefix)
 	{
-		Ptr<const Layout>* it = GeneratedLayouts.Find({ SourceLayout,MeshIDPrefix });
+		TSharedPtr<const FLayout>* it = GeneratedLayouts.Find({ SourceLayout,MeshIDPrefix });
 
 		if (it)
 		{
 			return *it;
 		}
 
-		Ptr<Layout> GeneratedLayout = new Layout;
+		TSharedPtr<FLayout> GeneratedLayout = MakeShared<FLayout>();
 		GeneratedLayout->Size = SourceLayout->Size;
 		GeneratedLayout->MaxSize = SourceLayout->MaxSize;
 		GeneratedLayout->Strategy = SourceLayout->Strategy;
@@ -427,7 +427,7 @@ namespace mu
 
 	Ptr<ASTOp> CodeGenerator::GenerateImageBlockPatch(Ptr<ASTOp> InBlockOp,
 		const NodeModifierSurfaceEdit::FTexture& Patch,
-		Ptr<Image> PatchMask,
+		TSharedPtr<FImage> PatchMask,
 		Ptr<ASTOp> conditionAd,
 		const FImageGenerationOptions& ImageOptions )
 	{
@@ -464,7 +464,7 @@ namespace mu
 			Ptr<ASTOp> RectConstantOp;
 			{
 				Ptr<NodeImageConstant> pNode = new NodeImageConstant();
-				pNode->SetValue(PatchMask.get());
+				pNode->SetValue(PatchMask);
 
 				FImageGenerationOptions ConstantOptions(-1);
 				FImageGenerationResult ConstantResult;
@@ -505,7 +505,7 @@ namespace mu
 			Ptr<ASTOp> conditionalAd;
 			{
 				Ptr<ASTOpConditional> op = new ASTOpConditional();
-				op->type = OP_TYPE::IM_CONDITIONAL;
+				op->type = EOpType::IM_CONDITIONAL;
 				op->no = InBlockOp;
 				op->yes = FinalOp;
 				op->condition = conditionAd;
@@ -592,7 +592,7 @@ namespace mu
 		}
 
 		Ptr<ASTOpInstanceAdd> InstanceOp = new ASTOpInstanceAdd();
-		InstanceOp->type = OP_TYPE::IN_ADDCOMPONENT;
+		InstanceOp->type = EOpType::IN_ADDCOMPONENT;
 		InstanceOp->instance = Options.BaseInstance;
 		InstanceOp->value = LODsOp;
 		InstanceOp->ExternalId = node.Id;
@@ -611,12 +611,12 @@ namespace mu
 			{
 				// TODO: This could be done earlier?
 				Ptr<ASTOpFixed> ConditionOp = new ASTOpFixed();
-				ConditionOp->op.type = OP_TYPE::BO_AND;
+				ConditionOp->op.type = EOpType::BO_AND;
 				ConditionOp->SetChild(ConditionOp->op.args.BoolBinary.a, Component.ObjectCondition);
 				ConditionOp->SetChild(ConditionOp->op.args.BoolBinary.b, Component.ComponentCondition);
 
 				Ptr<ASTOpConditional> IfOp = new ASTOpConditional();
-				IfOp->type = OP_TYPE::IN_CONDITIONAL;
+				IfOp->type = EOpType::IN_CONDITIONAL;
 				IfOp->no = Options.BaseInstance;
 				IfOp->yes = Result.op;
 				IfOp->condition = ConditionOp;
@@ -639,7 +639,7 @@ namespace mu
 		}
 
 		Ptr<ASTOpSwitch> Op = new ASTOpSwitch();
-		Op->type = OP_TYPE::IN_SWITCH;
+		Op->type = EOpType::IN_SWITCH;
 
 		// Variable value
 		if (Node->Parameter)
@@ -723,7 +723,7 @@ namespace mu
 			}
 
 			Ptr<ASTOpConditional> Conditional = new ASTOpConditional;
-			Conditional->type = OP_TYPE::IN_CONDITIONAL;
+			Conditional->type = EOpType::IN_CONDITIONAL;
 			Conditional->no = CurrentMeshOp;
 			Conditional->yes = VariationMeshOp;
 			Conditional->condition = FirstPass.Tags[TagIndex].GenericCondition;
@@ -753,7 +753,7 @@ namespace mu
 		}
 
 		Ptr<ASTOpFixed> BaseImage = new ASTOpFixed;
-		BaseImage->op.type = OP_TYPE::IM_PLAINCOLOUR;
+		BaseImage->op.type = EOpType::IM_PLAINCOLOUR;
 		BaseImage->op.args.ImagePlainColour.size[0] = Size[0];
 		BaseImage->op.args.ImagePlainColour.size[1] = Size[1];
 		BaseImage->op.args.ImagePlainColour.format = Format;
@@ -791,10 +791,10 @@ namespace mu
 	}
 
 
-	Ptr<Image> CodeGenerator::GenerateImageBlockPatchMask(const NodeModifierSurfaceEdit::FTexture& Patch, FIntPoint GridSize, int32 BlockPixelsX, int32 BlockPixelsY, box<FIntVector2> RectInCells )
+	TSharedPtr<FImage> CodeGenerator::GenerateImageBlockPatchMask(const NodeModifierSurfaceEdit::FTexture& Patch, FIntPoint GridSize, int32 BlockPixelsX, int32 BlockPixelsY, box<FIntVector2> RectInCells )
 	{
 		// Create a patching mask for the block
-		Ptr<Image> PatchMask;
+		TSharedPtr<FImage> PatchMask;
 
 		FIntVector2 SourceTextureSize = { GridSize[0] * BlockPixelsX, GridSize[1] * BlockPixelsY };
 
@@ -817,7 +817,7 @@ namespace mu
 				FInt32Point BlockSize = BlockRectInPixels.Size();
 				if (!PatchMask)
 				{
-					PatchMask = new mu::Image(BlockSize[0], BlockSize[1], 1, mu::EImageFormat::L_UByte, mu::EInitializationType::Black);
+					PatchMask = MakeShared<FImage>(BlockSize[0], BlockSize[1], 1, mu::EImageFormat::L_UByte, mu::EInitializationType::Black);
 				}
 
 				uint8* Pixels = PatchMask->GetMipData(0);
@@ -986,7 +986,7 @@ namespace mu
 				else
 				{
 					Ptr<ASTOpConstantResource> ConstantLayoutOp = new ASTOpConstantResource();
-					ConstantLayoutOp->Type = OP_TYPE::LA_CONSTANT;
+					ConstantLayoutOp->Type = EOpType::LA_CONSTANT;
 
 					ConstantLayoutOp->SetValue(
 							SurfaceReferenceLayouts[LayoutIndex].Layout, 
@@ -1035,7 +1035,7 @@ namespace mu
 					else
 					{
 						Ptr<ASTOpConstantResource> LayoutFragmentConstantOp = new ASTOpConstantResource();
-						LayoutFragmentConstantOp->Type = OP_TYPE::LA_CONSTANT;
+						LayoutFragmentConstantOp->Type = EOpType::LA_CONSTANT;
 
 						LayoutFragmentConstantOp->SetValue(
 								ExtraLayoutsData.GeneratedLayouts[LayoutIndex].Layout,
@@ -1050,7 +1050,7 @@ namespace mu
 						if (ExtraLayoutsData.Condition)
 						{
 							Ptr<ASTOpConditional> ConditionalOp = new ASTOpConditional();
-							ConditionalOp->type = OP_TYPE::LA_CONDITIONAL;
+							ConditionalOp->type = EOpType::LA_CONDITIONAL;
 							ConditionalOp->no = SurfaceLayoutOps[LayoutIndex];
 							ConditionalOp->yes = LayoutMergeOp;
 							ConditionalOp->condition = ExtraLayoutsData.Condition;
@@ -1101,7 +1101,7 @@ namespace mu
 				// Create the expression to apply the layout to the mesh
 				{
 					Ptr<ASTOpFixed> ApplyLayoutOp = new ASTOpFixed();
-					ApplyLayoutOp->op.type = OP_TYPE::ME_APPLYLAYOUT;
+					ApplyLayoutOp->op.type = EOpType::ME_APPLYLAYOUT;
 					ApplyLayoutOp->SetChild(ApplyLayoutOp->op.args.MeshApplyLayout.mesh, LastMeshOp);
 					ApplyLayoutOp->SetChild(ApplyLayoutOp->op.args.MeshApplyLayout.layout, SurfaceLayoutOps[LayoutIndex]);
 					ApplyLayoutOp->op.args.MeshApplyLayout.channel = (uint16)LayoutIndex;
@@ -1263,7 +1263,7 @@ namespace mu
 						}
 
 						Ptr<ASTOpInstanceAdd> op = new ASTOpInstanceAdd();
-						op->type = OP_TYPE::IN_ADDIMAGE;
+						op->type = EOpType::IN_ADDIMAGE;
 						op->instance = LastSurfOp;
 						op->value = imageAd;
 						op->name = SurfaceNode->Images[ImageIndex].Name;
@@ -1280,11 +1280,11 @@ namespace mu
 						}
 						else
 						{
-							const Layout* pLayout = MeshResults.GeneratedLayouts[LayoutIndex].Layout.get();
+							TSharedPtr<const FLayout> pLayout = MeshResults.GeneratedLayouts[LayoutIndex].Layout;
 							check(pLayout);
 
 							Ptr<ASTOpInstanceAdd> op = new ASTOpInstanceAdd();
-							op->type = OP_TYPE::IN_ADDIMAGE;
+							op->type = EOpType::IN_ADDIMAGE;
 							op->instance = LastSurfOp;
 
 							// Image
@@ -1311,7 +1311,7 @@ namespace mu
 							Ptr<ASTOp> imageAd;
 							{
 								BlankImageOp = new ASTOpFixed();
-								BlankImageOp->op.type = OP_TYPE::IM_BLANKLAYOUT;
+								BlankImageOp->op.type = EOpType::IM_BLANKLAYOUT;
 								BlankImageOp->SetChild(BlankImageOp->op.args.ImageBlankLayout.layout, MeshResults.LayoutOps[LayoutIndex]);
 								// The rest ok the op will be completed below
 								BlankImageOp->op.args.ImageBlankLayout.mipmapCount = 0;
@@ -1491,7 +1491,7 @@ namespace mu
 				if (Ptr<NodeColour> pVectorNode = SurfaceNode->Vectors[t].Vector)
 				{
 					Ptr<ASTOpInstanceAdd> op = new ASTOpInstanceAdd();
-					op->type = OP_TYPE::IN_ADDVECTOR;
+					op->type = EOpType::IN_ADDVECTOR;
 					op->instance = LastSurfOp;
 
 					// Vector
@@ -1515,7 +1515,7 @@ namespace mu
 				if (NodeScalarPtr pScalarNode = SurfaceNode->Scalars[t].Scalar)
 				{
 					Ptr<ASTOpInstanceAdd> op = new ASTOpInstanceAdd();
-					op->type = OP_TYPE::IN_ADDSCALAR;
+					op->type = EOpType::IN_ADDSCALAR;
 					op->instance = LastSurfOp;
 
 					// Scalar
@@ -1537,7 +1537,7 @@ namespace mu
 				if (NodeStringPtr pStringNode = SurfaceNode->Strings[t].String)
 				{
 					Ptr<ASTOpInstanceAdd> op = new ASTOpInstanceAdd();
-					op->type = OP_TYPE::IN_ADDSTRING;
+					op->type = EOpType::IN_ADDSTRING;
 					op->instance = LastSurfOp;
 
 					FStringGenerationResult StringResult;
@@ -1609,7 +1609,7 @@ namespace mu
                 }
 
                 Ptr<ASTOpInstanceAdd> SurfaceOp = new ASTOpInstanceAdd();
-				SurfaceOp->type = OP_TYPE::IN_ADDSURFACE;
+				SurfaceOp->type = EOpType::IN_ADDSURFACE;
 				SurfaceOp->name = its.Node->Name;
 				SurfaceOp->instance = lastCompOp;
 
@@ -1626,7 +1626,7 @@ namespace mu
 
                 {
                     Ptr<ASTOpConditional> op = new ASTOpConditional();
-                    op->type = OP_TYPE::IN_CONDITIONAL;
+                    op->type = EOpType::IN_CONDITIONAL;
                     op->no = lastCompOp;
                     op->yes = SurfaceOp;
                     op->condition = SurfaceConditionOp;
@@ -1641,7 +1641,7 @@ namespace mu
                     Ptr<ASTOp> added = its.ResultMeshOp;
  
                     Ptr<ASTOpFixed> mop = new ASTOpFixed();
-                    mop->op.type = OP_TYPE::ME_MERGE;
+                    mop->op.type = EOpType::ME_MERGE;
                     mop->SetChild(mop->op.args.MeshMerge.base, lastMeshOp );
                     mop->SetChild(mop->op.args.MeshMerge.added, added );
                     mop->op.args.MeshMerge.newSurfaceID = surfaceID;
@@ -1651,7 +1651,7 @@ namespace mu
                 if (SurfaceConditionOp)
                 {
                     Ptr<ASTOpConditional> op = new ASTOpConditional();
-                    op->type = OP_TYPE::ME_CONDITIONAL;
+                    op->type = EOpType::ME_CONDITIONAL;
                     op->no = lastMeshOp;
                     op->yes = mergeAd;
                     op->condition = SurfaceConditionOp;
@@ -1676,7 +1676,7 @@ namespace mu
         // Add the component mesh
         {
             Ptr<ASTOpInstanceAdd> iop = new ASTOpInstanceAdd();
-            iop->type = OP_TYPE::IN_ADDMESH;
+            iop->type = EOpType::IN_ADDMESH;
             iop->instance = lastCompOp;
             iop->value = lastMeshOp;
 
@@ -1836,7 +1836,7 @@ namespace mu
 				if (SavedOp.Condition.get())
 				{
 					Ptr<ASTOpConditional> ConditionOp = new ASTOpConditional();
-					ConditionOp->type = OP_TYPE::IN_CONDITIONAL;
+					ConditionOp->type = EOpType::IN_CONDITIONAL;
 					ConditionOp->no = RootOp;
 					ConditionOp->yes = ExtensionPinOp;
 					ConditionOp->condition = ASTChild(ConditionOp, SavedOp.Condition);
@@ -2082,7 +2082,7 @@ namespace mu
 					BaseMeshResult.ExtraMeshLayouts[EditIndex] = data;
 
 					Ptr<ASTOpFixed> mop = new ASTOpFixed();
-					mop->op.type = OP_TYPE::ME_MERGE;
+					mop->op.type = EOpType::ME_MERGE;
 					mop->SetChild(mop->op.args.MeshMerge.base, LastMeshOp);
 					mop->SetChild(mop->op.args.MeshMerge.added, AddedMeshOp);
 					// will merge the meshes under the same surface
@@ -2092,7 +2092,7 @@ namespace mu
 					if (m.FinalCondition)
 					{
 						Ptr<ASTOpConditional> conditionalAd = new ASTOpConditional();
-						conditionalAd->type = OP_TYPE::ME_CONDITIONAL;
+						conditionalAd->type = EOpType::ME_CONDITIONAL;
 						conditionalAd->no = LastMeshOp;
 						conditionalAd->yes = mop;
 						conditionalAd->condition = m.FinalCondition;
@@ -2138,7 +2138,7 @@ namespace mu
 					GenerateMesh(RemoveMeshOptions, removeResults, pRemove);
 
 					Ptr<ASTOpFixed> maskOp = new ASTOpFixed();
-					maskOp->op.type = OP_TYPE::ME_MASKDIFF;
+					maskOp->op.type = EOpType::ME_MASKDIFF;
 
 					// By default, remove from the base
 					Ptr<ASTOp> removeFrom = BaseMeshResult.BaseMeshOp;
@@ -2173,7 +2173,7 @@ namespace mu
 
 				check(OriginalMeshNode);
 
-				Ptr<Mesh> TargetMesh = OriginalMeshNode->FindMorph(Edit->MeshMorph);
+				TSharedPtr<FMesh> TargetMesh = OriginalMeshNode->FindMorph(Edit->MeshMorph);
 				if (!TargetMesh)
 				{
 					continue;
@@ -2182,7 +2182,7 @@ namespace mu
 				{
 					// Target mesh
 					Ptr<ASTOpConstantResource> TargetMeshOp = new ASTOpConstantResource;
-					TargetMeshOp->Type = OP_TYPE::ME_CONSTANT;
+					TargetMeshOp->Type = EOpType::ME_CONSTANT;
 					TargetMeshOp->SetValue(TargetMesh->Clone(), CompilerOptions->OptimisationOptions.DiskCacheContext);
 					TargetMeshOp->SourceDataDescriptor = OriginalMeshNode->SourceDataDescriptor;
 
@@ -2234,7 +2234,7 @@ namespace mu
 					if (m.FinalCondition)
 					{
 						Ptr<ASTOpConditional> conditionalAd = new ASTOpConditional();
-						conditionalAd->type = OP_TYPE::ME_CONDITIONAL;
+						conditionalAd->type = EOpType::ME_CONDITIONAL;
 						conditionalAd->no = LastMeshOp;
 						conditionalAd->yes = morphAd;
 						conditionalAd->condition = m.FinalCondition;
@@ -2260,7 +2260,7 @@ namespace mu
 				Ptr<ASTOpMeshMaskClipMesh> op = new ASTOpMeshMaskClipMesh();
 				op->source = PreModifiersMesh;
 
-				// Parameters
+				// FParameters
 				FMeshGenerationOptions ClipOptions( Options.ComponentId );
 				ClipOptions.bLayouts = false;
 				ClipOptions.State = Options.State;
@@ -2300,9 +2300,9 @@ namespace mu
 				// TODO: Optimize, by caching.
 				// TODO: Optimize by formatting and keeping only UVs
 				check(OriginalMeshNode);
-				const Mesh* OriginalMesh = OriginalMeshNode->GetPrivate()->Value.get();
+				const FMesh* OriginalMesh = OriginalMeshNode->GetPrivate()->Value.Get();
 				Ptr<ASTOpConstantResource> UVMeshOp = new ASTOpConstantResource();
-				UVMeshOp->Type = OP_TYPE::ME_CONSTANT;
+				UVMeshOp->Type = EOpType::ME_CONSTANT;
 				UVMeshOp->SetValue(OriginalMesh->Clone(), CompilerOptions->OptimisationOptions.DiskCacheContext);
 				UVMeshOp->SourceDataDescriptor = OriginalMeshNode->SourceDataDescriptor;
 
@@ -2318,7 +2318,7 @@ namespace mu
 
 				if (TypedClipNode->ClipMask)
 				{
-					// Parameters to generate the mask image
+					// FParameters to generate the mask image
 					FImageGenerationOptions ClipOptions(Options.ComponentId);
 					ClipOptions.ImageLayoutStrategy = CompilerOptions::TextureLayoutStrategy::None;
 					ClipOptions.LayoutBlockId = FLayoutBlock::InvalidBlockId;
@@ -2340,10 +2340,10 @@ namespace mu
 				else if (TypedClipNode->ClipLayout)
 				{
 					// Generate the layout with blocks to extract
-					Ptr<const Layout> Layout = GenerateLayout(TypedClipNode->ClipLayout, 0);
+					TSharedPtr<const FLayout> Layout = GenerateLayout(TypedClipNode->ClipLayout, 0);
 
 					Ptr<ASTOpConstantResource> LayoutOp = new ASTOpConstantResource();
-					LayoutOp->Type = OP_TYPE::LA_CONSTANT;
+					LayoutOp->Type = EOpType::LA_CONSTANT;
 					LayoutOp->SetValue(Layout, CompilerOptions->OptimisationOptions.DiskCacheContext);
 					op->MaskLayout = LayoutOp;
 				}
@@ -2422,7 +2422,7 @@ namespace mu
 					op->vertexSelectionBoneMaxRadius = TypedNode->Parameters.MaxEffectRadius;
 				}
 
-				// Parameters
+				// FParameters
 				op->dist = TypedNode->Parameters.DistanceToPlane;
 				op->factor = TypedNode->Parameters.LinearityFactor;
 
@@ -2431,7 +2431,7 @@ namespace mu
 				Ptr<ASTOp> fullCondition = m.FinalCondition;
 
 				Ptr<ASTOpConditional> ConditionalOp = new ASTOpConditional();
-				ConditionalOp->type = OP_TYPE::ME_CONDITIONAL;
+				ConditionalOp->type = EOpType::ME_CONDITIONAL;
 				ConditionalOp->no = LastMeshOp;
 				ConditionalOp->yes = modifiedMeshOp;
 				ConditionalOp->condition = fullCondition;
@@ -2482,7 +2482,7 @@ namespace mu
 				Ptr<ASTOp> FullCondition = M.FinalCondition;
 
 				Ptr<ASTOpConditional> Op = new ASTOpConditional();
-				Op->type = OP_TYPE::ME_CONDITIONAL;
+				Op->type = EOpType::ME_CONDITIONAL;
 				Op->no = LastMeshOp;
 				Op->yes = ModifiedMeshOp;
 				Op->condition = FullCondition;
@@ -2513,7 +2513,7 @@ namespace mu
 					
 					if (TypedTransformNode->BoundingMesh)
 					{
-						// Parameters
+						// FParameters
 						FMeshGenerationOptions MeshOptions(Options.ComponentId);
 						MeshOptions.bLayouts = false;
 						MeshOptions.State = Options.State;
@@ -2533,7 +2533,7 @@ namespace mu
 					if (m.FinalCondition)
 					{
 						Ptr<ASTOpConditional> conditionalAd = new ASTOpConditional();
-						conditionalAd->type = OP_TYPE::ME_CONDITIONAL;
+						conditionalAd->type = EOpType::ME_CONDITIONAL;
 						conditionalAd->no = LastMeshOp;
 						conditionalAd->yes = transformOp;
 						conditionalAd->condition = m.FinalCondition;
@@ -2593,7 +2593,7 @@ namespace mu
 				if (MatchingEdit->PatchImage.get())
 				{
 					// Does the current block need to be patched? Find out by building a mask.
-					Ptr<Image> PatchMask = GenerateImageBlockPatchMask(*MatchingEdit, GridSize, LayoutBlockDesc.BlockPixelsX, LayoutBlockDesc.BlockPixelsY, RectInCells);
+					TSharedPtr<FImage> PatchMask = GenerateImageBlockPatchMask(*MatchingEdit, GridSize, LayoutBlockDesc.BlockPixelsX, LayoutBlockDesc.BlockPixelsY, RectInCells);
 
 					if (PatchMask)
 					{
@@ -2690,7 +2690,7 @@ namespace mu
 				}
 				else
 				{
-					Ptr<const Layout> pExtendLayout = ExtraLayouts[LayoutIndex].Layout;
+					TSharedPtr<const FLayout> pExtendLayout = ExtraLayouts[LayoutIndex].Layout;
 
 					Ptr<ASTOp> lastBase = LastImageOp;
 
@@ -2748,7 +2748,7 @@ namespace mu
 					{
 						Ptr<ASTOp> conditionalAd;
 						Ptr<ASTOpConditional> cop = new ASTOpConditional();
-						cop->type = OP_TYPE::IM_CONDITIONAL;
+						cop->type = EOpType::IM_CONDITIONAL;
 						cop->no = LastImageOp;
 						cop->yes = lastBase;
 						cop->condition = m.FinalCondition;
diff --git a/Engine/Plugins/Mutable/Source/MutableTools/Private/MuT/CodeGenerator.h b/Engine/Plugins/Mutable/Source/MutableTools/Private/MuT/CodeGenerator.h
--- a/Engine/Plugins/Mutable/Source/MutableTools/Private/MuT/CodeGenerator.h
+++ b/Engine/Plugins/Mutable/Source/MutableTools/Private/MuT/CodeGenerator.h
@@ -8,7 +8,6 @@
 #include "MuR/Image.h"
 #include "MuR/Mesh.h"
 #include "MuR/MutableMath.h"
-#include "MuR/MutableMemory.h"
 #include "MuR/Operations.h"
 #include "MuR/Parameters.h"
 #include "MuR/Ptr.h"
@@ -45,7 +44,7 @@
 namespace mu
 {
 	class ASTOpParameter;
-	class Layout;
+	class FLayout;
 	class NodeColourArithmeticOperation;
 	class NodeColourConstant;
 	class NodeColourFromScalars;
@@ -243,13 +242,13 @@ namespace mu
 		*/
 		struct FGeneratedConstantMesh
 		{
-			Ptr<Mesh> Mesh;
+			TSharedPtr<FMesh> Mesh;
 			Ptr<ASTOp> LastMeshOp;
 		};
 		TMap<uint64,TArray<FGeneratedConstantMesh>> GeneratedConstantMeshes;
 
         /** List of image resources for every image formats that have been generated so far as palceholders for missing images. */
-        Ptr<Image> MissingImage[size_t(EImageFormat::Count)];
+        TSharedPtr<FImage> MissingImage[size_t(EImageFormat::Count)];
 
 
         //! List of already used vertex ID groups that must be unique.
@@ -307,7 +306,7 @@ namespace mu
 		/** The key for generated tables is made of the source table and a parameter name. */
 		struct FTableCacheKey
 		{
-			Ptr<const Table> Table;
+			Ptr<const FTable> Table;
 			FString ParameterName;
 
 			friend FORCEINLINE uint32 GetTypeHash(const FTableCacheKey& InKey)
@@ -370,7 +369,7 @@ namespace mu
         Ptr<ASTOp> GenerateMissingBoolCode(const TCHAR* strWhere, bool value, const void* errorContext );
 
         //!
-		template<class NODE_TABLE, ETableColumnType TYPE, OP_TYPE OPTYPE, typename F>
+		template<class NODE_TABLE, ETableColumnType TYPE, EOpType OPTYPE, typename F>
 		Ptr<ASTOp> GenerateTableSwitch( const NODE_TABLE& node, F&& GenerateOption );
 
 
@@ -396,7 +395,7 @@ namespace mu
 
 			/** Layout block that we are trying to generate if any. */
 			uint64 LayoutBlockId = FLayoutBlock::InvalidBlockId;
-			Ptr<const Layout> LayoutToApply;
+			TSharedPtr<const FLayout> LayoutToApply;
 
 			friend FORCEINLINE uint32 GetTypeHash(const FImageGenerationOptions& InKey)
 			{
@@ -405,7 +404,7 @@ namespace mu
 				KeyHash = HashCombineFast(KeyHash, ::GetTypeHash(InKey.ImageLayoutStrategy));
 				KeyHash = HashCombineFast(KeyHash, GetTypeHash(InKey.RectSize));
 				KeyHash = HashCombineFast(KeyHash, ::GetTypeHash(InKey.LayoutBlockId));
-				KeyHash = HashCombineFast(KeyHash, ::GetTypeHash(InKey.LayoutToApply.get()));
+				KeyHash = HashCombineFast(KeyHash, ::GetTypeHash(InKey.LayoutToApply.Get()));
 				return KeyHash;
 			}
 
@@ -472,7 +471,7 @@ namespace mu
 		void GenerateImage_Transform(const FImageGenerationOptions&, FImageGenerationResult&, const NodeImageTransform*);
 
 		//!
-		Ptr<Image> GenerateMissingImage(EImageFormat);
+		TSharedPtr<FImage> GenerateMissingImage(EImageFormat);
 
 		//!
 		Ptr<ASTOp> GenerateMissingImageCode(const TCHAR* strWhere, EImageFormat, const void* errorContext, const FImageGenerationOptions& Options);
@@ -493,10 +492,10 @@ namespace mu
 		Ptr<ASTOp> ApplyTiling(Ptr<ASTOp> Source, UE::Math::TIntVector2<int32> Size, EImageFormat Format);
 
 		/** Generate a layout block-sized image with a mask including all pixels in the blocks defined in the patch node. */
-		Ptr<Image> GenerateImageBlockPatchMask(const NodeModifierSurfaceEdit::FTexture&, FIntPoint GridSize, int32 BlockPixelsX, int32 BlockPixelsY, box<FIntVector2> RectInCells);
+		TSharedPtr<FImage> GenerateImageBlockPatchMask(const NodeModifierSurfaceEdit::FTexture&, FIntPoint GridSize, int32 BlockPixelsX, int32 BlockPixelsY, box<FIntVector2> RectInCells);
 
 		/** Generate all the operations to apply the block patching on top of the BlockOp, and masking with PatchMask. */
-		Ptr<ASTOp> GenerateImageBlockPatch(Ptr<ASTOp> BlockOp, const NodeModifierSurfaceEdit::FTexture&, Ptr<Image> PatchMask, Ptr<ASTOp> ConditionOp, const FImageGenerationOptions&);
+		Ptr<ASTOp> GenerateImageBlockPatch(Ptr<ASTOp> BlockOp, const NodeModifierSurfaceEdit::FTexture&, TSharedPtr<FImage> PatchMask, Ptr<ASTOp> ConditionOp, const FImageGenerationOptions&);
 
         //-----------------------------------------------------------------------------------------
         // Meshes
@@ -504,7 +503,7 @@ namespace mu
 		/** */
 		struct FGeneratedLayout
 		{
-			Ptr<const Layout> Layout;
+			TSharedPtr<const FLayout> Layout;
 			Ptr<const NodeLayout> Source;
 
 			FORCEINLINE bool operator==(const FGeneratedLayout& Other) const
@@ -630,7 +629,7 @@ namespace mu
 
 			FORCEINLINE bool operator==(const FGeneratedLayoutKey& Other) const = default;
 		};
-		TMap<FGeneratedLayoutKey, Ptr<const Layout>> GeneratedLayouts;
+		TMap<FGeneratedLayoutKey, TSharedPtr<const FLayout>> GeneratedLayouts;
 
         void GenerateMesh(const FMeshGenerationOptions&, FMeshGenerationResult& result, const NodeMeshPtrConst&);
         void GenerateMesh_Constant(const FMeshGenerationOptions&, FMeshGenerationResult&, const NodeMeshConstant* );
@@ -651,14 +650,14 @@ namespace mu
 		void GenerateMesh_ClipDeform(const FMeshGenerationOptions&, FMeshGenerationResult&, const NodeMeshClipDeform*);
 
 		//-----------------------------------------------------------------------------------------
-		void PrepareMeshForLayout(const FGeneratedLayout&, Ptr<Mesh>,
+		void PrepareMeshForLayout(const FGeneratedLayout&, TSharedPtr<FMesh>,
 			int32 currentLayoutChannel,
 			const void* errorContext,
 			const FMeshGenerationOptions&,
 			bool bUseAbsoluteBlockIds);
 
 		//!
-		Ptr<const Layout> GenerateLayout(Ptr<const NodeLayout> SourceLayout, uint32 MeshIDPrefix);
+		TSharedPtr<const FLayout> GenerateLayout(Ptr<const NodeLayout> SourceLayout, uint32 MeshIDPrefix);
 
 		struct FExtensionDataGenerationResult
 		{
@@ -806,7 +805,7 @@ namespace mu
         void GenerateSurface( FSurfaceGenerationResult&, const FSurfaceGenerationOptions&, Ptr<const NodeSurfaceNew> );
 
 		//-----------------------------------------------------------------------------------------
-		// Default Table Parameters
+		// Default Table FParameters
 		Ptr<ASTOp> GenerateDefaultTableValue(ETableColumnType NodeType);
 
 
@@ -866,10 +865,10 @@ namespace mu
 
 	
     //---------------------------------------------------------------------------------------------
-    template<class NODE_TABLE, ETableColumnType TYPE, OP_TYPE OPTYPE, typename F>
+    template<class NODE_TABLE, ETableColumnType TYPE, EOpType OPTYPE, typename F>
     Ptr<ASTOp> CodeGenerator::GenerateTableSwitch( const NODE_TABLE& node, F&& GenerateOption )
     {
-        Ptr<const Table> NodeTable = node.Table;
+        Ptr<const FTable> NodeTable = node.Table;
         Ptr<ASTOp> Variable;
 
 		FTableCacheKey CacheKey = FTableCacheKey{ node.Table, node.ParameterName };
diff --git a/Engine/Plugins/Mutable/Source/MutableTools/Private/MuT/CodeGenerator_Bool.cpp b/Engine/Plugins/Mutable/Source/MutableTools/Private/MuT/CodeGenerator_Bool.cpp
--- a/Engine/Plugins/Mutable/Source/MutableTools/Private/MuT/CodeGenerator_Bool.cpp
+++ b/Engine/Plugins/Mutable/Source/MutableTools/Private/MuT/CodeGenerator_Bool.cpp
@@ -75,7 +75,7 @@ namespace mu
 		const NodeBoolConstant::Private& node = *Typed->GetPrivate();
 
 		Ptr<ASTOpConstantBool> op = new ASTOpConstantBool();
-		op->value = node.m_value;
+		op->value = node.Value;
 
 		result.op = op;
 	}
@@ -91,21 +91,21 @@ namespace mu
 		if (!it)
 		{
 			FParameterDesc param;
-			param.m_name = node.m_name;
-			const TCHAR* CStr = ToCStr(node.m_uid);
-			param.m_uid.ImportTextItem(CStr, 0, nullptr, nullptr);
-			param.m_type = EParameterType::Bool;
-			param.m_defaultValue.Set<FParamBoolType>(node.m_defaultValue);
+			param.Name = node.Name;
+			const TCHAR* CStr = ToCStr(node.UID);
+			param.UID.ImportTextItem(CStr, 0, nullptr, nullptr);
+			param.Type = EParameterType::Bool;
+			param.DefaultValue.Set<FParamBoolType>(node.DefaultValue);
 
 			op = new ASTOpParameter();
-			op->type = OP_TYPE::BO_PARAMETER;
+			op->type = EOpType::BO_PARAMETER;
 			op->parameter = param;
 
 			// Generate the code for the ranges
-			for (int32 a = 0; a < node.m_ranges.Num(); ++a)
+			for (int32 a = 0; a < node.Ranges.Num(); ++a)
 			{
 				FRangeGenerationResult rangeResult;
-				GenerateRange(rangeResult, Options, node.m_ranges[a]);
+				GenerateRange(rangeResult, Options, node.Ranges[a]);
 				op->ranges.Emplace(op.get(), rangeResult.sizeOp, rangeResult.rangeName, rangeResult.rangeUID);
 			}
 
@@ -125,7 +125,7 @@ namespace mu
 		const NodeBoolNot::Private& node = *Typed->GetPrivate();
 
 		Ptr<ASTOpFixed> op = new ASTOpFixed();
-		op->op.type = OP_TYPE::BO_NOT;
+		op->op.type = EOpType::BO_NOT;
 
 		FBoolGenerationResult ChildResult;
 		GenerateBool(ChildResult, Options, node.m_pSource);
@@ -140,7 +140,7 @@ namespace mu
 		const NodeBoolAnd::Private& node = *Typed->GetPrivate();
 
 		Ptr<ASTOpFixed> op = new ASTOpFixed();
-		op->op.type = OP_TYPE::BO_AND;
+		op->op.type = EOpType::BO_AND;
 
 		FBoolGenerationResult ChildResult;
 		GenerateBool(ChildResult, Options, node.m_pA);
diff --git a/Engine/Plugins/Mutable/Source/MutableTools/Private/MuT/CodeGenerator_Color.cpp b/Engine/Plugins/Mutable/Source/MutableTools/Private/MuT/CodeGenerator_Color.cpp
--- a/Engine/Plugins/Mutable/Source/MutableTools/Private/MuT/CodeGenerator_Color.cpp
+++ b/Engine/Plugins/Mutable/Source/MutableTools/Private/MuT/CodeGenerator_Color.cpp
@@ -122,7 +122,7 @@ namespace mu
 		const NodeColourConstant& node = *Typed;
 
 		Ptr<ASTOpFixed> op = new ASTOpFixed();
-		op->op.type = OP_TYPE::CO_CONSTANT;
+		op->op.type = EOpType::CO_CONSTANT;
 		op->op.args.ColourConstant.value[0] = node.Value.X;
 		op->op.args.ColourConstant.value[1] = node.Value.Y;
 		op->op.args.ColourConstant.value[2] = node.Value.Z;
@@ -144,10 +144,10 @@ namespace mu
 		if (!it)
 		{
 			FParameterDesc param;
-			param.m_name = node.Name;
+			param.Name = node.Name;
 			const TCHAR* CStr = ToCStr(node.Uid);
-			param.m_uid.ImportTextItem(CStr, 0, nullptr, nullptr);
-			param.m_type = EParameterType::Color;
+			param.UID.ImportTextItem(CStr, 0, nullptr, nullptr);
+			param.Type = EParameterType::Color;
 
 			FParamColorType Value;
 			Value[0] =  node.DefaultValue[0];
@@ -155,10 +155,10 @@ namespace mu
 			Value[2] = node.DefaultValue[2];
 			Value[3] = node.DefaultValue[3];
 
-			param.m_defaultValue.Set<FParamColorType>(Value);
+			param.DefaultValue.Set<FParamColorType>(Value);
 			
 			op = new ASTOpParameter();
-			op->type = OP_TYPE::CO_PARAMETER;
+			op->type = EOpType::CO_PARAMETER;
 			op->parameter = param;
 
 			// Generate the code for the ranges
@@ -196,7 +196,7 @@ namespace mu
 		}
 
 		Ptr<ASTOpSwitch> op = new ASTOpSwitch();
-		op->type = OP_TYPE::CO_SWITCH;
+		op->type = EOpType::CO_SWITCH;
 
 		// Variable value
 		if (node.Parameter)
@@ -279,7 +279,7 @@ namespace mu
 
 
 			Ptr<ASTOpConditional> conditional = new ASTOpConditional;
-			conditional->type = OP_TYPE::CO_CONDITIONAL;
+			conditional->type = EOpType::CO_CONDITIONAL;
 			conditional->no = currentOp;
 			conditional->yes = variationOp;
 			conditional->condition = FirstPass.Tags[tagIndex].GenericCondition;
@@ -298,7 +298,7 @@ namespace mu
 
 		// Generate the code
 		Ptr<ASTOpFixed> op = new ASTOpFixed();
-		op->op.type = OP_TYPE::CO_SAMPLEIMAGE;
+		op->op.type = EOpType::CO_SAMPLEIMAGE;
 
 		// Source image
 		int32 ComponentId = -1; // TODO.
@@ -368,7 +368,7 @@ namespace mu
 		const NodeColourFromScalars& node = *Typed;
 
 		Ptr<ASTOpFixed> op = new ASTOpFixed();
-		op->op.type = OP_TYPE::CO_FROMSCALARS;
+		op->op.type = EOpType::CO_FROMSCALARS;
 
 		FScalarGenerationResult ChildResult;
 
@@ -438,7 +438,7 @@ namespace mu
 		const NodeColourArithmeticOperation& node = *Typed;
 
 		Ptr<ASTOpFixed> op = new ASTOpFixed();
-		op->op.type = OP_TYPE::CO_ARITHMETIC;
+		op->op.type = EOpType::CO_ARITHMETIC;
 
 		switch (node.Operation)
 		{
@@ -487,7 +487,7 @@ namespace mu
 	{
 		const NodeColourTable& node = *Typed;
 
-		result.op = GenerateTableSwitch<NodeColourTable, ETableColumnType::Color, OP_TYPE::CO_SWITCH>(node,
+		result.op = GenerateTableSwitch<NodeColourTable, ETableColumnType::Color, EOpType::CO_SWITCH>(node,
 			[this, &Options](const NodeColourTable& node, int32 colIndex, int32 row, mu::ErrorLog* pErrorLog)
 			{
 				Ptr<NodeColourConstant> CellData = new NodeColourConstant();
diff --git a/Engine/Plugins/Mutable/Source/MutableTools/Private/MuT/CodeGenerator_ExtensionData.cpp b/Engine/Plugins/Mutable/Source/MutableTools/Private/MuT/CodeGenerator_ExtensionData.cpp
--- a/Engine/Plugins/Mutable/Source/MutableTools/Private/MuT/CodeGenerator_ExtensionData.cpp
+++ b/Engine/Plugins/Mutable/Source/MutableTools/Private/MuT/CodeGenerator_ExtensionData.cpp
@@ -68,11 +68,11 @@ class Node;
 		Ptr<ASTOpConstantExtensionData> Op = new ASTOpConstantExtensionData();
 		OutResult.Op = Op;
 
-		ExtensionDataPtrConst Data = Constant->Value;
+		TSharedPtr<const FExtensionData> Data = Constant->Value;
 		if (!Data)
 		{
 			// Data can't be null, so make an empty one
-			Data = new ExtensionData();
+			Data = MakeShared<FExtensionData>();
 			
 			// Log an error message
 			ErrorLog->GetPrivate()->Add("Constant extension data not set", ELMT_WARNING, Constant->GetMessageContext());
@@ -94,7 +94,7 @@ class Node;
 		}
 
 		Ptr<ASTOpSwitch> Op = new ASTOpSwitch;
-		Op->type = OP_TYPE::ED_SWITCH;
+		Op->type = EOpType::ED_SWITCH;
 
 		// Variable name
 		if (Switch->Parameter)
@@ -176,7 +176,7 @@ class Node;
 			}
 
 			Ptr<ASTOpConditional> Conditional = new ASTOpConditional;
-			Conditional->type = OP_TYPE::ED_CONDITIONAL;
+			Conditional->type = EOpType::ED_CONDITIONAL;
 			Conditional->no = CurrentOp;
 			Conditional->yes = VariationOp;
 			Conditional->condition = FirstPass.Tags[TagIndex].GenericCondition;
diff --git a/Engine/Plugins/Mutable/Source/MutableTools/Private/MuT/CodeGenerator_FirstPass.cpp b/Engine/Plugins/Mutable/Source/MutableTools/Private/MuT/CodeGenerator_FirstPass.cpp
--- a/Engine/Plugins/Mutable/Source/MutableTools/Private/MuT/CodeGenerator_FirstPass.cpp
+++ b/Engine/Plugins/Mutable/Source/MutableTools/Private/MuT/CodeGenerator_FirstPass.cpp
@@ -383,7 +383,7 @@ namespace mu
 		{
 			// Create a comparison operation as the boolean parameter for the child
 			Ptr<ASTOpFixed> ParamOp = new ASTOpFixed();
-			ParamOp->op.type = OP_TYPE::BO_EQUAL_INT_CONST;
+			ParamOp->op.type = EOpType::BO_EQUAL_INT_CONST;
 			ParamOp->SetChild(ParamOp->op.args.BoolEqualScalarConst.value, ScalarResult.op);
 			ParamOp->op.args.BoolEqualScalarConst.constant = (int16)t;
 
@@ -391,7 +391,7 @@ namespace mu
 			if (CurrentCondition.Last().ObjectCondition)
 			{
 				Ptr<ASTOpFixed> op = new ASTOpFixed();
-				op->op.type = OP_TYPE::BO_AND;
+				op->op.type = EOpType::BO_AND;
 				op->SetChild(op->op.args.BoolBinary.a, CurrentCondition.Last().ObjectCondition);
 				op->SetChild(op->op.args.BoolBinary.b, ParamOp);
 				ParamOp = op;
@@ -515,7 +515,7 @@ namespace mu
 		{
 			// Create a comparison operation as the boolean parameter for the child
 			Ptr<ASTOpFixed> ParamOp = new ASTOpFixed();
-			ParamOp->op.type = OP_TYPE::BO_EQUAL_INT_CONST;
+			ParamOp->op.type = EOpType::BO_EQUAL_INT_CONST;
 			ParamOp->SetChild(ParamOp->op.args.BoolEqualScalarConst.value, ScalarResult.op);
 			ParamOp->op.args.BoolEqualScalarConst.constant = (int16)t;
 
@@ -523,7 +523,7 @@ namespace mu
 			if (CurrentCondition.Last().ObjectCondition)
 			{
 				Ptr<ASTOpFixed> op = new ASTOpFixed();
-				op->op.type = OP_TYPE::BO_AND;
+				op->op.type = EOpType::BO_AND;
 				op->SetChild(op->op.args.BoolBinary.a, CurrentCondition.Last().ObjectCondition);
 				op->SetChild(op->op.args.BoolBinary.b, ParamOp);
 				ParamOp = op;
@@ -615,24 +615,24 @@ namespace mu
 
 		// Prepare the enumeration parameter if necessary
         Ptr<ASTOpParameter> enumOp;
-        if (Private->m_type==NodeObjectGroup::CS_ALWAYS_ONE ||
-			Private->m_type==NodeObjectGroup::CS_ONE_OR_NONE )
+        if (Private->Type==NodeObjectGroup::CS_ALWAYS_ONE ||
+			Private->Type==NodeObjectGroup::CS_ONE_OR_NONE )
         {
             Ptr<ASTOpParameter> op = new ASTOpParameter();
-            op->type = OP_TYPE::NU_PARAMETER;
+            op->type = EOpType::NU_PARAMETER;
 
-            op->parameter.m_name = Private->Name;
+            op->parameter.Name = Private->Name;
 			const TCHAR* CStr = ToCStr(Private->Uid);
-			op->parameter.m_uid.ImportTextItem(CStr, 0, nullptr, nullptr);
-            op->parameter.m_type = EParameterType::Int;
-            op->parameter.m_defaultValue.Set<FParamIntType>(Private->DefaultValue);
+			op->parameter.UID.ImportTextItem(CStr, 0, nullptr, nullptr);
+            op->parameter.Type = EParameterType::Int;
+            op->parameter.DefaultValue.Set<FParamIntType>(Private->DefaultValue);
 
-            if (Private->m_type==NodeObjectGroup::CS_ONE_OR_NONE )
+            if (Private->Type==NodeObjectGroup::CS_ONE_OR_NONE )
             {
                 FParameterDesc::FIntValueDesc nullValue;
-                nullValue.m_value = -1;
-                nullValue.m_name = "None";
-                op->parameter.m_possibleValues.Add( nullValue );
+                nullValue.Value = -1;
+                nullValue.Name = "None";
+                op->parameter.PossibleValues.Add( nullValue );
             }
 
 			ParameterNodes.Add(InNode, op);
@@ -648,7 +648,7 @@ namespace mu
             {
                 // Overwrite the implicit condition
                 Ptr<ASTOp> paramOp = 0;
-                switch (Private->m_type )
+                switch (Private->Type )
                 {
                     case NodeObjectGroup::CS_TOGGLE_EACH:
                     {           
@@ -662,13 +662,13 @@ namespace mu
 						{
 							// Create a new boolean parameter
 							Ptr<ASTOpParameter> op = new ASTOpParameter();
-							op->type = OP_TYPE::BO_PARAMETER;
+							op->type = EOpType::BO_PARAMETER;
 
-							op->parameter.m_name = pChildNode->GetName();
+							op->parameter.Name = pChildNode->GetName();
 							const TCHAR* CStr = ToCStr(pChildNode->GetUid());
-							op->parameter.m_uid.ImportTextItem(CStr, 0, nullptr, nullptr);
-							op->parameter.m_type = EParameterType::Bool;
-							op->parameter.m_defaultValue.Set<FParamBoolType>(false);
+							op->parameter.UID.ImportTextItem(CStr, 0, nullptr, nullptr);
+							op->parameter.Type = EParameterType::Bool;
+							op->parameter.DefaultValue.Set<FParamBoolType>(false);
 
 							ParameterNodes.Add(pChildNode, op);
 
@@ -693,15 +693,15 @@ namespace mu
                     {
                         // Add the option to the enumeration parameter
                         FParameterDesc::FIntValueDesc value;
-                        value.m_value = (int16)t;
-                        value.m_name = pChildNode->GetName();
-                        enumOp->parameter.m_possibleValues.Add( value );
+                        value.Value = (int16)t;
+                        value.Name = pChildNode->GetName();
+                        enumOp->parameter.PossibleValues.Add( value );
 
                         check(enumOp);
 
                         // Create a comparison operation as the boolean parameter for the child
                         Ptr<ASTOpFixed> op = new ASTOpFixed();
-                        op->op.type = OP_TYPE::BO_EQUAL_INT_CONST;
+                        op->op.type = EOpType::BO_EQUAL_INT_CONST;
                         op->SetChild( op->op.args.BoolEqualScalarConst.value, enumOp);
                         op->op.args.BoolEqualScalarConst.constant = (int16_t)t;
 
@@ -717,7 +717,7 @@ namespace mu
                 if (CurrentCondition.Last().ObjectCondition)
                 {
                     Ptr<ASTOpFixed> op = new ASTOpFixed();
-                    op->op.type = OP_TYPE::BO_AND;
+                    op->op.type = EOpType::BO_AND;
                     op->SetChild( op->op.args.BoolBinary.a,CurrentCondition.Last().ObjectCondition);
                     op->SetChild( op->op.args.BoolBinary.b,paramOp);
                     paramOp = op;
diff --git a/Engine/Plugins/Mutable/Source/MutableTools/Private/MuT/CodeGenerator_FirstPass.h b/Engine/Plugins/Mutable/Source/MutableTools/Private/MuT/CodeGenerator_FirstPass.h
--- a/Engine/Plugins/Mutable/Source/MutableTools/Private/MuT/CodeGenerator_FirstPass.h
+++ b/Engine/Plugins/Mutable/Source/MutableTools/Private/MuT/CodeGenerator_FirstPass.h
@@ -2,7 +2,6 @@
 
 #pragma once
 
-#include "MuR/MutableMemory.h"
 #include "MuR/Ptr.h"
 #include "MuR/RefCounted.h"
 #include "MuT/AST.h"
@@ -29,7 +28,7 @@
 
 namespace mu
 {
-	class Layout;
+	class FLayout;
 	struct FObjectState;
 
 	/** First pass of the code generation process.
@@ -189,7 +188,7 @@ namespace mu
         typedef TArray< TPair<FObjectState, const Node*> > StateList;
         StateList States;
 
-		/** Parameters added for every node. */
+		/** FParameters added for every node. */
 		TMap< Ptr<const Node>, Ptr<ASTOpParameter> > ParameterNodes;
 
 	private:
diff --git a/Engine/Plugins/Mutable/Source/MutableTools/Private/MuT/CodeGenerator_Image.cpp b/Engine/Plugins/Mutable/Source/MutableTools/Private/MuT/CodeGenerator_Image.cpp
--- a/Engine/Plugins/Mutable/Source/MutableTools/Private/MuT/CodeGenerator_Image.cpp
+++ b/Engine/Plugins/Mutable/Source/MutableTools/Private/MuT/CodeGenerator_Image.cpp
@@ -170,7 +170,7 @@ namespace mu
 		const NodeImageConstant& node = *InNode;
 		
         // TODO: check duplicates
-        Ptr<const Image> pImage;
+        TSharedPtr<const FImage> pImage;
 		if (node.Proxy)
 		{
 			pImage = node.Proxy->Get();
@@ -189,7 +189,7 @@ namespace mu
 		if (pImage->IsReference())
 		{
 			Ptr<ASTOpReferenceResource> ReferenceOp = new ASTOpReferenceResource();
-			ReferenceOp->type = OP_TYPE::IM_REFERENCE;
+			ReferenceOp->type = EOpType::IM_REFERENCE;
 			ReferenceOp->ID = pImage->GetReferencedTexture();
 			ReferenceOp->bForceLoad = pImage->IsForceLoad();
 			ReferenceOp->SourceDataDescriptor = InNode->SourceDataDescriptor;
@@ -206,7 +206,7 @@ namespace mu
 		else
 		{
 			Ptr<ASTOpConstantResource> op = new ASTOpConstantResource();
-			op->Type = OP_TYPE::IM_CONSTANT;
+			op->Type = EOpType::IM_CONSTANT;
 			op->SetValue(pImage, CompilerOptions->OptimisationOptions.DiskCacheContext);
 			op->SourceDataDescriptor = InNode->SourceDataDescriptor;
 			Result.op = op;
@@ -281,19 +281,19 @@ namespace mu
         if ( !it )
         {
             op = new ASTOpParameter();
-            op->type = OP_TYPE::IM_PARAMETER;
+            op->type = EOpType::IM_PARAMETER;
 
-			op->parameter.m_name = node.m_name;
-        	const TCHAR* CStr = ToCStr(node.m_uid);
-        	op->parameter.m_uid.ImportTextItem(CStr, 0, nullptr, nullptr);
-        	op->parameter.m_type = EParameterType::Image;
-        	op->parameter.m_defaultValue.Set<FParamImageType>(node.m_defaultValue);
+			op->parameter.Name = node.Name;
+        	const TCHAR* CStr = ToCStr(node.UID);
+        	op->parameter.UID.ImportTextItem(CStr, 0, nullptr, nullptr);
+        	op->parameter.Type = EParameterType::Image;
+        	op->parameter.DefaultValue.Set<FParamImageType>(node.DefaultValue);
 
 			// Generate the code for the ranges
-			for (int32 a = 0; a < node.m_ranges.Num(); ++a)
+			for (int32 a = 0; a < node.Ranges.Num(); ++a)
 			{
 				FRangeGenerationResult rangeResult;
-				GenerateRange(rangeResult, Options, node.m_ranges[a]);
+				GenerateRange(rangeResult, Options, node.Ranges[a]);
 				op->ranges.Emplace(op.get(), rangeResult.sizeOp, rangeResult.rangeName, rangeResult.rangeUID);
 			}
 
@@ -317,7 +317,7 @@ namespace mu
 
         Ptr<ASTOpImageLayer> op = new ASTOpImageLayer();
 
-        op->blendType = node.m_type;
+        op->blendType = node.Type;
 
         // Base image
         Ptr<ASTOp> base;
@@ -393,7 +393,7 @@ namespace mu
 		MUTABLE_CPUPROFILER_SCOPE(NodeImageLayerColour);
 
         Ptr<ASTOpImageLayerColor> op = new ASTOpImageLayerColor();
-		op->blendType = node.m_type;
+		op->blendType = node.Type;
 
         // Base image
         Ptr<ASTOp> base;
@@ -463,7 +463,7 @@ namespace mu
 
         Ptr<ASTOpImageMultiLayer> op = new ASTOpImageMultiLayer();
 
-		op->blendType = node.m_type;
+		op->blendType = node.Type;
 
         // Base image
         Ptr<ASTOp> base;
@@ -697,7 +697,7 @@ namespace mu
 		MUTABLE_CPUPROFILER_SCOPE(NodeImageInterpolate);
 
         Ptr<ASTOpFixed> op = new ASTOpFixed();
-        op->op.type = OP_TYPE::IM_INTERPOLATE;
+        op->op.type = EOpType::IM_INTERPOLATE;
 
         // Factor
         if ( Node* pFactor = node.m_pFactor.get() )
@@ -911,7 +911,7 @@ namespace mu
 		const NodeImageSaturate::Private& node = *InNode->GetPrivate();
 
 		Ptr<ASTOpFixed> op = new ASTOpFixed();
-        op->op.type = OP_TYPE::IM_SATURATE;
+        op->op.type = EOpType::IM_SATURATE;
 
 
         // Source image
@@ -954,7 +954,7 @@ namespace mu
 		const NodeImageLuminance::Private& node = *InNode->GetPrivate();
 
 		Ptr<ASTOpFixed> op = new ASTOpFixed();
-        op->op.type = OP_TYPE::IM_LUMINANCE;
+        op->op.type = EOpType::IM_LUMINANCE;
 
         // Source image
         Ptr<ASTOp> base;
@@ -982,7 +982,7 @@ namespace mu
 		const NodeImageColourMap::Private& node = *InNode->GetPrivate();
 
 		Ptr<ASTOpFixed> op = new ASTOpFixed();
-        op->op.type = OP_TYPE::IM_COLOURMAP;
+        op->op.type = EOpType::IM_COLOURMAP;
 
         // Base image
         Ptr<ASTOp> base ;
@@ -1050,7 +1050,7 @@ namespace mu
 		const NodeImageGradient::Private& node = *InNode->GetPrivate();
 
 		Ptr<ASTOpFixed> op = new ASTOpFixed();
-        op->op.type = OP_TYPE::IM_GRADIENT;
+        op->op.type = EOpType::IM_GRADIENT;
 
         // First colour
         Ptr<ASTOp> colour0 = 0;
@@ -1091,7 +1091,7 @@ namespace mu
 		const NodeImageBinarise::Private& node = *InNode->GetPrivate();
 
 		Ptr<ASTOpFixed> op = new ASTOpFixed();
-        op->op.type = OP_TYPE::IM_BINARISE;
+        op->op.type = EOpType::IM_BINARISE;
 
         // A image
         Ptr<ASTOp> a;
@@ -1162,7 +1162,7 @@ namespace mu
         if ( node.m_relative )
         {
             Ptr<ASTOpFixed> op = new ASTOpFixed();
-            op->op.type = OP_TYPE::IM_RESIZEREL;
+            op->op.type = EOpType::IM_RESIZEREL;
             op->op.args.ImageResizeRel.factor[0] = node.m_sizeX;
             op->op.args.ImageResizeRel.factor[1] = node.m_sizeY;
             op->SetChild( op->op.args.ImageResizeRel.source, base);
@@ -1192,7 +1192,7 @@ namespace mu
 			}
 
             Ptr<ASTOpFixed> op = new ASTOpFixed();
-            op->op.type = OP_TYPE::IM_RESIZE;
+            op->op.type = EOpType::IM_RESIZE;
             op->op.args.ImageResize.size[0] = (uint16)FinalImageSize[0];
             op->op.args.ImageResize.size[1] = (uint16)FinalImageSize[1];
             op->SetChild( op->op.args.ImageResize.source, base);
@@ -1253,7 +1253,7 @@ namespace mu
 
 
         Ptr<ASTOpFixed> op = new ASTOpFixed();
-        op->op.type = OP_TYPE::IM_PLAINCOLOUR;
+        op->op.type = EOpType::IM_PLAINCOLOUR;
         op->SetChild( op->op.args.ImagePlainColour.colour, base);
 		op->op.args.ImagePlainColour.format = node.Format;
 		op->op.args.ImagePlainColour.size[0] = FinalImageSize[0];
@@ -1271,7 +1271,7 @@ namespace mu
 
 		MUTABLE_CPUPROFILER_SCOPE(NodeImageSwitch);
 
-        if (node.m_options.Num() == 0)
+        if (node.Options.Num() == 0)
 		{
 			// No options in the switch!
             Ptr<ASTOp> missingOp = GenerateMissingImageCode(TEXT("Switch option"), EImageFormat::RGBA_UByte, InNode->GetMessageContext(), Options);
@@ -1280,7 +1280,7 @@ namespace mu
 		}
 
         Ptr<ASTOpSwitch> op = new ASTOpSwitch();
-        op->type = OP_TYPE::IM_SWITCH;
+        op->type = EOpType::IM_SWITCH;
 
 		// Variable value
 		if ( node.m_pParameter )
@@ -1294,14 +1294,14 @@ namespace mu
 		}
 
 		// Options
-        for ( std::size_t t=0; t< node.m_options.Num(); ++t )
+        for ( std::size_t t=0; t< node.Options.Num(); ++t )
         {
             Ptr<ASTOp> branch;
 
-            if (node.m_options[t])
+            if (node.Options[t])
             {
 				FImageGenerationResult BaseResult;
-				GenerateImage(Options, BaseResult, node.m_options[t]);
+				GenerateImage(Options, BaseResult, node.Options[t]);
 				branch = BaseResult.op;
 			}
             else
@@ -1329,7 +1329,7 @@ namespace mu
         //if (desc.m_size[0]!=0 && desc.m_size[1]!=0)
         //{
         //    Ptr<ASTOpFixed> sop = new ASTOpFixed();
-        //    sop->op.type = OP_TYPE::IM_RESIZE;
+        //    sop->op.type = EOpType::IM_RESIZE;
         //    sop->op.args.ImageResize.size[0] = desc.m_size[0];
         //    sop->op.args.ImageResize.size[1] = desc.m_size[1];
         //    sop->SetChild( sop->op.args.ImageResize.source, switchAt );
@@ -1338,7 +1338,7 @@ namespace mu
 
         //{
         //    Ptr<ASTOpFixed> fop = new ASTOpFixed();
-        //    fop->op.type = OP_TYPE::IM_PIXELFORMAT;
+        //    fop->op.type = EOpType::IM_PIXELFORMAT;
         //    fop->op.args.ImagePixelFormat.format = desc.m_format;
         //    fop->SetChild( fop->op.args.ImagePixelFormat.source, switchAt );
         //    switchAt = fop;
@@ -1354,12 +1354,12 @@ namespace mu
 		const NodeImageConditional::Private& node = *InNode->GetPrivate();
 
 		Ptr<ASTOpConditional> op = new ASTOpConditional();
-        op->type = OP_TYPE::IM_CONDITIONAL;
+        op->type = EOpType::IM_CONDITIONAL;
 
         // Condition
-        if ( node.m_parameter )
+        if ( node.Parameter )
         {
-            op->condition = Generate_Generic( node.m_parameter.get(), Options);
+            op->condition = Generate_Generic( node.Parameter.get(), Options);
         }
         else
         {
@@ -1390,7 +1390,7 @@ namespace mu
         // Mesh project operation
         //------------------------------
         Ptr<ASTOpFixed> ProjectOp = new ASTOpFixed();
-		ProjectOp->op.type = OP_TYPE::ME_PROJECT;
+		ProjectOp->op.type = EOpType::ME_PROJECT;
 
         Ptr<ASTOp> LastMeshOp = ProjectOp;
 
@@ -1431,7 +1431,7 @@ namespace mu
 			GenerateMesh( MeshOptions, MeshResult, node.m_pMesh );
 
 			// Match the block id of the block we are generating with the id that resulted in the generated mesh			
-			Ptr<const Layout> Layout = MeshResult.GeneratedLayouts.IsValidIndex(node.m_layout) ? MeshResult.GeneratedLayouts[node.m_layout].Layout : nullptr;
+			TSharedPtr<const FLayout> Layout = MeshResult.GeneratedLayouts.IsValidIndex(node.m_layout) ? MeshResult.GeneratedLayouts[node.m_layout].Layout : nullptr;
 			if (Layout && Layout->Blocks.IsValidIndex(LayoutBlockIndex))
 			{
 				GeneratedLayoutBlockId = Layout->Blocks[LayoutBlockIndex].Id;
@@ -1459,10 +1459,10 @@ namespace mu
                 // For wrapping projector we need the entire mesh. The actual project operation
                 // will remove the faces that are not in the layout block we are generating.
                 Ptr<ASTOpConstantResource> cop = new ASTOpConstantResource();
-                cop->Type = OP_TYPE::ME_CONSTANT;
+                cop->Type = EOpType::ME_CONSTANT;
 				cop->SourceDataDescriptor = CurrentMeshToProjectOp->GetSourceDataDescriptor();
-				Ptr<Mesh> FormatMeshResult = new Mesh();
-				CreateMeshOptimisedForWrappingProjection(FormatMeshResult.get(), node.m_layout);
+				TSharedPtr<FMesh> FormatMeshResult = MakeShared<FMesh>();
+				CreateMeshOptimisedForWrappingProjection(FormatMeshResult.Get(), node.m_layout);
 
                 cop->SetValue(FormatMeshResult, CompilerOptions->OptimisationOptions.DiskCacheContext);
 
@@ -1490,11 +1490,11 @@ namespace mu
 
                 // Reformat the mesh to a more efficient format for this operation
                 Ptr<ASTOpConstantResource> cop = new ASTOpConstantResource();
-                cop->Type = OP_TYPE::ME_CONSTANT;
+                cop->Type = EOpType::ME_CONSTANT;
 				cop->SourceDataDescriptor = CurrentMeshToProjectOp->GetSourceDataDescriptor();
 
-				Ptr<Mesh> FormatMeshResult = new Mesh();
-                CreateMeshOptimisedForProjection(FormatMeshResult.get(), node.m_layout);
+				TSharedPtr<FMesh> FormatMeshResult = MakeShared<FMesh>();
+                CreateMeshOptimisedForProjection(FormatMeshResult.Get(), node.m_layout);
 
                 cop->SetValue(FormatMeshResult, CompilerOptions->OptimisationOptions.DiskCacheContext);
 
@@ -1513,9 +1513,9 @@ namespace mu
         else
         {
             // This argument is required
-            Ptr<const Mesh> TempMesh = new Mesh();
+            TSharedPtr<const FMesh> TempMesh = MakeShared<FMesh>();
             Ptr<ASTOpConstantResource> cop = new ASTOpConstantResource();
-            cop->Type = OP_TYPE::ME_CONSTANT;
+            cop->Type = EOpType::ME_CONSTANT;
             cop->SetValue(TempMesh, CompilerOptions->OptimisationOptions.DiskCacheContext);
 			ProjectOp->SetChild(ProjectOp->op.args.MeshProject.mesh, cop );
             ErrorLog->GetPrivate()->Add( "Projector mesh not set.", ELMT_ERROR, InNode->GetMessageContext() );
@@ -1681,7 +1681,7 @@ namespace mu
 			}
 
 			Ptr<ASTOpFixed> DisplaceOp = new ASTOpFixed();
-			DisplaceOp->op.type = OP_TYPE::IM_DISPLACE;
+			DisplaceOp->op.type = EOpType::IM_DISPLACE;
 			DisplaceOp->SetChild(DisplaceOp->op.args.ImageDisplace.displacementMap, MakeGrowMapOp);
 			DisplaceOp->SetChild(DisplaceOp->op.args.ImageDisplace.source, ImageRasterOp);
 
@@ -1764,7 +1764,7 @@ namespace mu
 		const NodeImageInvert::Private& node = *InNode->GetPrivate();
 
 		Ptr<ASTOpFixed> op = new ASTOpFixed();
-		op->op.type = OP_TYPE::IM_INVERT;
+		op->op.type = EOpType::IM_INVERT;
 
 		// A image
 		Ptr<ASTOp> a;
@@ -1838,7 +1838,7 @@ namespace mu
 
 
             Ptr<ASTOpConditional> conditional = new ASTOpConditional;
-            conditional->type = OP_TYPE::IM_CONDITIONAL;
+            conditional->type = EOpType::IM_CONDITIONAL;
             conditional->no = currentOp;
             conditional->yes = variationOp;
             conditional->condition = FirstPass.Tags[tagIndex].GenericCondition;
@@ -1863,7 +1863,7 @@ namespace mu
         if ( desc.m_size[0] != 0 && desc.m_size[1] != 0 )
         {
             Ptr<ASTOpFixed> sop = new ASTOpFixed();
-            sop->op.type = OP_TYPE::IM_RESIZE;
+            sop->op.type = EOpType::IM_RESIZE;
             sop->op.args.ImageResize.size[0] = desc.m_size[0];
             sop->op.args.ImageResize.size[1] = desc.m_size[1];
             sop->SetChild( sop->op.args.ImageResize.source, currentOp );
@@ -1886,13 +1886,13 @@ namespace mu
 	{
 		const NodeImageTable& node = *InNode;
 
-		Result.op = GenerateTableSwitch<NodeImageTable, ETableColumnType::Image, OP_TYPE::IM_SWITCH>(node,
+		Result.op = GenerateTableSwitch<NodeImageTable, ETableColumnType::Image, EOpType::IM_SWITCH>(node,
 			[this, InNode, Options](const NodeImageTable& node, int32 colIndex, int32 row, mu::ErrorLog* pErrorLog)
 			{
 				const FTableValue& CellData = node.Table->GetPrivate()->Rows[row].Values[colIndex];
-				Ptr<const Image> pImage = nullptr;
+				TSharedPtr<const FImage> pImage = nullptr;
 
-				if (Ptr<ResourceProxy<Image>> pProxyImage = CellData.ProxyImage)
+				if (Ptr<TResourceProxy<FImage>> pProxyImage = CellData.ProxyImage)
 				{
 					pImage = pProxyImage->Get();
 				}
@@ -1909,7 +1909,7 @@ namespace mu
 				else
 				{
 					Ptr<NodeImageConstant> ImageConst = new NodeImageConstant();
-					ImageConst->SetValue(pImage.get());
+					ImageConst->SetValue(pImage);
 
 					// TODO: We probably want to get the data tags from the table row.
 					ImageConst->SourceDataDescriptor = InNode->SourceDataDescriptor;
@@ -1930,7 +1930,7 @@ namespace mu
 						// Use a relative resize, because at this point we may be generating a layout block and not the full image
 						float Factor = FMath::Min(MaxTextureSize / (float)(pImage->GetSizeX()), MaxTextureSize / (float)(pImage->GetSizeY()));
 						Ptr<ASTOpFixed> op = new ASTOpFixed();
-						op->op.type = OP_TYPE::IM_RESIZEREL;
+						op->op.type = EOpType::IM_RESIZEREL;
 						op->op.args.ImageResizeRel.factor[0] = Factor;
 						op->op.args.ImageResizeRel.factor[1] = Factor;
 						op->SetChild(op->op.args.ImageResizeRel.source, ImageOp);
@@ -1943,7 +1943,7 @@ namespace mu
 	}
 
 
-	Ptr<Image> CodeGenerator::GenerateMissingImage(EImageFormat Format)
+	TSharedPtr<FImage> CodeGenerator::GenerateMissingImage(EImageFormat Format)
 	{
 		// Create the image node if it hasn't been created yet.
 		if (!MissingImage[SIZE_T(Format)])
@@ -1951,7 +1951,7 @@ namespace mu
 			// Make a checkered debug image
 			const FImageSize Size(16, 16);
 
-			Ptr<Image> GeneratedImage = new Image(Size[0], Size[1], 1, Format, EInitializationType::NotInitialized);
+			TSharedPtr<FImage> GeneratedImage = MakeShared<FImage>(Size[0], Size[1], 1, Format, EInitializationType::NotInitialized);
 
 			switch (Format)
 			{
@@ -2030,7 +2030,7 @@ namespace mu
 			MissingImage[(SIZE_T)Format] = GeneratedImage;
 		}
 
-		return MissingImage[(SIZE_T)Format].get();
+		return MissingImage[(SIZE_T)Format];
 	}
 
 
@@ -2041,10 +2041,10 @@ namespace mu
 		ErrorLog->GetPrivate()->Add( Msg, ELMT_ERROR, errorContext );
 
 		// Make a checkered debug image
-		Ptr<Image> GeneratedImage = GenerateMissingImage( format );
+		TSharedPtr<FImage> GeneratedImage = GenerateMissingImage( format );
 
 		Ptr<NodeImageConstant> pNode = new NodeImageConstant();
-		pNode->SetValue(GeneratedImage.get());
+		pNode->SetValue(GeneratedImage);
 
 		FImageGenerationResult Result;
 		GenerateImage(Options, Result, pNode);
@@ -2122,7 +2122,7 @@ namespace mu
 			if (UE::Math::TIntVector2<int32>(at->GetImageDesc().m_size) != size)
 			{
 				Ptr<ASTOpFixed> op = new ASTOpFixed();
-				op->op.type = OP_TYPE::IM_RESIZE;
+				op->op.type = EOpType::IM_RESIZE;
 				op->SetChild(op->op.args.ImageResize.source, at);
 				op->op.args.ImageResize.size[0] = size[0];
 				op->op.args.ImageResize.size[1] = size[1];
diff --git a/Engine/Plugins/Mutable/Source/MutableTools/Private/MuT/CodeGenerator_Matrix.cpp b/Engine/Plugins/Mutable/Source/MutableTools/Private/MuT/CodeGenerator_Matrix.cpp
--- a/Engine/Plugins/Mutable/Source/MutableTools/Private/MuT/CodeGenerator_Matrix.cpp
+++ b/Engine/Plugins/Mutable/Source/MutableTools/Private/MuT/CodeGenerator_Matrix.cpp
@@ -55,14 +55,14 @@ namespace mu
 		if (!it)
 		{
 			FParameterDesc param;
-			param.m_name = Typed->Name;
+			param.Name = Typed->Name;
 			const TCHAR* CStr = ToCStr(Typed->Uid);
-			param.m_uid.ImportTextItem(CStr, 0, nullptr, nullptr);
-			param.m_type = EParameterType::Matrix;
-			param.m_defaultValue.Set<FParamMatrixType>(Typed->DefaultValue);
+			param.UID.ImportTextItem(CStr, 0, nullptr, nullptr);
+			param.Type = EParameterType::Matrix;
+			param.DefaultValue.Set<FParamMatrixType>(Typed->DefaultValue);
 
 			op = new ASTOpParameter();
-			op->type = OP_TYPE::MA_PARAMETER;
+			op->type = EOpType::MA_PARAMETER;
 			op->parameter = param;
 
 			// Generate the code for the ranges
