[mutable] More enforcement of coding standards.

diff --git a/Engine/Plugins/Mutable/Source/MutableTools/Private/MuT/CodeGenerator_Mesh.cpp b/Engine/Plugins/Mutable/Source/MutableTools/Private/MuT/CodeGenerator_Mesh.cpp
--- a/Engine/Plugins/Mutable/Source/MutableTools/Private/MuT/CodeGenerator_Mesh.cpp
+++ b/Engine/Plugins/Mutable/Source/MutableTools/Private/MuT/CodeGenerator_Mesh.cpp
@@ -108,7 +108,7 @@ namespace mu
 
     //---------------------------------------------------------------------------------------------
 	void CodeGenerator::PrepareMeshForLayout( const FGeneratedLayout& GeneratedLayout,
-		Ptr<Mesh> Mesh,
+		TSharedPtr<FMesh> Mesh,
 		int32 LayoutChannel,
 		const void* errorContext,
 		const FMeshGenerationOptions& MeshOptions,
@@ -126,7 +126,7 @@ namespace mu
 		check(GeneratedLayout.Layout->Blocks.IsEmpty() || GeneratedLayout.Layout->Blocks[0].Id != FLayoutBlock::InvalidBlockId);
 
 		// 
-		Ptr<const Layout> Layout = GeneratedLayout.Layout;
+		TSharedPtr<const FLayout> Layout = GeneratedLayout.Layout;
 		Mesh->AddLayout(Layout);
 
 		const int32 NumVertices = Mesh->GetVertexCount();
@@ -155,7 +155,7 @@ namespace mu
 		TArray<int32> OverlappingBlocks;
 		for (int32 BlockIndex = 0; BlockIndex < NumBlocks; ++BlockIndex)
 		{
-			bool bBlockHasMask = GeneratedLayout.Source->Blocks[BlockIndex].Mask.get() != nullptr;
+			bool bBlockHasMask = GeneratedLayout.Source->Blocks[BlockIndex].Mask != nullptr;
 
 			// Fill the block rect
 			FIntVector2 Min = Layout->Blocks[BlockIndex].Min;
@@ -271,7 +271,7 @@ namespace mu
 			VertexToFaceMap.Reserve(NumVertices);
 			Triangles.SetNumUninitialized(NumTriangles);
 
-			MeshCreateCollapsedVertexMap(Mesh.get(), CollapsedVertices);
+			MeshCreateCollapsedVertexMap(Mesh.Get(), CollapsedVertices);
 		}
 
 		TArray<int32> ConflictiveTriangles;
@@ -314,7 +314,7 @@ namespace mu
 					// First: Assign the vertices to masked blocks in order
 					for (int32 CandidateBlockIndex = 0; CandidateBlockIndex < NumBlocks; ++CandidateBlockIndex)
 					{
-						const mu::Image* Mask = GeneratedLayout.Source->Blocks[CandidateBlockIndex].Mask.get();
+						const mu::FImage* Mask = GeneratedLayout.Source->Blocks[CandidateBlockIndex].Mask.Get();
 						if (Mask)
 						{
 							// First discard with block limits.
@@ -804,7 +804,7 @@ namespace mu
 			if (Node->Layout)
 			{
 				// Generate the layout with blocks to extract
-				Ptr<const Layout> Layout = GenerateLayout(Node->Layout, 0);
+				TSharedPtr<const FLayout> Layout = GenerateLayout(Node->Layout, 0);
 				BaseOptions.OverrideLayouts.Empty();
 				BaseOptions.OverrideLayouts.Add({ Layout, Node->Layout });
 			}
@@ -831,7 +831,7 @@ namespace mu
 
         // Generate the code
         Ptr<ASTOpFixed> op = new ASTOpFixed();
-        op->op.type = OP_TYPE::ME_INTERPOLATE;
+        op->op.type = EOpType::ME_INTERPOLATE;
         OutResult.MeshOp = op;
 
         // Factor
@@ -914,7 +914,7 @@ namespace mu
     {
         NodeMeshSwitch::Private& node = *SwitchNode->GetPrivate();
 
-        if (node.m_options.Num() == 0)
+        if (node.Options.Num() == 0)
         {
             // No options in the switch!
             // TODO
@@ -923,7 +923,7 @@ namespace mu
         }
 
         Ptr<ASTOpSwitch> op = new ASTOpSwitch();
-        op->type = OP_TYPE::ME_SWITCH;
+        op->type = EOpType::ME_SWITCH;
 
         // Factor
         if ( node.m_pParameter )
@@ -938,11 +938,11 @@ namespace mu
 
         // Options
 		bool bFirstValidConnectionFound = false;
-        for ( int32 t=0; t< node.m_options.Num(); ++t )
+        for ( int32 t=0; t< node.Options.Num(); ++t )
         {
 			FMeshGenerationOptions TargetOptions = InOptions;
 
-            if ( node.m_options[t] )
+            if ( node.Options[t] )
             {
 				// Take the layouts from the first non-null connection.
 				// \TODO: Take them from the first connection that actually returns layouts?
@@ -952,7 +952,7 @@ namespace mu
 				}
 				
 				FMeshGenerationResult BranchResults;
-                GenerateMesh(TargetOptions, BranchResults, node.m_options[t] );
+                GenerateMesh(TargetOptions, BranchResults, node.Options[t] );
 
                 Ptr<ASTOp> branch = BranchResults.MeshOp;
                 op->cases.Emplace((int16)t,op,branch);
@@ -977,10 +977,10 @@ namespace mu
 		int32 t = 0;
 		bool bFirstRowGenerated = false;
 
-		Ptr<ASTOp> Op = GenerateTableSwitch<NodeMeshTable, ETableColumnType::Mesh, OP_TYPE::ME_SWITCH>(*TableNode,
+		Ptr<ASTOp> Op = GenerateTableSwitch<NodeMeshTable, ETableColumnType::Mesh, EOpType::ME_SWITCH>(*TableNode,
 			[this, &NewResult, &bFirstRowGenerated, &InOptions] (const NodeMeshTable& node, int32 colIndex, int32 row, mu::ErrorLog* pErrorLog)
 			{
-				mu::Ptr<mu::Mesh> pMesh = node.Table->GetPrivate()->Rows[row].Values[colIndex].Mesh;
+				TSharedPtr<mu::FMesh> pMesh = node.Table->GetPrivate()->Rows[row].Values[colIndex].Mesh;
 				FMeshGenerationResult BranchResults;
 
 				if (pMesh)
@@ -1100,7 +1100,7 @@ namespace mu
             }
 
             Ptr<ASTOpConditional> conditional = new ASTOpConditional;
-            conditional->type = OP_TYPE::ME_CONDITIONAL;
+            conditional->type = EOpType::ME_CONDITIONAL;
             conditional->no = currentMeshOp;
             conditional->yes = variationMeshOp;            
             conditional->condition = FirstPass.Tags[tagIndex].GenericCondition;
@@ -1121,7 +1121,7 @@ namespace mu
         NodeMeshConstant::Private& Node = *InNode->GetPrivate();
 
         Ptr<ASTOpConstantResource> ConstantOp = new ASTOpConstantResource();
-		ConstantOp->Type = OP_TYPE::ME_CONSTANT;
+		ConstantOp->Type = EOpType::ME_CONSTANT;
 		ConstantOp->SourceDataDescriptor = InNode->SourceDataDescriptor;
 		OutResult.BaseMeshOp = ConstantOp;
 		OutResult.MeshOp = ConstantOp;
@@ -1129,11 +1129,11 @@ namespace mu
 
 		bool bIsOverridingLayouts = !InOptions.OverrideLayouts.IsEmpty();
 
-        Ptr<Mesh> pMesh = Node.Value.get();
+        TSharedPtr<FMesh> pMesh = Node.Value;
 		if (!pMesh)
 		{
 			// This data is required
-			MeshPtr EmptyMesh = new Mesh();
+			TSharedPtr<FMesh> EmptyMesh = MakeShared<FMesh>();
 			ConstantOp->SetValue(EmptyMesh, CompilerOptions->OptimisationOptions.DiskCacheContext);
 			EmptyMesh->MeshIDPrefix = ConstantOp->GetValueHash();
 
@@ -1146,7 +1146,7 @@ namespace mu
 		if (pMesh->IsReference())
 		{
 			Ptr<ASTOpReferenceResource> ReferenceOp = new ASTOpReferenceResource();
-			ReferenceOp->type = OP_TYPE::ME_REFERENCE;
+			ReferenceOp->type = EOpType::ME_REFERENCE;
 			ReferenceOp->ID = pMesh->GetReferencedMesh();
 			ReferenceOp->bForceLoad = pMesh->IsForceLoad();
 
@@ -1160,7 +1160,7 @@ namespace mu
 		TArray<FString> Tags = pMesh->Tags;
 		if (Tags.Num())
 		{
-			Ptr<Mesh> TaglessMesh = CloneOrTakeOver(pMesh.get());
+			TSharedPtr<FMesh> TaglessMesh = CloneOrTakeOver(pMesh);
 			TaglessMesh->Tags.SetNum(0, EAllowShrinking::No);
 			pMesh = TaglessMesh;
 		}
@@ -1220,7 +1220,7 @@ namespace mu
 				{
 					for (int32 l = 0; l < DuplicateOf.Mesh->GetLayoutCount(); ++l)
 					{
-						const Layout* DuplicatedLayout = DuplicateOf.Mesh->GetLayout(l);
+						TSharedPtr<const FLayout> DuplicatedLayout = DuplicateOf.Mesh->GetLayout(l);
 						OutResult.GeneratedLayouts.Add({ DuplicatedLayout,0 });
 					}
 				}
@@ -1232,7 +1232,7 @@ namespace mu
 		else
 		{
 			// We need to clone the mesh in the node because we will modify it.
-			Ptr<Mesh> ClonedMesh = pMesh->Clone();
+			TSharedPtr<FMesh> ClonedMesh = pMesh->Clone();
 			ClonedMesh->EnsureSurfaceData();
 
 			ConstantOp->SetValue(ClonedMesh, CompilerOptions->OptimisationOptions.DiskCacheContext);
@@ -1292,7 +1292,7 @@ namespace mu
 					for (int32 LayoutIndex = 0; LayoutIndex < InOptions.OverrideLayouts.Num(); ++LayoutIndex)
 					{
 						const FGeneratedLayout& OverrideData = InOptions.OverrideLayouts[LayoutIndex];
-						Ptr<const Layout> GeneratedLayout = OverrideData.Layout;
+						TSharedPtr<const FLayout> GeneratedLayout = OverrideData.Layout;
 						const void* Context = InOptions.OverrideContext.Get(InNode->GetMessageContext());
 
 						// In this case we need the layout block ids to use the ids in the parent layout, and not be prefixed with
@@ -1344,7 +1344,7 @@ namespace mu
             op->Source = baseResult.MeshOp;
             op->Flags = 0;
 
-            Ptr<Mesh> FormatMesh = new Mesh();
+            TSharedPtr<FMesh> FormatMesh = MakeShared<FMesh>();
 
             if (node.VertexBuffers.GetBufferCount())
             {
@@ -1364,7 +1364,7 @@ namespace mu
 			}
 
             Ptr<ASTOpConstantResource> cop = new ASTOpConstantResource();
-            cop->Type = OP_TYPE::ME_CONSTANT;
+            cop->Type = EOpType::ME_CONSTANT;
             cop->SetValue( FormatMesh, CompilerOptions->OptimisationOptions.DiskCacheContext );
 			if (baseResult.BaseMeshOp)
 			{
@@ -1469,7 +1469,7 @@ namespace mu
 			op->vertexSelectionBoneMaxRadius = ClipNode->Parameters.MaxEffectRadius;
         }
  
-        // Parameters
+        // FParameters
         op->dist = ClipNode->Parameters.DistanceToPlane;
         op->factor = ClipNode->Parameters.LinearityFactor;
 
@@ -1482,7 +1482,7 @@ namespace mu
                                                    const NodeMeshClipWithMesh* ClipNode)
     {
         Ptr<ASTOpFixed> op = new ASTOpFixed();
-        op->op.type = OP_TYPE::ME_CLIPWITHMESH;
+        op->op.type = EOpType::ME_CLIPWITHMESH;
 
         // Base
         if (ClipNode->Source)
diff --git a/Engine/Plugins/Mutable/Source/MutableTools/Private/MuT/CodeGenerator_Projector.cpp b/Engine/Plugins/Mutable/Source/MutableTools/Private/MuT/CodeGenerator_Projector.cpp
--- a/Engine/Plugins/Mutable/Source/MutableTools/Private/MuT/CodeGenerator_Projector.cpp
+++ b/Engine/Plugins/Mutable/Source/MutableTools/Private/MuT/CodeGenerator_Projector.cpp
@@ -28,7 +28,7 @@ namespace mu
 	FProjector ProjectorFromNode(const NodeProjectorConstant::Private& node)
 	{
 		FProjector p;
-		p.type = node.m_type;
+		p.type = node.Type;
 		p.position[0] = node.m_position[0];
 		p.position[1] = node.m_position[1];
 		p.position[2] = node.m_position[2];
@@ -113,23 +113,23 @@ namespace mu
 		if (!it)
 		{
 			FParameterDesc param;
-			param.m_name = node.m_name;
-			const TCHAR* CStr = ToCStr(node.m_uid);
-			param.m_uid.ImportTextItem(CStr, 0, nullptr, nullptr);
-			param.m_type = EParameterType::Projector;
+			param.Name = node.Name;
+			const TCHAR* CStr = ToCStr(node.UID);
+			param.UID.ImportTextItem(CStr, 0, nullptr, nullptr);
+			param.Type = EParameterType::Projector;
 
 			FProjector p = ProjectorFromNode(node);
-			param.m_defaultValue.Set<FParamProjectorType>(p);
+			param.DefaultValue.Set<FParamProjectorType>(p);
 
 			op = new ASTOpParameter();
-			op->type = OP_TYPE::PR_PARAMETER;
+			op->type = EOpType::PR_PARAMETER;
 			op->parameter = param;
 
 			// Generate the code for the ranges
-			for (int32 a = 0; a < node.m_ranges.Num(); ++a)
+			for (int32 a = 0; a < node.Ranges.Num(); ++a)
 			{
 				FRangeGenerationResult rangeResult;
-				GenerateRange(rangeResult, Options, node.m_ranges[a]);
+				GenerateRange(rangeResult, Options, node.Ranges[a]);
 				op->ranges.Emplace(op.get(), rangeResult.sizeOp, rangeResult.rangeName, rangeResult.rangeUID);
 			}
 
@@ -140,7 +140,7 @@ namespace mu
 			op = *it;
 		}
 
-		const FProjector& Projector = op->parameter.m_defaultValue.Get<FParamProjectorType>();
+		const FProjector& Projector = op->parameter.DefaultValue.Get<FParamProjectorType>();
 		result.type = Projector.type;
 		result.op = op;
 	}
diff --git a/Engine/Plugins/Mutable/Source/MutableTools/Private/MuT/CodeGenerator_Scalar.cpp b/Engine/Plugins/Mutable/Source/MutableTools/Private/MuT/CodeGenerator_Scalar.cpp
--- a/Engine/Plugins/Mutable/Source/MutableTools/Private/MuT/CodeGenerator_Scalar.cpp
+++ b/Engine/Plugins/Mutable/Source/MutableTools/Private/MuT/CodeGenerator_Scalar.cpp
@@ -123,8 +123,8 @@ namespace mu
 		const NodeScalarConstant::Private& node = *Typed->GetPrivate();
 
 		Ptr<ASTOpFixed> op = new ASTOpFixed();
-		op->op.type = OP_TYPE::SC_CONSTANT;
-		op->op.args.ScalarConstant.value = node.m_value;
+		op->op.type = EOpType::SC_CONSTANT;
+		op->op.args.ScalarConstant.value = node.Value;
 
 		result.op = op;
 	}
@@ -141,21 +141,21 @@ namespace mu
 		if (!it)
 		{
 			FParameterDesc param;
-			param.m_name = node.m_name;
-			const TCHAR* CStr = ToCStr(node.m_uid);
-			param.m_uid.ImportTextItem(CStr, 0, nullptr, nullptr);
-			param.m_type = EParameterType::Float;
-			param.m_defaultValue.Set<FParamFloatType>(node.m_defaultValue);
+			param.Name = node.Name;
+			const TCHAR* CStr = ToCStr(node.UID);
+			param.UID.ImportTextItem(CStr, 0, nullptr, nullptr);
+			param.Type = EParameterType::Float;
+			param.DefaultValue.Set<FParamFloatType>(node.DefaultValue);
 
 			op = new ASTOpParameter();
-			op->type = OP_TYPE::SC_PARAMETER;
+			op->type = EOpType::SC_PARAMETER;
 			op->parameter = param;
 
 			// Generate the code for the ranges
-			for (int32 a = 0; a < node.m_ranges.Num(); ++a)
+			for (int32 a = 0; a < node.Ranges.Num(); ++a)
 			{
 				FRangeGenerationResult rangeResult;
-				GenerateRange(rangeResult, Options, node.m_ranges[a]);
+				GenerateRange(rangeResult, Options, node.Ranges[a]);
 				op->ranges.Emplace(op.get(), rangeResult.sizeOp, rangeResult.rangeName, rangeResult.rangeUID);
 			}
 
@@ -181,28 +181,28 @@ namespace mu
 		if (!it)
 		{
 			FParameterDesc param;
-			param.m_name = node.m_name;
-			const TCHAR* CStr = ToCStr(node.m_uid);
-			param.m_uid.ImportTextItem(CStr, 0, nullptr, nullptr);
-			param.m_type = EParameterType::Int;
-			param.m_defaultValue.Set<FParamIntType>(node.m_defaultValue);
-
-			param.m_possibleValues.SetNum(node.m_options.Num());
-			for (int32 i = 0; i < node.m_options.Num(); ++i)
+			param.Name = node.Name;
+			const TCHAR* CStr = ToCStr(node.UID);
+			param.UID.ImportTextItem(CStr, 0, nullptr, nullptr);
+			param.Type = EParameterType::Int;
+			param.DefaultValue.Set<FParamIntType>(node.DefaultValue);
+
+			param.PossibleValues.SetNum(node.Options.Num());
+			for (int32 i = 0; i < node.Options.Num(); ++i)
 			{
-				param.m_possibleValues[i].m_value = (int16_t)node.m_options[i].value;
-				param.m_possibleValues[i].m_name = node.m_options[i].name;
+				param.PossibleValues[i].Value = (int16_t)node.Options[i].value;
+				param.PossibleValues[i].Name = node.Options[i].name;
 			}
 
 			op = new ASTOpParameter();
-			op->type = OP_TYPE::NU_PARAMETER;
+			op->type = EOpType::NU_PARAMETER;
 			op->parameter = param;
 
 			// Generate the code for the ranges
-			for (int32 a = 0; a < node.m_ranges.Num(); ++a)
+			for (int32 a = 0; a < node.Ranges.Num(); ++a)
 			{
 				FRangeGenerationResult rangeResult;
-				GenerateRange(rangeResult, Options, node.m_ranges[a]);
+				GenerateRange(rangeResult, Options, node.Ranges[a]);
 				op->ranges.Emplace(op.get(), rangeResult.sizeOp, rangeResult.rangeName, rangeResult.rangeUID);
 			}
 
@@ -233,7 +233,7 @@ namespace mu
 		}
 
 		Ptr<ASTOpSwitch> op = new ASTOpSwitch();
-		op->type = OP_TYPE::SC_SWITCH;
+		op->type = EOpType::SC_SWITCH;
 
 		// Variable value
 		if (node.Parameter)
@@ -323,7 +323,7 @@ namespace mu
 
 
 			Ptr<ASTOpConditional> conditional = new ASTOpConditional;
-			conditional->type = OP_TYPE::SC_CONDITIONAL;
+			conditional->type = EOpType::SC_CONDITIONAL;
 			conditional->no = op;
 			conditional->yes = variationOp;
 			conditional->condition = FirstPass.Tags[tagIndex].GenericCondition;
@@ -362,7 +362,7 @@ namespace mu
 		const NodeScalarArithmeticOperation::Private& node = *Typed->GetPrivate();
 
 		Ptr<ASTOpFixed> op = new ASTOpFixed();
-		op->op.type = OP_TYPE::SC_ARITHMETIC;
+		op->op.type = EOpType::SC_ARITHMETIC;
 
 		switch (node.m_operation)
 		{
@@ -409,7 +409,7 @@ namespace mu
 	{
 		const NodeScalarTable& node = *Typed;
 
-		Ptr<ASTOp> Op = GenerateTableSwitch<NodeScalarTable, ETableColumnType::Scalar, OP_TYPE::SC_SWITCH>(node,
+		Ptr<ASTOp> Op = GenerateTableSwitch<NodeScalarTable, ETableColumnType::Scalar, EOpType::SC_SWITCH>(node,
 			[this,&Options](const NodeScalarTable& node, int32 colIndex, int32 row, mu::ErrorLog* pErrorLog)
 			{
 				NodeScalarConstantPtr pCell = new NodeScalarConstant();
diff --git a/Engine/Plugins/Mutable/Source/MutableTools/Private/MuT/CodeGenerator_SecondPass.cpp b/Engine/Plugins/Mutable/Source/MutableTools/Private/MuT/CodeGenerator_SecondPass.cpp
--- a/Engine/Plugins/Mutable/Source/MutableTools/Private/MuT/CodeGenerator_SecondPass.cpp
+++ b/Engine/Plugins/Mutable/Source/MutableTools/Private/MuT/CodeGenerator_SecondPass.cpp
@@ -27,22 +27,22 @@ namespace mu
 
 
 	//---------------------------------------------------------------------------------------------
-	mu::Ptr<ASTOp> SecondPassGenerator::GenerateTagCondition(size_t tagIndex,
-		const set<size_t>& posSurf,
-		const set<size_t>& negSurf,
-		const set<size_t>& posTag,
-		const set<size_t>& negTag)
+	mu::Ptr<ASTOp> SecondPassGenerator::GenerateTagCondition(int32 tagIndex,
+		const TSet<int32>& posSurf,
+		const TSet<int32>& negSurf,
+		const TSet<int32>& posTag,
+		const TSet<int32>& negTag)
 	{
 		FirstPassGenerator::FTag& t = FirstPass->Tags[tagIndex];
 
 		// If this tag is already in the list of positive tags, return true as condition
-		if (posTag.find(tagIndex) != posTag.end())
+		if (posTag.Contains(tagIndex))
 		{
 			return m_opPool.Add(new ASTOpConstantBool(true));
 		}
 
 		// If this tag is already in the list of negative tags, return false as condition
-		if (negTag.find(tagIndex) != negTag.end())
+		if (negTag.Contains(tagIndex))
 		{
 			return m_opPool.Add(new ASTOpConstantBool(false));
 		}
@@ -54,16 +54,16 @@ namespace mu
 		//    key.posSurf = posSurf;
 		//    key.negTag = negTag;
 		//    key.posTag = posTag;
-		for (auto s : negTag) { if (TagsPerTag[tagIndex].count(s) > 0) { key.negTag.insert(s); } }
-		for (auto s : posTag) { if (TagsPerTag[tagIndex].count(s) > 0) { key.posTag.insert(s); } }
-		for (auto s : negSurf) { if (SurfacesPerTag[tagIndex].count(s) > 0) { key.negSurf.insert(s); } }
-		for (auto s : posSurf) { if (SurfacesPerTag[tagIndex].count(s) > 0) { key.posSurf.insert(s); } }
+		for (auto s : negTag) { if (TagsPerTag[tagIndex].Contains(s)) { key.negTag.Add(s); } }
+		for (auto s : posTag) { if (TagsPerTag[tagIndex].Contains(s)) { key.posTag.Add(s); } }
+		for (auto s : negSurf) { if (SurfacesPerTag[tagIndex].Contains(s)) { key.negSurf.Add(s); } }
+		for (auto s : posSurf) { if (SurfacesPerTag[tagIndex].Contains(s)) { key.posSurf.Add(s); } }
 
 		{
-			auto it = TagConditionGenerationCache.find(key);
-			if (it != TagConditionGenerationCache.end())
+			Ptr<ASTOp>* Found = TagConditionGenerationCache.Find(key);
+			if (Found)
 			{
-				return it->second;
+				return *Found;
 			}
 		}
 		
@@ -72,14 +72,14 @@ namespace mu
 		// Condition expression for all the surfaces that activate the tag
 		for (int32 surfIndex : t.Surfaces)
 		{
-			if (posSurf.find(surfIndex) != posSurf.end())
+			if (posSurf.Contains(surfIndex))
 			{
 				// This surface is already a positive requirement higher up in the condition so
 				// we can ignore it here.
 				continue;
 			}
 
-			if (negSurf.find(surfIndex) != negSurf.end())
+			if (negSurf.Contains(surfIndex))
 			{
 				// This surface is a negative requirement higher up in the condition so
 				// this branch never be true.
@@ -89,7 +89,7 @@ namespace mu
 			const FirstPassGenerator::FSurface& surface = FirstPass->Surfaces[surfIndex];
 
 			auto PositiveTags = posTag;
-			PositiveTags.insert(tagIndex);
+			PositiveTags.Add(tagIndex);
 
 			Ptr<ASTOp> surfCondition = GenerateDataCodition(surfIndex,
 				FirstPass->Surfaces[surfIndex].PositiveTags,
@@ -100,7 +100,7 @@ namespace mu
 				negTag);
 
 			// If the surface is a constant false, we can skip adding it
-			if (surfCondition && surfCondition->GetOpType()==OP_TYPE::BO_CONSTANT)
+			if (surfCondition && surfCondition->GetOpType()==EOpType::BO_CONSTANT)
 			{
 				const ASTOpConstantBool* constOp = static_cast<const ASTOpConstantBool*>(surfCondition.get());
 				if (constOp->value == false)
@@ -113,7 +113,7 @@ namespace mu
 			if (surfCondition)
 			{
 				Ptr<ASTOpFixed> f = new ASTOpFixed();
-				f->op.type = OP_TYPE::BO_AND;
+				f->op.type = EOpType::BO_AND;
 				f->SetChild(f->op.args.BoolBinary.a, surface.ObjectCondition);
 				f->SetChild(f->op.args.BoolBinary.b, surfCondition);
 				fullCondition = m_opPool.Add(f);
@@ -131,7 +131,7 @@ namespace mu
 			else
 			{
 				Ptr<ASTOpFixed> o = new ASTOpFixed();
-				o->op.type = OP_TYPE::BO_OR;
+				o->op.type = EOpType::BO_OR;
 				o->SetChild(o->op.args.BoolBinary.a, fullCondition);
 				o->SetChild(o->op.args.BoolBinary.b, c);
 				c = m_opPool.Add(o);
@@ -141,28 +141,28 @@ namespace mu
 			//PartialOptimise( c, CompilerOptions->OptimisationOptions );
 		}
 
-		TagConditionGenerationCache.insert(std::make_pair<>(key, c));
+		TagConditionGenerationCache.Add(key, c);
 
 		return c;
 	}
 
 	//---------------------------------------------------------------------------------------------
-	mu::Ptr<ASTOp> SecondPassGenerator::GenerateDataCodition(size_t Index,
+	mu::Ptr<ASTOp> SecondPassGenerator::GenerateDataCodition(int32 Index,
 		const TArray<FString>& PositiveTags,
 		const TArray<FString>& NegativeTags,
-		const set<size_t>& posSurf,
-		const set<size_t>& negSurf,
-		const set<size_t>& posTag,
-		const set<size_t>& negTag)
+		const TSet<int32>& posSurf,
+		const TSet<int32>& negSurf,
+		const TSet<int32>& posTag,
+		const TSet<int32>& negTag)
 	{
 		// If this surface is already in the list of positive surfaces, return true as condition
-		if (posSurf.find(Index) != posSurf.end())
+		if (posSurf.Contains(Index))
 		{
 			return m_opPool.Add(new ASTOpConstantBool(true));
 		}
 
 		// If this surface is already in the list of negative surfaces, return false as condition
-		if (negSurf.find(Index) != negSurf.end())
+		if (negSurf.Contains(Index))
 		{
 			return m_opPool.Add(new ASTOpConstantBool(false));
 		}
@@ -179,10 +179,10 @@ namespace mu
 				return m_opPool.Add(new ASTOpConstantBool(false));
 			}
 
-			size_t tagIndex = it - &FirstPass->Tags[0];
+			int32 tagIndex = int32(it - &FirstPass->Tags[0]);
 
-			set<size_t> positiveSurfacesVisited = posSurf;
-			positiveSurfacesVisited.insert(Index);
+			TSet<int32> positiveSurfacesVisited = posSurf;
+			positiveSurfacesVisited.Add(Index);
 
 			Ptr<ASTOp> tagCondition = GenerateTagCondition(tagIndex,
 				positiveSurfacesVisited,
@@ -201,7 +201,7 @@ namespace mu
 			// If the tag is a constant ...
 			bool isConstant = false;
 			bool constantValue = false;
-			if (tagCondition->GetOpType()==OP_TYPE::BO_CONSTANT)
+			if (tagCondition->GetOpType()==EOpType::BO_CONSTANT)
 			{
 				const ASTOpConstantBool* constOp = static_cast<const ASTOpConstantBool*>(tagCondition.get());
 				isConstant = true;
@@ -217,7 +217,7 @@ namespace mu
 				else
 				{
 					Ptr<ASTOpFixed> o = new ASTOpFixed;
-					o->op.type = OP_TYPE::BO_AND;
+					o->op.type = EOpType::BO_AND;
 					o->SetChild(o->op.args.BoolBinary.a, tagCondition);
 					o->SetChild(o->op.args.BoolBinary.b, c);
 					c = m_opPool.Add(o);
@@ -248,13 +248,13 @@ namespace mu
 				continue;
 			}
 
-			size_t tagIndex = it - FirstPass->Tags.GetData();
+			int32 tagIndex = int32(it - FirstPass->Tags.GetData());
 
-			set<size_t> positiveSurfacesVisited = negSurf;
-			set<size_t> negativeSurfacesVisited = posSurf;
-			negativeSurfacesVisited.insert(Index);
-			set<size_t> positiveTagsVisited = negTag;
-			set<size_t> negativeTagsVisited = posTag;
+			TSet<int32> positiveSurfacesVisited = negSurf;
+			TSet<int32> negativeSurfacesVisited = posSurf;
+			negativeSurfacesVisited.Add(Index);
+			TSet<int32> positiveTagsVisited = negTag;
+			TSet<int32> negativeTagsVisited = posTag;
 			Ptr<ASTOp> tagCondition = GenerateTagCondition(tagIndex,
 				positiveSurfacesVisited,
 				negativeSurfacesVisited,
@@ -274,7 +274,7 @@ namespace mu
 			// If the tag is a constant ...
 			bool isConstant = false;
 			bool constantValue = false;
-			if (tagCondition && tagCondition->GetOpType()== OP_TYPE::BO_CONSTANT)
+			if (tagCondition && tagCondition->GetOpType()== EOpType::BO_CONSTANT)
 			{
 				const ASTOpConstantBool* constOp = static_cast<const ASTOpConstantBool*>(tagCondition.get());
 				isConstant = true;
@@ -285,7 +285,7 @@ namespace mu
 			if (!isConstant && tagCondition)
 			{
 				Ptr<ASTOpFixed> n = new ASTOpFixed;
-				n->op.type = OP_TYPE::BO_NOT;
+				n->op.type = EOpType::BO_NOT;
 				n->SetChild(n->op.args.BoolNot.source, tagCondition);
 
 				if (!c)
@@ -295,7 +295,7 @@ namespace mu
 				else
 				{
 					Ptr<ASTOpFixed> o = new ASTOpFixed;
-					o->op.type = OP_TYPE::BO_AND;
+					o->op.type = EOpType::BO_AND;
 					o->SetChild(o->op.args.BoolBinary.a, n);
 					o->SetChild(o->op.args.BoolBinary.b, c);
 					c = m_opPool.Add(o);
@@ -329,31 +329,30 @@ namespace mu
 		SurfacesPerTag.SetNum(FirstPass->Tags.Num());
 		TagsPerTag.Empty();
 		TagsPerTag.SetNum(FirstPass->Tags.Num());
-		for (size_t t = 0; t < FirstPass->Tags.Num(); ++t)
+		for (int32 t = 0; t < FirstPass->Tags.Num(); ++t)
 		{
-			set<size_t> pendingSurfs;
-			for (size_t s : FirstPass->Tags[t].Surfaces)
+			TSet<int32> PendingSurfaces;
+			for (int32 s : FirstPass->Tags[t].Surfaces)
 			{
-				pendingSurfs.insert(s);
+				PendingSurfaces.Add(s);
 			}
 
-			set<size_t> processedSurfs;
-			while (!pendingSurfs.empty())
+			TSet<int32> processedSurfs;
+			while (!PendingSurfaces.IsEmpty())
 			{
-				size_t cs = *pendingSurfs.begin();
-				pendingSurfs.erase(pendingSurfs.begin());
+				int32 cs = *PendingSurfaces.begin();
+				PendingSurfaces.Remove(cs);
 
-				if (std::find(processedSurfs.begin(), processedSurfs.end(), cs)
-					!= processedSurfs.end())
+				if (processedSurfs.Contains(cs))
 				{
 					continue;
 				}
 
-				processedSurfs.insert(cs);
+				processedSurfs.Add(cs);
 
-				SurfacesPerTag[t].insert(cs);
+				SurfacesPerTag[t].Add(cs);
 
-				auto& csurf = FirstPass->Surfaces[cs];
+				FirstPassGenerator::FSurface& csurf = FirstPass->Surfaces[cs];
 				for (const FString& sct : csurf.PositiveTags)
 				{
 					const FirstPassGenerator::FTag* it = FirstPass->Tags.FindByPredicate([&](const FirstPassGenerator::FTag& e) { return e.Tag == sct; });
@@ -363,15 +362,15 @@ namespace mu
 						continue;
 					}
 
-					size_t ct = it - &FirstPass->Tags[0];
+					int32 ct = int32(it - &FirstPass->Tags[0]);
 
-					TagsPerTag[t].insert(ct);
+					TagsPerTag[t].Add(ct);
 
-					for (size_t s : FirstPass->Tags[ct].Surfaces)
+					for (int32 s : FirstPass->Tags[ct].Surfaces)
 					{
-						if (SurfacesPerTag[t].find(s) != SurfacesPerTag[t].end())
+						if (SurfacesPerTag[t].Contains(s))
 						{
-							pendingSurfs.insert(s);
+							PendingSurfaces.Add(s);
 						}
 					}
 				}
@@ -386,13 +385,13 @@ namespace mu
 
 					size_t ct = it - &FirstPass->Tags[0];
 
-					TagsPerTag[t].insert(ct);
+					TagsPerTag[t].Add(ct);
 
 					for (size_t s : FirstPass->Tags[ct].Surfaces)
 					{
-						if (SurfacesPerTag[t].find(s) != SurfacesPerTag[t].end())
+						if (SurfacesPerTag[t].Contains(s))
 						{
-							pendingSurfs.insert(s);
+							PendingSurfaces.Add(s);
 						}
 					}
 				}
@@ -400,9 +399,9 @@ namespace mu
 		}
 
 		// Create the conditions for every surface, modifier, component and individual tag.
-		TagConditionGenerationCache.clear();
+		TagConditionGenerationCache.Reset();
 
-		set<size_t> Empty;
+		TSet<int32> Empty;
 
 		for (int32 SurfaceIndex = 0; SurfaceIndex < FirstPass->Surfaces.Num(); ++SurfaceIndex)
 		{
@@ -416,7 +415,7 @@ namespace mu
 					Empty, Empty, Empty, Empty);
 
 				Ptr<ASTOpFixed> ConditionOp = new ASTOpFixed();
-				ConditionOp->op.type = OP_TYPE::BO_AND;
+				ConditionOp->op.type = EOpType::BO_AND;
 				ConditionOp->SetChild(ConditionOp->op.args.BoolBinary.a, Surface.ObjectCondition);
 				ConditionOp->SetChild(ConditionOp->op.args.BoolBinary.b, c);
 
@@ -435,7 +434,7 @@ namespace mu
 			//		Empty, Empty, Empty, Empty);
 
 			//	Ptr<ASTOpFixed> OpAnd = new ASTOpFixed;
-			//	OpAnd->op.type = OP_TYPE::BO_AND;
+			//	OpAnd->op.type = EOpType::BO_AND;
 			//	OpAnd->SetChild(OpAnd->op.args.BoolBinary.a, Edit.Condition);
 			//	OpAnd->SetChild(OpAnd->op.args.BoolBinary.b, c);
 			//	c = m_opPool.Add(OpAnd);
@@ -453,7 +452,7 @@ namespace mu
 				Empty, Empty, Empty, Empty);
 			
 			Ptr<ASTOpFixed> ConditionOp = new ASTOpFixed();
-			ConditionOp->op.type = OP_TYPE::BO_AND;
+			ConditionOp->op.type = EOpType::BO_AND;
 			ConditionOp->SetChild(ConditionOp->op.args.BoolBinary.a, FirstPass->Modifiers[ModifierIndex].ObjectCondition);
 			ConditionOp->SetChild(ConditionOp->op.args.BoolBinary.b, c);
 
diff --git a/Engine/Plugins/Mutable/Source/MutableTools/Private/MuT/CodeGenerator_SecondPass.h b/Engine/Plugins/Mutable/Source/MutableTools/Private/MuT/CodeGenerator_SecondPass.h
--- a/Engine/Plugins/Mutable/Source/MutableTools/Private/MuT/CodeGenerator_SecondPass.h
+++ b/Engine/Plugins/Mutable/Source/MutableTools/Private/MuT/CodeGenerator_SecondPass.h
@@ -3,19 +3,22 @@
 #pragma once
 
 #include "HAL/PlatformCrt.h"
-#include "MuR/MutableMemory.h"
 #include "MuR/Ptr.h"
 #include "MuR/RefCounted.h"
 #include "MuT/AST.h"
 #include "MuT/CodeGenerator_FirstPass.h"
 #include "MuT/Compiler.h"
 #include "MuT/Node.h"
-#include "map"
-#include "set"
-
+#include "Containers/Set.h"
+#include "Containers/Map.h"
 
 namespace mu
 {
+	inline bool SetsEquals(const TSet<int32>& Left, const TSet<int32>& Right)
+	{
+		return Left.Num() == Right.Num() && Left.Includes(Right);
+	}
+
 
 	/** Second pass of the code generation process.
      * It solves surface and modifier conditions from tags and variations
@@ -41,43 +44,49 @@ namespace mu
         //!
         struct FConditionGenerationKey
         {
-            size_t tagOrSurfIndex = 0;
-            set<size_t> posSurf;
-            set<size_t> negSurf;
-            set<size_t> posTag;
-            set<size_t> negTag;
-
-            inline bool operator<(const FConditionGenerationKey& o) const
-            {
-                if (tagOrSurfIndex<o.tagOrSurfIndex) return true;
-                if (tagOrSurfIndex>o.tagOrSurfIndex) return false;
-                if (posSurf<o.posSurf) return true;
-                if (posSurf>o.posSurf) return false;
-                if (negSurf<o.negSurf) return true;
-                if (negSurf>o.negSurf) return false;
-                if (posTag<o.posTag) return true;
-                if (posTag>o.posTag) return false;
-                if (negTag<o.negTag) return true;
-                if (negTag>o.negTag) return false;
-                return false;
-            }
-        };
+            int32 tagOrSurfIndex = 0;
+			TSet<int32> posSurf;
+            TSet<int32> negSurf;
+            TSet<int32> posTag;
+            TSet<int32> negTag;
+
+			friend FORCEINLINE uint32 GetTypeHash(const FConditionGenerationKey& InKey)
+			{
+				uint32 KeyHash = 0;
+				KeyHash = HashCombineFast(KeyHash, ::GetTypeHash(InKey.tagOrSurfIndex));
+				KeyHash = HashCombineFast(KeyHash, ::GetTypeHash(InKey.posSurf.Num()));
+				KeyHash = HashCombineFast(KeyHash, ::GetTypeHash(InKey.negSurf.Num()));
+				KeyHash = HashCombineFast(KeyHash, ::GetTypeHash(InKey.posTag.Num()));
+				KeyHash = HashCombineFast(KeyHash, ::GetTypeHash(InKey.negTag.Num()));
+				return KeyHash;
+			}
+
+			FORCEINLINE bool operator==(const FConditionGenerationKey& Other) const
+			{
+				if (tagOrSurfIndex != Other.tagOrSurfIndex) return false;
+				if (!SetsEquals(posSurf, Other.posSurf)) return false;
+				if (!SetsEquals(negSurf, Other.negSurf)) return false;
+				if (!SetsEquals(posTag, Other.posTag)) return false;
+				if (!SetsEquals(negTag, Other.negTag)) return false;
+				return true;
+			}
+			
+		};
 
         // List of surfaces that activate or deactivate every tag, or another surface that activates a tag in this set.
-		// \TODO: Change to UE containers
-		TArray< set<size_t> > SurfacesPerTag;
-		TArray< set<size_t> > TagsPerTag;
+		TArray< TSet<int32> > SurfacesPerTag;
+		TArray< TSet<int32> > TagsPerTag;
 
-        std::map<FConditionGenerationKey,Ptr<ASTOp>> TagConditionGenerationCache;
+        TMap<FConditionGenerationKey,Ptr<ASTOp>> TagConditionGenerationCache;
 
         UniqueOpPool m_opPool;
 
         //!
-        Ptr<ASTOp> GenerateTagCondition( size_t tagIndex,
-                                         const set<size_t>& posSurf,
-                                         const set<size_t>& negSurf,
-                                         const set<size_t>& posTag,
-                                         const set<size_t>& negTag );
+        Ptr<ASTOp> GenerateTagCondition( int32 tagIndex,
+                                         const TSet<int32>& posSurf,
+                                         const TSet<int32>& negSurf,
+                                         const TSet<int32>& posTag,
+                                         const TSet<int32>& negTag );
 
         /** Generate Surface, Edit or Modifier condition.
     	 * @param Index Surface, Edit, Component or Modifier index.
@@ -87,13 +96,13 @@ namespace mu
     	 * @param negSurf already visited Surfaces, Edits, or Modifiers that participate negatively in the condition.
       	 * @param posSurf Tags that already belong to the condition (positively).
 		 * @param negSurf Tags that already belong to the condition (negatively). */
-        Ptr<ASTOp> GenerateDataCodition(size_t Index,
+        Ptr<ASTOp> GenerateDataCodition(int32 Index,
 									        const TArray<FString>& PositiveTags,
 											const TArray<FString>& NegativeTags,
-                                            const set<size_t>& posSurf,
-                                            const set<size_t>& negSurf,
-                                            const set<size_t>& posTag,
-                                            const set<size_t>& negTag);
+                                            const TSet<int32>& posSurf,
+                                            const TSet<int32>& negSurf,
+                                            const TSet<int32>& posTag,
+                                            const TSet<int32>& negTag);
     };
 
 }
diff --git a/Engine/Plugins/Mutable/Source/MutableTools/Private/MuT/CodeGenerator_String.cpp b/Engine/Plugins/Mutable/Source/MutableTools/Private/MuT/CodeGenerator_String.cpp
--- a/Engine/Plugins/Mutable/Source/MutableTools/Private/MuT/CodeGenerator_String.cpp
+++ b/Engine/Plugins/Mutable/Source/MutableTools/Private/MuT/CodeGenerator_String.cpp
@@ -69,7 +69,7 @@ namespace mu
 		const NodeStringConstant::Private& node = *Typed->GetPrivate();
 
 		Ptr<ASTOpConstantString> op = new ASTOpConstantString();
-		op->value = node.m_value;
+		op->value = node.Value;
 
 		result.op = op;
 	}
@@ -86,21 +86,21 @@ namespace mu
 		if (!it)
 		{
 			FParameterDesc param;
-			param.m_name = node.m_name;
-			const TCHAR* CStr = ToCStr(node.m_uid);
-			param.m_uid.ImportTextItem(CStr, 0, nullptr, nullptr);
-			param.m_type = EParameterType::Float;
-			param.m_defaultValue.Set<FParamStringType>(node.m_defaultValue);
+			param.Name = node.Name;
+			const TCHAR* CStr = ToCStr(node.UID);
+			param.UID.ImportTextItem(CStr, 0, nullptr, nullptr);
+			param.Type = EParameterType::Float;
+			param.DefaultValue.Set<FParamStringType>(node.DefaultValue);
 
 			op = new ASTOpParameter();
-			op->type = OP_TYPE::ST_PARAMETER;
+			op->type = EOpType::ST_PARAMETER;
 			op->parameter = param;
 
 			// Generate the code for the ranges
-			for (int32 a = 0; a < node.m_ranges.Num(); ++a)
+			for (int32 a = 0; a < node.Ranges.Num(); ++a)
 			{
 				FRangeGenerationResult rangeResult;
-				GenerateRange(rangeResult, Options, node.m_ranges[a]);
+				GenerateRange(rangeResult, Options, node.Ranges[a]);
 				op->ranges.Emplace(op.get(), rangeResult.sizeOp, rangeResult.rangeName, rangeResult.rangeUID);
 			}
 
diff --git a/Engine/Plugins/Mutable/Source/MutableTools/Private/MuT/CodeOptimiser.cpp b/Engine/Plugins/Mutable/Source/MutableTools/Private/MuT/CodeOptimiser.cpp
--- a/Engine/Plugins/Mutable/Source/MutableTools/Private/MuT/CodeOptimiser.cpp
+++ b/Engine/Plugins/Mutable/Source/MutableTools/Private/MuT/CodeOptimiser.cpp
@@ -35,7 +35,7 @@ namespace mu
 
 		struct FMeshEntry
 		{
-			Ptr<const Mesh> Mesh;
+			TSharedPtr<const FMesh> Mesh;
 			Ptr<ASTOpConstantResource> Op;
 
 			bool operator==(const FMeshEntry& o) const
@@ -46,7 +46,7 @@ namespace mu
 
 		struct FImageEntry
 		{
-			Ptr< const Image> Image;
+			TSharedPtr<const FImage> Image;
 			Ptr<ASTOpConstantResource> Op;
 
 			bool operator==(const FImageEntry& o) const
@@ -57,7 +57,7 @@ namespace mu
 
 		struct FLayoutEntry
 		{
-			Ptr< const Layout> Layout;
+			TSharedPtr< const FLayout> Layout;
 			Ptr<ASTOpConstantResource> Op;
 
 			bool operator==(const FLayoutEntry& o) const
@@ -68,7 +68,7 @@ namespace mu
 
 		struct custom_mesh_equal
 		{
-			bool operator()( const MeshPtrConst& a, const MeshPtrConst& b ) const
+			bool operator()( const TSharedPtr<const FMesh>& a, const TSharedPtr<const FMesh>& b ) const
 			{
 				return a==b || *a==*b;
 			}
@@ -76,7 +76,7 @@ namespace mu
 
 		struct custom_image_equal
 		{
-			bool operator()( const ImagePtrConst& a, const ImagePtrConst& b ) const
+			bool operator()( const TSharedPtr<const FImage>& a, const TSharedPtr<const FImage>& b ) const
 			{
 				return a==b || *a==*b;
 			}
@@ -84,7 +84,7 @@ namespace mu
 
 		struct custom_layout_equal
 		{
-			bool operator()( const Ptr<const Layout>& a, const Ptr<const Layout>& b ) const
+			bool operator()( const TSharedPtr<const FLayout>& a, const TSharedPtr<const FLayout>& b ) const
 			{
 				return a==b || *a==*b;
 			}
@@ -110,21 +110,21 @@ namespace mu
 				switch ( n->GetOpType() )
 				{
 
-				case OP_TYPE::ME_CONSTANT:
+				case EOpType::ME_CONSTANT:
 				{
 					ASTOpConstantResource* typedNode = static_cast<ASTOpConstantResource*>(n.get());
 					AllMeshOps.Add(typedNode);
 					break;
 				}
 
-				case OP_TYPE::IM_CONSTANT:
+				case EOpType::IM_CONSTANT:
 				{
 					ASTOpConstantResource* typedNode = static_cast<ASTOpConstantResource*>(n.get());
 					AllImageOps.Add(typedNode);
 					break;
 				}
 
-				case OP_TYPE::LA_CONSTANT:
+				case EOpType::LA_CONSTANT:
 				{
 					ASTOpConstantResource* typedNode = static_cast<ASTOpConstantResource*>(n.get());
 					AllLayoutOps.Add(typedNode);
@@ -133,12 +133,12 @@ namespace mu
 
 				//    These should be part of the duplicated code removal, in AST.
 				//            // Names
-				//            case OP_TYPE::IN_ADDMESH:
-				//            case OP_TYPE::IN_ADDIMAGE:
-				//            case OP_TYPE::IN_ADDVECTOR:
-				//            case OP_TYPE::IN_ADDSCALAR:
-				//            case OP_TYPE::IN_ADDCOMPONENT:
-				//            case OP_TYPE::IN_ADDSURFACE:
+				//            case EOpType::IN_ADDMESH:
+				//            case EOpType::IN_ADDIMAGE:
+				//            case EOpType::IN_ADDVECTOR:
+				//            case EOpType::IN_ADDSCALAR:
+				//            case EOpType::IN_ADDCOMPONENT:
+				//            case EOpType::IN_ADDSURFACE:
 
 				default:
 					break;
@@ -166,13 +166,13 @@ namespace mu
 
 				if (!Candidates.IsEmpty())
 				{
-					Ptr<const Mesh> mesh = static_cast<const Mesh*>(typedNode->GetValue().get());
+					TSharedPtr<const FMesh> mesh = StaticCastSharedPtr<const FMesh>(typedNode->GetValue());
 
 					for (FMeshEntry* It : Candidates)
 					{
 						if (!It->Mesh)
 						{
-							It->Mesh = static_cast<const Mesh*>(It->Op->GetValue().get());
+							It->Mesh = StaticCastSharedPtr<const FMesh>(It->Op->GetValue());
 						}
 
 						if (custom_mesh_equal()(mesh, It->Mesh))
@@ -214,13 +214,13 @@ namespace mu
 				
 				if (!Candidates.IsEmpty())
 				{
-					Ptr<const Image> image = static_cast<const Image*>(typedNode->GetValue().get());
+					TSharedPtr<const FImage> image = StaticCastSharedPtr<const FImage>(typedNode->GetValue());
 
 					for (FImageEntry* It: Candidates)
 					{
 						if (!It->Image)
 						{
-							It->Image = static_cast<const Image*>(It->Op->GetValue().get());
+							It->Image = StaticCastSharedPtr<const FImage>(It->Op->GetValue());
 						}
 
 						if (custom_image_equal()(image, It->Image))
@@ -262,13 +262,13 @@ namespace mu
 
 				if (!Candidates.IsEmpty())
 				{
-					Ptr<const Layout> layout = static_cast<const Layout*>(typedNode->GetValue().get());
+					TSharedPtr<const FLayout> layout = StaticCastSharedPtr<const FLayout>(typedNode->GetValue());
 
 					for (FLayoutEntry* It : Candidates)
 					{
 						if (!It->Layout)
 						{
-							It->Layout = static_cast<const Layout*>(It->Op->GetValue().get());
+							It->Layout = StaticCastSharedPtr<const FLayout>(It->Op->GetValue());
 						}
 
 						if (custom_layout_equal()(layout, It->Layout))
@@ -322,7 +322,7 @@ namespace mu
 		MUTABLE_CPUPROFILER_SCOPE(DuplicatedCodeRemoverAST);
 
 		// Visited nodes, per type
-		std::unordered_set<Ptr<ASTOp>,op_pointer_hash,op_pointer_equal> visited[int32(OP_TYPE::COUNT)];
+		std::unordered_set<Ptr<ASTOp>,op_pointer_hash,op_pointer_equal> visited[int32(EOpType::COUNT)];
 
 		ASTOp::Traverse_BottomUp_Unique_NonReentrant( roots, [&](Ptr<ASTOp>& n)
 		{
@@ -347,16 +347,16 @@ namespace mu
 		bool found = false;
 
 		// If the program added new operations, we haven't visited them.
-		m_visited.SetNumZeroed( program.m_opAddress.Num() );
+		m_visited.SetNumZeroed( program.OpAddress.Num() );
 
-		m_pending.Reset();
-		m_pending.Reserve( program.m_opAddress.Num()/4 );
-		m_pending.Add({ false,rootAt });
+		Pending.Reset();
+		Pending.Reserve( program.OpAddress.Num()/4 );
+		Pending.Add({ false,rootAt });
 
 		// Don't early out to be able to complete parent op cached flags
-		while ( !m_pending.IsEmpty() )
+		while ( !Pending.IsEmpty() )
 		{
-			TPair<bool,int32> item = m_pending.Pop();
+			TPair<bool,int32> item = Pending.Pop();
 			OP::ADDRESS at = item.Value;
 
 			// Not cached?
@@ -387,13 +387,13 @@ namespace mu
 						check(m_visited[at]==0);
 
 						m_visited[at] = 1;
-						m_pending.Add({ true,at });
+						Pending.Add({ true,at });
 
 						ForEachReference( program, at, [&](OP::ADDRESS ref)
 						{
 							if (ref && m_visited[ref]==0)
 							{
-								m_pending.Add({ false,ref });
+								Pending.Add({ false,ref });
 							}
 						});
 					}
@@ -415,13 +415,13 @@ namespace mu
 	//---------------------------------------------------------------------------------------------
 	IsConstantVisitor::IsConstantVisitor()
 	{
-		TArray<OP_TYPE> parameterTypes;
-		parameterTypes.Add( OP_TYPE::BO_PARAMETER );
-		parameterTypes.Add( OP_TYPE::NU_PARAMETER );
-		parameterTypes.Add( OP_TYPE::SC_PARAMETER );
-		parameterTypes.Add( OP_TYPE::CO_PARAMETER );
-		parameterTypes.Add( OP_TYPE::PR_PARAMETER );
-		parameterTypes.Add( OP_TYPE::IM_PARAMETER );
+		TArray<EOpType> parameterTypes;
+		parameterTypes.Add( EOpType::BO_PARAMETER );
+		parameterTypes.Add( EOpType::NU_PARAMETER );
+		parameterTypes.Add( EOpType::SC_PARAMETER );
+		parameterTypes.Add( EOpType::CO_PARAMETER );
+		parameterTypes.Add( EOpType::PR_PARAMETER );
+		parameterTypes.Add( EOpType::IM_PARAMETER );
 		m_findOpTypeVisitor = MakeUnique<FindOpTypeVisitor>(parameterTypes);
 	}
 
@@ -470,15 +470,15 @@ namespace mu
 
 			// This runs in a worker thread
 
-			OP_TYPE type = SourceCloned->GetOpType();
+			EOpType type = SourceCloned->GetOpType();
 			DATATYPE DataType = GetOpDataType(type);
 
-			Ptr<Settings> pSettings = new Settings;
-			pSettings->SetProfile( false );
-			pSettings->SetImageCompressionQuality( ImageCompressionQuality );
-			SystemPtr pSystem = new System( pSettings );
+			FSettings Settings;
+			Settings.SetProfile( false );
+			Settings.SetImageCompressionQuality( ImageCompressionQuality );
+			TSharedPtr<FSystem> System = MakeShared<FSystem>(Settings);
 
-			pSystem->GetPrivate()->ImagePixelFormatOverride = ImOp.FormatImageOverride;
+			System->GetPrivate()->ImagePixelFormatOverride = ImOp.FormatImageOverride;
 
 			FSourceDataDescriptor SourceDataDescriptor;
 			if (DataType == DT_IMAGE || DataType == DT_MESH)
@@ -492,15 +492,15 @@ namespace mu
 			LinkerOptions.MinTextureResidentMipCount = 255;
 			LinkerOptions.bSeparateImageMips = false;
 
-			TSharedPtr<const Model> model = MakeShared<Model>();
-			OP::ADDRESS at = ASTOp::FullLink(SourceCloned, model->GetPrivate()->m_program, &LinkerOptions);
+			TSharedPtr<const FModel> model = MakeShared<FModel>();
+			OP::ADDRESS at = ASTOp::FullLink(SourceCloned, model->GetPrivate()->Program, &LinkerOptions);
 
 			FProgram::FState state;
-			state.m_root = at;
-			model->GetPrivate()->m_program.m_states.Add(state);
+			state.Root = at;
+			model->GetPrivate()->Program.States.Add(state);
 
-			Ptr<Parameters> LocalParams = Model::NewParameters(model);
-			pSystem->GetPrivate()->BeginBuild( model );
+			TSharedPtr<FParameters> LocalParams = FModel::NewParameters(model);
+			System->GetPrivate()->BeginBuild( model );
 
 			// Calculate the value and replace this op by a constant
 			switch( DataType )
@@ -509,13 +509,13 @@ namespace mu
 			{
 				MUTABLE_CPUPROFILER_SCOPE(ConstantMesh);
 
-				mu::Ptr<const Mesh> pMesh = pSystem->GetPrivate()->BuildMesh( model, LocalParams.get(), at );
+				TSharedPtr<const FMesh> pMesh = System->GetPrivate()->BuildMesh( model, LocalParams.Get(), at );
 
 				if (pMesh)
 				{
 					mu::Ptr<ASTOpConstantResource> ConstantOp = new ASTOpConstantResource();
 					ConstantOp->SourceDataDescriptor = SourceDataDescriptor;
-					ConstantOp->Type = OP_TYPE::ME_CONSTANT;
+					ConstantOp->Type = EOpType::ME_CONSTANT;
 					ConstantOp->SetValue( pMesh, DiskCacheContext );
 					Result = ConstantOp;
 				  }
@@ -526,13 +526,13 @@ namespace mu
 			{
 				MUTABLE_CPUPROFILER_SCOPE(ConstantImage);
 
-				mu::Ptr<const Image> pImage = pSystem->GetPrivate()->BuildImage( model, LocalParams.get(), at, 0, 0 );
+				TSharedPtr<const FImage> pImage = System->GetPrivate()->BuildImage( model, LocalParams.Get(), at, 0, 0 );
 
 				if (pImage)
 				{
 					mu::Ptr<ASTOpConstantResource> ConstantOp = new ASTOpConstantResource();
 					ConstantOp->SourceDataDescriptor = SourceDataDescriptor;
-					ConstantOp->Type = OP_TYPE::IM_CONSTANT;
+					ConstantOp->Type = EOpType::IM_CONSTANT;
 					ConstantOp->SetValue( pImage, DiskCacheContext );
 					Result = ConstantOp;
 				}
@@ -543,12 +543,12 @@ namespace mu
 			{
 				MUTABLE_CPUPROFILER_SCOPE(ConstantLayout);
 
-				mu::Ptr<const Layout> pLayout = pSystem->GetPrivate()->BuildLayout( model, LocalParams.get(), at );
+				TSharedPtr<const FLayout> pLayout = System->GetPrivate()->BuildLayout( model, LocalParams.Get(), at );
 
 				if (pLayout)
 				{
 					mu::Ptr<ASTOpConstantResource> constantOp = new ASTOpConstantResource();
-					constantOp->Type = OP_TYPE::LA_CONSTANT;
+					constantOp->Type = EOpType::LA_CONSTANT;
 					constantOp->SetValue( pLayout, DiskCacheContext);
 					Result = constantOp;
 				}
@@ -559,7 +559,7 @@ namespace mu
 			{
 				MUTABLE_CPUPROFILER_SCOPE(ConstantBool);
 
-				bool value = pSystem->GetPrivate()->BuildBool( model, LocalParams.get(), at );
+				bool value = System->GetPrivate()->BuildBool( model, LocalParams.Get(), at );
 
 				{
 					mu::Ptr<ASTOpConstantBool> constantOp = new ASTOpConstantBool();
@@ -574,11 +574,11 @@ namespace mu
 				MUTABLE_CPUPROFILER_SCOPE(ConstantBool);
 
 				FVector4f ResultColor(0, 0, 0, 0);
-				ResultColor = pSystem->GetPrivate()->BuildColour( model, LocalParams.get(), at );
+				ResultColor = System->GetPrivate()->BuildColour( model, LocalParams.Get(), at );
 
 				{
 					mu::Ptr<ASTOpFixed> constantOp = new ASTOpFixed();
-					constantOp->op.type = OP_TYPE::CO_CONSTANT;
+					constantOp->op.type = EOpType::CO_CONSTANT;
 					constantOp->op.args.ColourConstant.value[0] = ResultColor[0];
 					constantOp->op.args.ColourConstant.value[1] = ResultColor[1];
 					constantOp->op.args.ColourConstant.value[2] = ResultColor[2];
@@ -599,7 +599,7 @@ namespace mu
 				break;
 			}
 
-			pSystem->GetPrivate()->EndBuild();
+			System->GetPrivate()->EndBuild();
 		}
 
 	};
@@ -636,18 +636,18 @@ namespace mu
 				[&ConstantSubgraphs, Pass]
 				(Ptr<ASTOp>& SubgraphRoot)
 				{
-					OP_TYPE SubgraphType = SubgraphRoot->GetOpType();
+					EOpType SubgraphType = SubgraphRoot->GetOpType();
 
 					bool bIsConstantSubgraph = true;
 					switch (SubgraphType)
 					{
-					case OP_TYPE::BO_PARAMETER:
-					case OP_TYPE::NU_PARAMETER:
-					case OP_TYPE::SC_PARAMETER:
-					case OP_TYPE::CO_PARAMETER:
-					case OP_TYPE::PR_PARAMETER:
-					case OP_TYPE::IM_PARAMETER:
-					case OP_TYPE::MA_PARAMETER:
+					case EOpType::BO_PARAMETER:
+					case EOpType::NU_PARAMETER:
+					case EOpType::SC_PARAMETER:
+					case EOpType::CO_PARAMETER:
+					case EOpType::PR_PARAMETER:
+					case EOpType::IM_PARAMETER:
+					case EOpType::MA_PARAMETER:
 						bIsConstantSubgraph = false;
 						break;
 					default:
@@ -669,19 +669,19 @@ namespace mu
 					bool bHasSpecialOpInSubgraph = false;
 					switch (SubgraphType)
 					{
-					case OP_TYPE::IM_BLANKLAYOUT:
-					case OP_TYPE::IM_COMPOSE:
-					case OP_TYPE::ME_MERGE:
-					case OP_TYPE::ME_CLIPWITHMESH:
-					case OP_TYPE::ME_CLIPMORPHPLANE:
-					case OP_TYPE::ME_APPLYPOSE:
-					case OP_TYPE::ME_REMOVEMASK:
-					case OP_TYPE::ME_ADDTAGS:
-					case OP_TYPE::IM_PLAINCOLOUR:
+					case EOpType::IM_BLANKLAYOUT:
+					case EOpType::IM_COMPOSE:
+					case EOpType::ME_MERGE:
+					case EOpType::ME_CLIPWITHMESH:
+					case EOpType::ME_CLIPMORPHPLANE:
+					case EOpType::ME_APPLYPOSE:
+					case EOpType::ME_REMOVEMASK:
+					case EOpType::ME_ADDTAGS:
+					case EOpType::IM_PLAINCOLOUR:
 						bHasSpecialOpInSubgraph = true;
 						break;
 
-					case OP_TYPE::IM_RASTERMESH:
+					case EOpType::IM_RASTERMESH:
 					{
 						const ASTOpImageRasterMesh* Raster = static_cast<const ASTOpImageRasterMesh*>(SubgraphRoot.get());
 						// If this operation is only rastering the mesh UVs, reduce it to constant. Otherwise avoid reducing it
@@ -691,8 +691,8 @@ namespace mu
 						break;
 					}
 
-					case OP_TYPE::ME_REFERENCE:
-					case OP_TYPE::IM_REFERENCE:
+					case EOpType::ME_REFERENCE:
+					case EOpType::IM_REFERENCE:
 						// If we are in a reference-resolution optimization phase, then the ops are not special.
 						if (Pass < 2)
 						{
@@ -812,7 +812,7 @@ namespace mu
 						UE::Tasks::FTaskEvent& SubgraphCompletionEvent = ConstantSubgraphs[Index].CompletedEvent;
 
 						// Referenced images are resolved in its own task to prevent requesting them twice (and loading them twice)
-						bool bIsCompileTimeReferenceImage = (SubgraphRoot->GetOpType() == OP_TYPE::IM_REFERENCE);
+						bool bIsCompileTimeReferenceImage = (SubgraphRoot->GetOpType() == EOpType::IM_REFERENCE);
 
 						// Launch the task with its dependencies
 						if (bIsCompileTimeReferenceImage)
@@ -821,7 +821,7 @@ namespace mu
 							const ASTOpReferenceResource* Typed = static_cast<const ASTOpReferenceResource*>(SubgraphRoot.get());
 							uint32 ImageID = Typed->ID;
 
-							TSharedPtr< Ptr<Image> > ResolveImage = MakeShared<Ptr<Image>>();
+							TSharedPtr< TSharedPtr<FImage> > ResolveImage = MakeShared<TSharedPtr<FImage>>();
 
 							constexpr bool bRunImmediatlyIfPossible = false;
 							UE::Tasks::FTask ReferenceCompletion = InOptions->OptimisationOptions.ReferencedResourceProvider(ImageID, ResolveImage, bRunImmediatlyIfPossible);
@@ -833,12 +833,12 @@ namespace mu
 									{
 										MUTABLE_CPUPROFILER_SCOPE(MutableResolveComplete_CreateConstant);
 										ConstantOp = new ASTOpConstantResource;
-										ConstantOp->Type = OP_TYPE::IM_CONSTANT;
+										ConstantOp->Type = EOpType::IM_CONSTANT;
 										{
 											MUTABLE_CPUPROFILER_SCOPE(GetSourceDataDescriptor);
 											ConstantOp->SourceDataDescriptor = SubgraphRoot->GetSourceDataDescriptor();
 										}
-										ConstantOp->SetValue(ResolveImage->get(), InOptions->OptimisationOptions.DiskCacheContext);
+										ConstantOp->SetValue(*ResolveImage, InOptions->OptimisationOptions.DiskCacheContext);
 									}
 									{
 										MUTABLE_CPUPROFILER_SCOPE(MutableResolveComplete_Replace);
@@ -912,7 +912,7 @@ namespace mu
 					UE::Tasks::FTaskEvent& SubgraphCompletionEvent = ConstantSubgraphs[Index].CompletedEvent;
 
 					// Referenced images are resolved in its own task to prevent requesting them twice (and loading them twice)
-					bool bIsCompileTimeReferenceImage = (SubgraphRoot->GetOpType() == OP_TYPE::IM_REFERENCE);
+					bool bIsCompileTimeReferenceImage = (SubgraphRoot->GetOpType() == EOpType::IM_REFERENCE);
 
 					// Launch the task with its dependencies
 					if (bIsCompileTimeReferenceImage)
@@ -923,14 +923,14 @@ namespace mu
 
 						auto ImmediateCompleteFunc = [SubgraphRoot, InOptions, ImageID]()
 							{
-								TSharedPtr< Ptr<Image> > ResolveImage = MakeShared<Ptr<Image>>();
+								TSharedPtr< TSharedPtr<FImage> > ResolveImage = MakeShared<TSharedPtr<FImage>>();
 								constexpr bool bRunImmediatlyIfPossible = true;
 								UE::Tasks::FTask ReferenceCompletion = InOptions->OptimisationOptions.ReferencedResourceProvider(ImageID, ResolveImage, bRunImmediatlyIfPossible);
 								ReferenceCompletion.Wait();
 								Ptr<ASTOpConstantResource> ConstantOp = new ASTOpConstantResource;
 								ConstantOp->SourceDataDescriptor = SubgraphRoot->GetSourceDataDescriptor();
-								ConstantOp->Type = OP_TYPE::IM_CONSTANT;
-								ConstantOp->SetValue(ResolveImage->get(), InOptions->OptimisationOptions.DiskCacheContext);
+								ConstantOp->Type = EOpType::IM_CONSTANT;
+								ConstantOp->SetValue(*ResolveImage, InOptions->OptimisationOptions.DiskCacheContext);
 								ASTOp::Replace(SubgraphRoot, ConstantOp);
 							};
 						ConstantSubgraphs[Index].NonConcurrentTask = ImmediateCompleteFunc;
@@ -1034,9 +1034,9 @@ namespace mu
 	//-------------------------------------------------------------------------------------------------
 	//-------------------------------------------------------------------------------------------------
 	CodeOptimiser::CodeOptimiser(Ptr<CompilerOptions> InOptions, TArray<FStateCompilationData>& InStates )
-		: m_states( InStates )
+		: States( InStates )
 	{
-		m_options = InOptions;
+		Options = InOptions;
 	}
 
 
@@ -1057,13 +1057,13 @@ namespace mu
 			// All kind of optimisations that depend on the meaning of each operation
 			// \TODO: We are doing it for all states.
 			UE_LOG(LogMutableCore, Verbose, TEXT(" - semantic optimiser"));
-			modified |= SemanticOptimiserAST( roots, m_options->GetPrivate()->OptimisationOptions, Pass );
+			modified |= SemanticOptimiserAST( roots, Options->GetPrivate()->OptimisationOptions, Pass );
 			//AXE_INT_VALUE("Mutable", Verbose, "ast size", (int64_t)ASTOp::CountNodes(roots));
 			//UE_LOG(LogMutableCore, Verbose, TEXT("(int) %s : %ld"), TEXT("ast size"), int64(ASTOp::CountNodes(roots)));
 			ASTOp::LogHistogram(roots);
 
 			UE_LOG(LogMutableCore, Verbose, TEXT(" - sink optimiser"));
-			modified |= SinkOptimiserAST( roots, m_options->GetPrivate()->OptimisationOptions );
+			modified |= SinkOptimiserAST( roots, Options->GetPrivate()->OptimisationOptions );
 			//AXE_INT_VALUE("Mutable", Verbose, "ast size", (int64_t)ASTOp::CountNodes(roots));
 			ASTOp::LogHistogram(roots);
 
@@ -1076,7 +1076,7 @@ namespace mu
 			// TODO: Conservative logic optimiser
 			if (!modified)
 			{
-	//                AXE_INT_VALUE("Mutable", Verbose, "program size", (int64_t)program.m_opAddress.Num());
+	//                AXE_INT_VALUE("Mutable", Verbose, "program size", (int64_t)program.OpAddress.Num());
 	//                UE_LOG(LogMutableCore, Verbose, " - logic optimiser");
 	//                LocalLogicOptimiser log;
 	//                modified |= log.Apply( program, (int)s );
@@ -1104,7 +1104,7 @@ namespace mu
 			UE_LOG(LogMutableCore, Verbose, TEXT(" - constant generator"));
 
 			// Constant subtree generation
-			modified = ConstantGeneratorAST( m_options->GetPrivate(), Root, Pass );
+			modified = ConstantGeneratorAST( Options->GetPrivate(), Root, Pass );
 		}
 
 		//UE_LOG(LogMutableCore, Verbose, TEXT("(int) %s : %ld"), TEXT("ast size"), int64(ASTOp::CountNodes(roots)));
@@ -1124,7 +1124,7 @@ namespace mu
 
 		//if (!modified)
 		{
-			//for ( size_t s=0;  s<program.m_states.size(); ++s )
+			//for ( size_t s=0;  s<program.States.size(); ++s )
 			{
 				ASTOp::LogHistogram(roots);
 				UE_LOG(LogMutableCore, Verbose, TEXT(" - logic optimiser"));
@@ -1165,7 +1165,7 @@ namespace mu
 			switch (node->GetOpType())
 			{
 
-			case OP_TYPE::ME_CONSTANT:
+			case EOpType::ME_CONSTANT:
 			{
 				mu::Ptr<ASTOpConstantResource> typedOp = static_cast<ASTOpConstantResource*>(node.get());
 
@@ -1177,7 +1177,7 @@ namespace mu
 				return false;
 			}
 
-			case OP_TYPE::ME_CLIPMORPHPLANE:
+			case EOpType::ME_CLIPMORPHPLANE:
 			{
 				ASTOpMeshClipMorphPlane* typedOp = static_cast<ASTOpMeshClipMorphPlane*>(node.get());
 				if (typedOp->VertexSelectionType == EClipVertexSelectionType::BoneHierarchy)
@@ -1190,7 +1190,7 @@ namespace mu
 				return true;
 			}
 
-			case OP_TYPE::ME_APPLYPOSE:
+			case EOpType::ME_APPLYPOSE:
 			{
 				ASTOpMeshApplyPose* typedOp = static_cast<ASTOpMeshApplyPose*>(node.get());
 
@@ -1202,7 +1202,7 @@ namespace mu
 				break;
 			}
 
-			case OP_TYPE::ME_BINDSHAPE:
+			case EOpType::ME_BINDSHAPE:
 			{
 				ASTOpMeshBindShape* typedOp = static_cast<ASTOpMeshBindShape*>(node.get());
 				if (typedOp->bReshapeSkeleton)
@@ -1214,7 +1214,7 @@ namespace mu
 				break;
 			}
 
-			case OP_TYPE::ME_APPLYSHAPE:
+			case EOpType::ME_APPLYSHAPE:
 			{
 				ASTOpMeshApplyShape* typedOp = static_cast<ASTOpMeshApplyShape*>(node.get());
 				if (typedOp->bReshapeSkeleton)
@@ -1241,17 +1241,17 @@ namespace mu
 	// for this op.
 	struct FAddMeshSkeleton
 	{
-		mu::Ptr<ASTOp> m_pAddMeshOp;
-		TArray<mu::Ptr<ASTOpConstantResource>> m_contributingMeshes;
-		mu::Ptr<Skeleton> m_pFinalSkeleton;
+		mu::Ptr<ASTOp> AddMeshOp;
+		TArray<mu::Ptr<ASTOpConstantResource>> ContributingMeshes;
+		TSharedPtr<FSkeleton> FinalSkeleton;
 
-		FAddMeshSkeleton( const mu::Ptr<ASTOp>& pAddMeshOp,
-						  TArray<mu::Ptr<ASTOpConstantResource>>& contributingMeshes,
-						  const mu::Ptr<Skeleton>& pFinalSkeleton )
+		FAddMeshSkeleton( const mu::Ptr<ASTOp>& InAddMeshOp,
+						  TArray<mu::Ptr<ASTOpConstantResource>>& InContributingMeshes,
+						  const TSharedPtr<FSkeleton>& InFinalSkeleton )
 		{
-			m_pAddMeshOp = pAddMeshOp;
-			m_contributingMeshes = MoveTemp(contributingMeshes);
-			m_pFinalSkeleton = pFinalSkeleton;
+			AddMeshOp = InAddMeshOp;
+			ContributingMeshes = MoveTemp(InContributingMeshes);
+			FinalSkeleton = InFinalSkeleton;
 		}
 	};
 
@@ -1270,7 +1270,7 @@ namespace mu
 			// Only recurse instance construction ops.
 			bool processChildren = GetOpDataType(at->GetOpType())==DT_INSTANCE;
 
-			if ( at->GetOpType() == OP_TYPE::IN_ADDMESH )
+			if ( at->GetOpType() == EOpType::IN_ADDMESH )
 			{
 				ASTOpInstanceAdd* typedNode = static_cast<ASTOpInstanceAdd*>(at.get());
 				mu::Ptr<ASTOp> meshRoot = typedNode->value.child();
@@ -1284,7 +1284,7 @@ namespace mu
 					ASTOp::Traverse_TopDown_Unique_Imprecise( tempRoots, [&](mu::Ptr<ASTOp>& lat )
 					{
 						// \todo: refine to avoid instruction branches with irrelevant skeletons.
-						if ( lat->GetOpType() == OP_TYPE::ME_CONSTANT )
+						if ( lat->GetOpType() == EOpType::ME_CONSTANT )
 						{
 							mu::Ptr<ASTOpConstantResource> typedOp = static_cast<ASTOpConstantResource*>(lat.get());
 							if ( subtreeMeshes.Find(typedOp)
@@ -1298,18 +1298,18 @@ namespace mu
 					});
 
 					// Create a mesh just with the unified skeleton
-					SkeletonPtr pFinalSkeleton = new Skeleton;
+					TSharedPtr<FSkeleton> FinalSkeleton = MakeShared<FSkeleton>();
 					for (const auto& meshAt: subtreeMeshes)
 					{
-						MeshPtrConst pMesh = static_cast<const Mesh*>(meshAt->GetValue().get());
-						mu::Ptr<const Skeleton> sourceSkeleton = pMesh ? pMesh->GetSkeleton() : nullptr;
-						if (sourceSkeleton)
+						TSharedPtr<const FMesh> pMesh = StaticCastSharedPtr<const FMesh>(meshAt->GetValue());
+						TSharedPtr<const FSkeleton> SourceSkeleton = pMesh ? pMesh->GetSkeleton() : nullptr;
+						if (SourceSkeleton)
 						{
-							ExtendSkeleton(pFinalSkeleton.get(),sourceSkeleton.get());
+							ExtendSkeleton(FinalSkeleton.Get(),SourceSkeleton.Get());
 						}
 					}
 
-					replacementsFound.Emplace( at, subtreeMeshes, pFinalSkeleton );
+					replacementsFound.Emplace( at, subtreeMeshes, FinalSkeleton );
 				}
 			}
 
@@ -1320,26 +1320,26 @@ namespace mu
 		// Iterate all meshes again
 		ASTOp::Traverse_TopDown_Unique_Imprecise( roots, [&](mu::Ptr<ASTOp>& at )
 		{
-			if (at->GetOpType()==OP_TYPE::ME_CONSTANT)
+			if (at->GetOpType()==EOpType::ME_CONSTANT)
 			{
 				ASTOpConstantResource* typedOp = static_cast<ASTOpConstantResource*>(at.get());
 
 				for(FAddMeshSkeleton& Rep: replacementsFound)
 				{
-					if (Rep.m_contributingMeshes.Contains(at))
+					if (Rep.ContributingMeshes.Contains(at))
 					{
-						mu::Ptr<const Mesh> pMesh =  static_cast<const Mesh*>(typedOp->GetValue().get());
+						TSharedPtr<const FMesh> pMesh = StaticCastSharedPtr<const FMesh>(typedOp->GetValue());
 						pMesh->CheckIntegrity();
 
-						Ptr<Mesh> NewMesh = new Mesh();
+						TSharedPtr<FMesh> NewMesh = MakeShared<FMesh>();
 						bool bOutSuccess = false;
-						MeshRemapSkeleton(NewMesh.get(), pMesh.get(), Rep.m_pFinalSkeleton.get(), bOutSuccess);
+						MeshRemapSkeleton(NewMesh.Get(), pMesh.Get(), Rep.FinalSkeleton, bOutSuccess);
 
 						if (bOutSuccess)
 						{
 							NewMesh->CheckIntegrity();
 							mu::Ptr<ASTOpConstantResource> newOp = new ASTOpConstantResource();
-							newOp->Type = OP_TYPE::ME_CONSTANT;
+							newOp->Type = EOpType::ME_CONSTANT;
 							newOp->SetValue(NewMesh, options.DiskCacheContext);
 							newOp->SourceDataDescriptor = at->GetSourceDataDescriptor();
 
@@ -1353,18 +1353,18 @@ namespace mu
 
 	//    for( auto& rep: replacementsFound )
 	//    {
-	//        check(rep.m_pAddMeshOp->GetOpType()==OP_TYPE::IN_ADDMESH);
-	//        auto newAddMeshRootTyped = dynamic_cast<ASTOpInstanceAdd*>(rep.m_pAddMeshOp.get());
+	//        check(rep.AddMeshOp->GetOpType()==EOpType::IN_ADDMESH);
+	//        auto newAddMeshRootTyped = dynamic_cast<ASTOpInstanceAdd*>(rep.AddMeshOp.get());
 	//        auto meshRoot = newAddMeshRootTyped->value.child();
 
 	//        // Add new skeleton constant instruction
 	//        Ptr<ASTOpConstantResource> cop = new ASTOpConstantResource();
-	//        cop->type = OP_TYPE::ME_CONSTANT;
-	//        cop->SetValue( rep.m_pFinalSkeleton, options.bUseDiskCache );
+	//        cop->type = EOpType::ME_CONSTANT;
+	//        cop->SetValue( rep.FinalSkeleton, options.bUseDiskCache );
 
 	//        // Add new "apply skeleton" instruction
 	//        Ptr<ASTOpFixed> skop = new ASTOpFixed();
-	//        skop->op.type = OP_TYPE::ME_SETSKELETON;
+	//        skop->op.type = EOpType::ME_SETSKELETON;
 	//        skop->SetChild( skop->op.args.MeshSetSkeleton.source, meshRoot);
 	//        skop->SetChild( skop->op.args.MeshSetSkeleton.skeleton, cop);
 
@@ -1373,7 +1373,7 @@ namespace mu
 	//        newOp->value = skop;
 
 	//        // TODO BLEH UNSAFE?
-	//        ASTOp::Replace(rep.m_pAddMeshOp,newOp);
+	//        ASTOp::Replace(rep.AddMeshOp,newOp);
 	//    }
 	}
 
@@ -1385,18 +1385,18 @@ namespace mu
 
 		// Gather all the roots (one for each state)
 		TArray<Ptr<ASTOp>> roots;
-		for(const FStateCompilationData& s:m_states)
+		for(const FStateCompilationData& s:States)
 		{
 			roots.Add(s.root);
 		}
 
 		//UE_LOG(LogMutableCore, Verbose, TEXT("(int) %s : %ld"), TEXT("ast size"), int64(ASTOp::CountNodes(roots)));
 
-		if ( m_options->GetPrivate()->OptimisationOptions.bEnabled )
+		if ( Options->GetPrivate()->OptimisationOptions.bEnabled )
 		{
 			// We use 4 times the count because at the time we moved to sharing this count it
 			// was being used 4 times, and we want to keep the tests consistent.
-			m_optimizeIterationsLeft = m_options->GetPrivate()->OptimisationOptions.MaxOptimisationLoopCount * 4;
+			m_optimizeIterationsLeft = Options->GetPrivate()->OptimisationOptions.MaxOptimisationLoopCount * 4;
 			m_optimizeIterationsMax = m_optimizeIterationsLeft;
 
 			// The first duplicated data remover has the special mission of removing
@@ -1414,12 +1414,12 @@ namespace mu
 			//UE_LOG(LogMutableCore, Verbose, TEXT("(int) %s : %ld"), TEXT("ast size"), int64(ASTOp::CountNodes(roots)));
 
 			// Special optimization stages
-			if ( m_options->GetPrivate()->OptimisationOptions.bUniformizeSkeleton )
+			if ( Options->GetPrivate()->OptimisationOptions.bUniformizeSkeleton )
 			{
 				UE_LOG(LogMutableCore, Verbose, TEXT(" - skeleton cleaner"));
 				ASTOp::LogHistogram(roots);
 
-				SkeletonCleanerAST( roots, m_options->GetPrivate()->OptimisationOptions );
+				SkeletonCleanerAST( roots, Options->GetPrivate()->OptimisationOptions );
 				//UE_LOG(LogMutableCore, Verbose, TEXT("(int) %s : %ld"), TEXT("ast size"), int64(ASTOp::CountNodes(roots)));
 				ASTOp::LogHistogram(roots);
 			}
@@ -1471,7 +1471,7 @@ namespace mu
 				for (Ptr<ASTOp>& Root : roots)
 				{
 					// Constant subtree generation
-					modified = ConstantGeneratorAST(m_options->GetPrivate(), Root, Pass);
+					modified = ConstantGeneratorAST(Options->GetPrivate(), Root, Pass);
 				}
 
 				DuplicatedDataRemoverAST(roots);
@@ -1492,21 +1492,21 @@ namespace mu
 			// Analyse mesh constants to see which of them are in optimised mesh formats, and set the flags.
 			ASTOp::Traverse_BottomUp_Unique_NonReentrant( roots, [&](Ptr<ASTOp>& n)
 			{
-				if (n->GetOpType()==OP_TYPE::ME_CONSTANT)
+				if (n->GetOpType()==EOpType::ME_CONSTANT)
 				{
 					ASTOpConstantResource* typed = static_cast<ASTOpConstantResource*>(n.get());
-					auto pMesh = static_cast<const Mesh*>(typed->GetValue().get());
+					auto pMesh = StaticCastSharedPtr<const FMesh>(typed->GetValue());
 					pMesh->ResetStaticFormatFlags();
-					typed->SetValue( pMesh, m_options->GetPrivate()->OptimisationOptions.DiskCacheContext);
+					typed->SetValue( pMesh, Options->GetPrivate()->OptimisationOptions.DiskCacheContext);
 				}
 			});
 
 			ASTOp::LogHistogram(roots);
 
 			// Reset the state root operations in case they have changed due to optimization
-			for (int32 RootIndex = 0; RootIndex < m_states.Num(); ++RootIndex)
+			for (int32 RootIndex = 0; RootIndex < States.Num(); ++RootIndex)
 			{
-				m_states[RootIndex].root = roots[RootIndex];
+				States[RootIndex].root = roots[RootIndex];
 			}
 
 			{
@@ -1523,7 +1523,7 @@ namespace mu
 		}
 
 		// Minimal optimisation of constant subtrees
-		else if ( m_options->GetPrivate()->OptimisationOptions.bConstReduction )
+		else if ( Options->GetPrivate()->OptimisationOptions.bConstReduction )
 		{
 			// The first duplicated data remover has the special mission of removing
 			// duplicated data (meshes) that may have been specified in the source
@@ -1543,12 +1543,12 @@ namespace mu
 				FullOptimiseAST(roots, 2);
 			}
 
-			for ( int32 StateIndex=0; StateIndex <m_states.Num(); ++StateIndex)
+			for ( int32 StateIndex=0; StateIndex <States.Num(); ++StateIndex)
 			{
 				constexpr int32 Pass = 1;
 
 				UE_LOG(LogMutableCore, Verbose, TEXT(" - constant generator"));
-				ConstantGeneratorAST( m_options->GetPrivate(), roots[StateIndex], Pass);
+				ConstantGeneratorAST( Options->GetPrivate(), roots[StateIndex], Pass);
 				//AXE_INT_VALUE("Mutable", Verbose, "ast size", (int64_t)ASTOp::CountNodes(roots));
 			}
 
@@ -1561,9 +1561,9 @@ namespace mu
 			//AXE_INT_VALUE("Mutable", Verbose, "ast size", (int64_t)ASTOp::CountNodes(roots));
 
 			// Reset the state root operations in case they have changed due to optimization
-			for (int32 RootIndex = 0; RootIndex < m_states.Num(); ++RootIndex)
+			for (int32 RootIndex = 0; RootIndex < States.Num(); ++RootIndex)
 			{
-				m_states[RootIndex].root = roots[RootIndex];
+				States[RootIndex].root = roots[RootIndex];
 			}
 		}
 
diff --git a/Engine/Plugins/Mutable/Source/MutableTools/Private/MuT/CodeOptimiser.h b/Engine/Plugins/Mutable/Source/MutableTools/Private/MuT/CodeOptimiser.h
--- a/Engine/Plugins/Mutable/Source/MutableTools/Private/MuT/CodeOptimiser.h
+++ b/Engine/Plugins/Mutable/Source/MutableTools/Private/MuT/CodeOptimiser.h
@@ -23,10 +23,10 @@ namespace mu
 
     private:
 
-		Ptr<CompilerOptions> m_options;
+		Ptr<CompilerOptions> Options;
 
         //!
-		TArray<FStateCompilationData>& m_states;
+		TArray<FStateCompilationData>& States;
 
         //! The max number of optimize iterations is shared across several stages now.
         //! This is how many are left
@@ -52,7 +52,7 @@ namespace mu
     {
     public:
 
-        FindOpTypeVisitor(const TArray<OP_TYPE>& types)
+        FindOpTypeVisitor(const TArray<EOpType>& types)
             : m_typesToFind(types)
         {
         }
@@ -62,9 +62,9 @@ namespace mu
 
     private:
 
-        TArray<OP_TYPE> m_typesToFind;
+        TArray<EOpType> m_typesToFind;
 
-		TArray< TPair<bool,OP::ADDRESS> > m_pending;
+		TArray< TPair<bool,OP::ADDRESS> > Pending;
 
         // 0 not visited
         // 1 children pending
@@ -149,12 +149,10 @@ namespace mu
     };
 
 
-    //---------------------------------------------------------------------------------------------
-    //! Scan the code in the given subtree and return true if a state runtime parameters is
-    //! found.
-    //! Intermediate data is used betwen calls to apply, so don't remove program code or directly
-    //! change the instructions. Adding new instructions is ok.
-    //---------------------------------------------------------------------------------------------
+    /** Scan the code in the given subtree and return true if a state runtime parameters is found.
+    * Intermediate data is used betwen calls to apply, so don't remove program code or directly
+    * change the instructions. Adding new instructions is ok.
+	*/
     class RuntimeParameterVisitorAST
     {
     public:
@@ -165,41 +163,41 @@ namespace mu
 
     private:
 
-        const FStateCompilationData* m_pState;
+        const FStateCompilationData* State;
 
         //!
-        struct PENDING_ITEM
+        struct FPendingItem
         {
             //! 0: indicate subtree pending
             //! 1: indicate children finished
-            uint8_t itemType;
+            uint8 ItemType;
 
             //! 0: everything is relevant
             //! 1: only layouts are relevant
-            uint8_t onlyLayoutsRelevant;
+            uint8 OnlyLayoutsRelevant;
 
             //! Address to visit
-            Ptr<ASTOp> at;
+            Ptr<ASTOp> At;
         };
 
         //!
-        vector< PENDING_ITEM > m_pending;
+        TArray< FPendingItem > Pending;
 
         //! Possible op state
-        enum class OP_STATE : uint8_t
+        enum class EOpState : uint8
         {
-            NOT_VISITED = 0,
-            CHILDREN_PENDING_FULL,
-            CHILDREN_PENDING_PARTIAL,
-            VISITED_HASRUNTIME,
-            VISITED_FULL_DOESNTHAVERUNTIME,
-            VISITED_PARTIAL_DOESNTHAVERUNTIME
+            NotVisited = 0,
+            ChildrenPendingFull,
+            ChildrenPendingPartial,
+            VisitedHasRuntime,
+            VisitedFullDoesntHaveRuntime,
+            VisitedPartialDoesntHaveRuntime
         };
 
-        std::unordered_map<Ptr<ASTOp>,OP_STATE> m_visited;
+        TMap<Ptr<ASTOp>, EOpState> Visited;
 
         //!
-        void AddIfNeeded( const PENDING_ITEM& item );
+        void AddIfNeeded( const FPendingItem& item );
 
     };
 
@@ -263,40 +261,39 @@ namespace mu
     extern Ptr<ASTOp> EnsureValidMask( Ptr<ASTOp> mask, Ptr<ASTOp> base );
 
 
-    //---------------------------------------------------------------------------------------------
-    //! Calculate all the parameters found relevant under a particular operation. This may not
-    //! incldue all the parameters in the subtree (if because of the operations they are not
-    //! relevant)
-    //! It has an internal cache, so don't reuse if the program changes.
-    //---------------------------------------------------------------------------------------------
+    /** Calculate all the parameters found relevant under a particular operation.This may not
+    * incldue all the parameters in the subtree (if because of the operations they are not
+    * relevant)
+    * It has an internal cache, so don't reuse if the program changes.
+	*/
     class SubtreeRelevantParametersVisitorAST
     {
     public:
 
-        void Run( Ptr<ASTOp> root );
+        void Run( Ptr<ASTOp> Root );
 
         //! After Run, list of relevant parameters.
-        TSet< FString > m_params;
+        TSet< FString > Parameters;
 
     private:
 
         struct FState
         {
-            Ptr<ASTOp> op;
+            Ptr<ASTOp> Op;
             bool bOnlyLayoutIsRelevant=false;
 
-			FState( Ptr<ASTOp> o=nullptr, bool l=false) : op(o), bOnlyLayoutIsRelevant(l) {}
+			FState( Ptr<ASTOp> o=nullptr, bool l=false) : Op(o), bOnlyLayoutIsRelevant(l) {}
 
             bool operator==(const FState& o) const
             {
-                return  op == o.op &&
+                return  Op == o.Op &&
 					bOnlyLayoutIsRelevant == o.bOnlyLayoutIsRelevant;
             }
 
 			friend FORCEINLINE uint32 GetTypeHash(const FState& InKey)
 			{
 				uint32 KeyHash = 0;
-				KeyHash = HashCombineFast(KeyHash, ::GetTypeHash(InKey.op.get()));
+				KeyHash = HashCombineFast(KeyHash, ::GetTypeHash(InKey.Op.get()));
 				KeyHash = HashCombineFast(KeyHash, ::GetTypeHash(InKey.bOnlyLayoutIsRelevant));
 				return KeyHash;
 			}
diff --git a/Engine/Plugins/Mutable/Source/MutableTools/Private/MuT/CodeOptimiser_Logic.cpp b/Engine/Plugins/Mutable/Source/MutableTools/Private/MuT/CodeOptimiser_Logic.cpp
--- a/Engine/Plugins/Mutable/Source/MutableTools/Private/MuT/CodeOptimiser_Logic.cpp
+++ b/Engine/Plugins/Mutable/Source/MutableTools/Private/MuT/CodeOptimiser_Logic.cpp
@@ -24,15 +24,15 @@ namespace mu
 //        IntExpressionPtr pRes = new IntExpression;
 
 //        const OP& op = program.m_code[at];
-//        OP_TYPE type = op.type;
+//        EOpType type = op.type;
 //        switch ( type )
 //        {
-//        case OP_TYPE::NU_CONSTANT:
+//        case EOpType::NU_CONSTANT:
 //            pRes->type = IntExpression::IET_CONSTANT;
 //            pRes->value = op.args.IntConstant.value;
 //            break;
 
-//        case OP_TYPE::NU_PARAMETER:
+//        case EOpType::NU_PARAMETER:
 //            pRes->type = IntExpression::IET_PARAMETER;
 //            pRes->parameter = op.args.Parameter.variable;
 //            break;
@@ -54,46 +54,46 @@ namespace mu
 //                                         const FProgram& program,
 //                                         BOOL_EXPRESSION_CACHE* cache )
 //    {
-//        if (cache && (size_t)at<=cache->m_values.size() && cache->m_values[at])
+//        if (cache && (size_t)at<=cache->Values.size() && cache->Values[at])
 //        {
-//            return cache->m_values[at];
+//            return cache->Values[at];
 //        }
 
 //        BoolExpressionPtr pRes = new BoolExpression;
 
 //        const OP& op = program.m_code[at];
-//        OP_TYPE type = (OP_TYPE)op.type;
+//        EOpType type = (EOpType)op.type;
 //        switch ( type )
 //        {
-//        case OP_TYPE::BO_CONSTANT:
+//        case EOpType::BO_CONSTANT:
 //            pRes->type = op.args.BoolConstant.value
 //                ? BoolExpression::BET_TRUE
 //                : BoolExpression::BET_FALSE;
 //            break;
 
-//        case OP_TYPE::BO_PARAMETER:
+//        case EOpType::BO_PARAMETER:
 //            pRes->type = BoolExpression::BET_PARAMETER;
 //            pRes->parameter = op.args.Parameter.variable;
 //            break;
 
-//        case OP_TYPE::BO_AND:
+//        case EOpType::BO_AND:
 //            pRes->type = BoolExpression::BET_AND;
 //            pRes->a = GetBoolExpression( op.args.BoolBinary.a, program, cache );
 //            pRes->b = GetBoolExpression( op.args.BoolBinary.b, program, cache );
 //            break;
 
-//        case OP_TYPE::BO_OR:
+//        case EOpType::BO_OR:
 //            pRes->type = BoolExpression::BET_OR;
 //            pRes->a = GetBoolExpression( op.args.BoolBinary.a, program, cache );
 //            pRes->b = GetBoolExpression( op.args.BoolBinary.b, program, cache );
 //            break;
 
-//        case OP_TYPE::BO_NOT:
+//        case EOpType::BO_NOT:
 //            pRes->type = BoolExpression::BET_NOT;
 //            pRes->a = GetBoolExpression( op.args.BoolNot.source, program, cache );
 //            break;
 
-//        case OP_TYPE::BO_EQUAL_INT_CONST:
+//        case EOpType::BO_EQUAL_INT_CONST:
 //            pRes->type = BoolExpression::BET_INT_EQUAL_CONST;
 //            pRes->intExp = GetIntExpression( op.args.BoolEqualScalarConst.value, program );
 //            pRes->intValue = op.args.BoolEqualScalarConst.constant;
@@ -107,11 +107,11 @@ namespace mu
 
 //        if (cache)
 //        {
-//            if (cache->m_values.size()<program.m_opAddress.Num())
+//            if (cache->Values.size()<program.OpAddress.Num())
 //            {
-//                cache->m_values.resize(program.m_opAddress.Num());
+//                cache->Values.resize(program.OpAddress.Num());
 //            }
-//            cache->m_values[at]=pRes;
+//            cache->Values[at]=pRes;
 //        }
 
 //        return pRes;
@@ -124,7 +124,7 @@ namespace mu
 //    void LogicCallstackGather::Apply( FProgram& program, int state )
 //    {
 //        MUTABLE_CPUPROFILER_SCOPE(LogicCallstackGather);
-//        Traverse( program.m_states[state].m_root, program );
+//        Traverse( program.States[state].Root, program );
 //    }
 
 
@@ -132,20 +132,20 @@ namespace mu
 //    void LogicCallstackGather::Visit( OP::ADDRESS at, FProgram& program )
 //    {
 //        // We traverse top down
-//        OP_TYPE type = (OP_TYPE)program.m_code[at].type;
+//        EOpType type = (EOpType)program.m_code[at].type;
 //        OP::ARGS args = program.m_code[at].args;
 
 //        switch ( type )
 //        {
 
 //        //-----------------------------------------------------------------------------------------
-//        case OP_TYPE::NU_CONDITIONAL:
-//        case OP_TYPE::SC_CONDITIONAL:
-//        case OP_TYPE::CO_CONDITIONAL:
-//        case OP_TYPE::IM_CONDITIONAL:
-//        case OP_TYPE::ME_CONDITIONAL:
-//        case OP_TYPE::LA_CONDITIONAL:
-//        case OP_TYPE::IN_CONDITIONAL:
+//        case EOpType::NU_CONDITIONAL:
+//        case EOpType::SC_CONDITIONAL:
+//        case EOpType::CO_CONDITIONAL:
+//        case EOpType::IM_CONDITIONAL:
+//        case EOpType::ME_CONDITIONAL:
+//        case EOpType::LA_CONDITIONAL:
+//        case EOpType::IN_CONDITIONAL:
 //        {
 //            m_factsPerCallstack[at].push_back( m_facts );
 
@@ -182,19 +182,19 @@ namespace mu
 //            break;
 //        }
 
-//        case OP_TYPE::NU_SWITCH:
-//        case OP_TYPE::SC_SWITCH:
-//        case OP_TYPE::CO_SWITCH:
-//        case OP_TYPE::IM_SWITCH:
-//        case OP_TYPE::VO_SWITCH:
-//        case OP_TYPE::ME_SWITCH:
-//        case OP_TYPE::LA_SWITCH:
+//        case EOpType::NU_SWITCH:
+//        case EOpType::SC_SWITCH:
+//        case EOpType::CO_SWITCH:
+//        case EOpType::IM_SWITCH:
+//        case EOpType::VO_SWITCH:
+//        case EOpType::ME_SWITCH:
+//        case EOpType::LA_SWITCH:
 //        {
 //            m_factsPerCallstack[at].push_back( m_facts );
 
 //            IntExpressionPtr pValue = GetIntExpression( args.Switch.variable, program );
 
-//            const auto& options = program.m_constantSwitches[args.Switch.options].m_options;
+//            const auto& options = program.m_constantSwitches[args.Switch.options].Options;
 //            for ( size_t i=0; i<options.size(); ++i )
 //            {
 //                if ( options[i].at )
@@ -247,7 +247,7 @@ namespace mu
 //        m_callstacks.Apply( program, state );
 
 //        m_modified = false;
-//        Traverse( program.m_states[state].m_root, program );
+//        Traverse( program.States[state].Root, program );
 
 //        return m_modified;
 //    }
@@ -258,20 +258,20 @@ namespace mu
 //    {
 //        // We traverse top down
 
-//        OP_TYPE type = (OP_TYPE)program.m_code[at].type;
+//        EOpType type = (EOpType)program.m_code[at].type;
 //        OP::ARGS args = program.m_code[at].args;
 
 //        switch ( type )
 //        {
 
 //        //-----------------------------------------------------------------------------------------
-//        case OP_TYPE::NU_CONDITIONAL:
-//        case OP_TYPE::SC_CONDITIONAL:
-//        case OP_TYPE::CO_CONDITIONAL:
-//        case OP_TYPE::IM_CONDITIONAL:
-//        case OP_TYPE::ME_CONDITIONAL:
-//        case OP_TYPE::LA_CONDITIONAL:
-//        case OP_TYPE::IN_CONDITIONAL:
+//        case EOpType::NU_CONDITIONAL:
+//        case EOpType::SC_CONDITIONAL:
+//        case EOpType::CO_CONDITIONAL:
+//        case EOpType::IM_CONDITIONAL:
+//        case EOpType::ME_CONDITIONAL:
+//        case EOpType::LA_CONDITIONAL:
+//        case EOpType::IN_CONDITIONAL:
 //        {
 //            BOOL_EXPRESSION_CACHE localBoolExpressionCache;
 //            BoolExpressionPtr pCond = GetBoolExpression( args.Conditional.condition, program, &localBoolExpressionCache );
@@ -313,13 +313,13 @@ namespace mu
 //            break;
 //        }
 
-////        case OP_TYPE::NU_SWITCH:
-////        case OP_TYPE::SC_SWITCH:
-////        case OP_TYPE::CO_SWITCH:
-////        case OP_TYPE::IM_SWITCH:
-////        case OP_TYPE::VO_SWITCH:
-////        case OP_TYPE::ME_SWITCH:
-////        case OP_TYPE::LA_SWITCH:
+////        case EOpType::NU_SWITCH:
+////        case EOpType::SC_SWITCH:
+////        case EOpType::CO_SWITCH:
+////        case EOpType::IM_SWITCH:
+////        case EOpType::VO_SWITCH:
+////        case EOpType::ME_SWITCH:
+////        case EOpType::LA_SWITCH:
 ////        {
 ////            IntExpressionPtr pValue = GetIntExpression( args.Switch.variable, program );
 
@@ -397,9 +397,9 @@ namespace mu
 //    //---------------------------------------------------------------------------------------------
 //    bool LogicOptimiser::Apply( FProgram& program, int state )
 //    {
-//        m_initialCodeSize = program.m_opAddress.Num();
+//        m_initialCodeSize = program.OpAddress.Num();
 //        m_modified = false;
-//        Traverse( program.m_states[state].m_root, program );
+//        Traverse( program.States[state].Root, program );
 
 //        return m_modified;
 //    }
@@ -410,14 +410,14 @@ namespace mu
 //    {
 //        // Sanity stop
 //        const int sanityFactor = 4;
-//        if ( program.m_opAddress.Num()>m_initialCodeSize*sanityFactor)
+//        if ( program.OpAddress.Num()>m_initialCodeSize*sanityFactor)
 //        {
 //            return at;
 //        }
 
 //        // We traverse top down
 
-//        OP_TYPE type = (OP_TYPE)program.m_code[at].type;
+//        EOpType type = (EOpType)program.m_code[at].type;
 //        OP::ARGS args = program.m_code[at].args;
 
 //        // We only apply it to instances, for performance reasons.
@@ -430,13 +430,13 @@ namespace mu
 //        {
 
 //        //-----------------------------------------------------------------------------------------
-////        case OP_TYPE::NU_CONDITIONAL:
-////        case OP_TYPE::SC_CONDITIONAL:
-////        case OP_TYPE::CO_CONDITIONAL:
-////        case OP_TYPE::IM_CONDITIONAL:
-////        case OP_TYPE::ME_CONDITIONAL:
-////        case OP_TYPE::LA_CONDITIONAL:
-//        case OP_TYPE::IN_CONDITIONAL:
+////        case EOpType::NU_CONDITIONAL:
+////        case EOpType::SC_CONDITIONAL:
+////        case EOpType::CO_CONDITIONAL:
+////        case EOpType::IM_CONDITIONAL:
+////        case EOpType::ME_CONDITIONAL:
+////        case EOpType::LA_CONDITIONAL:
+//        case EOpType::IN_CONDITIONAL:
 //        {
 //            BoolExpressionPtr pCond = GetBoolExpression( args.Conditional.condition, program, &m_boolExpressionCache );
 
@@ -487,13 +487,13 @@ namespace mu
 //            break;
 //        }
 
-////        case OP_TYPE::NU_SWITCH:
-////        case OP_TYPE::SC_SWITCH:
-////        case OP_TYPE::CO_SWITCH:
-////        case OP_TYPE::IM_SWITCH:
-////        case OP_TYPE::VO_SWITCH:
-////        case OP_TYPE::ME_SWITCH:
-////        case OP_TYPE::LA_SWITCH:
+////        case EOpType::NU_SWITCH:
+////        case EOpType::SC_SWITCH:
+////        case EOpType::CO_SWITCH:
+////        case EOpType::IM_SWITCH:
+////        case EOpType::VO_SWITCH:
+////        case EOpType::ME_SWITCH:
+////        case EOpType::LA_SWITCH:
 ////        {
 ////            IntExpressionPtr pValue = GetIntExpression( args.Switch.variable, program );
 
@@ -584,7 +584,7 @@ namespace mu
 //        case IET_CONSTANT:
 //        {
 //            OP op;
-//            op.type = OP_TYPE::NU_CONSTANT;
+//            op.type = EOpType::NU_CONSTANT;
 //            op.args.IntConstant.value = value;
 //            at = program.AddOpUnique( op, cache );
 //            break;
@@ -593,7 +593,7 @@ namespace mu
 //        case IET_PARAMETER:
 //        {
 //            OP op;
-//            op.type = OP_TYPE::NU_PARAMETER;
+//            op.type = EOpType::NU_PARAMETER;
 //            op.args.Parameter.variable = parameter;
 //            at = program.AddOpUnique( op, cache );
 //            break;
@@ -760,7 +760,7 @@ namespace mu
 //        case BET_TRUE:
 //        {
 //            OP op;
-//            op.type = OP_TYPE::BO_CONSTANT;
+//            op.type = EOpType::BO_CONSTANT;
 //            op.args.BoolConstant.value = true;
 //            at = program.AddOpUnique( op, cache );
 //            break;
@@ -769,7 +769,7 @@ namespace mu
 //        case BET_FALSE:
 //        {
 //            OP op;
-//            op.type = OP_TYPE::BO_CONSTANT;
+//            op.type = EOpType::BO_CONSTANT;
 //            op.args.BoolConstant.value = false;
 //            at = program.AddOpUnique( op, cache );
 //            break;
@@ -778,7 +778,7 @@ namespace mu
 //        case BET_NOT:
 //        {
 //            OP op;
-//            op.type = OP_TYPE::BO_NOT;
+//            op.type = EOpType::BO_NOT;
 //            op.args.BoolNot.source = a->GenerateCode( program, cache );
 //            at = program.AddOpUnique( op, cache );
 //            break;
@@ -787,7 +787,7 @@ namespace mu
 //        case BET_AND:
 //        {
 //            OP op;
-//            op.type = OP_TYPE::BO_AND;
+//            op.type = EOpType::BO_AND;
 //            op.args.BoolBinary.a = a->GenerateCode( program, cache );
 //            op.args.BoolBinary.b = b->GenerateCode( program, cache );
 //            at = program.AddOpUnique( op, cache );
@@ -797,7 +797,7 @@ namespace mu
 //        case BET_OR:
 //        {
 //            OP op;
-//            op.type = OP_TYPE::BO_OR;
+//            op.type = EOpType::BO_OR;
 //            op.args.BoolBinary.a = a->GenerateCode( program, cache );
 //            op.args.BoolBinary.b = b->GenerateCode( program, cache );
 //            at = program.AddOpUnique( op, cache );
@@ -808,7 +808,7 @@ namespace mu
 //        case BET_PARAMETER:
 //        {
 //            OP op;
-//            op.type = OP_TYPE::BO_PARAMETER;
+//            op.type = EOpType::BO_PARAMETER;
 //            op.args.Parameter.variable = parameter;
 //            at = program.AddOpUnique( op, cache );
 //            break;
@@ -818,7 +818,7 @@ namespace mu
 //        case BET_INT_EQUAL_CONST:
 //        {
 //            OP op;
-//            op.type = OP_TYPE::BO_EQUAL_INT_CONST;
+//            op.type = EOpType::BO_EQUAL_INT_CONST;
 //            op.args.BoolEqualScalarConst.value = intExp->GenerateCode( program, cache );
 //            op.args.BoolEqualScalarConst.constant = (int16_t)intValue;
 //            at = program.AddOpUnique( op, cache );
@@ -1015,11 +1015,11 @@ namespace mu
             {
                 bool recurse = true;
 
-                if (n->GetOpType()==OP_TYPE::IN_CONDITIONAL)
+                if (n->GetOpType()==EOpType::IN_CONDITIONAL)
                 {
 					ASTOpConditional* topConditional = static_cast<ASTOpConditional*>(n.get());
 					Ptr<ASTOp> yes = topConditional->yes.child();
-                    if (yes && yes->GetOpType()==OP_TYPE::IN_ADDSURFACE)
+                    if (yes && yes->GetOpType()==EOpType::IN_ADDSURFACE)
                     {
 						ASTOpInstanceAdd* addSurface = static_cast<ASTOpInstanceAdd*>(yes.get());
 
@@ -1028,7 +1028,7 @@ namespace mu
                         {
                             ended = true;
 							Ptr<ASTOp> base = addSurface->instance.child();
-                            if ( base && base->GetOpType()==OP_TYPE::IN_CONDITIONAL )
+                            if ( base && base->GetOpType()==EOpType::IN_CONDITIONAL )
                             {
 								ASTOpConditional* bottomConditional = static_cast<ASTOpConditional*>(base.get());
 
@@ -1074,11 +1074,11 @@ namespace mu
                     }
                 }
 
-                else if (n->GetOpType()==OP_TYPE::ME_CONDITIONAL)
+                else if (n->GetOpType()==EOpType::ME_CONDITIONAL)
                 {
 					ASTOpConditional* topConditional = static_cast<ASTOpConditional*>(n.get());
                     Ptr<ASTOp> yes = topConditional->yes.child();
-                    if ( yes && yes->GetOpType()==OP_TYPE::ME_MERGE )
+                    if ( yes && yes->GetOpType()==EOpType::ME_MERGE )
                     {
                         Ptr<ASTOpFixed> merge = static_cast<ASTOpFixed*>(yes.get());
 
@@ -1087,7 +1087,7 @@ namespace mu
                         {
                             ended = true;
 							Ptr<ASTOp> base = merge->children[merge->op.args.MeshMerge.base].child();
-                            if ( base && base->GetOpType()==OP_TYPE::ME_CONDITIONAL )
+                            if ( base && base->GetOpType()==EOpType::ME_CONDITIONAL )
                             {
 								ASTOpConditional* bottomConditional = static_cast<ASTOpConditional*>(base.get());
 
@@ -1130,7 +1130,7 @@ namespace mu
                     }
 
 // todo, since the move to aggregated remove mask op
-//                    else if ( yes->GetOpType()==OP_TYPE::ME_REMOVEMASK )
+//                    else if ( yes->GetOpType()==EOpType::ME_REMOVEMASK )
 //                    {
 //                        Ptr<ASTOpFixed> remove = dynamic_cast<ASTOpFixed*>(yes.get());
 
@@ -1139,7 +1139,7 @@ namespace mu
 //                        {
 //                            ended = true;
 //                            auto base = remove->children[remove->op.args.MeshRemoveMask.source].child();
-//                            if ( base && base->GetOpType()==OP_TYPE::ME_CONDITIONAL )
+//                            if ( base && base->GetOpType()==EOpType::ME_CONDITIONAL )
 //                            {
 //                                auto bottomConditional = dynamic_cast<ASTOpConditional*>(base.get());
 
@@ -1204,7 +1204,7 @@ namespace mu
 
                 if (topConditional && topConditional->condition)
                 {
-                    if ( topConditional->condition->GetOpType()==OP_TYPE::BO_EQUAL_INT_CONST
+                    if ( topConditional->condition->GetOpType()==EOpType::BO_EQUAL_INT_CONST
                          &&
                          topConditional->no
                          &&
@@ -1234,7 +1234,7 @@ namespace mu
                                  &&
                                  conditional->condition
                                  &&
-                                 conditional->condition->GetOpType()==OP_TYPE::BO_EQUAL_INT_CONST )
+                                 conditional->condition->GetOpType()==EOpType::BO_EQUAL_INT_CONST )
                             {
 								ASTOpFixed* compare = static_cast<ASTOpFixed*>(conditional->condition.child().get());
                                 check(compare);
@@ -1295,7 +1295,7 @@ namespace mu
                 if (topSwitch)
                 {
                     bool first = true;
-					OP_TYPE caseType = OP_TYPE::NONE;
+					EOpType caseType = EOpType::NONE;
                     for (const ASTOpSwitch::FCase& c:topSwitch->cases)
                     {
                         if ( c.branch )
@@ -1309,7 +1309,7 @@ namespace mu
                             {
                                 if (caseType!=c.branch->GetOpType())
                                 {
-                                    caseType=OP_TYPE::NONE;
+                                    caseType=EOpType::NONE;
                                     break;
                                 }
                             }
@@ -1317,7 +1317,7 @@ namespace mu
                     }
 
                     // todo, since the move to aggregated remove mask op
-//                    if (caseType==OP_TYPE::ME_REMOVEMASK)
+//                    if (caseType==EOpType::ME_REMOVEMASK)
 //                    {
 //                        // Is the source of all remove operations the default option in the switch?
 //                        auto source = topSwitch->def.child();
@@ -1349,7 +1349,7 @@ namespace mu
 //                            newSwitch->def = nullptr;
 
 //                            Ptr<ASTOpFixed> newRemove = new ASTOpFixed();
-//                            newRemove->op.type = OP_TYPE::ME_REMOVEMASK;
+//                            newRemove->op.type = EOpType::ME_REMOVEMASK;
 //                            newRemove->SetChild( newRemove->op.args.MeshRemoveMask.mask, newSwitch );
 //                            newRemove->SetChild( newRemove->op.args.MeshRemoveMask.source, source);
 
