[mutable] More enforcement of coding standards.

diff --git a/Engine/Plugins/Mutable/Source/MutableTools/Private/MuT/CodeOptimiser_Parameters.cpp b/Engine/Plugins/Mutable/Source/MutableTools/Private/MuT/CodeOptimiser_Parameters.cpp
--- a/Engine/Plugins/Mutable/Source/MutableTools/Private/MuT/CodeOptimiser_Parameters.cpp
+++ b/Engine/Plugins/Mutable/Source/MutableTools/Private/MuT/CodeOptimiser_Parameters.cpp
@@ -16,7 +16,6 @@
 #include "MuT/Compiler.h"
 #include "MuT/CompilerPrivate.h"
 #include "MuT/DataPacker.h"
-#include "MuT/Platform.h"
 #include "MuR/Image.h"
 #include "MuR/ImagePrivate.h"
 #include "MuR/MutableTrace.h"
@@ -43,136 +42,148 @@ namespace mu
     //---------------------------------------------------------------------------------------------
     //---------------------------------------------------------------------------------------------
     //---------------------------------------------------------------------------------------------
-    RuntimeParameterVisitorAST::RuntimeParameterVisitorAST(const FStateCompilationData* pState)
-        : m_pState(pState)
+    RuntimeParameterVisitorAST::RuntimeParameterVisitorAST(const FStateCompilationData* InState)
+        : State(InState)
     {
     }
 
 
-    bool RuntimeParameterVisitorAST::HasAny( const Ptr<ASTOp>& root )
+    bool RuntimeParameterVisitorAST::HasAny( const Ptr<ASTOp>& Root )
     {
-        if (!m_pState->nodeState.RuntimeParams.Num())
+        if (!State->nodeState.RuntimeParams.Num())
         {
             return false;
         }
 
         // Shortcut flag: if true we already found a runtime parameter, don't process new ops,
         // but still store the results of processed ops.
-        bool found = false;
+        bool bRuntimeFound = false;
 
-        m_pending.clear();
+        Pending.Reset();
 
-        PENDING_ITEM start;
-        start.at = root;
-        start.itemType = 0;
-        start.onlyLayoutsRelevant = 0;
-        m_pending.push_back( start );
+        FPendingItem start;
+        start.At = Root;
+        start.ItemType = 0;
+        start.OnlyLayoutsRelevant = 0;
+		Pending.Add( start );
 
         // Don't early out to be able to complete parent op cached flags
-        while ( m_pending.size() )
+        while (!Pending.IsEmpty())
         {
-            PENDING_ITEM item = m_pending.back();
-            m_pending.pop_back();
-            Ptr<ASTOp> at = item.at;
+			FPendingItem item = Pending.Pop(EAllowShrinking::No);
+            Ptr<ASTOp> At = item.At;
 
-			if (!at)
+			if (!At)
 			{
 				continue;
 			}
 
             // Not cached?
-            if ( m_visited[at]!=OP_STATE::VISITED_HASRUNTIME
-                 &&
-                 m_visited[at]!=OP_STATE::VISITED_FULL_DOESNTHAVERUNTIME )
+			EOpState* FoundCached = Visited.Find(At);
+            if ( 
+				!FoundCached
+				||
+				(
+					*FoundCached !=EOpState::VisitedHasRuntime
+					&&
+					*FoundCached != EOpState::VisitedFullDoesntHaveRuntime
+					) 
+				)
             {
-                if (item.itemType)
+                if (item.ItemType)
                 {
                     // Item indicating we finished with all the children of a parent
-                    check( m_visited[at]==OP_STATE::CHILDREN_PENDING_FULL
-                                    ||
-                                    m_visited[at]==OP_STATE::CHILDREN_PENDING_PARTIAL
-                                    ||
-                                    m_visited[at]==OP_STATE::VISITED_PARTIAL_DOESNTHAVERUNTIME );
-
-                    bool subtreeFound = false;
-                    at->ForEachChild( [&](ASTChild& ref)
+                    check(*FoundCached ==EOpState::ChildrenPendingFull
+                        ||
+						*FoundCached ==EOpState::ChildrenPendingPartial
+                        ||
+						*FoundCached ==EOpState::VisitedPartialDoesntHaveRuntime );
+
+                    bool bSubtreeFound = false;
+                    At->ForEachChild( [&](ASTChild& ref)
                     {
-                        subtreeFound = subtreeFound || m_visited[ref.child()]==OP_STATE::VISITED_HASRUNTIME;
+						EOpState* FoundChild = Visited.Find(ref.child());
+						if (FoundChild && *FoundChild == EOpState::VisitedHasRuntime)
+						{
+							bSubtreeFound = true;
+						}
                     });
 
-                    if (subtreeFound)
+                    if (bSubtreeFound)
                     {
-                        m_visited[at] = OP_STATE::VISITED_HASRUNTIME;
+						*FoundCached = EOpState::VisitedHasRuntime;
                     }
                     else
                     {
-                        m_visited[at] = item.onlyLayoutsRelevant
-                                ? OP_STATE::VISITED_PARTIAL_DOESNTHAVERUNTIME
-                                : OP_STATE::VISITED_FULL_DOESNTHAVERUNTIME;
+						*FoundCached = item.OnlyLayoutsRelevant
+                                ? EOpState::VisitedPartialDoesntHaveRuntime
+                                : EOpState::VisitedFullDoesntHaveRuntime;
                     }
                 }
 
-                else if (!found)
+                else if (!bRuntimeFound)
                 {
                     // We need to process the subtree
-                    check( m_visited[at]==OP_STATE::NOT_VISITED
-                                    ||
-                                    ( m_visited[at]==OP_STATE::VISITED_PARTIAL_DOESNTHAVERUNTIME
-                                      &&
-                                      item.onlyLayoutsRelevant==0 )
-                                    );
+                    check(!FoundCached 
+						||
+						*FoundCached ==EOpState::NotVisited
+						||
+						(*FoundCached == EOpState::VisitedPartialDoesntHaveRuntime
+							&&
+							item.OnlyLayoutsRelevant==0 )
+					);
 
                     // Request the processing of the end of this instruction
-                    PENDING_ITEM endItem = item;
-                    endItem.itemType = 1;
-                    m_pending.push_back( endItem );
-                    m_visited[at] = item.onlyLayoutsRelevant
-                            ? OP_STATE::CHILDREN_PENDING_PARTIAL
-                            : OP_STATE::CHILDREN_PENDING_FULL;
+					FPendingItem endItem = item;
+                    endItem.ItemType = 1;
+                    Pending.Add( endItem );
+					Visited.Add(At, item.OnlyLayoutsRelevant
+						? EOpState::ChildrenPendingPartial
+						: EOpState::ChildrenPendingFull);
 
                     // Is it a special op type?
-                    switch ( at->GetOpType() )
+                    switch ( At->GetOpType() )
                     {
 
-                    case OP_TYPE::BO_PARAMETER:
-                    case OP_TYPE::NU_PARAMETER:
-                    case OP_TYPE::SC_PARAMETER:
-                    case OP_TYPE::CO_PARAMETER:
-                    case OP_TYPE::PR_PARAMETER:
-                    case OP_TYPE::IM_PARAMETER:
+                    case EOpType::BO_PARAMETER:
+                    case EOpType::NU_PARAMETER:
+                    case EOpType::SC_PARAMETER:
+                    case EOpType::CO_PARAMETER:
+                    case EOpType::PR_PARAMETER:
+                    case EOpType::IM_PARAMETER:
                     {
-						const ASTOpParameter* typed = static_cast<const ASTOpParameter*>(at.get());
-                        const TArray<FString>& params = m_pState->nodeState.RuntimeParams;
-                        if ( params.Find( typed->parameter.m_name)
+						const ASTOpParameter* typed = static_cast<const ASTOpParameter*>(At.get());
+                        const TArray<FString>& params = State->nodeState.RuntimeParams;
+                        if ( params.Find( typed->parameter.Name)
                              !=
                              INDEX_NONE )
                         {
-                            found = true;
-                            m_visited[at] = OP_STATE::VISITED_HASRUNTIME;
+							bRuntimeFound = true;
+                            Visited.Add(At,EOpState::VisitedHasRuntime);
                         }
                         break;
                     }
 
-                    case OP_TYPE::ME_INTERPOLATE:
+                    case EOpType::ME_INTERPOLATE:
                     {
-						const ASTOpFixed* typed = static_cast<const ASTOpFixed*>(at.get());
-                        PENDING_ITEM childItem;
-                        childItem.itemType = 0;
-                        childItem.onlyLayoutsRelevant = item.onlyLayoutsRelevant;
+						const ASTOpFixed* typed = static_cast<const ASTOpFixed*>(At.get());
+						FPendingItem ChildItem;
+                        ChildItem.ItemType = 0;
+                        ChildItem.OnlyLayoutsRelevant = item.OnlyLayoutsRelevant;
 
-                        if ( item.onlyLayoutsRelevant==0 )
+                        if ( item.OnlyLayoutsRelevant==0 )
                         {
-                            childItem.at = typed->children[typed->op.args.MeshInterpolate.factor].child();
-                            AddIfNeeded(childItem);
+                            ChildItem.At = typed->children[typed->op.args.MeshInterpolate.factor].child();
+                            AddIfNeeded(ChildItem);
                         }
 
-                        childItem.at = typed->children[typed->op.args.MeshInterpolate.base].child();
-                        AddIfNeeded(childItem);
+                        ChildItem.At = typed->children[typed->op.args.MeshInterpolate.base].child();
+                        AddIfNeeded(ChildItem);
 
                         for (int32 t=0;t<MUTABLE_OP_MAX_INTERPOLATE_COUNT-1;++t)
                         {
-                            childItem.at = typed->children[typed->op.args.MeshInterpolate.targets[t]].child();
-                            AddIfNeeded(childItem);
+                            ChildItem.At = typed->children[typed->op.args.MeshInterpolate.targets[t]].child();
+                            AddIfNeeded(ChildItem);
                         }
 
                         break;
@@ -180,13 +191,13 @@ namespace mu
 
                     default:
                     {
-                        at->ForEachChild([&](ASTChild& ref)
+                        At->ForEachChild([&](ASTChild& ref)
                         {
-                            PENDING_ITEM childItem;
-                            childItem.itemType = 0;
-                            childItem.at = ref.child();
-                            childItem.onlyLayoutsRelevant = item.onlyLayoutsRelevant;
-                            AddIfNeeded(childItem);
+							FPendingItem ChildItem;
+                            ChildItem.ItemType = 0;
+                            ChildItem.At = ref.child();
+                            ChildItem.OnlyLayoutsRelevant = item.OnlyLayoutsRelevant;
+                            AddIfNeeded(ChildItem);
                         });
                         break;
                     }
@@ -198,34 +209,35 @@ namespace mu
                 else
                 {
                     // We won't process it.
-                    m_visited[at] = OP_STATE::NOT_VISITED;
+					Visited.Add(At,EOpState::NotVisited);
                 }
             }
         }
 
-        return m_visited[root]==OP_STATE::VISITED_HASRUNTIME;
+        return Visited[Root]==EOpState::VisitedHasRuntime;
     }
 
 
-    void RuntimeParameterVisitorAST::AddIfNeeded( const PENDING_ITEM& item )
+    void RuntimeParameterVisitorAST::AddIfNeeded( const FPendingItem& Item )
     {
-        if (item.at)
+        if (Item.At)
         {
-            if (m_visited[item.at]==OP_STATE::NOT_VISITED)
+			const EOpState* Found = Visited.Find(Item.At);
+            if (!Found || *Found==EOpState::NotVisited)
             {
-                m_pending.push_back( item );
+                Pending.Add(Item);
             }
-            else if (m_visited[item.at]==OP_STATE::VISITED_PARTIAL_DOESNTHAVERUNTIME
-                     &&
-                     item.onlyLayoutsRelevant==0)
+            else if (*Found ==EOpState::VisitedPartialDoesntHaveRuntime
+				&&
+				Item.OnlyLayoutsRelevant==0)
             {
-                m_pending.push_back( item );
+				Pending.Add(Item);
             }
-            else if (m_visited[item.at]==OP_STATE::CHILDREN_PENDING_PARTIAL
-                     &&
-                     item.onlyLayoutsRelevant==0)
+            else if (*Found ==EOpState::ChildrenPendingPartial
+				&&
+				Item.OnlyLayoutsRelevant==0)
             {
-                m_pending.push_back( item );
+				Pending.Add(Item);
             }
         }
     }
@@ -239,14 +251,14 @@ namespace mu
         if ( !mask )
         {
             Ptr<ASTOpFixed> whiteOp = new ASTOpFixed;
-            whiteOp->op.type = OP_TYPE::CO_CONSTANT;
+            whiteOp->op.type = EOpType::CO_CONSTANT;
             whiteOp->op.args.ColourConstant.value[0] = 1;
             whiteOp->op.args.ColourConstant.value[1] = 1;
 			whiteOp->op.args.ColourConstant.value[2] = 1;
 			whiteOp->op.args.ColourConstant.value[3] = 1;
 
             Ptr<ASTOpFixed> wplainOp = new ASTOpFixed;
-            wplainOp->op.type = OP_TYPE::IM_PLAINCOLOUR;
+            wplainOp->op.type = EOpType::IM_PLAINCOLOUR;
             wplainOp->SetChild( wplainOp->op.args.ImagePlainColour.colour, whiteOp );
             wplainOp->op.args.ImagePlainColour.format = EImageFormat::L_UByte;
             wplainOp->op.args.ImagePlainColour.size[0] = 4;
@@ -254,7 +266,7 @@ namespace mu
 			wplainOp->op.args.ImagePlainColour.LODs = 1;
 
             Ptr<ASTOpFixed> wresizeOp = new ASTOpFixed;
-            wresizeOp->op.type = OP_TYPE::IM_RESIZELIKE;
+            wresizeOp->op.type = EOpType::IM_RESIZELIKE;
             wresizeOp->SetChild( wresizeOp->op.args.ImageResizeLike.source, wplainOp );
             wresizeOp->SetChild( wresizeOp->op.args.ImageResizeLike.sizeSource, base );
 
@@ -293,29 +305,29 @@ namespace mu
     }
 
 
-    Ptr<ASTOp> ParameterOptimiserAST::Visit( Ptr<ASTOp> at, bool& processChildren )
+    Ptr<ASTOp> ParameterOptimiserAST::Visit( Ptr<ASTOp> At, bool& processChildren )
     {
         // Only process children if there are runtime parameters in the subtree
-        processChildren = m_hasRuntimeParamVisitor.HasAny(at);
+        processChildren = m_hasRuntimeParamVisitor.HasAny(At);
 
-        OP_TYPE type = at->GetOpType();
+        EOpType type = At->GetOpType();
         switch ( type )
         {
         //-------------------------------------------------------------------------------------
         // Be careful with changing merge options and "mergesurfaces" flags
-//        case OP_TYPE::ME_MERGE:
+//        case EOpType::ME_MERGE:
 //        {
-//            OP::MeshMergeArgs mergeArgs = program.m_code[at].args.MeshMerge;
+//            OP::MeshMergeArgs mergeArgs = program.m_code[At].args.MeshMerge;
 
 //            RuntimeParameterVisitor paramVis;
 
 //            switch ( program.m_code[ mergeArgs.base ].type )
 //            {
-//            case OP_TYPE::ME_CONDITIONAL:
+//            case EOpType::ME_CONDITIONAL:
 //            {
 //                OP::ADDRESS conditionAt =
 //                        program.m_code[ mergeArgs.base ].args.Conditional.condition;
-//                bool conditionConst = !paramVis.HasAny( m_pModel.get(),
+//                bool conditionConst = !paramVis.HasAny( Model.get(),
 //                                                        m_state,
 //                                                        conditionAt );
 //                if ( conditionConst )
@@ -324,34 +336,34 @@ namespace mu
 //                    // the model data. Make this optimisation optional.
 //                    m_modified = true;
 
-//                    OP yesOp = program.m_code[at];
+//                    OP yesOp = program.m_code[At];
 //                    yesOp.args.MeshMerge.base = program.m_code[ mergeArgs.base ].args.Conditional.yes;
 
-//                    OP noOp = program.m_code[at];
+//                    OP noOp = program.m_code[At];
 //                    noOp.args.MeshMerge.base = program.m_code[ mergeArgs.base ].args.Conditional.no;
 
 //                    OP op = program.m_code[ mergeArgs.base ];
 //                    op.args.Conditional.yes = program.AddOp( yesOp );
 //                    op.args.Conditional.no = program.AddOp( noOp );
-//                    at = program.AddOp( op );
+//                    At = program.AddOp( op );
 //                }
 
 //                break;
 //            }
 
-//            case OP_TYPE::ME_MERGE:
+//            case EOpType::ME_MERGE:
 //            {
 //                OP::ADDRESS childBaseAt = program.m_code[ mergeArgs.base ].args.MeshMerge.base;
-//                bool childBaseConst = !paramVis.HasAny( m_pModel.get(),
+//                bool childBaseConst = !paramVis.HasAny( Model.get(),
 //                                                        m_state,
 //                                                        childBaseAt );
 
 //                OP::ADDRESS childAddAt = program.m_code[ mergeArgs.base ].args.MeshMerge.added;
-//                bool childAddConst = !paramVis.HasAny( m_pModel.get(),
+//                bool childAddConst = !paramVis.HasAny( Model.get(),
 //                                                       m_state,
 //                                                       childAddAt );
 
-//                bool addConst = !paramVis.HasAny( m_pModel.get(),
+//                bool addConst = !paramVis.HasAny( Model.get(),
 //                                                  m_state,
 //                                                  mergeArgs.added );
 
@@ -359,12 +371,12 @@ namespace mu
 //                {
 //                    m_modified = true;
 
-//                    OP bottom = program.m_code[at];
+//                    OP bottom = program.m_code[At];
 //                    bottom.args.MeshMerge.base = childAddAt;
 
 //                    OP top = program.m_code[ mergeArgs.base ];
 //                    top.args.MeshMerge.added = program.AddOp( bottom );
-//                    at = program.AddOp( top );
+//                    At = program.AddOp( top );
 //                }
 //                break;
 //            }
@@ -379,30 +391,30 @@ namespace mu
         //-------------------------------------------------------------------------------------
         //-------------------------------------------------------------------------------------
         //-------------------------------------------------------------------------------------
-        case OP_TYPE::IM_CONDITIONAL:
+        case EOpType::IM_CONDITIONAL:
         {
-			const ASTOpConditional* typedAt = static_cast<const ASTOpConditional*>(at.get());
+			const ASTOpConditional* TypedOp = static_cast<const ASTOpConditional*>(At.get());
 
             // If the condition is not runtime, but the branches are, try to move the
             // conditional down
             bool optimised = false;
 
-            if ( !m_hasRuntimeParamVisitor.HasAny( typedAt->condition.child() ) )
+            if ( !m_hasRuntimeParamVisitor.HasAny( TypedOp->condition.child() ) )
             {
-                OP_TYPE yesType = typedAt->yes->GetOpType();
-                OP_TYPE noType = typedAt->no->GetOpType();
+                EOpType yesType = TypedOp->yes->GetOpType();
+                EOpType noType = TypedOp->no->GetOpType();
 
-                bool yesHasAny = m_hasRuntimeParamVisitor.HasAny( typedAt->yes.child() );
-                bool noHasAny = m_hasRuntimeParamVisitor.HasAny( typedAt->no.child() );
+                bool yesHasAny = m_hasRuntimeParamVisitor.HasAny( TypedOp->yes.child() );
+                bool noHasAny = m_hasRuntimeParamVisitor.HasAny( TypedOp->no.child() );
 
                 if ( !optimised && yesHasAny && noHasAny && yesType==noType)
                 {
                     switch (yesType)
                     {
-                        case OP_TYPE::IM_COMPOSE:
+                        case EOpType::IM_COMPOSE:
                         {
-						const ASTOpImageCompose* typedYes = static_cast<const ASTOpImageCompose*>(typedAt->yes.child().get());
-						const ASTOpImageCompose* typedNo = static_cast<const ASTOpImageCompose*>(typedAt->no.child().get());
+						const ASTOpImageCompose* typedYes = static_cast<const ASTOpImageCompose*>(TypedOp->yes.child().get());
+						const ASTOpImageCompose* typedNo = static_cast<const ASTOpImageCompose*>(TypedOp->no.child().get());
                         if ( typedYes->BlockId == typedNo->BlockId
                              &&
                              (
@@ -415,32 +427,32 @@ namespace mu
                                 // Move the conditional down
                                 Ptr<ASTOpImageCompose> compOp = mu::Clone<ASTOpImageCompose>(typedYes);
 
-                                Ptr<ASTOpConditional> baseCond = mu::Clone<ASTOpConditional>(typedAt);
+                                Ptr<ASTOpConditional> baseCond = mu::Clone<ASTOpConditional>(TypedOp);
                                 baseCond->yes = typedYes->Base.child();
                                 baseCond->no = typedNo->Base.child();
                                 compOp->Base = baseCond;
 
-                                Ptr<ASTOpConditional> blockCond = mu::Clone<ASTOpConditional>(typedAt);
+                                Ptr<ASTOpConditional> blockCond = mu::Clone<ASTOpConditional>(TypedOp);
                                 blockCond->yes = typedYes->BlockImage.child();
                                 blockCond->no = typedNo->BlockImage.child();
                                 compOp->BlockImage = blockCond;
 
                                 if (typedYes->Mask)
                                 {
-                                    Ptr<ASTOpConditional> maskCond = mu::Clone<ASTOpConditional>(typedAt);
+                                    Ptr<ASTOpConditional> maskCond = mu::Clone<ASTOpConditional>(TypedOp);
                                     maskCond->yes = typedYes->Mask.child();
                                     maskCond->no = typedNo->Mask.child();
                                     compOp->Mask = maskCond;
                                 }
 
-                                Ptr<ASTOpConditional> layCond = mu::Clone<ASTOpConditional>(typedAt);
-                                layCond->type = OP_TYPE::LA_CONDITIONAL;
+                                Ptr<ASTOpConditional> layCond = mu::Clone<ASTOpConditional>(TypedOp);
+                                layCond->type = EOpType::LA_CONDITIONAL;
                                 layCond->yes = typedYes->Layout.child();
                                 layCond->no = typedNo->Layout.child();
                                 compOp->Layout = layCond;
 
 
-                                at = compOp;
+                                At = compOp;
                                 optimised = true;
                             }
                             break;
@@ -456,21 +468,21 @@ namespace mu
                 {
                     switch (yesType)
                     {
-                        case OP_TYPE::IM_LAYERCOLOUR:
+                        case EOpType::IM_LAYERCOLOUR:
                         {
                             optimised = true;
 
-							const ASTOpImageLayerColor* typedYes = static_cast<const ASTOpImageLayerColor*>(typedAt->yes.child().get());
+							const ASTOpImageLayerColor* typedYes = static_cast<const ASTOpImageLayerColor*>(TypedOp->yes.child().get());
 
                             Ptr<ASTOpFixed> blackOp = new ASTOpFixed;
-                            blackOp->op.type = OP_TYPE::CO_CONSTANT;
+                            blackOp->op.type = EOpType::CO_CONSTANT;
                             blackOp->op.args.ColourConstant.value[0] = 0;
                             blackOp->op.args.ColourConstant.value[1] = 0;
 							blackOp->op.args.ColourConstant.value[2] = 0;
 							blackOp->op.args.ColourConstant.value[3] = 1;
 
                             Ptr<ASTOpFixed> plainOp = new ASTOpFixed;
-                            plainOp->op.type = OP_TYPE::IM_PLAINCOLOUR;
+                            plainOp->op.type = EOpType::IM_PLAINCOLOUR;
                             plainOp->SetChild( plainOp->op.args.ImagePlainColour.colour, blackOp );
                             plainOp->op.args.ImagePlainColour.format = EImageFormat::L_UByte;
                             plainOp->op.args.ImagePlainColour.size[0] = 4;
@@ -478,57 +490,57 @@ namespace mu
 							plainOp->op.args.ImagePlainColour.LODs = 1;
 
                             Ptr<ASTOpFixed> resizeOp = new ASTOpFixed;
-                            resizeOp->op.type = OP_TYPE::IM_RESIZELIKE;
+                            resizeOp->op.type = EOpType::IM_RESIZELIKE;
                             resizeOp->SetChild( resizeOp->op.args.ImageResizeLike.source, plainOp );
                             resizeOp->SetChild( resizeOp->op.args.ImageResizeLike.sizeSource, typedYes->base );
 
-                            Ptr<ASTOpConditional> maskOp = mu::Clone<ASTOpConditional>(typedAt);
+                            Ptr<ASTOpConditional> maskOp = mu::Clone<ASTOpConditional>(TypedOp);
                             maskOp->no = resizeOp;
 
                             // If there is no mask (because it is optional), we need to make a
                             // white plain image
                             maskOp->yes = EnsureValidMask( typedYes->mask.child(), typedYes->base.child() );
 
-							Ptr<ASTOpConditional> baseOp = mu::Clone<ASTOpConditional>(typedAt);
+							Ptr<ASTOpConditional> baseOp = mu::Clone<ASTOpConditional>(TypedOp);
                             baseOp->yes = typedYes->base.child();
 
                             Ptr<ASTOpImageLayerColor> softOp = mu::Clone<ASTOpImageLayerColor>(typedYes);
                             softOp->base = baseOp;
                             softOp->mask = maskOp;
 
-                            at = softOp;
+                            At = softOp;
                             break;
                         }
 
                         // TODO
                         // It seems this is not worth since it replaces a conditional by a compose
-                        // (but only at build time, not update?) and it introduces the use of masks
+                        // (but only At build time, not update?) and it introduces the use of masks
                         // and resize likes... plus masks can't always be used if BC formats.
-//						case OP_TYPE::IM_COMPOSE:
+//						case EOpType::IM_COMPOSE:
 //						{
 //							optimised = true;
 
 //							OP blackOp;
-//							blackOp.type = OP_TYPE::CO_CONSTANT;
+//							blackOp.type = EOpType::CO_CONSTANT;
 //							blackOp.args.ColourConstant.value[0] = 0;
 //							blackOp.args.ColourConstant.value[1] = 0;
 //							blackOp.args.ColourConstant.value[2] = 0;
 //							blackOp.args.ColourConstant.value[3] = 1;
 
 //							OP plainOp;
-//							plainOp.type = OP_TYPE::IM_PLAINCOLOUR;
+//							plainOp.type = EOpType::IM_PLAINCOLOUR;
 //							plainOp.args.ImagePlainColour.colour = program.AddOp( blackOp );
 //							plainOp.args.ImagePlainColour.format = L_UByte;
 //							plainOp.args.ImagePlainColour.size = ;
 //							plainOp.args.ImagePlainColour.LODs = ;
 
 //							OP resizeOp;
-//							resizeOp.type = OP_TYPE::IM_RESIZELIKE;
+//							resizeOp.type = EOpType::IM_RESIZELIKE;
 //							resizeOp.args.ImageResizeLike.source = program.AddOp( plainOp );
 //							resizeOp.args.ImageResizeLike.sizeSource =
 //									program.m_code[args.yes].args.ImageCompose.blockImage;
 
-//							OP maskOp = program.m_code[at];
+//							OP maskOp = program.m_code[At];
 //							maskOp.args.Conditional.no = program.AddOp( resizeOp );
 
 //							// If there is no mask (because it is optional), we need to make a
@@ -538,7 +550,7 @@ namespace mu
 //									  program.m_code[args.yes].args.ImageCompose.base,
 //									  program );
 
-//							OP baseOp = program.m_code[at];
+//							OP baseOp = program.m_code[At];
 //							baseOp.args.Conditional.yes =
 //									program.m_code[args.yes].args.ImageCompose.base;
 
@@ -546,10 +558,10 @@ namespace mu
 //							composeOp.args.ImageCompose.base = program.AddOp( baseOp );
 //							composeOp.args.ImageCompose.mask = program.AddOp( maskOp );
 
-//							at = program.AddOp( composeOp );
+//							At = program.AddOp( composeOp );
 
 //							// Process the new children
-//							at = Recurse( at, program );
+//							At = Recurse( At, program );
 //							break;
 //						}
 
@@ -563,21 +575,21 @@ namespace mu
                 {
                     switch (noType)
                     {
-                        case OP_TYPE::IM_LAYERCOLOUR:
+                        case EOpType::IM_LAYERCOLOUR:
                         {
                             optimised = true;
 
-                            const ASTOpImageLayerColor* typedNo = static_cast<const ASTOpImageLayerColor*>(typedAt->no.child().get());
+                            const ASTOpImageLayerColor* typedNo = static_cast<const ASTOpImageLayerColor*>(TypedOp->no.child().get());
 
                             Ptr<ASTOpFixed> blackOp = new ASTOpFixed;
-                            blackOp->op.type = OP_TYPE::CO_CONSTANT;
+                            blackOp->op.type = EOpType::CO_CONSTANT;
                             blackOp->op.args.ColourConstant.value[0] = 0;
                             blackOp->op.args.ColourConstant.value[1] = 0;
 							blackOp->op.args.ColourConstant.value[2] = 0;
 							blackOp->op.args.ColourConstant.value[3] = 1;
 
                             Ptr<ASTOpFixed> plainOp = new ASTOpFixed;
-                            plainOp->op.type = OP_TYPE::IM_PLAINCOLOUR;
+                            plainOp->op.type = EOpType::IM_PLAINCOLOUR;
                             plainOp->SetChild( plainOp->op.args.ImagePlainColour.colour, blackOp );
                             plainOp->op.args.ImagePlainColour.format = EImageFormat::L_UByte;
                             plainOp->op.args.ImagePlainColour.size[0] = 4;
@@ -585,25 +597,25 @@ namespace mu
 							plainOp->op.args.ImagePlainColour.LODs = 1;
 
                             Ptr<ASTOpFixed> resizeOp = new ASTOpFixed;
-                            resizeOp->op.type = OP_TYPE::IM_RESIZELIKE;
+                            resizeOp->op.type = EOpType::IM_RESIZELIKE;
                             resizeOp->SetChild( resizeOp->op.args.ImageResizeLike.source, plainOp );
                             resizeOp->SetChild( resizeOp->op.args.ImageResizeLike.sizeSource, typedNo->base );
 
-                            Ptr<ASTOpConditional> maskOp = mu::Clone<ASTOpConditional>(typedAt);
+                            Ptr<ASTOpConditional> maskOp = mu::Clone<ASTOpConditional>(TypedOp);
                             maskOp->no = resizeOp;
 
                             // If there is no mask (because it is optional), we need to make a
                             // white plain image
                             maskOp->no = EnsureValidMask( typedNo->mask.child(), typedNo->base.child() );
 
-                            Ptr<ASTOpConditional> baseOp = mu::Clone<ASTOpConditional>(typedAt);
+                            Ptr<ASTOpConditional> baseOp = mu::Clone<ASTOpConditional>(TypedOp);
                             baseOp->no = typedNo->base.child();
 
                             Ptr<ASTOpImageLayerColor> softOp = mu::Clone<ASTOpImageLayerColor>(typedNo);
                             softOp->base = baseOp;
                             softOp->mask = maskOp;
 
-                            at = softOp;
+                            At = softOp;
                             break;
                         }
 
@@ -621,27 +633,27 @@ namespace mu
 
 
         //-------------------------------------------------------------------------------------
-        case OP_TYPE::IM_SWITCH:
+        case EOpType::IM_SWITCH:
         {
             // If the switch is not runtime, but the branches are, try to move the
             // switch down
 //            bool optimised = false;
 
-//            OP::ADDRESS variable = program.m_code[at].args.Switch.variable;
+//            OP::ADDRESS variable = program.m_code[At].args.Switch.variable;
 
 //            if ( !m_hasRuntimeParamVisitor.HasAny( variable, program ) )
 //            {
-//                SWITCH_CHAIN chain = GetSwitchChain( program, at );
+//                SWITCH_CHAIN chain = GetSwitchChain( program, At );
 
 //                bool branchHasAny = false;
-//                OP_TYPE branchType = (OP_TYPE)program.m_code[program.m_code[at].args.Switch.values[0]].type;
+//                EOpType branchType = (EOpType)program.m_code[program.m_code[At].args.Switch.values[0]].type;
 //                for ( map<uint16,OP::ADDRESS>::const_iterator it=chain.cases.begin();
 //                      it != chain.cases.end();
 //                      ++it )
 //                {
 //                    if ( program.m_code[it->second].type != branchType )
 //                    {
-//                        branchType = OP_TYPE::NONE;
+//                        branchType = EOpType::NONE;
 //                    }
 //                    else
 //                    {
@@ -654,7 +666,7 @@ namespace mu
 
 //                if ( chain.def && program.m_code[chain.def].type != branchType )
 //                {
-//                    branchType = OP_TYPE::NONE;
+//                    branchType = EOpType::NONE;
 //                }
 
 //                // Some branch in runtime
@@ -664,8 +676,8 @@ namespace mu
 //                    {
 
 //                    // TODO: Other operations
-//                    case OP_TYPE::IM_BLEND:
-//                    case OP_TYPE::IM_MULTIPLY:
+//                    case EOpType::IM_BLEND:
+//                    case EOpType::IM_MULTIPLY:
 //                    {
 //                        // Move the switch down the base
 //                        OP::ADDRESS baseAt = 0;
@@ -674,7 +686,7 @@ namespace mu
 //                              )
 //                        {
 //                            OP bsw;
-//                            bsw.type = OP_TYPE::IM_SWITCH;
+//                            bsw.type = EOpType::IM_SWITCH;
 //                            bsw.args.Switch.variable = variable;
 
 //                            for ( int b=0;
@@ -698,7 +710,7 @@ namespace mu
 //                              )
 //                        {
 //                            OP bsw;
-//                            bsw.type = OP_TYPE::IM_SWITCH;
+//                            bsw.type = EOpType::IM_SWITCH;
 //                            bsw.args.Switch.variable = variable;
 
 //                            for ( int b=0;
@@ -722,7 +734,7 @@ namespace mu
 //                              )
 //                        {
 //                            OP bsw;
-//                            bsw.type = OP_TYPE::IM_SWITCH;
+//                            bsw.type = EOpType::IM_SWITCH;
 //                            bsw.args.Switch.variable = variable;
 
 //                            for ( int b=0;
@@ -744,14 +756,14 @@ namespace mu
 //                        top.args.ImageLayer.base = baseAt;
 //                        top.args.ImageLayer.mask = maskAt;
 //                        top.args.ImageLayer.blended = blendedAt;
-//                        at = program.AddOp( top );
+//                        At = program.AddOp( top );
 //                        optimised = true;
 //                        break;
 //                    }
 
 
 //                    // TODO: Other operations
-//                    case OP_TYPE::IM_SOFTLIGHTCOLOUR:
+//                    case EOpType::IM_SOFTLIGHTCOLOUR:
 //                    {
 //                        // Move the switch down the base
 //                        OP::ADDRESS baseAt = 0;
@@ -760,7 +772,7 @@ namespace mu
 //                              )
 //                        {
 //                            OP bsw;
-//                            bsw.type = OP_TYPE::IM_SWITCH;
+//                            bsw.type = EOpType::IM_SWITCH;
 //                            bsw.args.Switch.variable = variable;
 
 //                            for ( int b=0;
@@ -784,7 +796,7 @@ namespace mu
 //                              )
 //                        {
 //                            OP bsw;
-//                            bsw.type = OP_TYPE::IM_SWITCH;
+//                            bsw.type = EOpType::IM_SWITCH;
 //                            bsw.args.Switch.variable = variable;
 
 //                            for ( int b=0;
@@ -808,7 +820,7 @@ namespace mu
 //                              )
 //                        {
 //                            OP bsw;
-//                            bsw.type = OP_TYPE::CO_SWITCH;
+//                            bsw.type = EOpType::CO_SWITCH;
 //                            bsw.args.Switch.variable = variable;
 
 //                            for ( int b=0;
@@ -830,7 +842,7 @@ namespace mu
 //                        top.args.ImageLayerColour.base = baseAt;
 //                        top.args.ImageLayerColour.mask = maskAt;
 //                        top.args.ImageLayerColour.colour = colourAt;
-//                        at = program.AddOp( top );
+//                        At = program.AddOp( top );
 //                        optimised = true;
 //                        break;
 //                    }
@@ -851,22 +863,22 @@ namespace mu
 
 
         //-----------------------------------------------------------------------------------------
-        case OP_TYPE::IM_COMPOSE:
+        case EOpType::IM_COMPOSE:
         {
-			const ASTOpImageCompose* typedAt = static_cast<const ASTOpImageCompose*>(at.get());
+			const ASTOpImageCompose* TypedOp = static_cast<const ASTOpImageCompose*>(At.get());
 
-            Ptr<ASTOp> blockAt = typedAt->BlockImage.child();
-			Ptr<ASTOp> baseAt = typedAt->Base.child();
-			Ptr<ASTOp> layoutAt = typedAt->Layout.child();
+            Ptr<ASTOp> blockAt = TypedOp->BlockImage.child();
+			Ptr<ASTOp> baseAt = TypedOp->Base.child();
+			Ptr<ASTOp> layoutAt = TypedOp->Layout.child();
 
 			if (!blockAt)
 			{
-				at = baseAt;
+				At = baseAt;
 				break;
 			}
 
-            OP_TYPE blockType = blockAt->GetOpType();
-            OP_TYPE baseType = baseAt->GetOpType();
+            EOpType blockType = blockAt->GetOpType();
+            EOpType baseType = baseAt->GetOpType();
 
             bool baseHasRuntime = m_hasRuntimeParamVisitor.HasAny( baseAt );
             bool blockHasRuntime = m_hasRuntimeParamVisitor.HasAny( blockAt );
@@ -881,7 +893,7 @@ namespace mu
                 {
                     switch ( blockType )
                     {
-                    case OP_TYPE::IM_LAYERCOLOUR:
+                    case EOpType::IM_LAYERCOLOUR:
                     {
                         optimised = true;
 
@@ -904,13 +916,13 @@ namespace mu
 							Ptr<ASTOp> newBaseMask = EnsureValidMask(baseMask, baseImage);
 							Ptr<ASTOp> newBlockMask = EnsureValidMask(blockMask, blockImage);
 
-                            maskOp = mu::Clone<ASTOpImageCompose>(typedAt);
+                            maskOp = mu::Clone<ASTOpImageCompose>(TypedOp);
                             maskOp->Base = newBaseMask;
                             maskOp->BlockImage = newBlockMask;
                         }
 
                         // The base is composition of the bases of both layer effect
-						Ptr<ASTOpImageCompose> baseOp = mu::Clone<ASTOpImageCompose>(typedAt);
+						Ptr<ASTOpImageCompose> baseOp = mu::Clone<ASTOpImageCompose>(TypedOp);
                         baseOp->Base = baseImage;
                         baseOp->BlockImage = blockImage;
 
@@ -919,11 +931,11 @@ namespace mu
                         nop->base = baseOp;
 
                         // Done
-                        at = nop;
+                        At = nop;
                         break;
                     }
 
-                    case OP_TYPE::IM_LAYER:
+                    case EOpType::IM_LAYER:
                     {
                         optimised = true;
 
@@ -948,18 +960,18 @@ namespace mu
 							Ptr<ASTOp> newBaseMask = EnsureValidMask(baseMask, baseImage);
 							Ptr<ASTOp> newBlockMask = EnsureValidMask(blockMask, blockImage);
 
-                            maskOp = mu::Clone<ASTOpImageCompose>(typedAt);
+                            maskOp = mu::Clone<ASTOpImageCompose>(TypedOp);
                             maskOp->Base = newBaseMask;
                             maskOp->BlockImage = newBlockMask;
                         }
 
                         // The base is composition of the bases of both layer effect
-						Ptr<ASTOpImageCompose> baseOp = mu::Clone<ASTOpImageCompose>(typedAt);
+						Ptr<ASTOpImageCompose> baseOp = mu::Clone<ASTOpImageCompose>(TypedOp);
                         baseOp->Base = baseImage;
                         baseOp->BlockImage = blockImage;
 
                         // The base is composition of the bases of both layer effect
-						Ptr<ASTOpImageCompose> blendedOp = mu::Clone<ASTOpImageCompose>(typedAt);
+						Ptr<ASTOpImageCompose> blendedOp = mu::Clone<ASTOpImageCompose>(TypedOp);
                         blendedOp->Base = baseBlended;
                         blendedOp->BlockImage = blockBlended;
 
@@ -969,7 +981,7 @@ namespace mu
                         nop->blend = blendedOp;
 
                         // Done
-                        at = nop;
+                        At = nop;
                         break;
                     }
 
@@ -983,7 +995,7 @@ namespace mu
             // Swap two composes
             if ( !optimised && baseHasRuntime && !blockHasRuntime
                  &&
-                 baseType == OP_TYPE::IM_COMPOSE )
+                 baseType == EOpType::IM_COMPOSE )
             {
 				const ASTOpImageCompose* typedBaseAt = static_cast<const ASTOpImageCompose*>(baseAt.get());
 
@@ -994,13 +1006,13 @@ namespace mu
                     optimised = true;
 
                     // Swap
-					Ptr<ASTOpImageCompose> childCompose = mu::Clone<ASTOpImageCompose>(at);
+					Ptr<ASTOpImageCompose> childCompose = mu::Clone<ASTOpImageCompose>(At);
                     childCompose->Base = typedBaseAt->Base.child();
 
 					Ptr<ASTOpImageCompose> parentCompose = mu::Clone<ASTOpImageCompose>(baseAt);
                     parentCompose->Base = childCompose;
 
-                    at = parentCompose;
+                    At = parentCompose;
                 }
             }
 
@@ -1018,7 +1030,7 @@ namespace mu
             {
                 switch ( blockType )
                 {
-                case OP_TYPE::IM_LAYERCOLOUR:
+                case EOpType::IM_LAYERCOLOUR:
                 {
                     optimised = true;
 
@@ -1031,21 +1043,21 @@ namespace mu
                     // no mask and the base of the layer opertation is a blanklayout, we can skip
                     // generating a mask.
                     Ptr<ASTOpImageCompose> maskOp;
-                    if (blockMask || baseType!=OP_TYPE::IM_BLANKLAYOUT)
+                    if (blockMask || baseType!=EOpType::IM_BLANKLAYOUT)
                     {
-                        maskOp = mu::Clone<ASTOpImageCompose>(at);
+                        maskOp = mu::Clone<ASTOpImageCompose>(At);
 						Ptr<ASTOp> newMaskBlock = EnsureValidMask(blockMask, blockImage);
                         maskOp->BlockImage = newMaskBlock;
 
                         Ptr<ASTOpFixed> blackOp = new ASTOpFixed;
-                        blackOp->op.type = OP_TYPE::CO_CONSTANT;
+                        blackOp->op.type = EOpType::CO_CONSTANT;
                         blackOp->op.args.ColourConstant.value[0] = 0;
                         blackOp->op.args.ColourConstant.value[1] = 0;
                         blackOp->op.args.ColourConstant.value[2] = 0;
 						blackOp->op.args.ColourConstant.value[3] = 1;
 
                         Ptr<ASTOpFixed> plainOp = new ASTOpFixed;
-                        plainOp->op.type = OP_TYPE::IM_PLAINCOLOUR;
+                        plainOp->op.type = EOpType::IM_PLAINCOLOUR;
                         plainOp->SetChild( plainOp->op.args.ImagePlainColour.colour, blackOp );
                         plainOp->op.args.ImagePlainColour.format = EImageFormat::L_UByte;
                         plainOp->op.args.ImagePlainColour.size[0] = 4;
@@ -1053,7 +1065,7 @@ namespace mu
 						plainOp->op.args.ImagePlainColour.LODs = 1;
 
                         Ptr<ASTOpFixed> baseResizeOp = new ASTOpFixed;
-                        baseResizeOp->op.type = OP_TYPE::IM_RESIZELIKE;
+                        baseResizeOp->op.type = EOpType::IM_RESIZELIKE;
                         baseResizeOp->SetChild( baseResizeOp->op.args.ImageResizeLike.sizeSource, baseAt );
                         baseResizeOp->SetChild( baseResizeOp->op.args.ImageResizeLike.source, plainOp );
 
@@ -1061,7 +1073,7 @@ namespace mu
                     }
 
                     // The base is composition of the layer base on the compose base
-                    Ptr<ASTOpImageCompose> baseOp = mu::Clone<ASTOpImageCompose>(at);
+                    Ptr<ASTOpImageCompose> baseOp = mu::Clone<ASTOpImageCompose>(At);
                     baseOp->BlockImage = typedBlockAt->base.child();
 
                     Ptr<ASTOpImageLayerColor> nop = mu::Clone<ASTOpImageLayerColor>(blockAt);
@@ -1069,11 +1081,11 @@ namespace mu
                     nop->base = baseOp;
 
                     // Done
-                    at = nop;
+                    At = nop;
                     break;
                 }
 
-                case OP_TYPE::IM_LAYER:
+                case EOpType::IM_LAYER:
                 {
                     optimised = true;
 
@@ -1087,22 +1099,22 @@ namespace mu
                     // no mask and the base of the layer opertation is a blanklayout, we can skip
                     // generating a mask.
                     Ptr<ASTOpImageCompose> maskOp;
-                    if (blockMask || baseType != OP_TYPE::IM_BLANKLAYOUT)
+                    if (blockMask || baseType != EOpType::IM_BLANKLAYOUT)
                     {
-                        maskOp = mu::Clone<ASTOpImageCompose>(at);
+                        maskOp = mu::Clone<ASTOpImageCompose>(At);
 						Ptr<ASTOp> newMaskBlock = EnsureValidMask(blockMask, blockImage);
                         maskOp->BlockImage = newMaskBlock;
 
 
                         Ptr<ASTOpFixed> blackOp = new ASTOpFixed;
-                        blackOp->op.type = OP_TYPE::CO_CONSTANT;
+                        blackOp->op.type = EOpType::CO_CONSTANT;
                         blackOp->op.args.ColourConstant.value[0] = 0;
                         blackOp->op.args.ColourConstant.value[1] = 0;
                         blackOp->op.args.ColourConstant.value[2] = 0;
 						blackOp->op.args.ColourConstant.value[3] = 1;
 
                         Ptr<ASTOpFixed> plainOp = new ASTOpFixed;
-                        plainOp->op.type = OP_TYPE::IM_PLAINCOLOUR;
+                        plainOp->op.type = EOpType::IM_PLAINCOLOUR;
                         plainOp->SetChild( plainOp->op.args.ImagePlainColour.colour, blackOp );
                         plainOp->op.args.ImagePlainColour.format = EImageFormat::L_UByte;
                         plainOp->op.args.ImagePlainColour.size[0] = 4;
@@ -1110,7 +1122,7 @@ namespace mu
 						plainOp->op.args.ImagePlainColour.LODs = 1;
 
                         Ptr<ASTOpFixed> baseResizeOp = new ASTOpFixed;
-                        baseResizeOp->op.type = OP_TYPE::IM_RESIZELIKE;
+                        baseResizeOp->op.type = EOpType::IM_RESIZELIKE;
                         baseResizeOp->SetChild( baseResizeOp->op.args.ImageResizeLike.sizeSource, baseAt );
                         baseResizeOp->SetChild( baseResizeOp->op.args.ImageResizeLike.source, plainOp );
 
@@ -1118,19 +1130,19 @@ namespace mu
                     }
 
                     // The blended is a compose of the blended image on a blank image
-					Ptr<ASTOpImageCompose> blendedOp = mu::Clone<ASTOpImageCompose>(at);
+					Ptr<ASTOpImageCompose> blendedOp = mu::Clone<ASTOpImageCompose>(At);
                     {
                         blendedOp->BlockImage = blockBlended;
 
                         Ptr<ASTOpFixed> blackOp = new ASTOpFixed;
-                        blackOp->op.type = OP_TYPE::CO_CONSTANT;
+                        blackOp->op.type = EOpType::CO_CONSTANT;
                         blackOp->op.args.ColourConstant.value[0] = 0;
                         blackOp->op.args.ColourConstant.value[1] = 0;
                         blackOp->op.args.ColourConstant.value[2] = 0;
 						blackOp->op.args.ColourConstant.value[3] = 1;
 
                         Ptr<ASTOpFixed> plainOp = new ASTOpFixed;
-                        plainOp->op.type = OP_TYPE::IM_PLAINCOLOUR;
+                        plainOp->op.type = EOpType::IM_PLAINCOLOUR;
                         plainOp->SetChild( plainOp->op.args.ImagePlainColour.colour, blackOp );
                         FImageDesc blendedDesc = baseAt->GetImageDesc();
                         plainOp->op.args.ImagePlainColour.format = blendedDesc.m_format;
@@ -1139,7 +1151,7 @@ namespace mu
 						plainOp->op.args.ImagePlainColour.LODs = 1;
 
                         Ptr<ASTOpFixed> resizeOp = new ASTOpFixed;
-                        resizeOp->op.type = OP_TYPE::IM_RESIZELIKE;
+                        resizeOp->op.type = EOpType::IM_RESIZELIKE;
                         resizeOp->SetChild( resizeOp->op.args.ImageResizeLike.sizeSource, baseAt );
                         resizeOp->SetChild( resizeOp->op.args.ImageResizeLike.source, plainOp );
 
@@ -1147,7 +1159,7 @@ namespace mu
                     }
 
                     // The base is composition of the softlight base on the compose base
-                    Ptr<ASTOpImageCompose> baseOp = mu::Clone<ASTOpImageCompose>(at);
+                    Ptr<ASTOpImageCompose> baseOp = mu::Clone<ASTOpImageCompose>(At);
                     baseOp->BlockImage = typedBlockAt->base.child();
 
                     Ptr<ASTOpImageLayer> nop = mu::Clone<ASTOpImageLayer>(blockAt);
@@ -1156,12 +1168,12 @@ namespace mu
                     nop->blend = blendedOp;
 
                     // Done
-                    at = nop;
+                    At = nop;
                     break;
                 }
 
 
-//                case OP_TYPE::IM_INTERPOLATE:
+//                case EOpType::IM_INTERPOLATE:
 //                {
 //                    optimised = true;
 //                    OP op = program.m_code[blockAt];
@@ -1171,7 +1183,7 @@ namespace mu
 //                    {
 //                        if ( op.args.ImageInterpolate.targets[t] )
 //                        {
-//                            OP targetOp = program.m_code[at];
+//                            OP targetOp = program.m_code[At];
 //                            targetOp.args.ImageCompose.blockImage =
 //                                    op.args.ImageInterpolate.targets[t];
 //                            op.args.ImageInterpolate.targets[t] = program.AddOp( targetOp );
@@ -1179,10 +1191,10 @@ namespace mu
 //                    }
 
 //                    // Done
-//                    at = program.AddOp( op );
+//                    At = program.AddOp( op );
 
 //                    // Reprocess the new children
-//                    at = Recurse( at, program );
+//                    At = Recurse( At, program );
 
 //                    break;
 //                }
@@ -1199,23 +1211,23 @@ namespace mu
             {
                 switch ( baseType )
                 {
-                case OP_TYPE::IM_LAYERCOLOUR:
+                case EOpType::IM_LAYERCOLOUR:
                 {
                     optimised = true;
 
                     const ASTOpImageLayerColor* typedBaseAt = static_cast<const ASTOpImageLayerColor*>(baseAt.get());
 
-					Ptr<ASTOpImageCompose> maskOp = mu::Clone<ASTOpImageCompose>(at);
+					Ptr<ASTOpImageCompose> maskOp = mu::Clone<ASTOpImageCompose>(At);
                     {
                         Ptr<ASTOpFixed> blackOp = new ASTOpFixed;
-                        blackOp->op.type = OP_TYPE::CO_CONSTANT;
+                        blackOp->op.type = EOpType::CO_CONSTANT;
                         blackOp->op.args.ColourConstant.value[0] = 0;
                         blackOp->op.args.ColourConstant.value[1] = 0;
                         blackOp->op.args.ColourConstant.value[2] = 0;
 						blackOp->op.args.ColourConstant.value[3] = 1;
 
                         Ptr<ASTOpFixed> plainOp = new ASTOpFixed;
-                        plainOp->op.type = OP_TYPE::IM_PLAINCOLOUR;
+                        plainOp->op.type = EOpType::IM_PLAINCOLOUR;
                         plainOp->SetChild( plainOp->op.args.ImagePlainColour.colour, blackOp );
                         plainOp->op.args.ImagePlainColour.format = EImageFormat::L_UByte; //TODO: FORMAT_LIKE
                         plainOp->op.args.ImagePlainColour.size[0] = 4;
@@ -1223,7 +1235,7 @@ namespace mu
 						plainOp->op.args.ImagePlainColour.LODs = 1;
 
                         Ptr<ASTOpFixed> blockResizeOp = new ASTOpFixed;
-                        blockResizeOp->op.type = OP_TYPE::IM_RESIZELIKE;
+                        blockResizeOp->op.type = EOpType::IM_RESIZELIKE;
                         blockResizeOp->SetChild( blockResizeOp->op.args.ImageResizeLike.sizeSource, blockAt );
                         blockResizeOp->SetChild( blockResizeOp->op.args.ImageResizeLike.source, plainOp );
 
@@ -1234,7 +1246,7 @@ namespace mu
                     }
 
                     // The base is composition of the softlight base on the compose base
-					Ptr<ASTOpImageCompose> baseOp = mu::Clone<ASTOpImageCompose>(at);
+					Ptr<ASTOpImageCompose> baseOp = mu::Clone<ASTOpImageCompose>(At);
                     baseOp->Base = typedBaseAt->base.child();
 
                     Ptr<ASTOpImageLayerColor> nop = mu::Clone<ASTOpImageLayerColor>(baseAt);
@@ -1242,27 +1254,27 @@ namespace mu
                     nop->mask = maskOp;
 
                     // Done
-                    at = nop;
+                    At = nop;
                     break;
                 }
 
-                case OP_TYPE::IM_LAYER:
+                case EOpType::IM_LAYER:
                 {
                     optimised = true;
 
                     const ASTOpImageLayer* typedBaseAt = static_cast<const ASTOpImageLayer*>(baseAt.get());
 
-					Ptr<ASTOpImageCompose> maskOp = mu::Clone<ASTOpImageCompose>(at);
+					Ptr<ASTOpImageCompose> maskOp = mu::Clone<ASTOpImageCompose>(At);
                     {
                         Ptr<ASTOpFixed> blackOp = new ASTOpFixed;
-                        blackOp->op.type = OP_TYPE::CO_CONSTANT;
+                        blackOp->op.type = EOpType::CO_CONSTANT;
                         blackOp->op.args.ColourConstant.value[0] = 0;
                         blackOp->op.args.ColourConstant.value[1] = 0;
                         blackOp->op.args.ColourConstant.value[2] = 0;
 						blackOp->op.args.ColourConstant.value[3] = 1;
 
                         Ptr<ASTOpFixed> plainOp = new ASTOpFixed;
-                        plainOp->op.type = OP_TYPE::IM_PLAINCOLOUR;
+                        plainOp->op.type = EOpType::IM_PLAINCOLOUR;
                         plainOp->SetChild( plainOp->op.args.ImagePlainColour.colour, blackOp );
                         plainOp->op.args.ImagePlainColour.format = EImageFormat::L_UByte; //TODO: FORMAT_LIKE
                         plainOp->op.args.ImagePlainColour.size[0] = 4;
@@ -1270,7 +1282,7 @@ namespace mu
 						plainOp->op.args.ImagePlainColour.LODs = 1;
 
                         Ptr<ASTOpFixed> blockResizeOp = new ASTOpFixed;
-                        blockResizeOp->op.type = OP_TYPE::IM_RESIZELIKE;
+                        blockResizeOp->op.type = EOpType::IM_RESIZELIKE;
                         blockResizeOp->SetChild( blockResizeOp->op.args.ImageResizeLike.sizeSource, blockAt );
                         blockResizeOp->SetChild( blockResizeOp->op.args.ImageResizeLike.source, plainOp );
 
@@ -1281,7 +1293,7 @@ namespace mu
                     }
 
                     // The base is composition of the effect base on the compose base
-                    Ptr<ASTOpImageCompose> baseOp = mu::Clone<ASTOpImageCompose>(at);
+                    Ptr<ASTOpImageCompose> baseOp = mu::Clone<ASTOpImageCompose>(At);
                     baseOp->Base = typedBaseAt->base.child();
 
                     Ptr<ASTOpImageLayer> nop = mu::Clone<ASTOpImageLayer>(baseAt);
@@ -1289,12 +1301,12 @@ namespace mu
                     nop->mask = maskOp;
 
                     // Done
-                    at = nop;
+                    At = nop;
                     break;
                 }
 
 
-//                case OP_TYPE::IM_INTERPOLATE:
+//                case EOpType::IM_INTERPOLATE:
 //                {
 //                    optimised = true;
 //                    OP op = program.m_code[baseAt];
@@ -1304,7 +1316,7 @@ namespace mu
 //                    {
 //                        if ( op.args.ImageInterpolate.targets[t] )
 //                        {
-//                            OP targetOp = program.m_code[at];
+//                            OP targetOp = program.m_code[At];
 //                            targetOp.args.ImageCompose.base =
 //                                    op.args.ImageInterpolate.targets[t];
 //                            op.args.ImageInterpolate.targets[t] = program.AddOp( targetOp );
@@ -1312,10 +1324,10 @@ namespace mu
 //                    }
 
 //                    // Done
-//                    at = program.AddOp( op );
+//                    At = program.AddOp( op );
 
 //                    // Reprocess the new children
-//                    at = Recurse( at, program );
+//                    At = Recurse( At, program );
 //                    break;
 //                }
 
@@ -1333,28 +1345,28 @@ namespace mu
 /*
         //-----------------------------------------------------------------------------------------
         // TODO: Other ops?
-        case OP_TYPE::IM_BLEND:
+        case EOpType::IM_BLEND:
         {
-            OP op = program.m_code[at];
+            OP op = program.m_code[At];
 
             if ( !m_hasRuntimeParamVisitor.HasAny( op.args.ImageLayer.mask, program ) )
             {
                 // If both the base and the blended have the same image layer operation with
                 // similar parameters, we can move that operation up.
-                OP_TYPE baseType = program.m_code[ op.args.ImageLayer.base ].type;
-                OP_TYPE blendedType = program.m_code[ op.args.ImageLayer.blended ].type;
+                EOpType baseType = program.m_code[ op.args.ImageLayer.base ].type;
+                EOpType blendedType = program.m_code[ op.args.ImageLayer.blended ].type;
                 if ( baseType == blendedType )
                 {
                     switch ( baseType )
                     {
-                    case OP_TYPE::IM_BLENDCOLOUR:
-                    case OP_TYPE::IM_SOFTLIGHTCOLOUR:
-                    case OP_TYPE::IM_HARDLIGHTCOLOUR:
-                    case OP_TYPE::IM_BURNCOLOUR:
-                    case OP_TYPE::IM_SCREENCOLOUR:
-                    case OP_TYPE::IM_OVERLAYCOLOUR:
-                    case OP_TYPE::IM_DODGECOLOUR:
-                    case OP_TYPE::IM_MULTIPLYCOLOUR:
+                    case EOpType::IM_BLENDCOLOUR:
+                    case EOpType::IM_SOFTLIGHTCOLOUR:
+                    case EOpType::IM_HARDLIGHTCOLOUR:
+                    case EOpType::IM_BURNCOLOUR:
+                    case EOpType::IM_SCREENCOLOUR:
+                    case EOpType::IM_OVERLAYCOLOUR:
+                    case EOpType::IM_DODGECOLOUR:
+                    case EOpType::IM_MULTIPLYCOLOUR:
                     {
                         if ( OptimisationOptions.m_optimiseOverlappedMasks )
                         {
@@ -1378,7 +1390,7 @@ namespace mu
                                 m_modified = true;
 
                                 OP baseOp;
-                                baseOp.type = OP_TYPE::IM_BLEND;
+                                baseOp.type = EOpType::IM_BLEND;
                                 baseOp.args.ImageLayer.base =
                                         program.m_code[ op.args.ImageLayer.base ].args.ImageLayerColour.base;
                                 baseOp.args.ImageLayer.blended =
@@ -1387,11 +1399,11 @@ namespace mu
 
                                 // TODO: Find out why these are not equivalent
     //							OP maskOp;
-    //							maskOp.type = OP_TYPE::IM_SCREEN;
+    //							maskOp.type = EOpType::IM_SCREEN;
     //							maskOp.args.ImageLayer.base = baseMaskAt;
     //							maskOp.args.ImageLayer.blended = blendedMaskAt;
                                 OP maskOp;
-                                maskOp.type = OP_TYPE::IM_BLEND;
+                                maskOp.type = EOpType::IM_BLEND;
                                 maskOp.args.ImageLayer.base = baseMaskAt;
                                 maskOp.args.ImageLayer.blended = blendedMaskAt;
                                 maskOp.args.ImageLayer.mask = blendedMaskAt;
@@ -1402,7 +1414,7 @@ namespace mu
                                 top.args.ImageLayerColour.base = program.AddOp( baseOp );
                                 top.args.ImageLayerColour.mask = program.AddOp( maskOp );
                                 top.args.ImageLayerColour.colour = baseColourAt;
-                                at = program.AddOp( top );
+                                At = program.AddOp( top );
                             }
                         }
 
@@ -1422,15 +1434,15 @@ namespace mu
 
         //-----------------------------------------------------------------------------------------
         // Sink the mipmap if worth it.
-        case OP_TYPE::IM_MIPMAP:
+        case EOpType::IM_MIPMAP:
         {
-			const ASTOpImageMipmap* typedAt = static_cast<const ASTOpImageMipmap*>(at.get());
+			const ASTOpImageMipmap* TypedOp = static_cast<const ASTOpImageMipmap*>(At.get());
 
-			Ptr<ASTOp> sourceOp = typedAt->Source.child();
+			Ptr<ASTOp> sourceOp = TypedOp->Source.child();
 
             switch ( sourceOp->GetOpType() )
             {
-            case OP_TYPE::IM_LAYERCOLOUR:
+            case EOpType::IM_LAYERCOLOUR:
             {
                 const ASTOpImageLayerColor* typedSource = static_cast<const ASTOpImageLayerColor*>(sourceOp.get());
 
@@ -1442,19 +1454,19 @@ namespace mu
 
 					Ptr<ASTOpImageLayerColor> top = mu::Clone<ASTOpImageLayerColor>(sourceOp);
 
-					Ptr<ASTOpImageMipmap> baseOp = mu::Clone<ASTOpImageMipmap>(at);
+					Ptr<ASTOpImageMipmap> baseOp = mu::Clone<ASTOpImageMipmap>(At);
                     baseOp->Source = typedSource->base.child();
                     top->base = baseOp;
 
 					Ptr<ASTOp> sourceMaskOp = typedSource->mask.child();
                     if (sourceMaskOp)
                     {
-						Ptr<ASTOpImageMipmap> maskOp = mu::Clone<ASTOpImageMipmap>(at);
+						Ptr<ASTOpImageMipmap> maskOp = mu::Clone<ASTOpImageMipmap>(At);
                         maskOp->Source = sourceMaskOp;
                         top->mask = maskOp;
                     }
 
-                    at = top;
+                    At = top;
                 }
 
                 break;
@@ -1472,7 +1484,7 @@ namespace mu
 
         }
 
-        return at;
+        return At;
     }
 
 
@@ -1494,21 +1506,21 @@ namespace mu
         }
 
 
-        bool Visit( const Ptr<ASTOp>& at ) override
+        bool Visit( const Ptr<ASTOp>& At ) override
         {
             bool recurse = false;
 
             const std::array<uint8_t, (size_t)EImageFormat::Count>& currentFormats = GetCurrentState();
 
             // Remove unsupported formats
-            if (GetOpDataType( at->GetOpType() )==DT_IMAGE)
+            if (GetOpDataType( At->GetOpType() )==DT_IMAGE)
             {
-				std::array<uint8_t, (size_t)EImageFormat::Count>* it = m_supportedFormats.Find(at);
+				std::array<uint8_t, (size_t)EImageFormat::Count>* it = m_supportedFormats.Find(At);
                 if (!it)
                 {
                     // Default to all supported
-                    m_supportedFormats.Add(at, m_allSupported);
-                    it = m_supportedFormats.Find(at);
+                    m_supportedFormats.Add(At, m_allSupported);
+                    it = m_supportedFormats.Find(At);
                 }
 
                 for ( unsigned f=0; f< (unsigned)EImageFormat::Count; ++f )
@@ -1520,34 +1532,34 @@ namespace mu
                 }
             }
 
-            switch ( at->GetOpType() )
+            switch ( At->GetOpType() )
             {
             // TODO: Code shared with the constant data format optimisation visitor
-            case OP_TYPE::IM_LAYERCOLOUR:
+            case EOpType::IM_LAYERCOLOUR:
             {
-                const ASTOpImageLayerColor* typedAt = static_cast<const ASTOpImageLayerColor*>(at.get());
+                const ASTOpImageLayerColor* TypedOp = static_cast<const ASTOpImageLayerColor*>(At.get());
 
-                RecurseWithCurrentState( typedAt->base.child() );
-                RecurseWithCurrentState( typedAt->color.child() );
+                RecurseWithCurrentState( TypedOp->base.child() );
+                RecurseWithCurrentState( TypedOp->color.child() );
 
-                if ( typedAt->mask )
+                if ( TypedOp->mask )
                 {
                     std::array<uint8_t, (size_t)EImageFormat::Count> newState;
                     newState.fill(0);
                     newState[(size_t)EImageFormat::L_UByte ] = 1;
                     newState[(size_t)EImageFormat::L_UByteRLE ] = 1;
 
-                    RecurseWithState( typedAt->mask.child(), newState );
+                    RecurseWithState( TypedOp->mask.child(), newState );
                 }
                 break;
             }
 
-            case OP_TYPE::IM_LAYER:
+            case EOpType::IM_LAYER:
             {
-                const ASTOpImageLayer* typedAt = static_cast<const ASTOpImageLayer*>(at.get());
+                const ASTOpImageLayer* TypedOp = static_cast<const ASTOpImageLayer*>(At.get());
 
-                RecurseWithCurrentState( typedAt->base.child() );
-                RecurseWithCurrentState( typedAt->blend.child() );
+                RecurseWithCurrentState( TypedOp->base.child() );
+                RecurseWithCurrentState( TypedOp->blend.child() );
 
                 std::array<uint8_t, (size_t)EImageFormat::Count> newState;
                 newState.fill(0);
@@ -1555,24 +1567,24 @@ namespace mu
                 //newState[ L_UByte ] = 1;
                 //newState[ L_UByteRLE ] = 1;
 
-                if ( typedAt->mask )
+                if ( TypedOp->mask )
                 {
-                    RecurseWithState( typedAt->mask.child(), newState );
+                    RecurseWithState( TypedOp->mask.child(), newState );
                 }
                 break;
             }
 
-            case OP_TYPE::IM_DISPLACE:
+            case EOpType::IM_DISPLACE:
             {
-				const ASTOpFixed* typedAt = static_cast<const ASTOpFixed*>(at.get());
+				const ASTOpFixed* TypedOp = static_cast<const ASTOpFixed*>(At.get());
 
-                RecurseWithCurrentState( typedAt->children[typedAt->op.args.ImageDisplace.source].child() );
+                RecurseWithCurrentState( TypedOp->children[TypedOp->op.args.ImageDisplace.source].child() );
 
                 std::array<uint8_t, (size_t)EImageFormat::Count> newState;
                 newState.fill(0);
                 newState[(size_t)EImageFormat::L_UByte ] = 1;
                 newState[(size_t)EImageFormat::L_UByteRLE ] = 1;
-                RecurseWithState( typedAt->children[typedAt->op.args.ImageDisplace.displacementMap].child(), newState );
+                RecurseWithState( TypedOp->children[TypedOp->op.args.ImageDisplace.displacementMap].child(), newState );
                 break;
             }
 
@@ -1586,9 +1598,9 @@ namespace mu
         }
 
 
-        bool IsSupportedFormat( Ptr<ASTOp> at, EImageFormat format ) const
+        bool IsSupportedFormat( Ptr<ASTOp> At, EImageFormat format ) const
         {
-            const std::array<uint8_t, (size_t)EImageFormat::Count>* it = m_supportedFormats.Find(at);
+            const std::array<uint8_t, (size_t)EImageFormat::Count>* it = m_supportedFormats.Find(At);
             if (!it)
             {
                 return false;
@@ -1620,7 +1632,7 @@ namespace mu
 		TSet<FString>* it = m_resultCache.Find( FState(root,false) );
         if (it)
         {
-            m_params = *it;
+            Parameters = *it;
             return;
         }
 
@@ -1628,25 +1640,25 @@ namespace mu
         {
             MUTABLE_CPUPROFILER_SCOPE(SubtreeRelevantParametersVisitorAST);
 
-            m_params.Empty();
+            Parameters.Empty();
 
             // The state is the onlyLayoutRelevant flag
             ASTOp::Traverse_TopDown_Unique_Imprecise_WithState<bool>( root, false,
-                [&]( Ptr<ASTOp>& at, bool& state, TArray<TPair<Ptr<ASTOp>,bool>>& pending )
+                [&]( Ptr<ASTOp>& At, bool& state, TArray<TPair<Ptr<ASTOp>,bool>>& Pending )
             {
                 (void)state;
 
-                switch ( at->GetOpType() )
+                switch ( At->GetOpType() )
                 {
-                case OP_TYPE::NU_PARAMETER:
-                case OP_TYPE::SC_PARAMETER:
-                case OP_TYPE::BO_PARAMETER:
-                case OP_TYPE::CO_PARAMETER:
-                case OP_TYPE::PR_PARAMETER:
-                case OP_TYPE::IM_PARAMETER:
+                case EOpType::NU_PARAMETER:
+                case EOpType::SC_PARAMETER:
+                case EOpType::BO_PARAMETER:
+                case EOpType::CO_PARAMETER:
+                case EOpType::PR_PARAMETER:
+                case EOpType::IM_PARAMETER:
                 {
-					const ASTOpParameter* typedAt = static_cast<const ASTOpParameter*>(at.get());
-                    m_params.Add(typedAt->parameter.m_name);
+					const ASTOpParameter* TypedOp = static_cast<const ASTOpParameter*>(At.get());
+                    Parameters.Add(TypedOp->parameter.Name);
 
                     // Not interested in the parameters from the parameters decorators.
                     return false;
@@ -1654,41 +1666,41 @@ namespace mu
                     break;
                 }
 
-				case OP_TYPE::LA_FROMMESH:
+				case EOpType::LA_FROMMESH:
 				{
 					// Manually choose how to recurse this op
-					const ASTOpLayoutFromMesh* pTyped = static_cast<const ASTOpLayoutFromMesh*>(at.get());
+					const ASTOpLayoutFromMesh* TypedOp = static_cast<const ASTOpLayoutFromMesh*>(At.get());
 
 					// For that mesh we only want to know about the layouts
-					if (const ASTChild& Mesh = pTyped->Mesh)
+					if (const ASTChild& Mesh = TypedOp->Mesh)
 					{
-						pending.Add({ Mesh.m_child, true });
+						Pending.Add({ Mesh.Child, true });
 					}
 
 					return false;
 				}
 
-                case OP_TYPE::ME_MORPH:
+                case EOpType::ME_MORPH:
                 {
                     // Manually choose how to recurse this op
-					const ASTOpMeshMorph* pTyped = static_cast<const ASTOpMeshMorph*>( at.get() );
+					const ASTOpMeshMorph* TypedOp = static_cast<const ASTOpMeshMorph*>( At.get() );
 
-                    if ( pTyped->Base )
+                    if ( TypedOp->Base )
                     {
-						pending.Add({ pTyped->Base.m_child, state });
+						Pending.Add({ TypedOp->Base.Child, state });
                     }
 
                     // Mesh morphs don't modify the layouts, so we can ignore the factor and morphs
                     if (!state)
                     {
-                        if ( pTyped->Factor )
+                        if ( TypedOp->Factor )
                         {
-							pending.Add({ pTyped->Factor.m_child, state });
+							Pending.Add({ TypedOp->Factor.Child, state });
                         }
 
-                        if ( pTyped->Target )
+                        if ( TypedOp->Target )
                         {
-							pending.Add({ pTyped->Target.m_child, state });
+							Pending.Add({ TypedOp->Target.Child, state });
                         }
                     }
 
@@ -1702,7 +1714,7 @@ namespace mu
                 return true;
             });
 
-            m_resultCache.Add( FState(root,false), m_params );
+            m_resultCache.Add( FState(root,false), Parameters );
         }
     }
 
@@ -1747,7 +1759,7 @@ namespace mu
 
 					// Temp copy
 					TArray<FString> ParamCopy;
-					for ( const FString& e: subtreeParams.m_params )
+					for ( const FString& e: subtreeParams.Parameters )
 					{
 						ParamCopy.Add(e);
 					}
@@ -1758,16 +1770,16 @@ namespace mu
         }
 
 
-        bool Visit( const Ptr<ASTOp>& at ) override
+        bool Visit( const Ptr<ASTOp>& At ) override
         {
-            bool thisIsRuntime = m_hasRuntimeParamVisitor.HasAny( at );
+            bool thisIsRuntime = m_hasRuntimeParamVisitor.HasAny( At );
 
             bool resourceRoot = GetCurrentState().Key;
             bool parentIsRuntime = GetCurrentState().Value;
 
-			m_cache.FindOrAdd(at, false);
+			m_cache.FindOrAdd(At, false);
 
-            OP_TYPE type = at->GetOpType();
+            EOpType type = At->GetOpType();
             if ( GetOpDesc( type ).cached )
             {
                 // If parent is runtime, but we are not
@@ -1779,37 +1791,37 @@ namespace mu
                 {
                     // We want to cache this result to update the instances.
                     // Mark this as update cache
-                    m_cache.Add(at, true);
+                    m_cache.Add(At, true);
                 }
             }
 
-            if ( !m_cache[at] && resourceRoot && thisIsRuntime )
+            if ( !m_cache[At] && resourceRoot && thisIsRuntime )
             {
-                m_dynamicResourceRoot.Add(at, true);
+                m_dynamicResourceRoot.Add(At, true);
             }
 
-            if ( !m_cache[at] && thisIsRuntime )
+            if ( !m_cache[At] && thisIsRuntime )
             {
                 switch( type )
                 {
-                case OP_TYPE::IN_ADDIMAGE:
-                case OP_TYPE::IN_ADDMESH:
-                case OP_TYPE::IN_ADDVECTOR:
-                case OP_TYPE::IN_ADDSCALAR:
-                case OP_TYPE::IN_ADDSTRING:
+                case EOpType::IN_ADDIMAGE:
+                case EOpType::IN_ADDMESH:
+                case EOpType::IN_ADDVECTOR:
+                case EOpType::IN_ADDSCALAR:
+                case EOpType::IN_ADDSTRING:
                 {
-					const ASTOpInstanceAdd* typedAt = static_cast<const ASTOpInstanceAdd*>(at.get());
+					const ASTOpInstanceAdd* TypedOp = static_cast<const ASTOpInstanceAdd*>(At.get());
 
                     TPair<bool,bool> newState;
                     newState.Key = false; //resource root
                     newState.Value = thisIsRuntime;
-                    RecurseWithState( typedAt->instance.child(), newState );
+                    RecurseWithState( TypedOp->instance.child(), newState );
 
-                    if ( typedAt->value )
+                    if ( TypedOp->value )
                     {
                         newState.Key = true; //resource root
                         newState.Value = thisIsRuntime;
-                        RecurseWithState( typedAt->value.child(), newState );
+                        RecurseWithState( TypedOp->value.child(), newState );
                     }
                     return false;
                 }
@@ -1858,11 +1870,11 @@ namespace mu
 
     protected:
 
-        Ptr<ASTOp> Visit( Ptr<ASTOp> at, bool& processChildren ) override
+        Ptr<ASTOp> Visit( Ptr<ASTOp> At, bool& processChildren ) override
         {
             processChildren = true;
 
-            bool isUpdateCache = m_state.m_updateCache.Contains(at);
+            bool isUpdateCache = m_state.m_updateCache.Contains(At);
 
             if ( isUpdateCache )
             {
@@ -1870,28 +1882,28 @@ namespace mu
                 processChildren = false;
 
                 // See if we can convert it to a more efficient format
-                if ( GetOpDataType( at->GetOpType() )==DT_IMAGE )
+                if ( GetOpDataType( At->GetOpType() )==DT_IMAGE )
                 {
-                    FImageDesc desc = at->GetImageDesc();
+                    FImageDesc desc = At->GetImageDesc();
 
                     if ( desc.m_format!=EImageFormat::L_UByteRLE
                          &&
-                         m_opFormats.IsSupportedFormat(at, EImageFormat::L_UByteRLE) )
+                         m_opFormats.IsSupportedFormat(At, EImageFormat::L_UByteRLE) )
                     {
                         Ptr<ASTOpImagePixelFormat> op = new ASTOpImagePixelFormat;
                         op->Format = EImageFormat::L_UByteRLE;
                         // Note: we have to clone here, to avoid a loop with the visitor system
                         // that updates visited children before processing a node.
 						ASTOp::MapChildFunc Identity = [](const Ptr<ASTOp>& o) {return o;};
-						op->Source = at->Clone(Identity);
+						op->Source = At->Clone(Identity);
 
-                        at = op;
+                        At = op;
                     }
                 }
 
             }
 
-            return at;
+            return At;
         }
 
     private:
@@ -1916,16 +1928,16 @@ namespace mu
     }
 
 
-    Ptr<ASTOp> RuntimeTextureCompressionRemoverAST::Visit( Ptr<ASTOp> at, bool& processChildren )
+    Ptr<ASTOp> RuntimeTextureCompressionRemoverAST::Visit( Ptr<ASTOp> At, bool& processChildren )
     {
-        OP_TYPE type = at->GetOpType();
+        EOpType type = At->GetOpType();
         processChildren = GetOpDataType(type)==DT_INSTANCE;
 
         // TODO: Finer grained: what if the runtime parameter just selects between compressed
         // textures? We don't want them uncompressed.
-        if( type==OP_TYPE::IN_ADDIMAGE )
+        if( type==EOpType::IN_ADDIMAGE )
         {
-			ASTOpInstanceAdd* typedAt = static_cast<ASTOpInstanceAdd*>(at.get());
+			ASTOpInstanceAdd* typedAt = static_cast<ASTOpInstanceAdd*>(At.get());
 			
 			if (typedAt && typedAt->value)
 			{			
@@ -1945,7 +1957,7 @@ namespace mu
 
 					if (isCompressedFormat)
 					{
-						Ptr<ASTOpInstanceAdd> newAt = mu::Clone<ASTOpInstanceAdd>(at);
+						Ptr<ASTOpInstanceAdd> newAt = mu::Clone<ASTOpInstanceAdd>(At);
 
 						// Add a new format operation to uncompress the image
 						Ptr<ASTOpImagePixelFormat> fop = new ASTOpImagePixelFormat;
@@ -1954,13 +1966,13 @@ namespace mu
 						fop->Source = imageAt;
 
 						newAt->value = fop;
-						at = newAt;
+						At = newAt;
 					}
 				}
 			}
         }
 
-        return at;
+        return At;
     }
 
 
@@ -1974,28 +1986,28 @@ namespace mu
     }
 
 
-    Ptr<ASTOp> LODCountReducerAST::Visit( Ptr<ASTOp> at, bool& processChildren )
+    Ptr<ASTOp> LODCountReducerAST::Visit( Ptr<ASTOp> At, bool& processChildren )
     {
         processChildren = true;
 
-        if( at->GetOpType()==OP_TYPE::IN_ADDLOD )
+        if( At->GetOpType()==EOpType::IN_ADDLOD )
         {
-			ASTOpAddLOD* typedAt = static_cast<ASTOpAddLOD*>(at.get());
+			ASTOpAddLOD* TypedOp = static_cast<ASTOpAddLOD*>(At.get());
 
-            if (typedAt->lods.Num()>(size_t)m_lodCount)
+            if (TypedOp->lods.Num()>(size_t)m_lodCount)
             {
-                Ptr<ASTOpAddLOD> newAt = mu::Clone<ASTOpAddLOD>(at);
+                Ptr<ASTOpAddLOD> newAt = mu::Clone<ASTOpAddLOD>(At);
                 while( newAt->lods.Num()>(size_t)m_lodCount )
                 {
                     newAt->lods.Pop();
                 }
-                at = newAt;
+                At = newAt;
             }
 
             processChildren = false;
         }
 
-        return at;
+        return At;
     }
 
 
@@ -2006,13 +2018,13 @@ namespace mu
     {
 		MUTABLE_CPUPROFILER_SCOPE(OptimiseStatesAST);
 
-		for ( int32 s=0; s<m_states.Num(); ++s )
+		for ( int32 s=0; s<States.Num(); ++s )
 		{
 			// Remove the unnecessary lods
-			FStateOptimizationOptions StateOptimization = m_states[s].nodeState.Optimisation;
+			FStateOptimizationOptions StateOptimization = States[s].nodeState.Optimisation;
 			if (StateOptimization.bOnlyFirstLOD)
 			{
-				LODCountReducerAST(m_states[s].root, StateOptimization.FirstLOD + 1 + StateOptimization.NumExtraLODsToBuildAfterFirstLOD);
+				LODCountReducerAST(States[s].root, StateOptimization.FirstLOD + 1 + StateOptimization.NumExtraLODsToBuildAfterFirstLOD);
 			}
 
 			// Apply texture compression strategy
@@ -2022,7 +2034,7 @@ namespace mu
 			case ETextureCompressionStrategy::DontCompressRuntime:
 			{
 				MUTABLE_CPUPROFILER_SCOPE(RuntimeTextureCompressionRemover);
-				RuntimeTextureCompressionRemoverAST r(&m_states[s], false);
+				RuntimeTextureCompressionRemoverAST r(&States[s], false);
 				bModified = true;
 				break;
 			}
@@ -2030,7 +2042,7 @@ namespace mu
 			case ETextureCompressionStrategy::NeverCompress:
 			{
 				MUTABLE_CPUPROFILER_SCOPE(RuntimeTextureCompressionRemover);
-				RuntimeTextureCompressionRemoverAST r(&m_states[s], true);
+				RuntimeTextureCompressionRemoverAST r(&States[s], true);
 				bModified = true;
 				break;
 			}
@@ -2040,7 +2052,7 @@ namespace mu
 			}
 
             // If a state has no runtime parameters, skip its optimisation alltogether
-            if (bModified || m_states[s].nodeState.RuntimeParams.Num())
+            if (bModified || States[s].nodeState.RuntimeParams.Num())
             {
                 // Promote the intructions that depend on runtime parameters, and sink new
                 // format instructions.
@@ -2057,23 +2069,23 @@ namespace mu
 
                     UE_LOG(LogMutableCore, Verbose, TEXT(" - before parameter optimiser"));
 
-                    ParameterOptimiserAST param( m_states[s], m_options->GetPrivate()->OptimisationOptions );
+                    ParameterOptimiserAST param( States[s], Options->GetPrivate()->OptimisationOptions );
                     modified = param.Apply();
 
                     TArray<Ptr<ASTOp>> roots;
-                    roots.Add(m_states[s].root);
+                    roots.Add(States[s].root);
 
                     UE_LOG(LogMutableCore, Verbose, TEXT(" - after parameter optimiser"));
 					//UE_LOG(LogMutableCore, Verbose, TEXT("(int) %s : %ld"), TEXT("ast size"), int64(ASTOp::CountNodes(roots)));
 
                     // All kind of optimisations that depend on the meaning of each operation
                     UE_LOG(LogMutableCore, Verbose, TEXT(" - semantic optimiser"));
-                    modified |= SemanticOptimiserAST( roots, m_options->GetPrivate()->OptimisationOptions, 1 );
+                    modified |= SemanticOptimiserAST( roots, Options->GetPrivate()->OptimisationOptions, 1 );
 					//UE_LOG(LogMutableCore, Verbose, TEXT("(int) %s : %ld"), TEXT("ast size"), int64(ASTOp::CountNodes(roots)));
 					//ASTOp::LogHistogram(roots);
 
                     UE_LOG(LogMutableCore, Verbose, TEXT(" - sink optimiser"));
-                    modified |= SinkOptimiserAST( roots, m_options->GetPrivate()->OptimisationOptions );
+                    modified |= SinkOptimiserAST( roots, Options->GetPrivate()->OptimisationOptions );
 					//UE_LOG(LogMutableCore, Verbose, TEXT("(int) %s : %ld"), TEXT("ast size"), int64(ASTOp::CountNodes(roots)));
 					//ASTOp::LogHistogram(roots);
 
@@ -2087,7 +2099,7 @@ namespace mu
                 }
 
                 TArray<Ptr<ASTOp>> roots;
-                roots.Add(m_states[s].root);
+                roots.Add(States[s].root);
 
                 UE_LOG(LogMutableCore, Verbose, TEXT(" - duplicated data remover"));
                 DuplicatedDataRemoverAST( roots );
@@ -2097,30 +2109,30 @@ namespace mu
                 DuplicatedCodeRemoverAST( roots );
 				//UE_LOG(LogMutableCore, Verbose, TEXT("(int) %s : %ld"), TEXT("ast size"), int64(ASTOp::CountNodes(roots)));
 
-                m_states[s].root = roots[0];
+                States[s].root = roots[0];
             }
         }
 
 
 		TArray<Ptr<ASTOp>> roots;
-		for (const FStateCompilationData& s : m_states)
+		for (const FStateCompilationData& s : States)
 		{
 			roots.Add(s.root);
 		}
 
         // Mark the instructions that don't depend on runtime parameters to be cached. This is
-        // necessary at this stage before GPU optimisation.
+        // necessary At this stage before GPU optimisation.
         {
             AccumulateAllImageFormatsOpAST opFormats;
             opFormats.Run(roots);
 
 			// Reset the state root operations in case they have changed due to optimization
-			for (int32 RootIndex = 0; RootIndex < m_states.Num(); ++RootIndex)
+			for (int32 RootIndex = 0; RootIndex < States.Num(); ++RootIndex)
 			{
-				m_states[RootIndex].root = roots[RootIndex];
+				States[RootIndex].root = roots[RootIndex];
 			}
 
-            for (FStateCompilationData& s: m_states )
+            for (FStateCompilationData& s: States )
             {
                 {
                     UE_LOG(LogMutableCore, Verbose, TEXT(" - state cache"));
@@ -2152,7 +2164,7 @@ namespace mu
                 modified = false;
 
                 UE_LOG(LogMutableCore, Verbose, TEXT(" - semantic optimiser"));
-                modified |= SemanticOptimiserAST( roots, m_options->GetPrivate()->OptimisationOptions, Pass );
+                modified |= SemanticOptimiserAST( roots, Options->GetPrivate()->OptimisationOptions, Pass );
 				//UE_LOG(LogMutableCore, Verbose, TEXT("(int) %s : %ld"), TEXT("ast size"), int64(ASTOp::CountNodes(roots)));
 
                 // Image size operations are treated separately
@@ -2164,7 +2176,7 @@ namespace mu
             for(Ptr<ASTOp>& Root : roots)
             {
                 UE_LOG(LogMutableCore, Verbose, TEXT(" - constant optimiser"));
-				modified = ConstantGeneratorAST( m_options->GetPrivate(), Root, Pass );
+				modified = ConstantGeneratorAST( Options->GetPrivate(), Root, Pass );
 			}
 
 			//UE_LOG(LogMutableCore, Verbose, TEXT("(int) %s : %ld"), TEXT("ast size"), int64(ASTOp::CountNodes(roots)));
@@ -2179,16 +2191,16 @@ namespace mu
 		}
 
 		// Reset the state root operations in case they have changed due to optimization
-		for (int32 RootIndex = 0; RootIndex < m_states.Num(); ++RootIndex)
+		for (int32 RootIndex = 0; RootIndex < States.Num(); ++RootIndex)
 		{
-			m_states[RootIndex].root = roots[RootIndex];
+			States[RootIndex].root = roots[RootIndex];
 		}
 
         // Optimise the data formats
         {
             MUTABLE_CPUPROFILER_SCOPE(DataFormats);
 
-            DataOptimise( m_options.get(), roots);
+            DataOptimise( Options.get(), roots);
 
             // After optimising the data formats, we may remove more constants
             DuplicatedDataRemoverAST( roots );
@@ -2196,7 +2208,7 @@ namespace mu
 
             // Update the marks for the instructions that don't depend on runtime parameters to be
             // cached.
-            for (FStateCompilationData& s:m_states)
+            for (FStateCompilationData& s:States)
             {
                 StateCacheDetectorAST c( &s );
             }
