[mutable] More enforcement of coding standards.

diff --git a/Engine/Plugins/Mutable/Source/MutableTools/Private/MuT/CodeOptimiser_Semantic.cpp b/Engine/Plugins/Mutable/Source/MutableTools/Private/MuT/CodeOptimiser_Semantic.cpp
--- a/Engine/Plugins/Mutable/Source/MutableTools/Private/MuT/CodeOptimiser_Semantic.cpp
+++ b/Engine/Plugins/Mutable/Source/MutableTools/Private/MuT/CodeOptimiser_Semantic.cpp
@@ -77,14 +77,14 @@ namespace mu
     {
         Ptr<ASTOp> at;
 
-        OP_TYPE type = GetOpType();
+        EOpType type = GetOpType();
         switch ( type )
         {
 
-        case OP_TYPE::BO_NOT:
+        case EOpType::BO_NOT:
         {
 			Ptr<ASTOp> sourceAt = children[op.args.BoolNot.source].child();
-            if  ( sourceAt->GetOpType() == OP_TYPE::BO_CONSTANT )
+            if  ( sourceAt->GetOpType() == EOpType::BO_CONSTANT )
             {
                 Ptr<ASTOpConstantBool> newAt = new ASTOpConstantBool();
                 newAt->value = !newAt->value;
@@ -93,7 +93,7 @@ namespace mu
             break;
         }
 
-        case OP_TYPE::BO_AND:
+        case EOpType::BO_AND:
         {
             bool changed = false;
 			Ptr<ASTOp> aAt = children[op.args.BoolBinary.a].child();
@@ -108,7 +108,7 @@ namespace mu
                 at = aAt;
                 changed = true;
             }
-            else if ( aAt->GetOpType() == OP_TYPE::BO_CONSTANT )
+            else if ( aAt->GetOpType() == EOpType::BO_CONSTANT )
             {
                 if (static_cast<const ASTOpConstantBool*>(aAt.get())->value)
                 {
@@ -121,7 +121,7 @@ namespace mu
                     changed = true;
                 }
             }
-            else if ( bAt->GetOpType() == OP_TYPE::BO_CONSTANT )
+            else if ( bAt->GetOpType() == EOpType::BO_CONSTANT )
             {
                 if (static_cast<const ASTOpConstantBool*>(bAt.get())->value)
                 {
@@ -175,7 +175,7 @@ namespace mu
             break;
         }
 
-        case OP_TYPE::BO_OR:
+        case EOpType::BO_OR:
         {
             bool changed = false;
 			Ptr<ASTOp> aAt = children[op.args.BoolBinary.a].child();
@@ -190,7 +190,7 @@ namespace mu
                 at = aAt;
                 changed = true;
             }
-            else if ( aAt->GetOpType() == OP_TYPE::BO_CONSTANT )
+            else if ( aAt->GetOpType() == EOpType::BO_CONSTANT )
             {
                 if (static_cast<const ASTOpConstantBool*>(aAt.get())->value)
                 {
@@ -203,7 +203,7 @@ namespace mu
                     changed = true;
                 }
             }
-            else if ( bAt->GetOpType() == OP_TYPE::BO_CONSTANT )
+            else if ( bAt->GetOpType() == EOpType::BO_CONSTANT )
             {
                 if (static_cast<const ASTOpConstantBool*>(bAt.get())->value)
                 {
@@ -258,7 +258,7 @@ namespace mu
             break;
         }
 
-		case OP_TYPE::CO_SWIZZLE:
+		case EOpType::CO_SWIZZLE:
 		{
 			// Optimizations that can be applied per-channel
 			{
@@ -276,7 +276,7 @@ namespace mu
 					{
 
 						// Swizzle + swizzle = swizzle
-					case OP_TYPE::CO_SWIZZLE:
+					case EOpType::CO_SWIZZLE:
 					{
 						if (!NewAt)
 						{
@@ -304,7 +304,7 @@ namespace mu
 			{
 				// Optimizations that depend on all channels.
 				bool bAllChannelsSameType = true;
-				OP_TYPE AllChannelsType = OP_TYPE::NONE;
+				EOpType AllChannelsType = EOpType::NONE;
 				for (int32 ChannelIndex = 0; ChannelIndex < MUTABLE_OP_MAX_SWIZZLE_CHANNELS; ++ChannelIndex)
 				{
 					Ptr<ASTOp> Candidate = children[op.args.ColourSwizzle.sources[ChannelIndex]].child();
@@ -328,12 +328,12 @@ namespace mu
 				{
 					switch (AllChannelsType)
 					{
-					case OP_TYPE::CO_FROMSCALARS:
+					case EOpType::CO_FROMSCALARS:
 					{
 						// We can remove the swizzle and replace it with a new FromScalars actually swizzling the inputs.
 
 						Ptr<ASTOpFixed> NewAt = new ASTOpFixed();
-						NewAt->op.type = OP_TYPE::CO_FROMSCALARS;
+						NewAt->op.type = EOpType::CO_FROMSCALARS;
 
 						for (int32 ChannelIndex = 0; ChannelIndex < MUTABLE_OP_MAX_SWIZZLE_CHANNELS; ++ChannelIndex)
 						{
@@ -360,17 +360,17 @@ namespace mu
 		}
 
 		//-------------------------------------------------------------------------------------
-		case OP_TYPE::IM_RESIZEREL:
+		case EOpType::IM_RESIZEREL:
 		{
 			Ptr<ASTOp> sourceAt = children[op.args.ImageResizeRel.source].child();
 
 			// The instruction can be sunk
-			OP_TYPE sourceType = sourceAt->GetOpType();
+			EOpType sourceType = sourceAt->GetOpType();
 			switch (sourceType)
 			{
 			// This is done here instead of in the OptimizeSize step to prevent removing 
 			// resize_rel too early.
-			case OP_TYPE::IM_RESIZE:
+			case EOpType::IM_RESIZE:
 			{
 				Ptr<ASTOpFixed> NewOp = mu::Clone<ASTOpFixed>(sourceAt.get());
 				NewOp->op.args.ImageResize.size[0] = int16(NewOp->op.args.ImageResize.size[0] * op.args.ImageResizeRel.factor[0]);
@@ -390,22 +390,22 @@ namespace mu
 		}
 
         //-------------------------------------------------------------------------------------
-        case OP_TYPE::IM_PIXELFORMAT:
-		case OP_TYPE::IM_SWIZZLE:
-		case OP_TYPE::IM_LAYER:
+        case EOpType::IM_PIXELFORMAT:
+		case EOpType::IM_SWIZZLE:
+		case EOpType::IM_LAYER:
 		{
 			check(false); // moved to its own operation class
 			break;
 		}
 
 		//-------------------------------------------------------------------------------------
-		case OP_TYPE::ME_MASKDIFF:
+		case EOpType::ME_MASKDIFF:
 		{
 			Ptr<ASTOp> Fragment = children[op.args.MeshMaskDiff.fragment].child();
-			OP_TYPE FragmentType = Fragment->GetOpType();
+			EOpType FragmentType = Fragment->GetOpType();
 			switch (FragmentType)
 			{
-			case OP_TYPE::ME_ADDTAGS:
+			case EOpType::ME_ADDTAGS:
 			{
 				// Tags in the fragment can be ignored.
 				const ASTOpMeshAddTags* Add = static_cast<const ASTOpMeshAddTags*>(Fragment.get());
@@ -415,7 +415,7 @@ namespace mu
 				at = NewAt;
 				break;
 			}
-			case OP_TYPE::ME_SWITCH:
+			case EOpType::ME_SWITCH:
 			{
 				// Move the mask diff down all the paths
 				Ptr<ASTOpSwitch> nop = mu::Clone<ASTOpSwitch>(Fragment.get());
@@ -449,13 +449,13 @@ namespace mu
 		}
 
 		//-------------------------------------------------------------------------------------
-		case OP_TYPE::ME_APPLYLAYOUT:
+		case EOpType::ME_APPLYLAYOUT:
 		{
 			Ptr<ASTOp> Base = children[op.args.MeshApplyLayout.mesh].child();
-			OP_TYPE BaseType = Base->GetOpType();
+			EOpType BaseType = Base->GetOpType();
 			switch (BaseType)
 			{
-			case OP_TYPE::ME_ADDTAGS:
+			case EOpType::ME_ADDTAGS:
 			{
 				// Add the tags after layout
 				Ptr<ASTOpMeshAddTags> NewAddTags = mu::Clone<ASTOpMeshAddTags>(Base);
@@ -524,11 +524,11 @@ namespace mu
     {
         Ptr<ASTOp> at;
 
-        OP_TYPE type = GetOpType();
+        EOpType type = GetOpType();
         switch ( type )
         {
 
-        case OP_TYPE::ME_PROJECT:
+        case EOpType::ME_PROJECT:
         {
             auto sourceAt = children[op.args.MeshProject.mesh].child();
             if (!sourceAt) { break; }
@@ -536,13 +536,13 @@ namespace mu
             auto projectorAt = children[op.args.MeshProject.projector].child();
             if (!projectorAt) { break; }
 
-            OP_TYPE sourceType = sourceAt->GetOpType();
+            EOpType sourceType = sourceAt->GetOpType();
             switch ( sourceType )
             {
 
-            case OP_TYPE::ME_CONDITIONAL:
+            case EOpType::ME_CONDITIONAL:
             {
-                if (projectorAt->GetOpType()==OP_TYPE::PR_CONSTANT)
+                if (projectorAt->GetOpType()==EOpType::PR_CONSTANT)
                 {
                     // We move the project down the two paths
                     auto nop = mu::Clone<ASTOpConditional>(sourceAt);
@@ -560,9 +560,9 @@ namespace mu
                 break;
             }
 
-            case OP_TYPE::ME_SWITCH:
+            case EOpType::ME_SWITCH:
             {
-                if (projectorAt->GetOpType()==OP_TYPE::PR_CONSTANT)
+                if (projectorAt->GetOpType()==EOpType::PR_CONSTANT)
                 {
                     // Move the format down all the paths
                     auto nop = mu::Clone<ASTOpSwitch>(sourceAt);
@@ -590,7 +590,7 @@ namespace mu
                 break;
             }
 
-			case OP_TYPE::ME_ADDTAGS:
+			case EOpType::ME_ADDTAGS:
 			{
 				// Add the tags after layout
 				Ptr<ASTOpMeshAddTags> NewAddTags = mu::Clone<ASTOpMeshAddTags>(sourceAt);
@@ -617,9 +617,9 @@ namespace mu
         //-----------------------------------------------------------------------------------------
         //-----------------------------------------------------------------------------------------
         //-----------------------------------------------------------------------------------------
-		case OP_TYPE::IM_SWIZZLE:
-		case OP_TYPE::LA_REMOVEBLOCKS:
-		case OP_TYPE::IM_RASTERMESH:
+		case EOpType::IM_SWIZZLE:
+		case EOpType::LA_REMOVEBLOCKS:
+		case EOpType::IM_RASTERMESH:
 		{
 			// Moved to their own operation class
 			check(false);
@@ -630,7 +630,7 @@ namespace mu
         //-----------------------------------------------------------------------------------------
         //-----------------------------------------------------------------------------------------
         //-----------------------------------------------------------------------------------------
-        case OP_TYPE::IM_DISPLACE:
+        case EOpType::IM_DISPLACE:
         {
 			Ptr<ASTOp> OriginalAt = at;
 			Ptr<ASTOp> sourceAt = children[op.args.ImageDisplace.source].child();
@@ -639,9 +639,9 @@ namespace mu
             switch ( sourceAt->GetOpType() )
             {
 
-            case OP_TYPE::IM_CONDITIONAL:
+            case EOpType::IM_CONDITIONAL:
             {
-                if (displaceMapAt->GetOpType()==OP_TYPE::IM_CONDITIONAL)
+                if (displaceMapAt->GetOpType()==EOpType::IM_CONDITIONAL)
                 {
 					const ASTOpConditional* typedSource = static_cast<const ASTOpConditional*>(sourceAt.get());
 					const ASTOpConditional* typedDisplacementMap = static_cast<const ASTOpConditional*>(displaceMapAt.get());
@@ -666,9 +666,9 @@ namespace mu
                 break;
             }
 
-            case OP_TYPE::IM_SWITCH:
+            case EOpType::IM_SWITCH:
             {
-                if (displaceMapAt->GetOpType()==OP_TYPE::IM_SWITCH)
+                if (displaceMapAt->GetOpType()==EOpType::IM_SWITCH)
                 {
 					const ASTOpSwitch* typedSource = static_cast<const ASTOpSwitch*>(sourceAt.get());
 					const ASTOpSwitch* typedDisplacementMap = static_cast<const ASTOpSwitch*>(displaceMapAt.get());
@@ -780,17 +780,17 @@ namespace mu
 
         Ptr<ASTOp> Apply( const ASTOp* root )
         {
-            m_root = root;
+            Root = root;
             OldToNew.Empty();
 
-            check(root->GetOpType()==OP_TYPE::IM_RESIZEREL);
+            check(root->GetOpType()==EOpType::IM_RESIZEREL);
 
 			const ASTOpFixed* typedRoot = static_cast<const ASTOpFixed*>(root);
-            m_initialSource = typedRoot->children[typedRoot->op.args.ImageResizeRel.source].child();
-            Ptr<ASTOp> newSource = Visit( m_initialSource, typedRoot );
+            InitialSource = typedRoot->children[typedRoot->op.args.ImageResizeRel.source].child();
+            Ptr<ASTOp> newSource = Visit( InitialSource, typedRoot );
 
             // If there is any change, it is the new root.
-            if (newSource!=m_initialSource)
+            if (newSource!=InitialSource)
             {
                 return newSource;
             }
@@ -800,8 +800,8 @@ namespace mu
 
     protected:
 
-        const ASTOp* m_root;
-        Ptr<ASTOp> m_initialSource;
+        const ASTOp* Root;
+        Ptr<ASTOp> InitialSource;
         //! For each operation we sink, the map from old instructions to new instructions.
 		TMap<FSinkerOldToNewKey, Ptr<ASTOp>> OldToNew;
 
@@ -823,7 +823,7 @@ namespace mu
 			switch (at->GetOpType())
 			{
 
-			case OP_TYPE::IM_CONDITIONAL:
+			case EOpType::IM_CONDITIONAL:
 			{
 				// We move the mask creation down the two paths
 				auto newOp = mu::Clone<ASTOpConditional>(at);
@@ -833,7 +833,7 @@ namespace mu
 				break;
 			}
 
-			case OP_TYPE::IM_PIXELFORMAT:
+			case EOpType::IM_PIXELFORMAT:
 			{
 				// \todo: only if shrinking?
 				Ptr<ASTOpImagePixelFormat> newOp = mu::Clone<ASTOpImagePixelFormat>(at);
@@ -842,7 +842,7 @@ namespace mu
 				break;
 			}
 
-			case OP_TYPE::IM_SWITCH:
+			case EOpType::IM_SWITCH:
 			{
 				// We move the mask creation down all the paths
 				Ptr<ASTOpSwitch> newOp = mu::Clone<ASTOpSwitch>(at);
@@ -855,7 +855,7 @@ namespace mu
 				break;
 			}
 
-			case OP_TYPE::IM_SWIZZLE:
+			case EOpType::IM_SWIZZLE:
 			{
 				Ptr<ASTOpImageSwizzle> newOp = mu::Clone<ASTOpImageSwizzle>(at);
 				for (int s = 0; s < MUTABLE_OP_MAX_SWIZZLE_CHANNELS; ++s)
@@ -870,7 +870,7 @@ namespace mu
 				break;
 			}
 
-			case OP_TYPE::IM_COMPOSE:
+			case EOpType::IM_COMPOSE:
 			{
 				// We can only optimise if the layout grid blocks size in pixels is
 				// still an integer after relative scale
@@ -912,7 +912,7 @@ namespace mu
 				break;
 			}
 
-			case OP_TYPE::IM_PATCH:
+			case EOpType::IM_PATCH:
 			{
 				Ptr<ASTOpImagePatch> newOp = mu::Clone<ASTOpImagePatch>(at);
 
@@ -928,7 +928,7 @@ namespace mu
 				break;
 			}
 
-			case OP_TYPE::IM_MIPMAP:
+			case EOpType::IM_MIPMAP:
 			{
 				Ptr<ASTOpImageMipmap> newOp = mu::Clone<ASTOpImageMipmap>(at.get());
 				Ptr<ASTOp> baseOp = newOp->Source.child();
@@ -937,7 +937,7 @@ namespace mu
 				break;
 			}
 
-			case OP_TYPE::IM_INTERPOLATE:
+			case EOpType::IM_INTERPOLATE:
 			{
 				Ptr<ASTOpFixed> newOp = mu::Clone<ASTOpFixed>(at);
 
@@ -952,7 +952,7 @@ namespace mu
 				break;
 			}
 
-			case OP_TYPE::IM_MULTILAYER:
+			case EOpType::IM_MULTILAYER:
 			{
 				Ptr<ASTOpImageMultiLayer> nop = mu::Clone<ASTOpImageMultiLayer>(at);
 				nop->base = Visit(nop->base.child(), currentSinkingOp);
@@ -962,7 +962,7 @@ namespace mu
 				break;
 			}
 
-			case OP_TYPE::IM_LAYER:
+			case EOpType::IM_LAYER:
 			{
 				Ptr<ASTOpImageLayer> nop = mu::Clone<ASTOpImageLayer>(at);
 
@@ -979,7 +979,7 @@ namespace mu
 				break;
 			}
 
-			case OP_TYPE::IM_LAYERCOLOUR:
+			case EOpType::IM_LAYERCOLOUR:
 			{
 				Ptr<ASTOpImageLayerColor> nop = mu::Clone<ASTOpImageLayerColor>(at);
 
@@ -993,7 +993,7 @@ namespace mu
 				break;
 			}
 
-			case OP_TYPE::IM_DISPLACE:
+			case EOpType::IM_DISPLACE:
 			{
 				Ptr<ASTOpFixed> nop = mu::Clone<ASTOpFixed>(at);
 
@@ -1004,7 +1004,7 @@ namespace mu
 				nop->SetChild(nop->op.args.ImageDisplace.displacementMap, Visit(mapOp, currentSinkingOp));
 
 				// Make sure we don't scale a constant displacement map, which is very wrong.
-				if (mapOp->GetOpType() == OP_TYPE::IM_MAKEGROWMAP)
+				if (mapOp->GetOpType() == EOpType::IM_MAKEGROWMAP)
 				{
 					newAt = nop;
 				}
@@ -1016,7 +1016,7 @@ namespace mu
 				break;
 			}
 
-			case OP_TYPE::IM_MAKEGROWMAP:
+			case EOpType::IM_MAKEGROWMAP:
 			{
 				Ptr<ASTOpImageMakeGrowMap> nop = mu::Clone<ASTOpImageMakeGrowMap>(at);
 				Ptr<ASTOp> maskOp = nop->Mask.child();
@@ -1025,7 +1025,7 @@ namespace mu
 				break;
 			}
 
-			case OP_TYPE::IM_INVERT:
+			case EOpType::IM_INVERT:
 			{
 				Ptr<ASTOpFixed> nop = mu::Clone<ASTOpFixed>(at);
 				Ptr<ASTOp> maskOp = nop->children[nop->op.args.ImageInvert.base].child();
@@ -1034,7 +1034,7 @@ namespace mu
 				break;
 			}
 
-			case OP_TYPE::IM_SATURATE:
+			case EOpType::IM_SATURATE:
 			{
 				Ptr<ASTOpFixed> nop = mu::Clone<ASTOpFixed>(at);
 				Ptr<ASTOp> maskOp = nop->children[nop->op.args.ImageSaturate.base].child();
@@ -1043,7 +1043,7 @@ namespace mu
 				break;
 			}
 
-			case OP_TYPE::IM_TRANSFORM:
+			case EOpType::IM_TRANSFORM:
 			{
 				// It can only sink in the transform if it doesn't have it's own size.
 				const ASTOpImageTransform* TypedAt = static_cast<const ASTOpImageTransform*>(at.get());
@@ -1058,7 +1058,7 @@ namespace mu
 				break;
 			}
 
-			case OP_TYPE::IM_RASTERMESH:
+			case EOpType::IM_RASTERMESH:
 			{
 				Ptr<ASTOpImageRasterMesh> nop = mu::Clone<ASTOpImageRasterMesh>(at);
 				Ptr<ASTOp> maskOp = nop->mask.child();
@@ -1081,10 +1081,10 @@ namespace mu
 			}
 
 			// end on line, replace with sinking op
-			if (at == newAt && at != m_initialSource)
+			if (at == newAt && at != InitialSource)
 			{
 				Ptr<ASTOpFixed> newOp = mu::Clone<ASTOpFixed>(currentSinkingOp);
-				check(newOp->GetOpType() == OP_TYPE::IM_RESIZEREL);
+				check(newOp->GetOpType() == EOpType::IM_RESIZEREL);
 
 				newOp->SetChild(newOp->op.args.ImageResizeRel.source, at);
 
@@ -1103,20 +1103,20 @@ namespace mu
     {
         Ptr<ASTOp> at;
 
-        OP_TYPE type = GetOpType();
+        EOpType type = GetOpType();
         switch ( type )
         {
 
         //-------------------------------------------------------------------------------------
-        case OP_TYPE::IM_RESIZE:
+        case EOpType::IM_RESIZE:
         {
             auto sourceAt = children[op.args.ImageResize.source].child();
 
             // The instruction can be sunk
-            OP_TYPE sourceType = sourceAt->GetOpType();
+            EOpType sourceType = sourceAt->GetOpType();
             switch ( sourceType )
             {
-            case OP_TYPE::IM_RESIZE:
+            case EOpType::IM_RESIZE:
             {
                 // Keep top resize
                 Ptr<const ASTOpFixed> sourceOp = static_cast<const ASTOpFixed*>(sourceAt.get());
@@ -1128,7 +1128,7 @@ namespace mu
                 break;
             }
 
-            case OP_TYPE::IM_PLAINCOLOUR:
+            case EOpType::IM_PLAINCOLOUR:
             {
                 // Set the size in the children and remove resize
 				Ptr<ASTOpFixed> sourceOp = mu::Clone<ASTOpFixed>(sourceAt.get());
@@ -1139,7 +1139,7 @@ namespace mu
                 break;
             }
 
-			case OP_TYPE::IM_TRANSFORM:
+			case EOpType::IM_TRANSFORM:
 			{
 				// Set the size in the children and remove resize
 				Ptr<ASTOpImageTransform> sourceOp = mu::Clone<ASTOpImageTransform>(sourceAt.get());
@@ -1149,7 +1149,7 @@ namespace mu
 				break;
 			}
 
-            case OP_TYPE::IM_CONDITIONAL:
+            case EOpType::IM_CONDITIONAL:
             {
                 // We move the resize down the two paths
                 auto newOp = mu::Clone<ASTOpConditional>(sourceAt);
@@ -1166,7 +1166,7 @@ namespace mu
                 break;
             }
 
-            case OP_TYPE::IM_SWITCH:
+            case EOpType::IM_SWITCH:
             {
                 // Move the resize down all the paths
                 Ptr<ASTOpSwitch> newOp = mu::Clone<ASTOpSwitch>(sourceAt);
@@ -1192,7 +1192,7 @@ namespace mu
                 break;
             }
 
-			case OP_TYPE::IM_SWIZZLE:
+			case EOpType::IM_SWIZZLE:
 			{
 				Ptr<ASTOpImageSwizzle> newOp = mu::Clone<ASTOpImageSwizzle>(sourceAt);
 				for (int s = 0; s < MUTABLE_OP_MAX_SWIZZLE_CHANNELS; ++s)
@@ -1209,7 +1209,7 @@ namespace mu
 				break;
 			}
 
-            case OP_TYPE::IM_COMPOSE:
+            case EOpType::IM_COMPOSE:
             {
                 Ptr<ASTOpImageCompose> newOp = mu::Clone<ASTOpImageCompose>(sourceAt);
 
@@ -1232,7 +1232,7 @@ namespace mu
                 break;
             }
 
-            case OP_TYPE::IM_RASTERMESH:
+            case EOpType::IM_RASTERMESH:
             {
                 Ptr<ASTOpImageRasterMesh> newOp = mu::Clone<ASTOpImageRasterMesh>(sourceAt);
 
@@ -1264,7 +1264,7 @@ namespace mu
                 break;
             }
 
-            case OP_TYPE::IM_INTERPOLATE:
+            case EOpType::IM_INTERPOLATE:
             {
                 Ptr<ASTOpFixed> newOp = mu::Clone<ASTOpFixed>(sourceAt);
 
@@ -1283,7 +1283,7 @@ namespace mu
                 break;
             }
 
-			case OP_TYPE::IM_INVERT:
+			case EOpType::IM_INVERT:
 			{
 				Ptr<ASTOpFixed> NewOp = mu::Clone<ASTOpFixed>(sourceAt);
 				Ptr<ASTOp> BaseAt = NewOp->children[NewOp->op.args.ImageInvert.base].child();
@@ -1297,7 +1297,7 @@ namespace mu
 				break;
 			}
 
-			case OP_TYPE::IM_PIXELFORMAT:
+			case EOpType::IM_PIXELFORMAT:
 			{				
 				// \todo: only if shrinking?
 				
@@ -1321,7 +1321,7 @@ namespace mu
 				break;
 			}
 
-            case OP_TYPE::IM_LAYER:
+            case EOpType::IM_LAYER:
             {
                 Ptr<ASTOpImageLayer> newOp = mu::Clone<ASTOpImageLayer>(sourceAt);
 
@@ -1345,7 +1345,7 @@ namespace mu
                 break;
             }
 
-            case OP_TYPE::IM_LAYERCOLOUR:
+            case EOpType::IM_LAYERCOLOUR:
             {
                 Ptr<ASTOpImageLayerColor> newOp = mu::Clone<ASTOpImageLayerColor>(sourceAt);
 
@@ -1365,14 +1365,14 @@ namespace mu
                 break;
             }
 
-			case OP_TYPE::IM_DISPLACE:
+			case EOpType::IM_DISPLACE:
 			{
 				// In the size optimization phase we can optimize the resize with the displace
 				// because the constants have not been collapsed yet.
 				// We will still check it and sink the size directly below the IM_MAKEGROWMAP op
 				Ptr<ASTOpFixed> SourceTyped = static_cast<ASTOpFixed*>(sourceAt.get());
-				Ptr<ASTOp> OriginalDisplacementMapOp = SourceTyped->children[SourceTyped->op.args.ImageDisplace.displacementMap].m_child;
-				if (OriginalDisplacementMapOp->GetOpType() == OP_TYPE::IM_MAKEGROWMAP)
+				Ptr<ASTOp> OriginalDisplacementMapOp = SourceTyped->children[SourceTyped->op.args.ImageDisplace.displacementMap].Child;
+				if (OriginalDisplacementMapOp->GetOpType() == EOpType::IM_MAKEGROWMAP)
 				{
 					Ptr<ASTOpFixed> newOp = mu::Clone<ASTOpFixed>(sourceAt);
 
@@ -1404,16 +1404,16 @@ namespace mu
 
 
         //-------------------------------------------------------------------------------------
-        case OP_TYPE::IM_RESIZEREL:
+        case EOpType::IM_RESIZEREL:
         {
             auto sourceAt = children[op.args.ImageResizeRel.source].child();
 
             // The instruction can be sunk
-            OP_TYPE sourceType = sourceAt->GetOpType();
+            EOpType sourceType = sourceAt->GetOpType();
             switch ( sourceType )
             {
 
-			case OP_TYPE::IM_BLANKLAYOUT:
+			case EOpType::IM_BLANKLAYOUT:
 			{
 				Ptr<ASTOpFixed> newOp = mu::Clone<ASTOpFixed>(sourceAt);
 
@@ -1429,7 +1429,7 @@ namespace mu
 				break;
 			}
 
-			case OP_TYPE::IM_PLAINCOLOUR:
+			case EOpType::IM_PLAINCOLOUR:
 			{
 				Ptr<ASTOpFixed> newOp = mu::Clone<ASTOpFixed>(sourceAt);
 
@@ -1440,7 +1440,7 @@ namespace mu
 				break;
 			}
 
-			case OP_TYPE::IM_TRANSFORM:
+			case EOpType::IM_TRANSFORM:
 			{
 				// We can only optimize here if we know the transform result size, otherwise, we will sink the op in the sinker.
 				const ASTOpImageTransform* typedAt = static_cast<const ASTOpImageTransform*>(sourceAt.get());
@@ -1458,7 +1458,7 @@ namespace mu
 				// Don't combine. ResizeRel sometimes can resize more children than Resize can do. (see RasterMesh)
 				// It can be combined in an optimization step further in the process, when normal sizes may have been 
 				// optimized already (see OptimizeSemantic)
-//                case OP_TYPE::IM_RESIZE:
+//                case EOpType::IM_RESIZE:
 //                {
 //                    break;
 //                }
@@ -1479,7 +1479,7 @@ namespace mu
 
 
         //-------------------------------------------------------------------------------------
-        case OP_TYPE::IM_RESIZELIKE:
+        case EOpType::IM_RESIZELIKE:
         {
             Ptr<ImageSizeExpression> sourceSize =
                     children[op.args.ImageResizeLike.source]->GetImageSizeExpression();
@@ -1494,7 +1494,7 @@ namespace mu
             else if ( sizeSourceSize->type == ImageSizeExpression::ISET_CONSTANT )
             {
                 Ptr<ASTOpFixed> newAt = new ASTOpFixed;
-                newAt->op.type = OP_TYPE::IM_RESIZE;
+                newAt->op.type = EOpType::IM_RESIZE;
                 newAt->SetChild( newAt->op.args.ImageResize.source, children[op.args.ImageResizeLike.source] );
                 newAt->op.args.ImageResize.size[0] = sizeSourceSize->size[0];
                 newAt->op.args.ImageResize.size[1] = sizeSourceSize->size[1];
@@ -1504,12 +1504,12 @@ namespace mu
             {
                 // TODO
                 // Skip intermediate ops until the layout
-//                    if ( program.m_code[sizeSourceAt].type != OP_TYPE::IM_BLANKLAYOUT )
+//                    if ( program.m_code[sizeSourceAt].type != EOpType::IM_BLANKLAYOUT )
 //                    {
 //                        m_modified = true;
 
 //                        OP blackLayoutOp;
-//                        blackLayoutOp.type = OP_TYPE::IM_BLANKLAYOUT;
+//                        blackLayoutOp.type = EOpType::IM_BLANKLAYOUT;
 //                        blackLayoutOp.args.ImageBlankLayout.layout = sizeSourceSize->layout;
 //                        blackLayoutOp.args.ImageBlankLayout.blockSize[0] = sizeSourceSize->factor[0];
 //                        blackLayoutOp.args.ImageBlankLayout.blockSize[1] = sizeSourceSize->factor[1];
diff --git a/Engine/Plugins/Mutable/Source/MutableTools/Private/MuT/Compiler.cpp b/Engine/Plugins/Mutable/Source/MutableTools/Private/MuT/Compiler.cpp
--- a/Engine/Plugins/Mutable/Source/MutableTools/Private/MuT/Compiler.cpp
+++ b/Engine/Plugins/Mutable/Source/MutableTools/Private/MuT/Compiler.cpp
@@ -211,10 +211,10 @@ namespace mu
     Compiler::Compiler( Ptr<CompilerOptions> options )
     {
         m_pD = new Private();
-        m_pD->m_options = options;
-        if (!m_pD->m_options)
+        m_pD->Options = options;
+        if (!m_pD->Options)
         {
-            m_pD->m_options = new CompilerOptions;
+            m_pD->Options = new CompilerOptions;
         }
     }
 
@@ -260,7 +260,7 @@ namespace mu
                     switch (Node->GetOpType())
                     {
 
-                    case OP_TYPE::IN_ADDMESH:
+                    case EOpType::IN_ADDMESH:
                     {
                         Ptr<ASTOpInstanceAdd> TypedNode = static_cast<ASTOpInstanceAdd*>(Node.get());
                         if (bRemoveMeshes && TypedNode)
@@ -269,7 +269,7 @@ namespace mu
                         }
                         break;
                     }
-                    case OP_TYPE::IN_ADDIMAGE:
+                    case EOpType::IN_ADDIMAGE:
                     {
                         Ptr<ASTOpInstanceAdd> TypedNode = static_cast<ASTOpInstanceAdd*>(Node.get());
                         if (bRemoveImages && TypedNode)
@@ -297,7 +297,7 @@ namespace mu
     }
 
     //---------------------------------------------------------------------------------------------
-	TSharedPtr<Model> Compiler::Compile( const Ptr<Node>& pNode )
+	TSharedPtr<FModel> Compiler::Compile( const Ptr<Node>& pNode )
     {
         MUTABLE_CPUPROFILER_SCOPE(Compile);
 
@@ -305,7 +305,7 @@ namespace mu
         Ptr<ErrorLog> genErrorLog;
 		TArray<FParameterDesc> Parameters;
         {
-            CodeGenerator gen( m_pD->m_options->GetPrivate() );
+            CodeGenerator gen( m_pD->Options->GetPrivate() );
 
             gen.GenerateRoot( pNode );
 
@@ -337,9 +337,9 @@ namespace mu
 			{
 				bool operator()(const FParameterDesc& A, const FParameterDesc& B) const
 				{
-					if (A.m_name < B.m_name) return true;
-					if (A.m_name > B.m_name) return false;
-					return A.m_uid < B.m_uid;
+					if (A.Name < B.Name) return true;
+					if (A.Name > B.Name) return false;
+					return A.UID < B.UID;
 				}
 			};
 			
@@ -347,7 +347,7 @@ namespace mu
 			Parameters.Sort(SortPredicate);
 		}
 
-        mu::Private::RemoveDisabledAddInstanceOpsValues(m_pD->m_options, states);
+        mu::Private::RemoveDisabledAddInstanceOpsValues(m_pD->Options, states);
 
         // Slow AST code verification for debugging.
         //TArray<Ptr<ASTOp>> roots;
@@ -359,54 +359,54 @@ namespace mu
 
         // Optimize the generated code
         {
-            CodeOptimiser optimiser( m_pD->m_options, states );
+            CodeOptimiser optimiser( m_pD->Options, states );
             optimiser.OptimiseAST( );
         }
 
         // Link the Program and generate state data.
-		TSharedPtr<Model> pResult = MakeShared<Model>();
-        FProgram& Program = pResult->GetPrivate()->m_program;
+		TSharedPtr<FModel> pResult = MakeShared<FModel>();
+        FProgram& Program = pResult->GetPrivate()->Program;
 
-		check(Program.m_parameters.IsEmpty());
-		Program.m_parameters = Parameters;
+		check(Program.Parameters.IsEmpty());
+		Program.Parameters = Parameters;
 
 		// Preallocate ample memory
-		Program.m_byteCode.Reserve(16 * 1024 * 1024);
-		Program.m_opAddress.Reserve(1024 * 1024);
+		Program.ByteCode.Reserve(16 * 1024 * 1024);
+		Program.OpAddress.Reserve(1024 * 1024);
 
 		// Keep the link options outside the scope because it is also used to cache constant data that has already been 
 		// added and could be reused across states.
-		FImageOperator ImOp = FImageOperator::GetDefault(m_pD->m_options->GetPrivate()->ImageFormatFunc);
+		FImageOperator ImOp = FImageOperator::GetDefault(m_pD->Options->GetPrivate()->ImageFormatFunc);
 		FLinkerOptions LinkerOptions(ImOp);
 
 		for(FStateCompilationData& s: states )
         {
-			LinkerOptions.MinTextureResidentMipCount = m_pD->m_options->GetPrivate()->MinTextureResidentMipCount;
+			LinkerOptions.MinTextureResidentMipCount = m_pD->Options->GetPrivate()->MinTextureResidentMipCount;
 
             if (s.root)
             {
-				s.state.m_root = ASTOp::FullLink(s.root, Program, &LinkerOptions);
+				s.state.Root = ASTOp::FullLink(s.root, Program, &LinkerOptions);
             }
             else
             {
-                s.state.m_root = 0;
+                s.state.Root = 0;
             }
         }
 
-		Program.m_byteCode.Shrink();
-		Program.m_opAddress.Shrink();
+		Program.ByteCode.Shrink();
+		Program.OpAddress.Shrink();
 
         // Set the runtime parameter indices.
         for(FStateCompilationData& s: states )
         {
-            for ( int32 p=0; p<s.nodeState.RuntimeParams.Num(); ++p )
+            for ( int32 RuntimeParameterIndex=0; RuntimeParameterIndex <s.nodeState.RuntimeParams.Num(); ++RuntimeParameterIndex)
             {
                 int32 paramIndex = -1;
-                for ( int32 i=0; paramIndex<0 && i<Program.m_parameters.Num(); ++i )
+                for ( int32 i=0; paramIndex<0 && i<Program.Parameters.Num(); ++i )
                 {
-                    if ( Program.m_parameters[i].m_name
+                    if ( Program.Parameters[i].Name
                          ==
-                         s.nodeState.RuntimeParams[p] )
+                         s.nodeState.RuntimeParams[RuntimeParameterIndex] )
                     {
                         paramIndex = (int)i;
                     }
@@ -422,7 +422,7 @@ namespace mu
 						"The state [%s] refers to a parameter [%s]  that has not been found in the model. This error can be "
 						"safely dismissed in case of partial compilation."), 
 						*s.nodeState.Name,
-						*s.nodeState.RuntimeParams[p]);
+						*s.nodeState.RuntimeParams[RuntimeParameterIndex]);
                     m_pD->ErrorLog->GetPrivate()->Add(Temp, ELMT_WARNING, pNode->GetMessageContext() );
                 }
             }
@@ -444,9 +444,9 @@ namespace mu
                 {
                     // Find the index in the model parameter list
                     int paramIndex = -1;
-                    for ( int32 i=0; paramIndex<0 && i<Program.m_parameters.Num(); ++i )
+                    for ( int32 i=0; paramIndex<0 && i<Program.Parameters.Num(); ++i )
                     {
-                        if ( Program.m_parameters[i].m_name == b )
+                        if ( Program.Parameters[i].Name == b )
                         {
                             paramIndex = (int)i;
                         }
@@ -475,19 +475,19 @@ namespace mu
             // Sort for performance and determinism
             s.state.m_dynamicResources.Sort();
 
-            Program.m_states.Add(s.state);
+            Program.States.Add(s.state);
         }
 
-		UE_LOG(LogMutableCore, Verbose, TEXT("(int) %s : %ld"), TEXT("Program size"), int64(Program.m_opAddress.Num()));
+		UE_LOG(LogMutableCore, Verbose, TEXT("(int) %s : %ld"), TEXT("Program size"), int64(Program.OpAddress.Num()));
 
         // Merge the log in the right order
         genErrorLog->Merge( m_pD->ErrorLog.get() );
         m_pD->ErrorLog = genErrorLog.get();
 
 		// Pack data
-		m_pD->GenerateRoms(pResult.Get(), m_pD->m_options.get(), LinkerOptions.AdditionalData );
+		m_pD->GenerateRoms(pResult.Get(), LinkerOptions.AdditionalData );
 
-		UE_LOG(LogMutableCore, Verbose, TEXT("(int) %s : %ld"), TEXT("Program size"), int64(Program.m_opAddress.Num()));
+		UE_LOG(LogMutableCore, Verbose, TEXT("(int) %s : %ld"), TEXT("Program size"), int64(Program.OpAddress.Num()));
 
         return pResult;
     }
@@ -524,14 +524,14 @@ namespace mu
 
 
 	//---------------------------------------------------------------------------------------------
-	void Compiler::Private::GenerateRoms(Model* p, const CompilerOptions* Options, const FLinkerOptions::FAdditionalData& AdditionalData )
+	void Compiler::Private::GenerateRoms(FModel* Model, const FLinkerOptions::FAdditionalData& AdditionalData )
 	{
 		LLM_SCOPE_BYNAME(TEXT("MutableRuntime"));
 		MUTABLE_CPUPROFILER_SCOPE(GenerateRoms);
 
 		uint64 EmbeddedDataBytesLimit = Options->GetPrivate()->EmbeddedDataBytesLimit;
 
-		mu::FProgram& Program = p->GetPrivate()->m_program;
+		mu::FProgram& Program = Model->GetPrivate()->Program;
 
 		// These are used for logging only.
 		int32 NumRoms = 0;
@@ -543,7 +543,7 @@ namespace mu
 
 		// Maximum number of roms
 		int32 MaxRomCount = Program.ConstantImageLODs.Num() + Program.ConstantMeshes.Num();
-		Program.m_roms.Reserve(MaxRomCount);
+		Program.Roms.Reserve(MaxRomCount);
 
 		TSet<uint32> UsedIds;
 		UsedIds.Reserve(MaxRomCount);
@@ -558,15 +558,15 @@ namespace mu
 			ParallelFor(Program.ConstantImageLODs.Num(),
 				[EmbeddedDataBytesLimit, &Program, &RomDatas](uint32 ResourceIndex)
 				{
-					TPair<int32, Ptr<const Image>>& ResData = Program.ConstantImageLODs[ResourceIndex];
+					TPair<int32, TSharedPtr<const FImage>>& ResData = Program.ConstantImageLODs[ResourceIndex];
 
 					// This shouldn't have been serialised with rom support before.
 					check(ResData.Key < 0);
 
 					// Serialize to find out final size of this rom
-					OutputSizeStream SizeStream;
-					OutputArchive MemoryArch(&SizeStream);
-					Image::Serialise(ResData.Value.get(), MemoryArch);
+					FOutputSizeStream SizeStream;
+					FOutputArchive MemoryArch(&SizeStream);
+					FImage::Serialise(ResData.Value.Get(), MemoryArch);
 
 					// If the resource uses less memory than the threshold, don't save it in a separate rom.
 					if (SizeStream.GetBufferSize() <= EmbeddedDataBytesLimit)
@@ -580,7 +580,7 @@ namespace mu
 					RomData.Size = SizeStream.GetBufferSize();
 					RomData.Flags = ERomFlags::None;
 
-					const Image* Resource = ResData.Value.get();
+					const FImage* Resource = ResData.Value.Get();
 					uint64 FullHash = CityHash64(reinterpret_cast<const char*>(Resource->GetLODData(0)), Resource->GetLODDataSize(0));
 					RomData.Id = GetTypeHash(FullHash);
 				});
@@ -663,9 +663,9 @@ namespace mu
 				// Ensure that the Id is unique
 				EnsureUniqueRomId(UsedIds, RomData.Id);
 
-				int32 RomIndex = Program.m_roms.Add(RomData);
+				int32 RomIndex = Program.Roms.Add(RomData);
 
-				TPair<int32, Ptr<const Image>>& ResData = Program.ConstantImageLODs[ResourceIndex];
+				TPair<int32, TSharedPtr<const FImage>>& ResData = Program.ConstantImageLODs[ResourceIndex];
 				ResData.Key = RomIndex;
 			}
 		}
@@ -678,13 +678,13 @@ namespace mu
 			ParallelFor(Program.ConstantMeshes.Num(),
 				[EmbeddedDataBytesLimit, &Program, &RomDatas](uint32 ResourceIndex)
 				{
-					TPair<int32, Ptr<const Mesh>>& ResData = Program.ConstantMeshes[ResourceIndex];
+					TPair<int32, TSharedPtr<const FMesh>>& ResData = Program.ConstantMeshes[ResourceIndex];
 
 					// This shouldn't have been serialised with rom support before.
 					check(ResData.Key < 0);
 
 					// Serialize to memory, to find out final size of this rom
-					const Mesh* Resource = ResData.Value.get();
+					const FMesh* Resource = ResData.Value.Get();
 					int32 ApproximateSize = Resource->GetDataSize();
 
 					// If the resource uses less memory than the threshold, don't save it in a separate rom.
@@ -693,9 +693,9 @@ namespace mu
 						return;
 					}
 
-					OutputMemoryStream MemStream(ApproximateSize + 64 * 1024);
-					OutputArchive MemoryArch(&MemStream);
-					Mesh::Serialise(ResData.Value.get(), MemoryArch);
+					FOutputMemoryStream MemStream(ApproximateSize + 64 * 1024);
+					FOutputArchive MemoryArch(&MemStream);
+					FMesh::Serialise(ResData.Value.Get(), MemoryArch);
 
 					FRomData& RomData = RomDatas[ResourceIndex];
 					RomData.ResourceType = DT_MESH;
@@ -704,7 +704,7 @@ namespace mu
 					RomData.Flags = ERomFlags::None;
 
 					// Ensure that the Id is unique
-					uint64 FullHash = CityHash64(static_cast<const char*>(MemStream.GetBuffer()), MemStream.GetBufferSize());
+					uint64 FullHash = CityHash64(reinterpret_cast<const char*>(MemStream.GetBuffer()), MemStream.GetBufferSize());
 					RomData.Id = GetTypeHash(FullHash);
 				});
 		}
@@ -730,9 +730,9 @@ namespace mu
 				// Ensure that the Id is unique
 				EnsureUniqueRomId(UsedIds, RomData.Id);
 
-				int32 RomIndex = Program.m_roms.Add(RomData);
+				int32 RomIndex = Program.Roms.Add(RomData);
 
-				TPair<int32, Ptr<const Mesh>>& ResData = Program.ConstantMeshes[ResourceIndex];
+				TPair<int32, TSharedPtr<const FMesh>>& ResData = Program.ConstantMeshes[ResourceIndex];
 				ResData.Key = RomIndex;
 			}
 		}
diff --git a/Engine/Plugins/Mutable/Source/MutableTools/Private/MuT/CompilerPrivate.h b/Engine/Plugins/Mutable/Source/MutableTools/Private/MuT/CompilerPrivate.h
--- a/Engine/Plugins/Mutable/Source/MutableTools/Private/MuT/CompilerPrivate.h
+++ b/Engine/Plugins/Mutable/Source/MutableTools/Private/MuT/CompilerPrivate.h
@@ -94,7 +94,7 @@ namespace mu
         Ptr<ErrorLog> ErrorLog;
 
         /** */
-        Ptr<CompilerOptions> m_options;
+        Ptr<CompilerOptions> Options;
 
 		/** Store for additional data generated during compilation, but not necessary for the runtime. */
 		struct FAdditionalData
@@ -111,7 +111,7 @@ namespace mu
 		};
 
 		//!
-		void GenerateRoms(Model*, const CompilerOptions*, const FLinkerOptions::FAdditionalData&);
+		void GenerateRoms(FModel*, const FLinkerOptions::FAdditionalData&);
 
     };
 
diff --git a/Engine/Plugins/Mutable/Source/MutableTools/Private/MuT/DataPacker.cpp b/Engine/Plugins/Mutable/Source/MutableTools/Private/MuT/DataPacker.cpp
--- a/Engine/Plugins/Mutable/Source/MutableTools/Private/MuT/DataPacker.cpp
+++ b/Engine/Plugins/Mutable/Source/MutableTools/Private/MuT/DataPacker.cpp
@@ -38,15 +38,15 @@ namespace mu
     //---------------------------------------------------------------------------------------------
     SubtreeSearchConstantVisitor::SubtreeSearchConstantVisitor( FProgram& program,
                                                                 OP::ADDRESS constant,
-                                                                OP_TYPE optype )
-        : m_program( program ),
-          m_constant( constant ),
-          m_opType( optype )
+                                                                EOpType optype )
+        : Program( program ),
+          Constant( constant ),
+          OpType( optype )
     {
-        m_visited.SetNum( program.m_opAddress.Num() );
-        if ( m_visited.Num() )
+        Visited.SetNum( program.OpAddress.Num() );
+        if ( Visited.Num() )
         {
-            FMemory::Memset( &m_visited[0], 0, m_visited.Num() );
+            FMemory::Memset( &Visited[0], 0, Visited.Num() );
         }
     }
 
@@ -61,7 +61,7 @@ namespace mu
         bool found = false;
 
         TArray< TPair<bool,int> > pending;
-        pending.Reserve(m_program.m_opAddress.Num()/4);
+        pending.Reserve(Program.OpAddress.Num()/4);
 
 		pending.Add({ false,root });
 
@@ -74,29 +74,29 @@ namespace mu
                 // Item indicating we finished with all the children of a parent.
 
                 // Propagate the "found" state upwards
-                m_visited[item.Value]=1;
+                Visited[item.Value]=1;
 
-                //OP cop = m_program.m_code[item.Value];
-                ForEachReference( m_program, (OP::ADDRESS)item.Value, [&](OP::ADDRESS ref)
+                //OP cop = Program.m_code[item.Value];
+                ForEachReference( Program, (OP::ADDRESS)item.Value, [&](OP::ADDRESS ref)
                 {
-                    if (ref && m_visited[ref]==2)
+                    if (ref && Visited[ref]==2)
                     {
-                        m_visited[item.Value]=2;
+                        Visited[item.Value]=2;
                     }
                 });
             }
             else
             {
-                if ( !m_visited[item.Value] )
+                if ( !Visited[item.Value] )
                 {
-                    OP_TYPE thisOpType = m_program.GetOpType(item.Value);
-                    if ( m_opType == thisOpType )
+                    EOpType thisOpType = Program.GetOpType(item.Value);
+                    if ( OpType == thisOpType )
                     {
-						OP::ResourceConstantArgs args = m_program.GetOpArgs<OP::ResourceConstantArgs>(item.Value);
-                        if ( m_constant == args.value )
+						OP::ResourceConstantArgs args = Program.GetOpArgs<OP::ResourceConstantArgs>(item.Value);
+                        if ( Constant == args.value )
                         {
                             found = true;
-                            m_visited[item.Value] = 2;
+                            Visited[item.Value] = 2;
                         }
                     }
 
@@ -104,10 +104,10 @@ namespace mu
                     {
 						pending.Add({ true,item.Value });
 
-                        //OP cop = m_program.m_code[item.Value];
-                        ForEachReference( m_program, item.Value, [&](OP::ADDRESS ref)
+                        //OP cop = Program.m_code[item.Value];
+                        ForEachReference( Program, item.Value, [&](OP::ADDRESS ref)
                         {
-                            if (ref && !m_visited[ref])
+                            if (ref && !Visited[ref])
                             {
 								pending.Add({ false,ref });
                             }
@@ -116,7 +116,7 @@ namespace mu
                 }
                 else
                 {
-                    found = (m_visited[item.Value]==2);
+                    found = (Visited[item.Value]==2);
                 }
             }
         }
@@ -134,20 +134,20 @@ namespace mu
     {
     public:
 
-        GatherParametersVisitor( FProgram& program, OP::ADDRESS constant, OP_TYPE opType )
+        GatherParametersVisitor( FProgram& program, OP::ADDRESS constant, EOpType opType )
             : m_constSearch( program, constant, opType )
         {
             MUTABLE_CPUPROFILER_SCOPE(GatherParametersVisitor);
 
 			TArray<int> currentParams;
-            currentParams.SetNumZeroed( program.m_parameters.Num() );
+            currentParams.SetNumZeroed( program.Parameters.Num() );
             SetDefaultState(currentParams);
 
             // TODO: we can optimize by precalculating what ops have the required optype below
-            m_opType = opType;
-            m_constant = constant;
+            OpType = opType;
+            Constant = constant;
 
-            m_allParams.SetNumZeroed( program.m_parameters.Num() );
+            m_allParams.SetNumZeroed( program.Parameters.Num() );
             FullTraverse( program );
 
             for (int32 i=0; i<m_allParams.Num(); ++i )
@@ -168,14 +168,14 @@ namespace mu
             {
                 switch ( program.GetOpType(at) )
                 {
-                case OP_TYPE::NU_CONDITIONAL:
-                case OP_TYPE::SC_CONDITIONAL:
-                case OP_TYPE::CO_CONDITIONAL:
-                case OP_TYPE::IM_CONDITIONAL:
-                case OP_TYPE::ME_CONDITIONAL:
-                case OP_TYPE::LA_CONDITIONAL:
-                case OP_TYPE::IN_CONDITIONAL:
-				case OP_TYPE::ED_CONDITIONAL:
+                case EOpType::NU_CONDITIONAL:
+                case EOpType::SC_CONDITIONAL:
+                case EOpType::CO_CONDITIONAL:
+                case EOpType::IM_CONDITIONAL:
+                case EOpType::ME_CONDITIONAL:
+                case EOpType::LA_CONDITIONAL:
+                case EOpType::IN_CONDITIONAL:
+				case EOpType::ED_CONDITIONAL:
                 {
 					OP::ConditionalArgs args = program.GetOpArgs<OP::ConditionalArgs>(at);
 
@@ -200,13 +200,13 @@ namespace mu
                     break;
                 }
 
-                case OP_TYPE::NU_SWITCH:
-                case OP_TYPE::SC_SWITCH:
-                case OP_TYPE::CO_SWITCH:
-                case OP_TYPE::IM_SWITCH:
-                case OP_TYPE::ME_SWITCH:
-                case OP_TYPE::LA_SWITCH:
-				case OP_TYPE::ED_SWITCH:
+                case EOpType::NU_SWITCH:
+                case EOpType::SC_SWITCH:
+                case EOpType::CO_SWITCH:
+                case EOpType::IM_SWITCH:
+                case EOpType::ME_SWITCH:
+                case EOpType::LA_SWITCH:
+				case EOpType::ED_SWITCH:
                 {
 					const uint8* data = program.GetOpArgsPointer(at);
 					
@@ -230,12 +230,12 @@ namespace mu
                 default:
                 {
                     // If we hit the constant we are analysing
-                    if ( m_opType == OP_TYPE::ME_CONSTANT
+                    if ( OpType == EOpType::ME_CONSTANT
                          &&
-                         program.GetOpType(at) == m_opType )
+                         program.GetOpType(at) == OpType )
                     {
 						OP::MeshConstantArgs args = program.GetOpArgs<OP::MeshConstantArgs>(at);
-                        if ( args.value == m_constant )
+                        if ( args.value == Constant )
                         {
                             // Accumulate the currently relevant parameters
                             const TArray<int>& currentParams = GetCurrentState();
@@ -245,10 +245,10 @@ namespace mu
                             }
                         }
                     }
-                    else if ( program.GetOpType(at) == m_opType )
+                    else if ( program.GetOpType(at) == OpType )
                     {
 						OP::ResourceConstantArgs args = program.GetOpArgs<OP::ResourceConstantArgs>(at);
-                        if ( args.value == m_constant )
+                        if ( args.value == Constant )
                         {
                             // Accumulate the currently relevant parameters
                             const TArray<int>& currentParams = GetCurrentState();
@@ -276,8 +276,8 @@ namespace mu
 
     private:
 
-        OP::ADDRESS m_constant;
-        OP_TYPE m_opType;
+        OP::ADDRESS Constant;
+        EOpType OpType;
 
 		TArray<int> m_allParams;
 
@@ -300,7 +300,7 @@ namespace mu
 
             // TODO: Use statically sized vectors (std::array)
             //vector<bool> initial( Count, true );
-            //std::fill( m_supportedFormats.begin(), m_supportedFormats.end(), initial );
+            //std::fill( SupportedFormats.begin(), SupportedFormats.end(), initial );
 
 			TArray<bool> defaultState;
 			defaultState.SetNumZeroed(int32(EImageFormat::Count));
@@ -322,44 +322,44 @@ namespace mu
             // Can we use the cache?
             if (allFalse)
             {
-                if (m_visited.count(node))
+                if (Visited.Contains(node))
                 {
                     return false;
                 }
 
-                m_visited.insert(node);
+                Visited.Add(node);
             }
 
             switch ( node->GetOpType() )
             {
-            case OP_TYPE::IM_CONSTANT:
+            case EOpType::IM_CONSTANT:
             {
                 // Remove unsupported formats
 				const ASTOpConstantResource* op = static_cast<const ASTOpConstantResource*>(node.get());
-                if (!m_supportedFormats.count(op))
+                if (!SupportedFormats.Contains(op))
                 {
 					TArray<bool> initial;
 					initial.Init( true, int32(EImageFormat::Count) );
-                    m_supportedFormats.insert( std::make_pair(op, std::move(initial)) );
+                    SupportedFormats.Add( op, std::move(initial) );
                 }
 
                 for ( unsigned f=0; f< unsigned(EImageFormat::Count); ++f )
                 {
                     if ( !currentFormats[f] )
                     {
-                        m_supportedFormats[op][f] = false;
+                        SupportedFormats[op][f] = false;
                     }
                 }
                 recurse = false;
                 break;
             }
 
-            case OP_TYPE::IM_SWITCH:
-            case OP_TYPE::IM_CONDITIONAL:
+            case EOpType::IM_SWITCH:
+            case EOpType::IM_CONDITIONAL:
                 // Switches and conditionals don't change the supported formats
                 break;
 
-            case OP_TYPE::IM_COMPOSE:
+            case EOpType::IM_COMPOSE:
             {
                 recurse = false;
 				const ASTOpImageCompose* op = static_cast<const ASTOpImageCompose*>(node.get());
@@ -378,7 +378,7 @@ namespace mu
                 break;
             }
 
-            case OP_TYPE::IM_LAYERCOLOUR:
+            case EOpType::IM_LAYERCOLOUR:
             {
                 recurse = false;
 				const ASTOpImageLayerColor* op = static_cast<const ASTOpImageLayerColor*>(node.get());
@@ -397,7 +397,7 @@ namespace mu
                 break;
             }
 
-            case OP_TYPE::IM_LAYER:
+            case EOpType::IM_LAYER:
             {
                 recurse = false;
 				const ASTOpImageLayer* op = static_cast<const ASTOpImageLayer*>(node.get());
@@ -416,7 +416,7 @@ namespace mu
                 break;
             }
 
-            case OP_TYPE::IM_MULTILAYER:
+            case EOpType::IM_MULTILAYER:
             {
                 recurse = false;
 				const ASTOpImageMultiLayer* op = static_cast<const ASTOpImageMultiLayer*>(node.get());
@@ -435,7 +435,7 @@ namespace mu
                 break;
             }
 
-            case OP_TYPE::IM_DISPLACE:
+            case EOpType::IM_DISPLACE:
             {
                 recurse = false;
 				const ASTOpFixed* op = static_cast<const ASTOpFixed*>(node.get());
@@ -479,12 +479,12 @@ namespace mu
 
         //! Result of this visitor:
         //! Formats known to be supported by every constant image.
-        std::unordered_map< Ptr<const ASTOpConstantResource>, TArray<bool> > m_supportedFormats;
+        TMap< Ptr<const ASTOpConstantResource>, TArray<bool> > SupportedFormats;
 
     private:
 
         //! Cache. Only valid is current formats are all false.
-        std::unordered_set<Ptr<ASTOp>> m_visited;
+        TSet<Ptr<ASTOp>> Visited;
 
     };
 
@@ -523,7 +523,7 @@ namespace mu
             switch ( node->GetOpType() )
             {
 
-            case OP_TYPE::ME_CONSTANT:
+            case EOpType::ME_CONSTANT:
             {
                 // Accumulate necessary semantics
 				const ASTOpConstantResource* op = static_cast<const ASTOpConstantResource*>(node.get());
@@ -544,9 +544,9 @@ namespace mu
             }
 
             // TODO: These could probably optimise something
-            //case OP_TYPE::IM_RASTERMESH: break;
+            //case EOpType::IM_RASTERMESH: break;
 
-            case OP_TYPE::ME_DIFFERENCE:
+            case EOpType::ME_DIFFERENCE:
             {
                 recurse = false;
 
@@ -560,7 +560,7 @@ namespace mu
                 break;
              }
 
-            case OP_TYPE::ME_REMOVEMASK:
+            case EOpType::ME_REMOVEMASK:
             {
                 recurse = false;
 
@@ -576,7 +576,7 @@ namespace mu
                 break;
              }
 
-            case OP_TYPE::ME_MORPH:
+            case EOpType::ME_MORPH:
             {
                 recurse = false;
 
@@ -589,7 +589,7 @@ namespace mu
                 break;
              }
 
-            case OP_TYPE::ME_APPLYLAYOUT:
+            case EOpType::ME_APPLYLAYOUT:
             {
                 recurse = false;
 
@@ -603,7 +603,7 @@ namespace mu
                 break;
             }
 
-			case OP_TYPE::ME_PROJECT:
+			case EOpType::ME_PROJECT:
 			{
 				recurse = false;
 
@@ -617,7 +617,7 @@ namespace mu
 				break;
 			}
 
-			case OP_TYPE::IM_RASTERMESH:
+			case EOpType::IM_RASTERMESH:
 			{
 				recurse = false;
 
@@ -634,7 +634,7 @@ namespace mu
 				break;
 			}
 
-            case OP_TYPE::ME_EXTRACTLAYOUTBLOCK:
+            case EOpType::ME_EXTRACTLAYOUTBLOCK:
             {
                 recurse = false;
 
@@ -649,7 +649,7 @@ namespace mu
                 break;
             }
 
-			case OP_TYPE::LA_FROMMESH:
+			case EOpType::LA_FROMMESH:
 			{
 				recurse = false;
 
@@ -662,7 +662,7 @@ namespace mu
 				break;
 			}
 
-            case OP_TYPE::IN_ADDMESH:
+            case EOpType::IN_ADDMESH:
             {
                 recurse = false;
 
@@ -699,7 +699,7 @@ namespace mu
     //---------------------------------------------------------------------------------------------
 
     // Todo: move to its own file
-    inline void MeshRemoveUnusedBufferSemantics( Mesh* pMesh, uint64 usedSemantics )
+    inline void MeshRemoveUnusedBufferSemantics( FMesh* pMesh, uint64 usedSemantics )
     {
         // right now we only remove entire buffers if no channel is used
         // TODO: remove from inside the buffer?
@@ -751,10 +751,10 @@ namespace mu
         // See if we can convert some constants to more efficient formats
         ASTOp::Traverse_BottomUp_Unique_NonReentrant( roots, [&](Ptr<ASTOp>& n)
         {
-            if (n->GetOpType()==OP_TYPE::IM_CONSTANT)
+            if (n->GetOpType()==EOpType::IM_CONSTANT)
             {
-				ASTOpConstantResource* typed = static_cast<ASTOpConstantResource*>(n.get());
-                Ptr<const Image> pOld = static_cast<const Image*>(typed->GetValue().get());
+				ASTOpConstantResource* Typed = static_cast<ASTOpConstantResource*>(n.get());
+                TSharedPtr<const FImage> pOld = StaticCastSharedPtr<const FImage>(Typed->GetValue());
 
 				FImageOperator ImOp = FImageOperator::GetDefault( Options->GetPrivate()->ImageFormatFunc );
 
@@ -766,14 +766,14 @@ namespace mu
 					// of data loading. 
 					// Warning This eliminates the mips. \TODO: Add support for mips in plaincolour instruction?
 					Ptr<ASTOpFixed> NewColor = new ASTOpFixed;
-					NewColor->op.type = OP_TYPE::CO_CONSTANT;
+					NewColor->op.type = EOpType::CO_CONSTANT;
 					NewColor->op.args.ColourConstant.value[0] = PlainColor[0];
 					NewColor->op.args.ColourConstant.value[1] = PlainColor[1];
 					NewColor->op.args.ColourConstant.value[2] = PlainColor[2];
 					NewColor->op.args.ColourConstant.value[3] = PlainColor[3];
 
 					Ptr<ASTOpFixed> NewPlain = new ASTOpFixed;
-					NewPlain->op.type = OP_TYPE::IM_PLAINCOLOUR;
+					NewPlain->op.type = EOpType::IM_PLAINCOLOUR;
 					NewPlain->SetChild( NewPlain->op.args.ImagePlainColour.colour, NewColor );
 					NewPlain->op.args.ImagePlainColour.format = pOld->GetFormat();
 					NewPlain->op.args.ImagePlainColour.size[0] = pOld->GetSizeX();
@@ -783,28 +783,28 @@ namespace mu
 					ASTOp::Replace(n, NewPlain);
 
 				}
-				else if ( accFormat.m_supportedFormats[typed][(int32)EImageFormat::L_UBitRLE] )
+				else if ( accFormat.SupportedFormats[Typed][(int32)EImageFormat::L_UBitRLE] )
                 {
-                    ImagePtr pNew = ImOp.ImagePixelFormat( ImageCompressionQuality, pOld.get(), EImageFormat::L_UBitRLE );
+                    TSharedPtr<FImage> pNew = ImOp.ImagePixelFormat( ImageCompressionQuality, pOld.Get(), EImageFormat::L_UBitRLE );
 
                     // Only replace if the compression was worth!
                     int32 oldSize = pOld->GetDataSize();
 					int32 newSize = pNew->GetDataSize();
                     if (float(oldSize) > float(newSize) * OptimizeOptions.MinRLECompressionGain)
                     {
-                        typed->SetValue(pNew, OptimizeOptions.DiskCacheContext);
+						Typed->SetValue(pNew, OptimizeOptions.DiskCacheContext);
                     }
                 }
-                else if ( accFormat.m_supportedFormats[typed][(int32)EImageFormat::L_UByteRLE] )
+                else if ( accFormat.SupportedFormats[Typed][(int32)EImageFormat::L_UByteRLE] )
                 {
-                    ImagePtr pNew = ImOp.ImagePixelFormat( ImageCompressionQuality, pOld.get(), EImageFormat::L_UByteRLE );
+                    TSharedPtr<FImage> pNew = ImOp.ImagePixelFormat( ImageCompressionQuality, pOld.Get(), EImageFormat::L_UByteRLE );
 
                     // Only replace if the compression was worth!
 					int32 oldSize = pOld->GetDataSize();
 					int32 newSize = pNew->GetDataSize();
                     if (float(oldSize) > float(newSize) * OptimizeOptions.MinRLECompressionGain)
                     {
-                        typed->SetValue(pNew, OptimizeOptions.DiskCacheContext);
+                        Typed->SetValue(pNew, OptimizeOptions.DiskCacheContext);
                     }
                 }
             }
@@ -818,12 +818,12 @@ namespace mu
         // See if we can remove some buffers from the constants
         ASTOp::Traverse_BottomUp_Unique_NonReentrant( roots, [&](Ptr<ASTOp>& n)
         {
-            if (n->GetOpType()==OP_TYPE::ME_CONSTANT)
+            if (n->GetOpType()==EOpType::ME_CONSTANT)
             {
-				ASTOpConstantResource* typed = static_cast<ASTOpConstantResource*>(n.get());
-                Ptr<Mesh> pMesh = static_cast<const Mesh*>(typed->GetValue().get())->Clone();
-                MeshRemoveUnusedBufferSemantics( pMesh.get(), meshSemanticsVisitor.m_requiredSemantics[typed]);
-                typed->SetValue(pMesh, OptimizeOptions.DiskCacheContext);
+				ASTOpConstantResource* Typed = static_cast<ASTOpConstantResource*>(n.get());
+                TSharedPtr<FMesh> Mesh = static_cast<const FMesh*>(Typed->GetValue().Get())->Clone();
+                MeshRemoveUnusedBufferSemantics(Mesh.Get(), meshSemanticsVisitor.m_requiredSemantics[Typed]);
+				Typed->SetValue(Mesh, OptimizeOptions.DiskCacheContext);
             }
         });
 
diff --git a/Engine/Plugins/Mutable/Source/MutableTools/Private/MuT/DataPacker.h b/Engine/Plugins/Mutable/Source/MutableTools/Private/MuT/DataPacker.h
--- a/Engine/Plugins/Mutable/Source/MutableTools/Private/MuT/DataPacker.h
+++ b/Engine/Plugins/Mutable/Source/MutableTools/Private/MuT/DataPacker.h
@@ -7,10 +7,6 @@
 
 namespace mu { struct FProgram; }
 
-
-#define MUTABLE_HASH_SEED					((uint32)0xcadababa)
-
-
 namespace mu
 {
 	class CompilerOptions;
@@ -25,21 +21,21 @@ namespace mu
     {
     public:
 
-        SubtreeSearchConstantVisitor( FProgram& program, OP::ADDRESS constant, OP_TYPE optype );
+        SubtreeSearchConstantVisitor( FProgram&, OP::ADDRESS Constant, EOpType );
 
         bool Run( OP::ADDRESS root );
 
 
     private:
 
-        FProgram& m_program;
-        OP::ADDRESS m_constant;
-        OP_TYPE m_opType;
+        FProgram& Program;
+        OP::ADDRESS Constant;
+        EOpType OpType;
 
         //! 0: not visited
         //! 1: visited and not found
         //! 2: visited and found
-        TArray<uint8> m_visited;
+        TArray<uint8> Visited;
     };
 
 }
diff --git a/Engine/Plugins/Mutable/Source/MutableTools/Private/MuT/ErrorLog.cpp b/Engine/Plugins/Mutable/Source/MutableTools/Private/MuT/ErrorLog.cpp
--- a/Engine/Plugins/Mutable/Source/MutableTools/Private/MuT/ErrorLog.cpp
+++ b/Engine/Plugins/Mutable/Source/MutableTools/Private/MuT/ErrorLog.cpp
@@ -349,13 +349,13 @@ namespace mu
 		TEXT("LA_FROMMESH	   "),
 	};
 
-	static_assert(sizeof(s_opNames) / sizeof(void*) == int32(OP_TYPE::COUNT));
+	static_assert(sizeof(s_opNames) / sizeof(void*) == int32(EOpType::COUNT));
 
     // clang-format on
 
-    const TCHAR* GetOpName( OP_TYPE type )
+    const TCHAR* GetOpName( EOpType type )
     {
-        //check( type>=0 && type<(int32)OP_TYPE::COUNT );
+        //check( type>=0 && type<(int32)EOpType::COUNT );
         return s_opNames[(int32)type];
     }
 
diff --git a/Engine/Plugins/Mutable/Source/MutableTools/Private/MuT/NodeBoolConstant.cpp b/Engine/Plugins/Mutable/Source/MutableTools/Private/MuT/NodeBoolConstant.cpp
--- a/Engine/Plugins/Mutable/Source/MutableTools/Private/MuT/NodeBoolConstant.cpp
+++ b/Engine/Plugins/Mutable/Source/MutableTools/Private/MuT/NodeBoolConstant.cpp
@@ -30,13 +30,13 @@ namespace mu
 	//---------------------------------------------------------------------------------------------
 	bool NodeBoolConstant::GetValue() const
 	{
-		return m_pD->m_value;
+		return m_pD->Value;
 	}
 
 	//---------------------------------------------------------------------------------------------
 	void NodeBoolConstant::SetValue( bool v )
 	{
-		m_pD->m_value = v;
+		m_pD->Value = v;
 	}
 
 
diff --git a/Engine/Plugins/Mutable/Source/MutableTools/Private/MuT/NodeBoolParameter.cpp b/Engine/Plugins/Mutable/Source/MutableTools/Private/MuT/NodeBoolParameter.cpp
--- a/Engine/Plugins/Mutable/Source/MutableTools/Private/MuT/NodeBoolParameter.cpp
+++ b/Engine/Plugins/Mutable/Source/MutableTools/Private/MuT/NodeBoolParameter.cpp
@@ -32,14 +32,14 @@ namespace mu
 	//---------------------------------------------------------------------------------------------
 	void NodeBoolParameter::SetName( const FString& strName )
 	{
-		m_pD->m_name = strName;
+		m_pD->Name = strName;
 	}
 
 
 	//---------------------------------------------------------------------------------------------
 	void NodeBoolParameter::SetDefaultValue( bool v )
 	{
-		m_pD->m_defaultValue = v;
+		m_pD->DefaultValue = v;
 	}
 
 
@@ -47,17 +47,17 @@ namespace mu
     void NodeBoolParameter::SetRangeCount( int i )
     {
         check(i>=0);
-        m_pD->m_ranges.SetNum(i);
+        m_pD->Ranges.SetNum(i);
     }
 
 
     //---------------------------------------------------------------------------------------------
     void NodeBoolParameter::SetRange( int i, NodeRangePtr pRange )
     {
-        check( i>=0 && i<int(m_pD->m_ranges.Num()) );
-        if ( i>=0 && i<int(m_pD->m_ranges.Num()) )
+        check( i>=0 && i<int(m_pD->Ranges.Num()) );
+        if ( i>=0 && i<int(m_pD->Ranges.Num()) )
         {
-            m_pD->m_ranges[i] = pRange;
+            m_pD->Ranges[i] = pRange;
         }
     }
 
diff --git a/Engine/Plugins/Mutable/Source/MutableTools/Private/MuT/NodeBoolPrivate.h b/Engine/Plugins/Mutable/Source/MutableTools/Private/MuT/NodeBoolPrivate.h
--- a/Engine/Plugins/Mutable/Source/MutableTools/Private/MuT/NodeBoolPrivate.h
+++ b/Engine/Plugins/Mutable/Source/MutableTools/Private/MuT/NodeBoolPrivate.h
@@ -23,7 +23,7 @@ namespace mu
 
 		static FNodeType s_type;
 
-		bool m_value;
+		bool Value;
 	};
 
 
@@ -33,11 +33,11 @@ namespace mu
 
 		static FNodeType s_type;
 
-		bool m_defaultValue;
-		FString m_name;
-		FString m_uid;
+		bool DefaultValue;
+		FString Name;
+		FString UID;
 
-        TArray<Ptr<NodeRange>> m_ranges;
+        TArray<Ptr<NodeRange>> Ranges;
 	};
 
 
diff --git a/Engine/Plugins/Mutable/Source/MutableTools/Private/MuT/NodeImageConditional.cpp b/Engine/Plugins/Mutable/Source/MutableTools/Private/MuT/NodeImageConditional.cpp
--- a/Engine/Plugins/Mutable/Source/MutableTools/Private/MuT/NodeImageConditional.cpp
+++ b/Engine/Plugins/Mutable/Source/MutableTools/Private/MuT/NodeImageConditional.cpp
@@ -30,14 +30,14 @@ namespace mu
 	//---------------------------------------------------------------------------------------------
     NodeBool* NodeImageConditional::GetParameter() const
 	{
-        return m_pD->m_parameter.get();
+        return m_pD->Parameter.get();
 	}
 
 
 	//---------------------------------------------------------------------------------------------
     void NodeImageConditional::SetParameter( NodeBool* pNode )
 	{
-        m_pD->m_parameter = pNode;
+        m_pD->Parameter = pNode;
 	}
 
 
diff --git a/Engine/Plugins/Mutable/Source/MutableTools/Private/MuT/NodeImageConditionalPrivate.h b/Engine/Plugins/Mutable/Source/MutableTools/Private/MuT/NodeImageConditionalPrivate.h
--- a/Engine/Plugins/Mutable/Source/MutableTools/Private/MuT/NodeImageConditionalPrivate.h
+++ b/Engine/Plugins/Mutable/Source/MutableTools/Private/MuT/NodeImageConditionalPrivate.h
@@ -17,7 +17,7 @@ namespace mu
 
 		static FNodeType s_type;
 
-        Ptr<NodeBool> m_parameter;
+        Ptr<NodeBool> Parameter;
         NodeImagePtr m_true;
         NodeImagePtr m_false;
 
diff --git a/Engine/Plugins/Mutable/Source/MutableTools/Private/MuT/NodeImageConstant.cpp b/Engine/Plugins/Mutable/Source/MutableTools/Private/MuT/NodeImageConstant.cpp
--- a/Engine/Plugins/Mutable/Source/MutableTools/Private/MuT/NodeImageConstant.cpp
+++ b/Engine/Plugins/Mutable/Source/MutableTools/Private/MuT/NodeImageConstant.cpp
@@ -19,14 +19,14 @@ namespace mu
 
 
 	//---------------------------------------------------------------------------------------------
-	void NodeImageConstant::SetValue(const Image* Value)
+	void NodeImageConstant::SetValue(TSharedPtr<const FImage> Value)
 	{
-		Proxy = new ResourceProxyMemory<Image>(Value);
+		Proxy = new TResourceProxyMemory<FImage>(Value);
 	}
 
 
 	//---------------------------------------------------------------------------------------------
-	void NodeImageConstant::SetValue(Ptr<ResourceProxy<Image>> InProxy)
+	void NodeImageConstant::SetValue(Ptr<TResourceProxy<FImage>> InProxy)
 	{
 		Proxy = InProxy;
 	}
diff --git a/Engine/Plugins/Mutable/Source/MutableTools/Private/MuT/NodeImageLayer.cpp b/Engine/Plugins/Mutable/Source/MutableTools/Private/MuT/NodeImageLayer.cpp
--- a/Engine/Plugins/Mutable/Source/MutableTools/Private/MuT/NodeImageLayer.cpp
+++ b/Engine/Plugins/Mutable/Source/MutableTools/Private/MuT/NodeImageLayer.cpp
@@ -71,14 +71,14 @@ namespace mu
 	//---------------------------------------------------------------------------------------------
 	EBlendType NodeImageLayer::GetBlendType() const
 	{
-		return m_pD->m_type;
+		return m_pD->Type;
 	}
 
 
 	//---------------------------------------------------------------------------------------------
 	void NodeImageLayer::SetBlendType(EBlendType t)
 	{
-		m_pD->m_type = t;
+		m_pD->Type = t;
 	}
 
 
diff --git a/Engine/Plugins/Mutable/Source/MutableTools/Private/MuT/NodeImageLayerColour.cpp b/Engine/Plugins/Mutable/Source/MutableTools/Private/MuT/NodeImageLayerColour.cpp
--- a/Engine/Plugins/Mutable/Source/MutableTools/Private/MuT/NodeImageLayerColour.cpp
+++ b/Engine/Plugins/Mutable/Source/MutableTools/Private/MuT/NodeImageLayerColour.cpp
@@ -67,14 +67,14 @@ namespace mu
 	//---------------------------------------------------------------------------------------------
 	EBlendType NodeImageLayerColour::GetBlendType() const
 	{
-		return m_pD->m_type;
+		return m_pD->Type;
 	}
 
 
 	//---------------------------------------------------------------------------------------------
 	void NodeImageLayerColour::SetBlendType(EBlendType t)
 	{
-		m_pD->m_type = t;
+		m_pD->Type = t;
 	}
 
 }
diff --git a/Engine/Plugins/Mutable/Source/MutableTools/Private/MuT/NodeImageParameter.cpp b/Engine/Plugins/Mutable/Source/MutableTools/Private/MuT/NodeImageParameter.cpp
--- a/Engine/Plugins/Mutable/Source/MutableTools/Private/MuT/NodeImageParameter.cpp
+++ b/Engine/Plugins/Mutable/Source/MutableTools/Private/MuT/NodeImageParameter.cpp
@@ -29,21 +29,21 @@ namespace mu
 	//---------------------------------------------------------------------------------------------
     void NodeImageParameter::SetName( const FString& Name )
 	{
-		m_pD->m_name = Name;
+		m_pD->Name = Name;
 	}
 
 
 	//---------------------------------------------------------------------------------------------
 	void NodeImageParameter::SetUid(const FString& Uid)
 	{
-		m_pD->m_uid = Uid;
+		m_pD->UID = Uid;
 	}
 
 	
 	//---------------------------------------------------------------------------------------------
 	void NodeImageParameter::SetDefaultValue(FName Value)
 	{
-    	m_pD->m_defaultValue = Value;
+    	m_pD->DefaultValue = Value;
 	}
 
 
@@ -51,17 +51,17 @@ namespace mu
 	void NodeImageParameter::SetRangeCount(int i)
 	{
 		check(i >= 0);
-		m_pD->m_ranges.SetNum(i);
+		m_pD->Ranges.SetNum(i);
 	}
 
 
 	//---------------------------------------------------------------------------------------------
 	void NodeImageParameter::SetRange(int i, NodeRangePtr pRange)
 	{
-		check(i >= 0 && i<int(m_pD->m_ranges.Num()));
-		if (i >= 0 && i<int(m_pD->m_ranges.Num()))
+		check(i >= 0 && i<int(m_pD->Ranges.Num()));
+		if (i >= 0 && i<int(m_pD->Ranges.Num()))
 		{
-			m_pD->m_ranges[i] = pRange;
+			m_pD->Ranges[i] = pRange;
 		}
 	}
 
diff --git a/Engine/Plugins/Mutable/Source/MutableTools/Private/MuT/NodeImageParameterPrivate.h b/Engine/Plugins/Mutable/Source/MutableTools/Private/MuT/NodeImageParameterPrivate.h
--- a/Engine/Plugins/Mutable/Source/MutableTools/Private/MuT/NodeImageParameterPrivate.h
+++ b/Engine/Plugins/Mutable/Source/MutableTools/Private/MuT/NodeImageParameterPrivate.h
@@ -16,11 +16,11 @@ namespace mu
 
 		static FNodeType s_type;
 
-    	FName m_defaultValue;
-		FString m_name;
-		FString m_uid;
+    	FName DefaultValue;
+		FString Name;
+		FString UID;
 
-		TArray<Ptr<NodeRange>> m_ranges;
+		TArray<Ptr<NodeRange>> Ranges;
 	};
 
 }
diff --git a/Engine/Plugins/Mutable/Source/MutableTools/Private/MuT/NodeImageSwitch.cpp b/Engine/Plugins/Mutable/Source/MutableTools/Private/MuT/NodeImageSwitch.cpp
--- a/Engine/Plugins/Mutable/Source/MutableTools/Private/MuT/NodeImageSwitch.cpp
+++ b/Engine/Plugins/Mutable/Source/MutableTools/Private/MuT/NodeImageSwitch.cpp
@@ -44,30 +44,30 @@ namespace mu
 	//---------------------------------------------------------------------------------------------
 	void NodeImageSwitch::SetOptionCount( int t )
 	{
-		m_pD->m_options.SetNum(t);
+		m_pD->Options.SetNum(t);
 	}
 
 
 	//---------------------------------------------------------------------------------------------
 	int NodeImageSwitch::GetOptionCount() const 
 	{
-		return int(m_pD->m_options.Num());
+		return int(m_pD->Options.Num());
 	}
 
 
 	//---------------------------------------------------------------------------------------------
 	NodeImagePtr NodeImageSwitch::GetOption( int t ) const
 	{
-		check( t>=0 && t<(int)m_pD->m_options.Num() );
-		return m_pD->m_options[t].get();
+		check( t>=0 && t<(int)m_pD->Options.Num() );
+		return m_pD->Options[t].get();
 	}
 
 
 	//---------------------------------------------------------------------------------------------
 	void NodeImageSwitch::SetOption( int t, NodeImagePtr pNode )
 	{
-		check( t>=0 && t<(int)m_pD->m_options.Num() );
-		m_pD->m_options[t] = pNode;
+		check( t>=0 && t<(int)m_pD->Options.Num() );
+		m_pD->Options[t] = pNode;
 	}
 
 
diff --git a/Engine/Plugins/Mutable/Source/MutableTools/Private/MuT/NodeLayout.cpp b/Engine/Plugins/Mutable/Source/MutableTools/Private/MuT/NodeLayout.cpp
--- a/Engine/Plugins/Mutable/Source/MutableTools/Private/MuT/NodeLayout.cpp
+++ b/Engine/Plugins/Mutable/Source/MutableTools/Private/MuT/NodeLayout.cpp
@@ -19,7 +19,7 @@ namespace mu
 	FNodeType NodeLayout::StaticType = FNodeType(Node::EType::Layout, Node::GetStaticType() );
 
 
-	void NodeLayout::GenerateLayoutBlocks(const Ptr<Mesh>& Mesh, int32 LayoutIndex )
+	void NodeLayout::GenerateLayoutBlocks(const TSharedPtr<FMesh>& Mesh, int32 LayoutIndex )
 	{
 		if (!(Mesh && LayoutIndex >= 0 && Size.X > 0 && Size.Y > 0))
 		{
@@ -225,7 +225,7 @@ namespace mu
 	}
 
 
-	void NodeLayout::GenerateLayoutBlocksFromUVIslands(const Ptr<Mesh>& Mesh, int32 LayoutIndex, bool bMergeChildBlocks)
+	void NodeLayout::GenerateLayoutBlocksFromUVIslands(const TSharedPtr<FMesh>& Mesh, int32 LayoutIndex, bool bMergeChildBlocks)
 	{
 		if (!(Mesh && LayoutIndex >= 0 && Size.X > 0 && Size.Y > 0))
 		{
@@ -247,7 +247,7 @@ namespace mu
 		VertexToFaceMap.Reserve(NumVertices*4);
 
 		// Find Unique Vertices
-		MeshCreateCollapsedVertexMap(Mesh.get(), CollapsedVertices);
+		MeshCreateCollapsedVertexMap(Mesh.Get(), CollapsedVertices);
 
 		UntypedMeshBufferIteratorConst ItIndices(Mesh->GetIndexBuffers(), EMeshBufferSemantic::VertexIndex);
 		for (int32 TriangleIndex = 0; TriangleIndex < NumTriangles; ++TriangleIndex)
@@ -487,7 +487,7 @@ namespace mu
 			// TODO: Size?
 			int32 SizeX = 1024;
 			int32 SizeY = 1024;
-			Ptr<Image> Mask = new Image(SizeX, SizeY, 1, EImageFormat::L_UByte, EInitializationType::Black);
+			TSharedPtr<FImage> Mask = MakeShared<FImage>(SizeX, SizeY, 1, EImageFormat::L_UByte, EInitializationType::Black);
 			const TArrayView<uint8> ImageData = Mask->DataStorage.GetLOD(0);
 
 			for (int32 TriangleIndex = 0; TriangleIndex < NumTriangles; ++TriangleIndex)
@@ -525,8 +525,8 @@ namespace mu
 			}
 
 			// TODO: Clamp UV islands always?
-			ImageGrow(Mask.get());
-			ImageGrow(Mask.get());
+			ImageGrow(Mask.Get());
+			ImageGrow(Mask.Get());
 
 			Block.Mask = Mask;
 		}
diff --git a/Engine/Plugins/Mutable/Source/MutableTools/Private/MuT/NodeMesh.cpp b/Engine/Plugins/Mutable/Source/MutableTools/Private/MuT/NodeMesh.cpp
--- a/Engine/Plugins/Mutable/Source/MutableTools/Private/MuT/NodeMesh.cpp
+++ b/Engine/Plugins/Mutable/Source/MutableTools/Private/MuT/NodeMesh.cpp
@@ -26,7 +26,7 @@ namespace mu
 	}
 
 
-	void MeshCreateCollapsedVertexMap(const mu::Mesh* Mesh, TArray<int32>& CollapsedVertices)
+	void MeshCreateCollapsedVertexMap(const mu::FMesh* Mesh, TArray<int32>& CollapsedVertices)
 	{
 		MUTABLE_CPUPROFILER_SCOPE(LayoutUV_CreateCollapsedVertexMap);
 
diff --git a/Engine/Plugins/Mutable/Source/MutableTools/Private/MuT/NodeMeshConstant.cpp b/Engine/Plugins/Mutable/Source/MutableTools/Private/MuT/NodeMeshConstant.cpp
--- a/Engine/Plugins/Mutable/Source/MutableTools/Private/MuT/NodeMeshConstant.cpp
+++ b/Engine/Plugins/Mutable/Source/MutableTools/Private/MuT/NodeMeshConstant.cpp
@@ -18,13 +18,13 @@ namespace mu
 	MUTABLE_IMPLEMENT_NODE( NodeMeshConstant );
 
 
-	Ptr<Mesh> NodeMeshConstant::GetValue() const
+	TSharedPtr<FMesh> NodeMeshConstant::GetValue() const
 	{
 		return m_pD->Value;
 	}
 
 
-	void NodeMeshConstant::SetValue( Ptr<Mesh> pValue )
+	void NodeMeshConstant::SetValue( TSharedPtr<FMesh> pValue )
 	{
 		m_pD->Value = pValue;
 
@@ -36,13 +36,13 @@ namespace mu
     }
 
 
-	void NodeMeshConstant::AddMorph(const FString& Name, Ptr<Mesh> Morphed)
+	void NodeMeshConstant::AddMorph(const FString& Name, TSharedPtr<FMesh> Morphed)
 	{
 		m_pD->Morphs.Add({Name,Morphed});
 	}
 
 
-	Ptr<Mesh> NodeMeshConstant::FindMorph(const FString& Name) const
+	TSharedPtr<FMesh> NodeMeshConstant::FindMorph(const FString& Name) const
 	{
 		for (Private::FMorph& Morph : m_pD->Morphs )
 		{
diff --git a/Engine/Plugins/Mutable/Source/MutableTools/Private/MuT/NodeMeshSwitch.cpp b/Engine/Plugins/Mutable/Source/MutableTools/Private/MuT/NodeMeshSwitch.cpp
--- a/Engine/Plugins/Mutable/Source/MutableTools/Private/MuT/NodeMeshSwitch.cpp
+++ b/Engine/Plugins/Mutable/Source/MutableTools/Private/MuT/NodeMeshSwitch.cpp
@@ -45,23 +45,23 @@ namespace mu
 	//---------------------------------------------------------------------------------------------
 	void NodeMeshSwitch::SetOptionCount( int t )
 	{
-		m_pD->m_options.SetNum(t);
+		m_pD->Options.SetNum(t);
 	}
 
 
 	//---------------------------------------------------------------------------------------------
 	NodeMeshPtr NodeMeshSwitch::GetOption( int t ) const
 	{
-		check( t>=0 && t<m_pD->m_options.Num() );
-		return m_pD->m_options[t].get();
+		check( t>=0 && t<m_pD->Options.Num() );
+		return m_pD->Options[t].get();
 	}
 
 
 	//---------------------------------------------------------------------------------------------
 	void NodeMeshSwitch::SetOption( int t, NodeMeshPtr pNode )
 	{
-		check( t>=0 && t<m_pD->m_options.Num() );
-		m_pD->m_options[t] = pNode;
+		check( t>=0 && t<m_pD->Options.Num() );
+		m_pD->Options[t] = pNode;
 	}
 
 }
diff --git a/Engine/Plugins/Mutable/Source/MutableTools/Private/MuT/NodeMeshSwitchPrivate.h b/Engine/Plugins/Mutable/Source/MutableTools/Private/MuT/NodeMeshSwitchPrivate.h
--- a/Engine/Plugins/Mutable/Source/MutableTools/Private/MuT/NodeMeshSwitchPrivate.h
+++ b/Engine/Plugins/Mutable/Source/MutableTools/Private/MuT/NodeMeshSwitchPrivate.h
@@ -18,7 +18,7 @@ namespace mu
 		static FNodeType s_type;
 
 		NodeScalarPtr m_pParameter;
-		TArray<NodeMeshPtr> m_options;
+		TArray<NodeMeshPtr> Options;
 
 	};
 
diff --git a/Engine/Plugins/Mutable/Source/MutableTools/Private/MuT/NodeProjectorConstant.cpp b/Engine/Plugins/Mutable/Source/MutableTools/Private/MuT/NodeProjectorConstant.cpp
--- a/Engine/Plugins/Mutable/Source/MutableTools/Private/MuT/NodeProjectorConstant.cpp
+++ b/Engine/Plugins/Mutable/Source/MutableTools/Private/MuT/NodeProjectorConstant.cpp
@@ -36,7 +36,7 @@ namespace mu
 		FVector3f* OutScale,
 		float* OutProjectionAngle ) const
 	{
-        if (OutType) *OutType = m_pD->m_type;
+        if (OutType) *OutType = m_pD->Type;
 
         if (OutPos) *OutPos = m_pD->m_position;
         if (OutDir) *OutDir = m_pD->m_direction;
@@ -50,7 +50,7 @@ namespace mu
 	//---------------------------------------------------------------------------------------------
     void NodeProjectorConstant::SetValue( EProjectorType type, FVector3f pos, FVector3f dir, FVector3f up, FVector3f scale, float projectionAngle )
 	{
-        m_pD->m_type = type;
+        m_pD->Type = type;
 		m_pD->m_position = pos;
 		m_pD->m_direction = dir;
 		m_pD->m_up = up;
diff --git a/Engine/Plugins/Mutable/Source/MutableTools/Private/MuT/NodeProjectorParameter.cpp b/Engine/Plugins/Mutable/Source/MutableTools/Private/MuT/NodeProjectorParameter.cpp
--- a/Engine/Plugins/Mutable/Source/MutableTools/Private/MuT/NodeProjectorParameter.cpp
+++ b/Engine/Plugins/Mutable/Source/MutableTools/Private/MuT/NodeProjectorParameter.cpp
@@ -34,19 +34,19 @@ namespace mu
 	//---------------------------------------------------------------------------------------------
 	void NodeProjectorParameter::SetName( const FString& strName )
 	{
-		m_pD->m_name = strName;
+		m_pD->Name = strName;
 	}
 
 
 	const FString& NodeProjectorParameter::GetUid() const
 	{
-		return m_pD->m_uid;
+		return m_pD->UID;
 	}
 
 
 	void NodeProjectorParameter::SetUid( const FString& strUid )
 	{
-		m_pD->m_uid = strUid;
+		m_pD->UID = strUid;
 	}
 
 
@@ -58,7 +58,7 @@ namespace mu
 		FVector3f scale,
 		float projectionAngle )
 	{
-        m_pD->m_type = type;
+        m_pD->Type = type;
         m_pD->m_position = pos;
 		m_pD->m_direction = dir;
 		m_pD->m_up = up;
@@ -71,17 +71,17 @@ namespace mu
     void NodeProjectorParameter::SetRangeCount( int i )
     {
         check(i>=0);
-        m_pD->m_ranges.SetNum(i);
+        m_pD->Ranges.SetNum(i);
     }
 
 
     //---------------------------------------------------------------------------------------------
     void NodeProjectorParameter::SetRange( int i, NodeRangePtr pRange )
     {
-        check( i>=0 && i<m_pD->m_ranges.Num() );
-        if ( i>=0 && i<m_pD->m_ranges.Num() )
+        check( i>=0 && i<m_pD->Ranges.Num() );
+        if ( i>=0 && i<m_pD->Ranges.Num() )
         {
-            m_pD->m_ranges[i] = pRange;
+            m_pD->Ranges[i] = pRange;
         }
     }
 
diff --git a/Engine/Plugins/Mutable/Source/MutableTools/Private/MuT/NodeRangeFromScalar.cpp b/Engine/Plugins/Mutable/Source/MutableTools/Private/MuT/NodeRangeFromScalar.cpp
--- a/Engine/Plugins/Mutable/Source/MutableTools/Private/MuT/NodeRangeFromScalar.cpp
+++ b/Engine/Plugins/Mutable/Source/MutableTools/Private/MuT/NodeRangeFromScalar.cpp
@@ -45,14 +45,14 @@ namespace mu
 	//---------------------------------------------------------------------------------------------
     const FString& NodeRangeFromScalar::GetName() const
 	{
-        return m_pD->m_name;
+        return m_pD->Name;
 	}
 
 
 	//---------------------------------------------------------------------------------------------
     void NodeRangeFromScalar::SetName( const FString& strName )
 	{
-        m_pD->m_name = strName;
+        m_pD->Name = strName;
 	}
 
 }
diff --git a/Engine/Plugins/Mutable/Source/MutableTools/Private/MuT/NodeRangeFromScalarPrivate.h b/Engine/Plugins/Mutable/Source/MutableTools/Private/MuT/NodeRangeFromScalarPrivate.h
--- a/Engine/Plugins/Mutable/Source/MutableTools/Private/MuT/NodeRangeFromScalarPrivate.h
+++ b/Engine/Plugins/Mutable/Source/MutableTools/Private/MuT/NodeRangeFromScalarPrivate.h
@@ -19,7 +19,7 @@ namespace mu
 		static FNodeType s_type;
 
         NodeScalarPtr m_pSize;
-        FString m_name;
+        FString Name;
 	};
 
 
diff --git a/Engine/Plugins/Mutable/Source/MutableTools/Private/MuT/NodeScalarConstant.cpp b/Engine/Plugins/Mutable/Source/MutableTools/Private/MuT/NodeScalarConstant.cpp
--- a/Engine/Plugins/Mutable/Source/MutableTools/Private/MuT/NodeScalarConstant.cpp
+++ b/Engine/Plugins/Mutable/Source/MutableTools/Private/MuT/NodeScalarConstant.cpp
@@ -29,13 +29,13 @@ namespace mu
 	//---------------------------------------------------------------------------------------------
 	float NodeScalarConstant::GetValue() const
 	{
-		return m_pD->m_value;
+		return m_pD->Value;
 	}
 
 	//---------------------------------------------------------------------------------------------
 	void NodeScalarConstant::SetValue( float v )
 	{
-		m_pD->m_value = v;
+		m_pD->Value = v;
 	}
 
 
diff --git a/Engine/Plugins/Mutable/Source/MutableTools/Private/MuT/NodeScalarConstantPrivate.h b/Engine/Plugins/Mutable/Source/MutableTools/Private/MuT/NodeScalarConstantPrivate.h
--- a/Engine/Plugins/Mutable/Source/MutableTools/Private/MuT/NodeScalarConstantPrivate.h
+++ b/Engine/Plugins/Mutable/Source/MutableTools/Private/MuT/NodeScalarConstantPrivate.h
@@ -16,7 +16,7 @@ namespace mu
 
 		static FNodeType s_type;
 
-		float m_value;
+		float Value;
 	};
 
 
diff --git a/Engine/Plugins/Mutable/Source/MutableTools/Private/MuT/NodeScalarEnumParameter.cpp b/Engine/Plugins/Mutable/Source/MutableTools/Private/MuT/NodeScalarEnumParameter.cpp
--- a/Engine/Plugins/Mutable/Source/MutableTools/Private/MuT/NodeScalarEnumParameter.cpp
+++ b/Engine/Plugins/Mutable/Source/MutableTools/Private/MuT/NodeScalarEnumParameter.cpp
@@ -27,64 +27,64 @@ namespace mu
 	//---------------------------------------------------------------------------------------------
 	const FString& NodeScalarEnumParameter::GetName() const
 	{
-		return m_pD->m_name;
+		return m_pD->Name;
 	}
 
 
 	//---------------------------------------------------------------------------------------------
 	void NodeScalarEnumParameter::SetName( const FString& strName )
 	{
-		m_pD->m_name = strName;
+		m_pD->Name = strName;
 	}
 
 
 	const FString& NodeScalarEnumParameter::GetUid() const
 	{
-		return m_pD->m_uid;
+		return m_pD->UID;
 	}
 
 
 	//---------------------------------------------------------------------------------------------
 	void NodeScalarEnumParameter::SetUid( const FString& strUid )
 	{
-		m_pD->m_uid = strUid;
+		m_pD->UID = strUid;
 	}
 
 
 	//---------------------------------------------------------------------------------------------
 	int NodeScalarEnumParameter::GetDefaultValueIndex() const
 	{
-		return m_pD->m_defaultValue;
+		return m_pD->DefaultValue;
 	}
 
 
 	//---------------------------------------------------------------------------------------------
 	void NodeScalarEnumParameter::SetDefaultValueIndex( int i )
 	{
-		m_pD->m_defaultValue = i;
+		m_pD->DefaultValue = i;
 	}
 
 
 	//---------------------------------------------------------------------------------------------
 	void NodeScalarEnumParameter::SetValueCount( int i )
 	{
-		m_pD->m_options.SetNum(i);
+		m_pD->Options.SetNum(i);
 	}
 
 
 	//---------------------------------------------------------------------------------------------
 	int NodeScalarEnumParameter::GetValueCount() const
 	{
-		return (int)m_pD->m_options.Num();
+		return (int)m_pD->Options.Num();
 	}
 
 
 	//---------------------------------------------------------------------------------------------
 	void NodeScalarEnumParameter::SetValue( int i, float value, const FString& strName )
 	{
-		check( i>=0 && i<(int)m_pD->m_options.Num() );
-		m_pD->m_options[i].name = strName;
-		m_pD->m_options[i].value = value;
+		check( i>=0 && i<(int)m_pD->Options.Num() );
+		m_pD->Options[i].name = strName;
+		m_pD->Options[i].value = value;
 	}
 
 
@@ -92,17 +92,17 @@ namespace mu
     void NodeScalarEnumParameter::SetRangeCount( int i )
     {
         check(i>=0);
-        m_pD->m_ranges.SetNum(i);
+        m_pD->Ranges.SetNum(i);
     }
 
 
     //---------------------------------------------------------------------------------------------
     void NodeScalarEnumParameter::SetRange( int i, NodeRangePtr pRange )
     {
-        check( i>=0 && i<int(m_pD->m_ranges.Num()) );
-        if ( i>=0 && i<int(m_pD->m_ranges.Num()) )
+        check( i>=0 && i<int(m_pD->Ranges.Num()) );
+        if ( i>=0 && i<int(m_pD->Ranges.Num()) )
         {
-            m_pD->m_ranges[i] = pRange;
+            m_pD->Ranges[i] = pRange;
         }
     }
 
diff --git a/Engine/Plugins/Mutable/Source/MutableTools/Private/MuT/NodeScalarParameter.cpp b/Engine/Plugins/Mutable/Source/MutableTools/Private/MuT/NodeScalarParameter.cpp
--- a/Engine/Plugins/Mutable/Source/MutableTools/Private/MuT/NodeScalarParameter.cpp
+++ b/Engine/Plugins/Mutable/Source/MutableTools/Private/MuT/NodeScalarParameter.cpp
@@ -29,21 +29,21 @@ namespace mu
 	//---------------------------------------------------------------------------------------------
 	void NodeScalarParameter::SetName( const FString& Name )
 	{
-		m_pD->m_name = Name;
+		m_pD->Name = Name;
 	}
 
 
 	//---------------------------------------------------------------------------------------------
 	void NodeScalarParameter::SetUid( const FString& Uid )
 	{
-		m_pD->m_uid = Uid;
+		m_pD->UID = Uid;
 	}
 
 
 	//---------------------------------------------------------------------------------------------
 	void NodeScalarParameter::SetDefaultValue( float v )
 	{
-		m_pD->m_defaultValue = v;
+		m_pD->DefaultValue = v;
 	}
 
 
@@ -51,17 +51,17 @@ namespace mu
     void NodeScalarParameter::SetRangeCount( int i )
     {
         check(i>=0);
-        m_pD->m_ranges.SetNum(i);
+        m_pD->Ranges.SetNum(i);
     }
 
 
     //---------------------------------------------------------------------------------------------
     void NodeScalarParameter::SetRange( int i, NodeRangePtr pRange )
     {
-        check( i>=0 && i<int(m_pD->m_ranges.Num()) );
-        if ( i>=0 && i<int(m_pD->m_ranges.Num()) )
+        check( i>=0 && i<int(m_pD->Ranges.Num()) );
+        if ( i>=0 && i<int(m_pD->Ranges.Num()) )
         {
-            m_pD->m_ranges[i] = pRange;
+            m_pD->Ranges[i] = pRange;
         }
     }
 
diff --git a/Engine/Plugins/Mutable/Source/MutableTools/Private/MuT/NodeStringConstant.cpp b/Engine/Plugins/Mutable/Source/MutableTools/Private/MuT/NodeStringConstant.cpp
--- a/Engine/Plugins/Mutable/Source/MutableTools/Private/MuT/NodeStringConstant.cpp
+++ b/Engine/Plugins/Mutable/Source/MutableTools/Private/MuT/NodeStringConstant.cpp
@@ -29,7 +29,7 @@ namespace mu
 	//---------------------------------------------------------------------------------------------
 	void NodeStringConstant::SetValue( const FString& v )
 	{
-		m_pD->m_value = v;
+		m_pD->Value = v;
 	}
 
 
diff --git a/Engine/Plugins/Mutable/Source/MutableTools/Private/MuT/NodeStringConstantPrivate.h b/Engine/Plugins/Mutable/Source/MutableTools/Private/MuT/NodeStringConstantPrivate.h
--- a/Engine/Plugins/Mutable/Source/MutableTools/Private/MuT/NodeStringConstantPrivate.h
+++ b/Engine/Plugins/Mutable/Source/MutableTools/Private/MuT/NodeStringConstantPrivate.h
@@ -17,7 +17,7 @@ namespace mu
 
 		static FNodeType s_type;
 
-		FString m_value;
+		FString Value;
 
 	};
 
diff --git a/Engine/Plugins/Mutable/Source/MutableTools/Private/MuT/NodeStringParameter.cpp b/Engine/Plugins/Mutable/Source/MutableTools/Private/MuT/NodeStringParameter.cpp
--- a/Engine/Plugins/Mutable/Source/MutableTools/Private/MuT/NodeStringParameter.cpp
+++ b/Engine/Plugins/Mutable/Source/MutableTools/Private/MuT/NodeStringParameter.cpp
@@ -29,21 +29,21 @@ namespace mu
 	//---------------------------------------------------------------------------------------------
 	void NodeStringParameter::SetName( const FString& Name )
 	{
-		m_pD->m_name = Name;
+		m_pD->Name = Name;
 	}
 
 
 	//---------------------------------------------------------------------------------------------
 	void NodeStringParameter::SetUid( const FString& Uid )
 	{
-		m_pD->m_uid = Uid;
+		m_pD->UID = Uid;
 	}
 
 
 	//---------------------------------------------------------------------------------------------
 	void NodeStringParameter::SetDefaultValue( const FString& v )
 	{
-		m_pD->m_defaultValue = v;
+		m_pD->DefaultValue = v;
 	}
 
 
@@ -51,17 +51,17 @@ namespace mu
     void NodeStringParameter::SetRangeCount( int i )
     {
         check(i>=0);
-        m_pD->m_ranges.SetNum(i);
+        m_pD->Ranges.SetNum(i);
     }
 
 
     //---------------------------------------------------------------------------------------------
     void NodeStringParameter::SetRange( int i, NodeRangePtr pRange )
     {
-        check( i>=0 && i<int(m_pD->m_ranges.Num()) );
-        if ( i>=0 && i<int(m_pD->m_ranges.Num()) )
+        check( i>=0 && i<int(m_pD->Ranges.Num()) );
+        if ( i>=0 && i<int(m_pD->Ranges.Num()) )
         {
-            m_pD->m_ranges[i] = pRange;
+            m_pD->Ranges[i] = pRange;
         }
     }
 
diff --git a/Engine/Plugins/Mutable/Source/MutableTools/Private/MuT/NodeStringParameterPrivate.h b/Engine/Plugins/Mutable/Source/MutableTools/Private/MuT/NodeStringParameterPrivate.h
--- a/Engine/Plugins/Mutable/Source/MutableTools/Private/MuT/NodeStringParameterPrivate.h
+++ b/Engine/Plugins/Mutable/Source/MutableTools/Private/MuT/NodeStringParameterPrivate.h
@@ -18,13 +18,13 @@ namespace mu
 
 		static FNodeType s_type;
 
-		FString m_defaultValue;
-		FString m_name;
-		FString m_uid;
+		FString DefaultValue;
+		FString Name;
+		FString UID;
 
         TArray<Ptr<NodeImage>> m_additionalImages;
 
-		TArray<Ptr<NodeRange>> m_ranges;
+		TArray<Ptr<NodeRange>> Ranges;
 
 	};
 
diff --git a/Engine/Plugins/Mutable/Source/MutableTools/Private/MuT/Table.cpp b/Engine/Plugins/Mutable/Source/MutableTools/Private/MuT/Table.cpp
--- a/Engine/Plugins/Mutable/Source/MutableTools/Private/MuT/Table.cpp
+++ b/Engine/Plugins/Mutable/Source/MutableTools/Private/MuT/Table.cpp
@@ -19,14 +19,14 @@ namespace mu
 
 
 	//---------------------------------------------------------------------------------------------
-	Table::Table()
+	FTable::FTable()
 	{
 		m_pD = new Private();
 	}
 
 
 	//---------------------------------------------------------------------------------------------
-	Table::~Table()
+	FTable::~FTable()
 	{
         check( m_pD );
 		delete m_pD;
@@ -34,25 +34,25 @@ namespace mu
 	}
 
 
-	Table::Private* Table::GetPrivate() const
+	FTable::Private* FTable::GetPrivate() const
 	{
 		return m_pD;
 	}
 
 
-	void Table::SetName( const FString& InName )
+	void FTable::SetName( const FString& InName )
 	{
 		m_pD->Name = InName;
 	}
 
 
-	const FString& Table::GetName() const
+	const FString& FTable::GetName() const
 	{
 		return m_pD->Name;
 	}
 
 
-	int32 Table::AddColumn(const FString& Name, ETableColumnType type )
+	int32 FTable::AddColumn(const FString& Name, ETableColumnType type )
 	{
 		int32 result = m_pD->Columns.Num();
 
@@ -73,7 +73,7 @@ namespace mu
 	}
 
 
-	int32 Table::FindColumn(const FString& Name ) const
+	int32 FTable::FindColumn(const FString& Name ) const
 	{
 		int32 res = -1;
 
@@ -89,7 +89,7 @@ namespace mu
 	}
 
 
-    void Table::AddRow( uint32 id )
+    void FTable::AddRow( uint32 id )
 	{
 		check( m_pD->FindRow(id)<0 );
 
@@ -100,7 +100,7 @@ namespace mu
 	}
 
 
-    void Table::SetCell( int32 Column, uint32 RowId, float Value, const void* ErrorContext)
+    void FTable::SetCell( int32 Column, uint32 RowId, float Value, const void* ErrorContext)
 	{
 		int32 Row = m_pD->FindRow(RowId);
 		check( Row>=0 );
@@ -113,7 +113,7 @@ namespace mu
 	}
 
 
-    void Table::SetCell( int32 Column, uint32 RowId, const FVector4f& Value, const void* ErrorContext)
+    void FTable::SetCell( int32 Column, uint32 RowId, const FVector4f& Value, const void* ErrorContext)
 	{
 		int32 Row = m_pD->FindRow(RowId);
 		check( Row>=0 );
@@ -126,7 +126,7 @@ namespace mu
 	}
 
 
-    void Table::SetCell(int32 Column, uint32 RowId, ResourceProxy<Image>* Value, const void* ErrorContext)
+    void FTable::SetCell(int32 Column, uint32 RowId, TResourceProxy<FImage>* Value, const void* ErrorContext)
 	{
 		int32 Row = m_pD->FindRow(RowId);
 		check( Row>=0 );
@@ -139,7 +139,7 @@ namespace mu
 	}
 
 
-    void Table::SetCell( int32 Column, uint32 RowId, Mesh* Value, const void* ErrorContext )
+    void FTable::SetCell( int32 Column, uint32 RowId, const TSharedPtr<FMesh>& Value, const void* ErrorContext )
 	{
 		int32 Row = m_pD->FindRow(RowId);
 		check( Row>=0 );
@@ -152,7 +152,7 @@ namespace mu
 	}
 
 
-    void Table::SetCell( int32 Column, uint32 RowId, const FString& Value, const void* ErrorContext )
+    void FTable::SetCell( int32 Column, uint32 RowId, const FString& Value, const void* ErrorContext )
 	{
 		int32 Row = m_pD->FindRow(RowId);
 		check( Row>=0 );
diff --git a/Engine/Plugins/Mutable/Source/MutableTools/Private/MuT/UnrealPixelFormatOverride.cpp b/Engine/Plugins/Mutable/Source/MutableTools/Private/MuT/UnrealPixelFormatOverride.cpp
--- a/Engine/Plugins/Mutable/Source/MutableTools/Private/MuT/UnrealPixelFormatOverride.cpp
+++ b/Engine/Plugins/Mutable/Source/MutableTools/Private/MuT/UnrealPixelFormatOverride.cpp
@@ -140,23 +140,27 @@ void FillBuildSettingsFromMutableFormat(FTextureBuildSettings& Settings, bool& b
 }
 
 
-void MutableToImageCore(const mu::Image* InMutable, FImage& CoreImage, int32 LOD)
+void MutableToImageCore(const mu::FImage* InMutable, FImage& CoreImage, int32 LOD)
 {
 	MUTABLE_CPUPROFILER_SCOPE(MutableToImageCore);
 
-	mu::Ptr<const mu::Image> Mutable = InMutable;
+	TSharedPtr<const mu::FImage> TempImage;
 
 	ERawImageFormat::Type CoreImageFormat;
-	switch (Mutable->GetFormat())
+	switch (InMutable->GetFormat())
 	{
-	case mu::EImageFormat::BGRA_UByte: CoreImageFormat = ERawImageFormat::BGRA8; break;
-		//case mu::EImageFormat::L_UByte: CoreImageFormat = ERawImageFormat::G8; break;
+	case mu::EImageFormat::BGRA_UByte: 
+		CoreImageFormat = ERawImageFormat::BGRA8; 
+		break;
+		
+	//case mu::EImageFormat::L_UByte: CoreImageFormat = ERawImageFormat::G8; break;
 
 	default:
 	{
 		// Unsupported format: force conversion
 		mu::FImageOperator ImOp = mu::FImageOperator::GetDefault(mu::FImageOperator::FImagePixelFormatFunc());
-		Mutable = ImOp.ImagePixelFormat(4, InMutable, mu::EImageFormat::BGRA_UByte,LOD);
+		TempImage = ImOp.ImagePixelFormat(4, InMutable, mu::EImageFormat::BGRA_UByte,LOD);
+		InMutable = TempImage.Get();
 		// We are extracting one LOD, so always access LOD 0 of the resulting mutable image 
 		LOD = 0;
 		CoreImageFormat = ERawImageFormat::BGRA8;
@@ -165,13 +169,13 @@ void MutableToImageCore(const mu::Image* InMutable, FImage& CoreImage, int32 LOD
 
 	}
 
-	FIntVector2 MipSize = Mutable->CalculateMipSize(LOD);
+	FIntVector2 MipSize = InMutable->CalculateMipSize(LOD);
 	CoreImage.Init(MipSize.X, MipSize.Y, CoreImageFormat, EGammaSpace::Linear);
-	FMemory::Memcpy(CoreImage.RawData.GetData(), Mutable->GetMipData(LOD), CoreImage.GetImageSizeBytes());
+	FMemory::Memcpy(CoreImage.RawData.GetData(), InMutable->GetMipData(LOD), CoreImage.GetImageSizeBytes());
 }
 
 
-bool ImageCoreToMutable(const FCompressedImage2D& Compressed, mu::Image* Mutable, int32 LOD)
+bool ImageCoreToMutable(const FCompressedImage2D& Compressed, mu::FImage* Mutable, int32 LOD)
 {
 	TArrayView<uint8> MutableView = Mutable->DataStorage.GetLOD(LOD);
 
@@ -251,7 +255,7 @@ mu::EImageFormat QualityAndPerformanceFix(mu::EImageFormat Format)
 }
 
 
-void UnrealPixelFormatFunc(bool& bOutSuccess, int32 Quality, mu::Image* Target, const mu::Image* Source, int32 OnlyLOD)
+void UnrealPixelFormatFunc(bool& bOutSuccess, int32 Quality, mu::FImage* Target, const mu::FImage* Source, int32 OnlyLOD)
 {
 	// If this fails, PrepareUnrealCompression wasn't called before.
 	check(STextureFormatManager);
