[mutable] More enforcement of coding standards.

diff --git a/Engine/Plugins/Mutable/Source/MutableTools/Public/MuT/Compiler.h b/Engine/Plugins/Mutable/Source/MutableTools/Public/MuT/Compiler.h
--- a/Engine/Plugins/Mutable/Source/MutableTools/Public/MuT/Compiler.h
+++ b/Engine/Plugins/Mutable/Source/MutableTools/Public/MuT/Compiler.h
@@ -14,12 +14,7 @@ namespace mu
 {
     // Forward declarations
     class Compiler;
-    typedef Ptr<Compiler> CompilerPtr;
-    typedef Ptr<const Compiler> CompilerPtrConst;
-
-    class Model;
-    using ModelPtr=Ptr<Model>;
-    using ModelPtrConst=Ptr<const Model>;
+    class FModel;
 
     class Transform;
     using TransformPtr=Ptr<Transform>;
@@ -32,7 +27,7 @@ namespace mu
     class Node;
     class NodeTransformedObject;
 
-	typedef TFunction<UE::Tasks::FTask(int32, TSharedPtr<Ptr<Image>>, bool)> FReferencedResourceFunc;
+	typedef TFunction<UE::Tasks::FTask(int32, TSharedPtr<TSharedPtr<FImage>>, bool)> FReferencedResourceFunc;
 
     //! \brief Options used to compile the models with a compiler.
     class MUTABLETOOLS_API CompilerOptions : public RefCounted
@@ -65,7 +60,7 @@ namespace mu
         void SetUseDiskCache( bool bEnabled );
 
 		//! Set the quality for the image compression algorithms. The level value is used internally
-		//! with System::SetImagecompressionQuality
+		//! with FSystem::SetImagecompressionQuality
 		void SetImageCompressionQuality(int32 Quality);
 
 		/** Set the image tiling strategy :
@@ -138,61 +133,20 @@ namespace mu
 		bool bOnlyFirstLOD = false;
 		ETextureCompressionStrategy TextureCompressionStrategy = ETextureCompressionStrategy::None;
 
-        void Serialise( OutputArchive& arch ) const
+        void Serialise( FOutputArchive& arch ) const
         {
-            const int32_t ver = 4;
-            arch << ver;
-
 			arch << FirstLOD;
 			arch << bOnlyFirstLOD;
             arch << TextureCompressionStrategy;
 			arch << NumExtraLODsToBuildAfterFirstLOD;
         }
 
-
-        void Unserialise( InputArchive& arch )
+        void Unserialise( FInputArchive& arch )
         {
-            int32_t ver = 0;
-            arch >> ver;
-			check(ver <= 4);
-
-			if (ver >= 2)
-			{
-				arch >> FirstLOD;
-			}
-			else
-			{
-				FirstLOD = 0;
-			}
-
+			arch >> FirstLOD;
             arch >> bOnlyFirstLOD;
-
-			if (ver >= 4)
-			{
-				arch >> TextureCompressionStrategy;
-			}
-			else
-			{
-				bool bAvoidRuntimeCompression;
-				arch >> bAvoidRuntimeCompression;
-				TextureCompressionStrategy = bAvoidRuntimeCompression ? ETextureCompressionStrategy::DontCompressRuntime : ETextureCompressionStrategy::None;
-			}
-
-			if (ver == 3)
-			{
-				int32 OldNumExtraLODsToBuildAfterFirstLOD;
-				arch >> OldNumExtraLODsToBuildAfterFirstLOD;
-				NumExtraLODsToBuildAfterFirstLOD = OldNumExtraLODsToBuildAfterFirstLOD;
-			}
-			else if (ver >= 4)
-			{
-				arch >> NumExtraLODsToBuildAfterFirstLOD;
-			}
-			else
-			{
-				NumExtraLODsToBuildAfterFirstLOD = 0;
-			}
-
+			arch >> TextureCompressionStrategy;
+			arch >> NumExtraLODsToBuildAfterFirstLOD;
 		}
     };
 
@@ -222,7 +176,7 @@ namespace mu
         Compiler( Ptr<CompilerOptions> Options=nullptr );
 
         //! Compile the expression into a run-time model.
-		TSharedPtr<Model> Compile( const Ptr<Node>& pNode );
+		TSharedPtr<FModel> Compile( const Ptr<Node>& pNode );
 
         //! Return the log of messages of all the compile operations executed so far.
         ErrorLogPtrConst GetLog() const;
diff --git a/Engine/Plugins/Mutable/Source/MutableTools/Public/MuT/NodeColourTable.h b/Engine/Plugins/Mutable/Source/MutableTools/Public/MuT/NodeColourTable.h
--- a/Engine/Plugins/Mutable/Source/MutableTools/Public/MuT/NodeColourTable.h
+++ b/Engine/Plugins/Mutable/Source/MutableTools/Public/MuT/NodeColourTable.h
@@ -19,7 +19,7 @@ namespace mu
 	public:
 
 		FString ParameterName;
-		Ptr<Table> Table;
+		Ptr<FTable> Table;
 		FString ColumnName;
 		bool bNoneOption = false;
 		FString DefaultRowName;
diff --git a/Engine/Plugins/Mutable/Source/MutableTools/Public/MuT/NodeExtensionDataConstant.h b/Engine/Plugins/Mutable/Source/MutableTools/Public/MuT/NodeExtensionDataConstant.h
--- a/Engine/Plugins/Mutable/Source/MutableTools/Public/MuT/NodeExtensionDataConstant.h
+++ b/Engine/Plugins/Mutable/Source/MutableTools/Public/MuT/NodeExtensionDataConstant.h
@@ -17,7 +17,7 @@ namespace mu
 	{
 	public:
 
-		Ptr<const ExtensionData> Value;
+		TSharedPtr<const FExtensionData> Value;
 
 	public:
 
@@ -32,7 +32,7 @@ namespace mu
 		//-----------------------------------------------------------------------------------------
 
 		/** Deprecated. Access Value attribute directly. */
-		void SetValue(const Ptr<const ExtensionData>& In) { Value = In; }
+		void SetValue(const TSharedPtr<const FExtensionData>& In) { Value = In; }
 
 	protected:
 
diff --git a/Engine/Plugins/Mutable/Source/MutableTools/Public/MuT/NodeImageConstant.h b/Engine/Plugins/Mutable/Source/MutableTools/Public/MuT/NodeImageConstant.h
--- a/Engine/Plugins/Mutable/Source/MutableTools/Public/MuT/NodeImageConstant.h
+++ b/Engine/Plugins/Mutable/Source/MutableTools/Public/MuT/NodeImageConstant.h
@@ -20,7 +20,7 @@ namespace mu
 	public:
 
 		/** */
-		Ptr<ResourceProxy<Image>> Proxy;
+		Ptr<TResourceProxy<FImage>> Proxy;
 		
 		/** */
 		FSourceDataDescriptor SourceDataDescriptor;
@@ -39,10 +39,10 @@ namespace mu
 		//-----------------------------------------------------------------------------------------
 
         //! Set the image to be output by this node
-        void SetValue( const Image* );
+        void SetValue(TSharedPtr<const FImage>);
 
         //! Set the image proxy that will provide the image for this node when necessary
-        void SetValue( Ptr<ResourceProxy<Image>> pImage );
+        void SetValue( Ptr<TResourceProxy<FImage>> );
 
 	protected:
 
diff --git a/Engine/Plugins/Mutable/Source/MutableTools/Public/MuT/NodeImageTable.h b/Engine/Plugins/Mutable/Source/MutableTools/Public/MuT/NodeImageTable.h
--- a/Engine/Plugins/Mutable/Source/MutableTools/Public/MuT/NodeImageTable.h
+++ b/Engine/Plugins/Mutable/Source/MutableTools/Public/MuT/NodeImageTable.h
@@ -21,7 +21,7 @@ namespace mu
 
 		FString ParameterName;
 
-		Ptr<Table> Table;
+		Ptr<FTable> Table;
 
 		FString ColumnName;
 
diff --git a/Engine/Plugins/Mutable/Source/MutableTools/Public/MuT/NodeLayout.h b/Engine/Plugins/Mutable/Source/MutableTools/Public/MuT/NodeLayout.h
--- a/Engine/Plugins/Mutable/Source/MutableTools/Public/MuT/NodeLayout.h
+++ b/Engine/Plugins/Mutable/Source/MutableTools/Public/MuT/NodeLayout.h
@@ -12,13 +12,13 @@
 
 namespace mu
 {
-	class Mesh;
+	class FMesh;
 
 	/** Data for a layout block before it is compiled. */
 	struct FSourceLayoutBlock
 	{
 		/** Optional mask image that selects the vertices to include in the block. */
-		Ptr<Image> Mask;
+		TSharedPtr<FImage> Mask;
 
 		FIntVector2 Min = { 0, 0 };
 		FIntVector2 Size = { 0, 0 };
@@ -70,8 +70,8 @@ namespace mu
 		/** Generate the blocks of a layout using the UV of the meshes. 
 		* A list of existing blocks may be provided to specify starting blocks that shouldn't be modified unless it is mandatory because they contain partial face islands. 
 		*/
-		void GenerateLayoutBlocks(const Ptr<Mesh>&, int32 LayoutIndex );
-		void GenerateLayoutBlocksFromUVIslands(const Ptr<Mesh>&, int32 LayoutIndex, bool bMergeChildBlocks);
+		void GenerateLayoutBlocks(const TSharedPtr<FMesh>&, int32 LayoutIndex );
+		void GenerateLayoutBlocksFromUVIslands(const TSharedPtr<FMesh>&, int32 LayoutIndex, bool bMergeChildBlocks);
 
 	protected:
 
diff --git a/Engine/Plugins/Mutable/Source/MutableTools/Public/MuT/NodeMesh.h b/Engine/Plugins/Mutable/Source/MutableTools/Public/MuT/NodeMesh.h
--- a/Engine/Plugins/Mutable/Source/MutableTools/Public/MuT/NodeMesh.h
+++ b/Engine/Plugins/Mutable/Source/MutableTools/Public/MuT/NodeMesh.h
@@ -13,7 +13,7 @@ namespace mu
 {
 
 	// Forward definitions
-	class Mesh;
+	class FMesh;
 	class NodeMesh;
 	typedef Ptr<NodeMesh> NodeMeshPtr;
 	typedef Ptr<const NodeMesh> NodeMeshPtrConst;
@@ -66,7 +66,7 @@ namespace mu
 
 	/** Create a map from vertices into vertices, collapsing vertices that have the same position.
 	*/
-	extern void MeshCreateCollapsedVertexMap(const mu::Mesh* Mesh, TArray<int32>& CollapsedVertices);
+	extern void MeshCreateCollapsedVertexMap(const mu::FMesh* Mesh, TArray<int32>& CollapsedVertices);
 
 }
 
diff --git a/Engine/Plugins/Mutable/Source/MutableTools/Public/MuT/NodeMeshConstant.h b/Engine/Plugins/Mutable/Source/MutableTools/Public/MuT/NodeMeshConstant.h
--- a/Engine/Plugins/Mutable/Source/MutableTools/Public/MuT/NodeMeshConstant.h
+++ b/Engine/Plugins/Mutable/Source/MutableTools/Public/MuT/NodeMeshConstant.h
@@ -7,7 +7,7 @@
 #include "MuT/Node.h"
 #include "MuT/NodeMesh.h"
 #include "MuT/NodeImage.h"
-
+#include "Templates/SharedPointer.h"
 
 namespace mu
 {
@@ -17,7 +17,7 @@ namespace mu
 	typedef Ptr<NodeMeshConstant> NodeMeshConstantPtr;
 	typedef Ptr<const NodeMeshConstant> NodeMeshConstantPtrConst;
 
-	class Mesh;
+	class FMesh;
 	class NodeLayout;
 
 	//! Node that outputs a constant mesh.
@@ -44,16 +44,16 @@ namespace mu
 		//-----------------------------------------------------------------------------------------
 
 		//! Get the constant mesh that will be returned.
-		Ptr<Mesh> GetValue() const;
+		TSharedPtr<FMesh> GetValue() const;
 
 		//! Set the constant mesh that will be returned.
-		void SetValue(Ptr<Mesh>);
+		void SetValue(TSharedPtr<FMesh>);
 
 		/** */
-		void AddMorph(const FString& Name, Ptr<Mesh>);
+		void AddMorph(const FString& Name, TSharedPtr<FMesh>);
 
 		/** */
-		Ptr<Mesh> FindMorph(const FString& Name) const;
+		TSharedPtr<FMesh> FindMorph(const FString& Name) const;
 
 		//! Get the number of layouts defined in this mesh.
 		int32 GetLayoutCount() const;
diff --git a/Engine/Plugins/Mutable/Source/MutableTools/Public/MuT/Table.h b/Engine/Plugins/Mutable/Source/MutableTools/Public/MuT/Table.h
--- a/Engine/Plugins/Mutable/Source/MutableTools/Public/MuT/Table.h
+++ b/Engine/Plugins/Mutable/Source/MutableTools/Public/MuT/Table.h
@@ -10,12 +10,9 @@
 namespace mu
 {
 	// Forward declarations
-	class Table;
-	typedef Ptr<Table> TablePtr;
-	typedef Ptr<const Table> TablePtrConst;
-
-	class Mesh;
-	class Image;
+	class FTable;
+	class FMesh;
+	class FImage;
 
 	/** Types of the values for the table cells. */
 	enum class ETableColumnType : uint32
@@ -33,11 +30,11 @@ namespace mu
 	//! colours, etc. for every column. It is useful to define a big number of similarly structured
 	//! objects, by using the NodeDatabase in a model expression.
 	//! \ingroup model
-	class MUTABLETOOLS_API Table : public RefCounted
+	class MUTABLETOOLS_API FTable : public RefCounted
 	{
 	public:
 
-		Table();
+		FTable();
 
 		//-----------------------------------------------------------------------------------------
 		// Own interface
@@ -59,8 +56,8 @@ namespace mu
 		//!
         void SetCell( int32 Column, uint32 RowId, float Value, const void* ErrorContext = nullptr);
         void SetCell( int32 Column, uint32 RowId, const FVector4f& Value, const void* ErrorContext = nullptr);
-		void SetCell( int32 Column, uint32 RowId, ResourceProxy<Image>* Value, const void* ErrorContext = nullptr);
-		void SetCell( int32 Column, uint32 RowId, Mesh* Value, const void* ErrorContext = nullptr);
+		void SetCell( int32 Column, uint32 RowId, TResourceProxy<FImage>* Value, const void* ErrorContext = nullptr);
+		void SetCell( int32 Column, uint32 RowId, const TSharedPtr<FMesh>& Value, const void* ErrorContext = nullptr);
         void SetCell( int32 Column, uint32 RowId, const FString& Value, const void* ErrorContext = nullptr);
 
 
@@ -73,7 +70,7 @@ namespace mu
 	protected:
 
 		//! Forbidden. Manage with the Ptr<> template.
-		~Table();
+		~FTable();
 
 	private:
 
diff --git a/Engine/Plugins/Mutable/Source/MutableTools/Public/MuT/UnrealPixelFormatOverride.h b/Engine/Plugins/Mutable/Source/MutableTools/Public/MuT/UnrealPixelFormatOverride.h
--- a/Engine/Plugins/Mutable/Source/MutableTools/Public/MuT/UnrealPixelFormatOverride.h
+++ b/Engine/Plugins/Mutable/Source/MutableTools/Public/MuT/UnrealPixelFormatOverride.h
@@ -17,10 +17,10 @@ MUTABLETOOLS_API extern void PrepareUnrealCompression();
 void FillBuildSettingsFromMutableFormat(FTextureBuildSettings& Settings, bool& bOutHasAlpha, mu::EImageFormat Format);
 
 /** */
-void MutableToImageCore(const mu::Image* InMutable, FImage& CoreImage, int32 LOD);
+void MutableToImageCore(const mu::FImage* InMutable, FImage& CoreImage, int32 LOD);
 
 /** */
-bool ImageCoreToMutable(const FCompressedImage2D& Compressed, mu::Image* Mutable, int32 LOD);
+bool ImageCoreToMutable(const FCompressedImage2D& Compressed, mu::FImage* Mutable, int32 LOD);
 
 /** Convert an Unreal platform pixel format to Mutable. */
 MUTABLETOOLS_API extern mu::EImageFormat UnrealToMutablePixelFormat(EPixelFormat PlatformFormat, bool bHasAlpha);
@@ -32,4 +32,4 @@ MUTABLETOOLS_API extern mu::EImageFormat QualityAndPerformanceFix(mu::EImageForm
 * Can be called from any thread.
 * PrepareUnrealCompression needs to be called from the game thread once prior to using this function. 
 */
-MUTABLETOOLS_API extern void UnrealPixelFormatFunc(bool& bOutSuccess, int32 Quality, mu::Image* Target, const mu::Image* Source, int32 OnlyLOD);
+MUTABLETOOLS_API extern void UnrealPixelFormatFunc(bool& bOutSuccess, int32 Quality, mu::FImage* Target, const mu::FImage* Source, int32 OnlyLOD);
diff --git a/Engine/Plugins/Mutable/Source/MutableValidation/Private/MuV/CustomizableObjectValidationCommandlet.cpp b/Engine/Plugins/Mutable/Source/MutableValidation/Private/MuV/CustomizableObjectValidationCommandlet.cpp
--- a/Engine/Plugins/Mutable/Source/MutableValidation/Private/MuV/CustomizableObjectValidationCommandlet.cpp
+++ b/Engine/Plugins/Mutable/Source/MutableValidation/Private/MuV/CustomizableObjectValidationCommandlet.cpp
@@ -112,7 +112,7 @@ int32 UCustomizableObjectValidationCommandlet::Main(const FString& Params)
 	
 	// GHet the total size of the streaming data of the model ---------------------------------------------- //
 	{
-		const TSharedPtr<const mu::Model> MutableModel = ToTestCustomizableObject->GetPrivate()->GetModel();
+		const TSharedPtr<const mu::FModel> MutableModel = ToTestCustomizableObject->GetPrivate()->GetModel();
 		check (MutableModel);
 
 		// Roms ---------------------- //
