[mutable] Fix handling of "overlay" layout strategy with extended meshes.
- add warnings if strategies don't match.
- use base mesh strategy in that case.

The code to extract mesh layout blocks to add has been moved to the time the modifier is applied to a mesh, and not earlier when the modifier was generated becase there wasn't enough information.

#rnx

diff --git a/Engine/Plugins/Mutable/Source/CustomizableObjectEditor/Private/MuCOE/CustomizableInstanceDetails.cpp b/Engine/Plugins/Mutable/Source/CustomizableObjectEditor/Private/MuCOE/CustomizableInstanceDetails.cpp
--- a/Engine/Plugins/Mutable/Source/CustomizableObjectEditor/Private/MuCOE/CustomizableInstanceDetails.cpp
+++ b/Engine/Plugins/Mutable/Source/CustomizableObjectEditor/Private/MuCOE/CustomizableInstanceDetails.cpp
@@ -2173,7 +2173,7 @@ void FCustomizableInstanceDetails::OnProjectorTextureParameterComboBoxChanged(TS
 {
 	if (Selection.IsValid())
 	{
-		FScopedTransaction LocalTransaction(LOCTEXT("SelectProjectorImage", "Select Projecto Image"));
+		FScopedTransaction LocalTransaction(LOCTEXT("SelectProjectorImage", "Select Projector Image"));
 		CustomInstance->Modify();
 		CustomInstance->SetIntParameterSelectedOption(ParamName, *Selection, RangeIndex);
 		UpdateInstance();
diff --git a/Engine/Plugins/Mutable/Source/CustomizableObjectEditor/Private/MuCOE/GenerateMutableSource/GenerateMutableSourceModifier.cpp b/Engine/Plugins/Mutable/Source/CustomizableObjectEditor/Private/MuCOE/GenerateMutableSource/GenerateMutableSourceModifier.cpp
--- a/Engine/Plugins/Mutable/Source/CustomizableObjectEditor/Private/MuCOE/GenerateMutableSource/GenerateMutableSourceModifier.cpp
+++ b/Engine/Plugins/Mutable/Source/CustomizableObjectEditor/Private/MuCOE/GenerateMutableSource/GenerateMutableSourceModifier.cpp
@@ -327,28 +327,6 @@ mu::Ptr<mu::NodeModifier> GenerateMutableSourceModifier(const UEdGraphPin * Pin,
 			{
 				mu::Ptr<mu::NodeMesh> MeshPtr = AddMeshNode;
 
-				const TArray<UCustomizableObjectLayout*> Layouts = TypedNodeExt->GetLayouts();
-
-				if (Layouts.Num())
-				{
-					mu::Ptr<mu::NodeMeshFragment> MeshFrag = new mu::NodeMeshFragment();
-
-					MeshFrag->SourceMesh = MeshPtr;
-					//TODO: Implement support for multiple UV channels (e.g. Add warning for vertices which have a block in a layout but not in the other)
-					MeshFrag->LayoutIndex = 0;
-
-					// For this case we don't want to create another layout: we will use the one defined in the mesh to be added since we want to add
-					// any block defined there.
-					//bool bWasEmpty = false;
-					//MeshFrag->Layout = CreateMutableLayoutNode(GenerationContext, Layouts[MeshFrag->LayoutIndex], true, bWasEmpty);
-
-					MeshPtr = MeshFrag;
-				}
-				else
-				{
-					GenerationContext.Log(LOCTEXT("ExtendMaterialLayoutMissing", "Skeletal Mesh without Layout Node linked to an Extend Material. A 4x4 layout will be added as default layout."), Node);
-				}
-
 				mu::Ptr<mu::NodeMeshFormat> MeshFormat = new mu::NodeMeshFormat();
 				SetSurfaceFormat(GenerationContext,
 					MeshFormat->GetVertexBuffers(), MeshFormat->GetIndexBuffers(), MeshData,
diff --git a/Engine/Plugins/Mutable/Source/MutableRuntime/Private/MuR/CodeRunner.cpp b/Engine/Plugins/Mutable/Source/MutableRuntime/Private/MuR/CodeRunner.cpp
--- a/Engine/Plugins/Mutable/Source/MutableRuntime/Private/MuR/CodeRunner.cpp
+++ b/Engine/Plugins/Mutable/Source/MutableRuntime/Private/MuR/CodeRunner.cpp
@@ -1907,8 +1907,8 @@ namespace mu
             FMemory::Memcpy( &source, data, sizeof(OP::ADDRESS) );
             data += sizeof(OP::ADDRESS);
 
-            uint16 layout;
-			FMemory::Memcpy( &layout, data, sizeof(uint16) );
+            uint16 LayoutIndex;
+			FMemory::Memcpy( &LayoutIndex, data, sizeof(uint16) );
             data += sizeof(uint16);
 
             uint16 blockCount;
@@ -1948,11 +1948,11 @@ namespace mu
 
 					if (blockCount > 0)
 					{
-						MeshExtractLayoutBlock(Result.Get(), Source.Get(), layout, blockCount, blocks, bOutSuccess);
+						MeshExtractLayoutBlock(Result.Get(), Source.Get(), LayoutIndex, blockCount, blocks, bOutSuccess);
 					}
 					else
 					{
-						MeshExtractLayoutBlock(Result.Get(), Source.Get(), layout, bOutSuccess);
+						MeshExtractLayoutBlock(Result.Get(), Source.Get(), LayoutIndex, bOutSuccess);
 					}
 
 					if (!bOutSuccess)
diff --git a/Engine/Plugins/Mutable/Source/MutableTools/Private/MuT/CodeGenerator.cpp b/Engine/Plugins/Mutable/Source/MutableTools/Private/MuT/CodeGenerator.cpp
--- a/Engine/Plugins/Mutable/Source/MutableTools/Private/MuT/CodeGenerator.cpp
+++ b/Engine/Plugins/Mutable/Source/MutableTools/Private/MuT/CodeGenerator.cpp
@@ -2053,6 +2053,7 @@ namespace mu
 					// Store the data necessary to apply modifiers for the pre-normal operations stage.
 					FMeshGenerationOptions MergedMeshOptions(Options);
 					MergedMeshOptions.ActiveTags = Edit->EnableTags; // TODO: Append to current?
+					MergedMeshOptions.bEnsureAllVerticesHaveLayoutBlock = false;
 
 					if (SharedMeshResults)
 					{
@@ -2063,6 +2064,51 @@ namespace mu
 					FMeshGenerationResult AddResults;
 					GenerateMesh(MergedMeshOptions, AddResults, pAdd);
 
+					// Warn about discrepancies on layout strategy between the added and the base
+					int32 LayoutIndexThatHasBlocks = -1;
+					{
+						if (BaseMeshResult.GeneratedLayouts.Num() != AddResults.GeneratedLayouts.Num())
+						{
+							// When extending a mesh section the added mesh sectio will use the layout strategy of the base one
+							FString Msg = FString::Printf(TEXT("Extended mesh section layout count is differenta than the mesh being extended."));
+							ErrorLog->GetPrivate()->Add(Msg, ELMT_WARNING, Edit->GetMessageContext(), ErrorContext);
+						}
+
+						for (int32 LayoutIndex = 0; LayoutIndex < BaseMeshResult.GeneratedLayouts.Num(); ++LayoutIndex)
+						{
+							if (BaseMeshResult.GeneratedLayouts[LayoutIndex].Layout
+								&&
+								BaseMeshResult.GeneratedLayouts[LayoutIndex].Layout->Strategy != EPackStrategy::Overlay)
+							{
+								LayoutIndexThatHasBlocks = LayoutIndex;
+							}
+
+							if (BaseMeshResult.GeneratedLayouts[LayoutIndex].Layout
+								&&
+								AddResults.GeneratedLayouts.IsValidIndex(LayoutIndex)
+								&&
+								AddResults.GeneratedLayouts[LayoutIndex].Layout
+								&&
+								BaseMeshResult.GeneratedLayouts[LayoutIndex].Layout->Strategy != AddResults.GeneratedLayouts[LayoutIndex].Layout->Strategy)
+							{
+								// When extending a mesh section the added mesh sectio will use the layout strategy of the base one
+								FString Msg = FString::Printf(TEXT("Extended mesh section layout [%d] is using a different strategy than the section being extended. The base strategy will be used."), LayoutIndex);
+								ErrorLog->GetPrivate()->Add(Msg, ELMT_WARNING, Edit->GetMessageContext(), ErrorContext);
+							}
+						}
+					}
+
+					// Add the operation to extract the relevant layout blocks if necessary
+					// TODO: Handle multiple layouts defining blocks: what to extract?
+					if (LayoutIndexThatHasBlocks>=0)
+					{
+						Ptr<ASTOpMeshExtractLayoutBlocks> ExtractOp = new ASTOpMeshExtractLayoutBlocks();
+						ExtractOp->Source = AddResults.MeshOp;
+						ExtractOp->LayoutIndex = LayoutIndexThatHasBlocks;
+
+						AddResults.MeshOp = ExtractOp;
+					}
+
 					// Apply the modifier for the post-normal operations stage to the added mesh
 					FMeshGenerationOptions ModifierOptions(Options);
 					ModifierOptions.ActiveTags = Edit->EnableTags;
diff --git a/Engine/Plugins/Mutable/Source/MutableTools/Private/MuT/CodeGenerator_Mesh.cpp b/Engine/Plugins/Mutable/Source/MutableTools/Private/MuT/CodeGenerator_Mesh.cpp
--- a/Engine/Plugins/Mutable/Source/MutableTools/Private/MuT/CodeGenerator_Mesh.cpp
+++ b/Engine/Plugins/Mutable/Source/MutableTools/Private/MuT/CodeGenerator_Mesh.cpp
@@ -132,8 +132,6 @@ namespace mu
 		const int32 NumVertices = Mesh->GetVertexCount();
 		const int32 NumBlocks = Layout->GetBlockCount();
 
-		bool bIsSingleFullBlock = (NumBlocks == 1) && (Layout->Blocks[0].Min == FIntVector2(0, 0) && Layout->Blocks[0].Size == Layout->Size);
-
 		// Find block ids for each block in the grid. Calculate a grid size that contains all blocks
 		FIntPoint LayoutGrid = Layout->GetGridSize();
 		FIntPoint WorkingGrid = LayoutGrid;
@@ -266,7 +264,7 @@ namespace mu
 		TMultiMap<int32, uint32> VertexToFaceMap;
 
 		// Find Unique Vertices
-		if (!bIsSingleFullBlock && MeshOptions.bClampUVIslands)
+		if (MeshOptions.bClampUVIslands)
 		{
 			VertexToFaceMap.Reserve(NumVertices);
 			Triangles.SetNumUninitialized(NumTriangles);
@@ -357,7 +355,7 @@ namespace mu
 			const uint16 BlockIndexV1 = AssignOneVertex(Index1);
 			const uint16 BlockIndexV2 = AssignOneVertex(Index2);
 
-			if (!bIsSingleFullBlock && MeshOptions.bClampUVIslands)
+			if (MeshOptions.bClampUVIslands)
 			{
 				if (BlockIndexV0 != BlockIndexV1 || BlockIndexV0 != BlockIndexV2)
 				{
