[Mutable] Do not set generated meshes to NM_DedicatedServer components.

#rnx

diff --git a/Engine/Plugins/Mutable/Source/CustomizableObject/Private/MuCO/CustomizableObjectInstanceUsage.cpp b/Engine/Plugins/Mutable/Source/CustomizableObject/Private/MuCO/CustomizableObjectInstanceUsage.cpp
--- a/Engine/Plugins/Mutable/Source/CustomizableObject/Private/MuCO/CustomizableObjectInstanceUsage.cpp
+++ b/Engine/Plugins/Mutable/Source/CustomizableObject/Private/MuCO/CustomizableObjectInstanceUsage.cpp
@@ -631,7 +631,12 @@ void UCustomizableObjectInstanceUsagePrivate::Tick(float DeltaTime)
 
 		USkeletalMesh* SkeletalMesh = nullptr;
 
-		const bool bInstanceGenerated = CustomizableObjectInstance->GetPrivate()->SkeletalMeshStatus == ESkeletalMeshStatus::Success;
+		bool bInstanceGenerated = CustomizableObjectInstance->GetPrivate()->SkeletalMeshStatus == ESkeletalMeshStatus::Success;
+		
+#if WITH_EDITOR
+		bInstanceGenerated = bInstanceGenerated && !IsNetMode(NM_DedicatedServer);
+#endif
+		
 		if (bInstanceGenerated)
 		{
 			// Generated SkeletalMesh to set, can be null if the component is empty
@@ -653,6 +658,10 @@ void UCustomizableObjectInstanceUsagePrivate::Tick(float DeltaTime)
 		{
 			SetSkeletalMeshAndOverrideMaterials(*Parent, SkeletalMesh, *CustomizableObjectInstance, nullptr, nullptr);
 		}
+		else
+		{
+			SetPendingSetSkeletalMesh(false);
+		}
 	}
 }
 
