#Mutable Disable a warning for now

diff --git a/Engine/Plugins/Mutable/Source/MutableTools/Private/MuT/CodeGenerator.cpp b/Engine/Plugins/Mutable/Source/MutableTools/Private/MuT/CodeGenerator.cpp
--- a/Engine/Plugins/Mutable/Source/MutableTools/Private/MuT/CodeGenerator.cpp
+++ b/Engine/Plugins/Mutable/Source/MutableTools/Private/MuT/CodeGenerator.cpp
@@ -2070,8 +2070,10 @@ namespace mu
 						if (BaseMeshResult.GeneratedLayouts.Num() != AddResults.GeneratedLayouts.Num())
 						{
 							// When extending a mesh section the added mesh sectio will use the layout strategy of the base one
+							/*
 							FString Msg = FString::Printf(TEXT("Extended mesh section layout count is differenta than the mesh being extended."));
 							ErrorLog->GetPrivate()->Add(Msg, ELMT_WARNING, Edit->GetMessageContext(), ErrorContext);
+							*/
 						}
 
 						for (int32 LayoutIndex = 0; LayoutIndex < BaseMeshResult.GeneratedLayouts.Num(); ++LayoutIndex)
@@ -2092,8 +2094,10 @@ namespace mu
 								BaseMeshResult.GeneratedLayouts[LayoutIndex].Layout->Strategy != AddResults.GeneratedLayouts[LayoutIndex].Layout->Strategy)
 							{
 								// When extending a mesh section the added mesh sectio will use the layout strategy of the base one
+								/*
 								FString Msg = FString::Printf(TEXT("Extended mesh section layout [%d] is using a different strategy than the section being extended. The base strategy will be used."), LayoutIndex);
 								ErrorLog->GetPrivate()->Add(Msg, ELMT_WARNING, Edit->GetMessageContext(), ErrorContext);
+								*/
 							}
 						}
 					}
