[mutable] Delay the cache key generation until PostLoad is complete for all related CustomizableObjects.
#rnx

diff --git a/Engine/Plugins/Mutable/Source/CustomizableObject/Internal/MuCO/CustomizableObjectPrivate.h b/Engine/Plugins/Mutable/Source/CustomizableObject/Internal/MuCO/CustomizableObjectPrivate.h
--- a/Engine/Plugins/Mutable/Source/CustomizableObject/Internal/MuCO/CustomizableObjectPrivate.h
+++ b/Engine/Plugins/Mutable/Source/CustomizableObject/Internal/MuCO/CustomizableObjectPrivate.h
@@ -1279,8 +1279,6 @@ public:
 
 	void SetModelStreamableBulkData(const TSharedPtr<FModelStreamableBulkData>& StreamableData, bool bIsCooking);
 #endif
-
-	USkeletalMesh* GetRefSkeletalMesh(const FName& ComponentName) const;
 	
 	TSharedPtr<FModelStreamableBulkData> GetModelStreamableBulkData(bool bIsCooking = false);
 
diff --git a/Engine/Plugins/Mutable/Source/CustomizableObject/Internal/MuCO/CustomizableObjectSystemPrivate.h b/Engine/Plugins/Mutable/Source/CustomizableObject/Internal/MuCO/CustomizableObjectSystemPrivate.h
--- a/Engine/Plugins/Mutable/Source/CustomizableObject/Internal/MuCO/CustomizableObjectSystemPrivate.h
+++ b/Engine/Plugins/Mutable/Source/CustomizableObject/Internal/MuCO/CustomizableObjectSystemPrivate.h
@@ -858,6 +858,11 @@ public:
 	/** Mutable default image provider. Used by the COIEditor and Instance/Descriptor APIs. */
 	UPROPERTY(Transient)
 	TObjectPtr<UEditorImageProvider> EditorImageProvider = nullptr;
+
+	/** List of objects pending to complete loading. */
+	UPROPERTY(Transient)
+	TArray<TObjectPtr<UCustomizableObject>> ObjectsPendingLoad;
+
 #endif
 
 	FLogBenchmarkUtil LogBenchmarkUtil;
diff --git a/Engine/Plugins/Mutable/Source/CustomizableObject/Internal/MuCO/ICustomizableObjectEditorModule.h b/Engine/Plugins/Mutable/Source/CustomizableObject/Internal/MuCO/ICustomizableObjectEditorModule.h
--- a/Engine/Plugins/Mutable/Source/CustomizableObject/Internal/MuCO/ICustomizableObjectEditorModule.h
+++ b/Engine/Plugins/Mutable/Source/CustomizableObject/Internal/MuCO/ICustomizableObjectEditorModule.h
@@ -71,7 +71,10 @@ public:
 
 	/** See GraphTraversal::GetRootObject(...) */
 	virtual const UCustomizableObject* GetRootObject(const UCustomizableObject* ChildObject) const = 0;
-	
+
+	/** Return all the CustomizableObjects related to the given one. */
+	virtual void GetRelatedObjects(UCustomizableObject*, TSet<UCustomizableObject*> OutRelated ) const = 0;
+
 	/**
 	 * Execute this method in order to bake the provided instance. It will schedule a special type of instance update before proceeding with the bake itself.
 	 * @param InTargetInstance The instance we want to bake
diff --git a/Engine/Plugins/Mutable/Source/CustomizableObject/Private/MuCO/CustomizableObject.cpp b/Engine/Plugins/Mutable/Source/CustomizableObject/Private/MuCO/CustomizableObject.cpp
--- a/Engine/Plugins/Mutable/Source/CustomizableObject/Private/MuCO/CustomizableObject.cpp
+++ b/Engine/Plugins/Mutable/Source/CustomizableObject/Private/MuCO/CustomizableObject.cpp
@@ -370,16 +370,7 @@ void UCustomizableObject::PostLoad()
 	if (!IsRunningCookCommandlet())
 	{
 		GetPrivate()->Status.NextState(FCustomizableObjectStatusTypes::EState::Loading);
-	
-		const FAssetRegistryModule& AssetRegistryModule = FModuleManager::LoadModuleChecked<FAssetRegistryModule>("AssetRegistry");
-		if (AssetRegistryModule.Get().IsLoadingAssets())
-		{
-			AssetRegistryModule.Get().OnFilesLoaded().AddUObject(GetPrivate(), &UCustomizableObjectPrivate::LoadCompiledDataFromDisk);
-		}
-		else
-		{
-			GetPrivate()->LoadCompiledDataFromDisk();
-		}
+		UCustomizableObjectSystem::GetInstance()->AddPendingLoad(this);
 	}
 #endif
 }
diff --git a/Engine/Plugins/Mutable/Source/CustomizableObject/Private/MuCO/CustomizableObjectSystem.cpp b/Engine/Plugins/Mutable/Source/CustomizableObject/Private/MuCO/CustomizableObjectSystem.cpp
--- a/Engine/Plugins/Mutable/Source/CustomizableObject/Private/MuCO/CustomizableObjectSystem.cpp
+++ b/Engine/Plugins/Mutable/Source/CustomizableObject/Private/MuCO/CustomizableObjectSystem.cpp
@@ -1361,6 +1361,14 @@ void UCustomizableObjectSystem::ClearResourceCacheProtected()
 
 
 #if WITH_EDITOR
+
+void UCustomizableObjectSystem::AddPendingLoad(UCustomizableObject* CO)
+{
+	check(IsInGameThread());
+	GetPrivate()->ObjectsPendingLoad.AddUnique(CO);
+}
+
+
 bool UCustomizableObjectSystem::LockObject(UCustomizableObject* InObject)
 {
 	check(InObject != nullptr);
@@ -3619,6 +3627,36 @@ int32 UCustomizableObjectSystem::TickInternal(const bool bBlocking)
 	{
 		return 0;
 	}
+
+	// See if any COs pending to load can be completed.
+	for ( TArray<TObjectPtr<UCustomizableObject>>::TIterator It = GetPrivate()->ObjectsPendingLoad.CreateIterator(); It; ++It )
+	{
+		UCustomizableObject* CO = (*It);
+
+		bool bReady = true;
+
+		if (ICustomizableObjectEditorModule* Module = ICustomizableObjectEditorModule::Get())
+		{
+			TSet<UCustomizableObject*> RelatedCustomizableObjects;
+			Module->GetRelatedObjects(CO, RelatedCustomizableObjects);
+
+			for (const UObject* CustObject : RelatedCustomizableObjects)
+			{
+				if (CustObject && !CustObject->HasAnyFlags(EObjectFlags::RF_LoadCompleted))
+				{
+					bReady = false;
+					break;
+				}
+			}
+		}
+
+		if (bReady)
+		{
+			CO->GetPrivate()->LoadCompiledDataFromDisk();
+			It.RemoveCurrent();
+		}
+	}
+
 #endif
 
 	Private->UpdateStats();
@@ -4392,4 +4430,3 @@ int32 UCustomizableObjectSystemPrivate::BlockTillAllRequestsFinished(float TimeL
 
 	return 0;
 }
-
diff --git a/Engine/Plugins/Mutable/Source/CustomizableObject/Public/MuCO/CustomizableObjectSystem.h b/Engine/Plugins/Mutable/Source/CustomizableObject/Public/MuCO/CustomizableObjectSystem.h
--- a/Engine/Plugins/Mutable/Source/CustomizableObject/Public/MuCO/CustomizableObjectSystem.h
+++ b/Engine/Plugins/Mutable/Source/CustomizableObject/Public/MuCO/CustomizableObjectSystem.h
@@ -209,6 +209,12 @@ public:
 	void ClearResourceCacheProtected();
 
 #if WITH_EDITOR
+
+	/** Add a CO to the pending - load list.COs need to wait until all related objects are fully loaded before being able
+	 * to do things like check-if-up-to-date or compile.
+	 */
+	void AddPendingLoad(UCustomizableObject*);
+
 	// Lock a CustomizableObjects, preventing the generation or update of any of its instances
 	// Will return true if successful, false if it fails to lock because an update is already underway
 	// This is usually only used in the editor
diff --git a/Engine/Plugins/Mutable/Source/CustomizableObjectEditor/Private/MuCOE/CustomizableObjectEditorModule.cpp b/Engine/Plugins/Mutable/Source/CustomizableObjectEditor/Private/MuCOE/CustomizableObjectEditorModule.cpp
--- a/Engine/Plugins/Mutable/Source/CustomizableObjectEditor/Private/MuCOE/CustomizableObjectEditorModule.cpp
+++ b/Engine/Plugins/Mutable/Source/CustomizableObjectEditor/Private/MuCOE/CustomizableObjectEditorModule.cpp
@@ -687,6 +687,13 @@ const UCustomizableObject* FCustomizableObjectEditorModule::GetRootObject(const
 }
 
 
+void FCustomizableObjectEditorModule::GetRelatedObjects(UCustomizableObject* CO, TSet<UCustomizableObject*> OutRelated) const
+{
+	UCustomizableObject* RootObject = GraphTraversal::GetRootObject(CO);
+	GetAllObjectsInGraph(RootObject, OutRelated);
+}
+
+
 void FCustomizableObjectEditorModule::BakeCustomizableObjectInstance(UCustomizableObjectInstance* InTargetInstance, const FBakingConfiguration& InBakingConfig)
 {
 	UCustomizableObjectInstanceBaker* InstanceBaker = NewObject<UCustomizableObjectInstanceBaker>();
diff --git a/Engine/Plugins/Mutable/Source/CustomizableObjectEditor/Private/MuCOE/CustomizableObjectEditorModule.h b/Engine/Plugins/Mutable/Source/CustomizableObjectEditor/Private/MuCOE/CustomizableObjectEditorModule.h
--- a/Engine/Plugins/Mutable/Source/CustomizableObjectEditor/Private/MuCOE/CustomizableObjectEditorModule.h
+++ b/Engine/Plugins/Mutable/Source/CustomizableObjectEditor/Private/MuCOE/CustomizableObjectEditorModule.h
@@ -43,6 +43,7 @@ public:
 	virtual FString GetCurrentReleaseVersionForObject(const UCustomizableObject& Object) const override;
 	virtual UCustomizableObject* GetRootObject(UCustomizableObject* ChildObject) const override;
 	virtual const UCustomizableObject* GetRootObject(const UCustomizableObject* ChildObject) const override;
+	virtual void GetRelatedObjects(UCustomizableObject*, TSet<UCustomizableObject*> OutRelated) const override;
 	virtual void BakeCustomizableObjectInstance(UCustomizableObjectInstance* InTargetInstance, const FBakingConfiguration& InBakingConfig) override;
 	virtual USkeletalMesh* GetReferenceSkeletalMesh(const UCustomizableObject& Object, const FName& ComponentName) const override;
 	virtual TMap<FName, FGuid> GetParticipatingObjects(const UCustomizableObject* Object, bool bLoadObjects, const FCompilationOptions* Options = nullptr) const override;
