[mutable] Fix crash when trying to share layouts with mesh that doesn't have them.
#rnx

diff --git a/Engine/Plugins/Mutable/Source/CustomizableObjectEditor/Private/MuCOE/CustomizableObjectInstanceEditor.cpp b/Engine/Plugins/Mutable/Source/CustomizableObjectEditor/Private/MuCOE/CustomizableObjectInstanceEditor.cpp
--- a/Engine/Plugins/Mutable/Source/CustomizableObjectEditor/Private/MuCOE/CustomizableObjectInstanceEditor.cpp
+++ b/Engine/Plugins/Mutable/Source/CustomizableObjectEditor/Private/MuCOE/CustomizableObjectInstanceEditor.cpp
@@ -676,7 +676,10 @@ void FCustomizableObjectInstanceEditor::OnCustomizableObjectStatusChanged(FCusto
 	{
 		if (CurrentState == FCustomizableObjectStatusTypes::EState::ModelLoaded)
 		{
-			Viewport->CreatePreviewActor(CustomizableObjectInstance);
+			if (Viewport)
+			{
+				Viewport->CreatePreviewActor(CustomizableObjectInstance);
+			}
 			CustomizableObjectInstance->UpdateSkeletalMeshAsync(true, true);
 		}
 		else if (CurrentState == FCustomizableObjectStatusTypes::EState::NoModel)
diff --git a/Engine/Plugins/Mutable/Source/CustomizableObjectEditor/Private/MuCOE/GenerateMutableSource/GenerateMutableSourceMesh.cpp b/Engine/Plugins/Mutable/Source/CustomizableObjectEditor/Private/MuCOE/GenerateMutableSource/GenerateMutableSourceMesh.cpp
--- a/Engine/Plugins/Mutable/Source/CustomizableObjectEditor/Private/MuCOE/GenerateMutableSource/GenerateMutableSourceMesh.cpp
+++ b/Engine/Plugins/Mutable/Source/CustomizableObjectEditor/Private/MuCOE/GenerateMutableSource/GenerateMutableSourceMesh.cpp
@@ -375,7 +375,7 @@ bool IsSkeletalMeshCompatibleWithRefSkeleton(FMutableComponentInfo& ComponentInf
 			// Different skeletons can't be used if they are incompatible with the reference skeleton.
 			OutErrorMessage = FString::Printf(
 				TEXT("The SkeletalMesh [%s] with Skeleton [%s] is incompatible with the reference mesh [%s] which has [%s]. "
-					"Bone [%s] has a differnt parent on the Skeleton from the reference mesh."),
+					"Bone [%s] has a different parent on the Skeleton from the reference mesh."),
 				*InSkeletalMesh->GetName(), *Skeleton->GetName(),
 				*ComponentInfo.RefSkeletalMesh->GetName(), *ComponentInfo.RefSkeleton->GetName(),
 				*Bone.ExportName);
diff --git a/Engine/Plugins/Mutable/Source/MutableTools/Private/MuT/CodeGenerator.cpp b/Engine/Plugins/Mutable/Source/MutableTools/Private/MuT/CodeGenerator.cpp
--- a/Engine/Plugins/Mutable/Source/MutableTools/Private/MuT/CodeGenerator.cpp
+++ b/Engine/Plugins/Mutable/Source/MutableTools/Private/MuT/CodeGenerator.cpp
@@ -978,9 +978,14 @@ namespace mu
 
 				SurfaceReferenceLayouts[LayoutIndex] = MeshResults.GeneratedLayouts[LayoutIndex];
 
-				if (SharedMeshResults)
+				bool bSharedHasThisLayout = SharedMeshResults 
+					&& 
+					SharedMeshResults->LayoutOps.IsValidIndex(LayoutIndex)
+					&&
+					SharedMeshResults->LayoutOps[LayoutIndex];
+
+				if (bSharedHasThisLayout)
 				{
-					check(SharedMeshResults->LayoutOps.IsValidIndex(LayoutIndex));
 					SurfaceLayoutOps[LayoutIndex] = SharedMeshResults->LayoutOps[LayoutIndex];
 				}
 				else
