[mutable] Fix wrong detection of related objects when checking of ready to generate DDC key.
#rnx

diff --git a/Engine/Plugins/Mutable/Source/CustomizableObject/Internal/MuCO/ICustomizableObjectEditorModule.h b/Engine/Plugins/Mutable/Source/CustomizableObject/Internal/MuCO/ICustomizableObjectEditorModule.h
--- a/Engine/Plugins/Mutable/Source/CustomizableObject/Internal/MuCO/ICustomizableObjectEditorModule.h
+++ b/Engine/Plugins/Mutable/Source/CustomizableObject/Internal/MuCO/ICustomizableObjectEditorModule.h
@@ -73,7 +73,7 @@ public:
 	virtual const UCustomizableObject* GetRootObject(const UCustomizableObject* ChildObject) const = 0;
 
 	/** Return all the CustomizableObjects related to the given one. */
-	virtual void GetRelatedObjects(UCustomizableObject*, TSet<UCustomizableObject*> OutRelated ) const = 0;
+	virtual void GetRelatedObjects(UCustomizableObject*, TSet<UCustomizableObject*>& OutRelated ) const = 0;
 
 	/**
 	 * Execute this method in order to bake the provided instance. It will schedule a special type of instance update before proceeding with the bake itself.
diff --git a/Engine/Plugins/Mutable/Source/CustomizableObjectEditor/Private/MuCOE/CustomizableObjectEditorModule.cpp b/Engine/Plugins/Mutable/Source/CustomizableObjectEditor/Private/MuCOE/CustomizableObjectEditorModule.cpp
--- a/Engine/Plugins/Mutable/Source/CustomizableObjectEditor/Private/MuCOE/CustomizableObjectEditorModule.cpp
+++ b/Engine/Plugins/Mutable/Source/CustomizableObjectEditor/Private/MuCOE/CustomizableObjectEditorModule.cpp
@@ -687,7 +687,7 @@ const UCustomizableObject* FCustomizableObjectEditorModule::GetRootObject(const
 }
 
 
-void FCustomizableObjectEditorModule::GetRelatedObjects(UCustomizableObject* CO, TSet<UCustomizableObject*> OutRelated) const
+void FCustomizableObjectEditorModule::GetRelatedObjects(UCustomizableObject* CO, TSet<UCustomizableObject*>& OutRelated) const
 {
 	UCustomizableObject* RootObject = GraphTraversal::GetRootObject(CO);
 	GetAllObjectsInGraph(RootObject, OutRelated);
diff --git a/Engine/Plugins/Mutable/Source/CustomizableObjectEditor/Private/MuCOE/CustomizableObjectEditorModule.h b/Engine/Plugins/Mutable/Source/CustomizableObjectEditor/Private/MuCOE/CustomizableObjectEditorModule.h
--- a/Engine/Plugins/Mutable/Source/CustomizableObjectEditor/Private/MuCOE/CustomizableObjectEditorModule.h
+++ b/Engine/Plugins/Mutable/Source/CustomizableObjectEditor/Private/MuCOE/CustomizableObjectEditorModule.h
@@ -43,7 +43,7 @@ public:
 	virtual FString GetCurrentReleaseVersionForObject(const UCustomizableObject& Object) const override;
 	virtual UCustomizableObject* GetRootObject(UCustomizableObject* ChildObject) const override;
 	virtual const UCustomizableObject* GetRootObject(const UCustomizableObject* ChildObject) const override;
-	virtual void GetRelatedObjects(UCustomizableObject*, TSet<UCustomizableObject*> OutRelated) const override;
+	virtual void GetRelatedObjects(UCustomizableObject*, TSet<UCustomizableObject*>& OutRelated) const override;
 	virtual void BakeCustomizableObjectInstance(UCustomizableObjectInstance* InTargetInstance, const FBakingConfiguration& InBakingConfig) override;
 	virtual USkeletalMesh* GetReferenceSkeletalMesh(const UCustomizableObject& Object, const FName& ComponentName) const override;
 	virtual TMap<FName, FGuid> GetParticipatingObjects(const UCustomizableObject* Object, bool bLoadObjects, const FCompilationOptions* Options = nullptr) const override;
