[mutable] Handle access to instance parameters as if they were multi-dimensional when they actually are not. It may happen because of optimization.
#rnx

diff --git a/Engine/Plugins/Mutable/Source/CustomizableObject/Private/MuCO/CustomizableObjectInstanceDescriptor.cpp b/Engine/Plugins/Mutable/Source/CustomizableObject/Private/MuCO/CustomizableObjectInstanceDescriptor.cpp
--- a/Engine/Plugins/Mutable/Source/CustomizableObject/Private/MuCO/CustomizableObjectInstanceDescriptor.cpp
+++ b/Engine/Plugins/Mutable/Source/CustomizableObject/Private/MuCO/CustomizableObjectInstanceDescriptor.cpp
@@ -1321,15 +1321,15 @@ const FString& FCustomizableObjectInstanceDescriptor::GetIntParameterSelectedOpt
 
 	if (ParameterIndexInObject >= 0 && IntParameters.IsValidIndex(ParameterIndexInInstance))
 	{
-		if (RangeIndex == -1)
+		// Due to optimizations a parameter may lose its multidimensionality (if it becomes constant). 
+		// In that case it means there is only one possible option so it is ok to set it as if wasn't multidimensional.
+		bool bIsMultidimensional = CustomizableObject->IsParameterMultidimensional(ParameterIndexInObject);
+		if (!bIsMultidimensional || RangeIndex == INDEX_NONE)
 		{
-			check(!CustomizableObject->IsParameterMultidimensional(ParameterIndexInObject)); // This param is multidimensional, it must have a RangeIndex of 0 or more
 			return IntParameters[ParameterIndexInInstance].ParameterValueName;
 		}
 		else
 		{
-			check(CustomizableObject->IsParameterMultidimensional(ParameterIndexInObject)); // This param is not multidimensional, it must have a RangeIndex of -1
-
 			if (IntParameters[ParameterIndexInInstance].ParameterRangeValueNames.IsValidIndex(RangeIndex))
 			{
 				return IntParameters[ParameterIndexInInstance].ParameterRangeValueNames[RangeIndex];
@@ -1378,17 +1378,15 @@ void FCustomizableObjectInstanceDescriptor::SetIntParameterSelectedOption(const
 		return;
 	}
 
-	if (RangeIndex == -1)
+	// Due to optimizations a parameter may lose its multidimensionality (if it becomes constant). 
+	// In that case it means there is only one possible option so it is ok to set it as if wasn't multidimensional.
+	bool bIsMultidimensional = CustomizableObject->IsParameterMultidimensional(ParameterIndexInObject);
+	if (!bIsMultidimensional || RangeIndex == INDEX_NONE)
 	{
-		// If this param were multidimensional, it must have a RangeIndex of 0 or more
-		check(!CustomizableObject->IsParameterMultidimensional(ParameterIndexInObject));
 		IntParameters[ParameterIndexInInstance].ParameterValueName = SelectedOption;
 	}
 	else
 	{
-		// If this param were not multidimensional, it must have a RangeIndex of -1
-		check(CustomizableObject->IsParameterMultidimensional(ParameterIndexInObject));
-
 		if (!IntParameters[ParameterIndexInInstance].ParameterRangeValueNames.IsValidIndex(RangeIndex))
 		{
 			const int32 InsertionIndex = IntParameters[ParameterIndexInInstance].ParameterRangeValueNames.Num();
@@ -1437,15 +1435,15 @@ float FCustomizableObjectInstanceDescriptor::GetFloatParameterSelectedOption(con
 
 	if (ParameterIndexInObject >= 0 && FloatParamIndex >= 0)
 	{
-		if (RangeIndex == -1)
+		// Due to optimizations a parameter may lose its multidimensionality (if it becomes constant). 
+		// In that case it means there is only one possible option so it is ok to set it as if wasn't multidimensional.
+		bool bIsMultidimensional = CustomizableObject->IsParameterMultidimensional(ParameterIndexInObject);
+		if (!bIsMultidimensional || RangeIndex == INDEX_NONE)
 		{
-			check(!CustomizableObject->IsParameterMultidimensional(ParameterIndexInObject)); // This param is multidimensional, it must have a RangeIndex of 0 or more
 			return FloatParameters[FloatParamIndex].ParameterValue;
 		}
 		else
 		{
-			check(CustomizableObject->IsParameterMultidimensional(ParameterIndexInObject)); // This param is not multidimensional, it must have a RangeIndex of -1
-
 			if (FloatParameters[FloatParamIndex].ParameterRangeValues.IsValidIndex(RangeIndex))
 			{
 				return FloatParameters[FloatParamIndex].ParameterRangeValues[RangeIndex];
@@ -1479,15 +1477,15 @@ void FCustomizableObjectInstanceDescriptor::SetFloatParameterSelectedOption(cons
 		return;
 	}
 
-	if (RangeIndex == INDEX_NONE)
+	// Due to optimizations a parameter may lose its multidimensionality (if it becomes constant). 
+	// In that case it means there is only one possible option so it is ok to set it as if wasn't multidimensional.
+	bool bIsMultidimensional = CustomizableObject->IsParameterMultidimensional(ParameterIndexInObject);
+	if (!bIsMultidimensional || RangeIndex == INDEX_NONE)
 	{
-		check(!CustomizableObject->IsParameterMultidimensional(ParameterIndexInObject)); // This param is multidimensional, it must have a RangeIndex of 0 or more
 		FloatParameters[ParameterIndexInInstance].ParameterValue = FloatValue;
 	}
 	else
 	{
-		check(CustomizableObject->IsParameterMultidimensional(ParameterIndexInObject)); // This param is not multidimensional, it must have a RangeIndex of -1
-
 		if (!FloatParameters[ParameterIndexInInstance].ParameterRangeValues.IsValidIndex(RangeIndex))
 		{
 			const int32 InsertionIndex = FloatParameters[ParameterIndexInInstance].ParameterRangeValues.Num();
@@ -1515,15 +1513,15 @@ FName FCustomizableObjectInstanceDescriptor::GetTextureParameterSelectedOption(c
 
 	if (ParameterIndexInObject >= 0 && ParameterIndexInInstance >= 0)
 	{
-		if (RangeIndex == -1)
+		// Due to optimizations a parameter may lose its multidimensionality (if it becomes constant). 
+		// In that case it means there is only one possible option so it is ok to set it as if wasn't multidimensional.
+		bool bIsMultidimensional = CustomizableObject->IsParameterMultidimensional(ParameterIndexInObject);
+		if (!bIsMultidimensional || RangeIndex == INDEX_NONE)
 		{
-			check(!CustomizableObject->IsParameterMultidimensional(ParameterIndexInObject)); // This param is multidimensional, it must have a RangeIndex of 0 or more
 			return TextureParameters[ParameterIndexInInstance].ParameterValue;
 		}
 		else
 		{
-			check(CustomizableObject->IsParameterMultidimensional(ParameterIndexInObject)); // This param is not multidimensional, it must have a RangeIndex of -1
-
 			if (TextureParameters[ParameterIndexInInstance].ParameterRangeValues.IsValidIndex(RangeIndex))
 			{
 				return TextureParameters[ParameterIndexInInstance].ParameterRangeValues[RangeIndex];
@@ -1557,15 +1555,15 @@ void FCustomizableObjectInstanceDescriptor::SetTextureParameterSelectedOption(co
 		return;
 	}
 
-	if (RangeIndex == INDEX_NONE)
+	// Due to optimizations a parameter may lose its multidimensionality (if it becomes constant). 
+	// In that case it means there is only one possible option so it is ok to set it as if wasn't multidimensional.
+	bool bIsMultidimensional = CustomizableObject->IsParameterMultidimensional(ParameterIndexInObject);
+	if (!bIsMultidimensional || RangeIndex == INDEX_NONE)
 	{
-		check(!CustomizableObject->IsParameterMultidimensional(ParameterIndexInObject)); // This param is multidimensional, it must have a RangeIndex of 0 or more
 		TextureParameters[ParameterIndexInInstance].ParameterValue = FName(TextureValue);
 	}
 	else
 	{
-		check(CustomizableObject->IsParameterMultidimensional(ParameterIndexInObject)); // This param is not multidimensional, it must have a RangeIndex of -1
-
 		if (!TextureParameters[ParameterIndexInInstance].ParameterRangeValues.IsValidIndex(RangeIndex))
 		{
 			const int32 InsertionIndex = TextureParameters[ParameterIndexInInstance].ParameterRangeValues.Num();
@@ -1596,6 +1594,7 @@ FLinearColor FCustomizableObjectInstanceDescriptor::GetColorParameterSelectedOpt
 		return FCustomizableObjectVectorParameterValue::DEFAULT_PARAMETER_VALUE;
 	}
 
+	// TODO: Multidimensional parameter support
 	return VectorParameters.IsValidIndex(ColorParamIndex) ? VectorParameters[ColorParamIndex].ParameterValue : FLinearColor();
 }
 
@@ -1610,6 +1609,7 @@ void FCustomizableObjectInstanceDescriptor::SetColorParameterSelectedOption(cons
 
 	RETURN_ON_UNCOMPILED_CO(CustomizableObject, TEXT("Error: Cannot set Int parameter "));
 
+	// TODO: Multidimensional parameter support
 	SetVectorParameterSelectedOption(ColorParamName, ColorValue);
 }
 
@@ -1630,6 +1630,7 @@ FTransform FCustomizableObjectInstanceDescriptor::GetTransformParameterSelectedO
 		return FCustomizableObjectTransformParameterValue::DEFAULT_PARAMETER_VALUE;
 	}
 
+	// TODO: Multidimensional parameter support
 	return TransformParameters.IsValidIndex(TransformParamIndex) ? TransformParameters[TransformParamIndex].ParameterValue : FTransform::Identity;
 }
 
@@ -1653,6 +1654,7 @@ void FCustomizableObjectInstanceDescriptor::SetTransformParameterSelectedOption(
 		return;
 	}
 
+	// TODO: Multidimensional parameter support
 	TransformParameters[ParameterIndexInInstance].ParameterValue = TransformValue;
 }
 
@@ -1761,15 +1763,15 @@ void FCustomizableObjectInstanceDescriptor::SetProjectorValue(const FString& Pro
 	ProjectorData.Angle = Angle;
 	ProjectorData.ProjectionType = ProjectorParameter.Value.ProjectionType;
 
-	if (RangeIndex == -1)
+	// Due to optimizations a parameter may lose its multidimensionality (if it becomes constant). 
+	// In that case it means there is only one possible option so it is ok to set it as if wasn't multidimensional.
+	bool bIsMultidimensional = CustomizableObject->IsParameterMultidimensional(ParameterIndexInObject);
+	if (!bIsMultidimensional || RangeIndex == INDEX_NONE)
 	{
-		check(!CustomizableObject->IsParameterMultidimensional(ParameterIndexInObject)); // This param is multidimensional, it must have a RangeIndex of 0 or more
 		ProjectorParameter.Value = ProjectorData;
 	}
 	else
 	{
-		check(CustomizableObject->IsParameterMultidimensional(ParameterIndexInObject)); // This param is not multidimensional, it must have a RangeIndex of -1
-
 		if (!ProjectorParameter.RangeValues.IsValidIndex(RangeIndex))
 		{
 			const int32 InsertionIndex = ProjectorParameter.RangeValues.Num();
@@ -1872,15 +1874,15 @@ void FCustomizableObjectInstanceDescriptor::GetProjectorValueF(const FString& Pr
 		const FCustomizableObjectProjectorParameterValue& ProjectorParameter = ProjectorParameters[ParameterIndexInInstance];
 		const FCustomizableObjectProjector* ProjectorData;
 
-		if (RangeIndex == -1)
+		// Due to optimizations a parameter may lose its multidimensionality (if it becomes constant). 
+		// In that case it means there is only one possible option so it is ok to set it as if wasn't multidimensional.
+		bool bIsMultidimensional = CustomizableObject->IsParameterMultidimensional(ParameterIndexInObject);
+		if (!bIsMultidimensional || RangeIndex == INDEX_NONE)
 		{
-			check(!CustomizableObject->IsParameterMultidimensional(ParameterIndexInObject)); // This param is multidimensional, it must have a RangeIndex of 0 or more
-
 			ProjectorData = &ProjectorParameter.Value;
 		}
 		else
 		{
-			check(CustomizableObject->IsParameterMultidimensional(ParameterIndexInObject)); // This param is not multidimensional, it must have a RangeIndex of -1
 			check(ProjectorParameter.RangeValues.IsValidIndex(RangeIndex));
 
 			ProjectorData = &ProjectorParameter.RangeValues[RangeIndex];
@@ -1916,7 +1918,10 @@ FVector FCustomizableObjectInstanceDescriptor::GetProjectorPosition(const FStrin
 
 	if (ParameterIndexInObject >= 0 && ParameterIndexInInstance >= 0)
 	{
-		if (RangeIndex == -1)
+		// Due to optimizations a parameter may lose its multidimensionality (if it becomes constant). 
+		// In that case it means there is only one possible option so it is ok to set it as if wasn't multidimensional.
+		bool bIsMultidimensional = CustomizableObject->IsParameterMultidimensional(ParameterIndexInObject);
+		if (!bIsMultidimensional || RangeIndex == INDEX_NONE)
 		{
 			return static_cast<FVector>(ProjectorParameters[ParameterIndexInInstance].Value.Position);
 		}
@@ -1945,7 +1950,10 @@ FVector FCustomizableObjectInstanceDescriptor::GetProjectorDirection(const FStri
 
 	if (ParameterIndexInObject >= 0 && ParameterIndexInInstance >= 0)
 	{
-		if (RangeIndex == -1)
+		// Due to optimizations a parameter may lose its multidimensionality (if it becomes constant). 
+		// In that case it means there is only one possible option so it is ok to set it as if wasn't multidimensional.
+		bool bIsMultidimensional = CustomizableObject->IsParameterMultidimensional(ParameterIndexInObject);
+		if (!bIsMultidimensional || RangeIndex == INDEX_NONE)
 		{
 			return static_cast<FVector>(ProjectorParameters[ParameterIndexInInstance].Value.Direction);
 		}
@@ -1974,7 +1982,10 @@ FVector FCustomizableObjectInstanceDescriptor::GetProjectorUp(const FString& Par
 
 	if (ParameterIndexInObject >= 0 && ParameterIndexInInstance >= 0)
 	{
-		if (RangeIndex == -1)
+		// Due to optimizations a parameter may lose its multidimensionality (if it becomes constant). 
+		// In that case it means there is only one possible option so it is ok to set it as if wasn't multidimensional.
+		bool bIsMultidimensional = CustomizableObject->IsParameterMultidimensional(ParameterIndexInObject);
+		if (!bIsMultidimensional || RangeIndex == INDEX_NONE)
 		{
 			return static_cast<FVector>(ProjectorParameters[ParameterIndexInInstance].Value.Up);
 		}
@@ -2003,7 +2014,10 @@ FVector FCustomizableObjectInstanceDescriptor::GetProjectorScale(const FString&
 
 	if (ParameterIndexInObject >= 0 && ParameterIndexInInstance >= 0)
 	{
-		if (RangeIndex == -1)
+		// Due to optimizations a parameter may lose its multidimensionality (if it becomes constant). 
+		// In that case it means there is only one possible option so it is ok to set it as if wasn't multidimensional.
+		bool bIsMultidimensional = CustomizableObject->IsParameterMultidimensional(ParameterIndexInObject);
+		if (!bIsMultidimensional || RangeIndex == INDEX_NONE)
 		{
 			return static_cast<FVector>(ProjectorParameters[ParameterIndexInInstance].Value.Scale);
 		}
@@ -2032,7 +2046,10 @@ float FCustomizableObjectInstanceDescriptor::GetProjectorAngle(const FString& Pa
 
 	if (ParameterIndexInObject >= 0 && ParameterIndexInInstance >= 0)
 	{
-		if (RangeIndex == -1)
+		// Due to optimizations a parameter may lose its multidimensionality (if it becomes constant). 
+		// In that case it means there is only one possible option so it is ok to set it as if wasn't multidimensional.
+		bool bIsMultidimensional = CustomizableObject->IsParameterMultidimensional(ParameterIndexInObject);
+		if (!bIsMultidimensional || RangeIndex == INDEX_NONE)
 		{
 			return ProjectorParameters[ParameterIndexInInstance].Value.Angle;
 		}
@@ -2061,7 +2078,10 @@ ECustomizableObjectProjectorType FCustomizableObjectInstanceDescriptor::GetProje
 
 	if (ParameterIndexInObject >= 0 && ParameterIndexInInstance >= 0)
 	{
-		if (RangeIndex == -1)
+		// Due to optimizations a parameter may lose its multidimensionality (if it becomes constant). 
+		// In that case it means there is only one possible option so it is ok to set it as if wasn't multidimensional.
+		bool bIsMultidimensional = CustomizableObject->IsParameterMultidimensional(ParameterIndexInObject);
+		if (!bIsMultidimensional || RangeIndex == INDEX_NONE)
 		{
 			return ProjectorParameters[ParameterIndexInInstance].Value.ProjectionType;
 		}
@@ -2079,21 +2099,25 @@ ECustomizableObjectProjectorType FCustomizableObjectInstanceDescriptor::GetProje
 
 FCustomizableObjectProjector FCustomizableObjectInstanceDescriptor::GetProjector(const FString& ParamName, const int32 RangeIndex) const
 {
-	const int32 Index = FindTypedParameterIndex(ParamName, EMutableParameterType::Projector);
+	const int32 ParameterIndexInObject = CustomizableObject->FindParameter(ParamName);
+	const int32 ParameterIndexInInstance = FindTypedParameterIndex(ParamName, EMutableParameterType::Projector);
 
-	if (Index != INDEX_NONE)
+	if (ParameterIndexInObject >= 0 && ParameterIndexInInstance >= 0)
 	{
-		if (RangeIndex == -1)
+		// Due to optimizations a parameter may lose its multidimensionality (if it becomes constant). 
+		// In that case it means there is only one possible option so it is ok to set it as if wasn't multidimensional.
+		bool bIsMultidimensional = CustomizableObject->IsParameterMultidimensional(ParameterIndexInObject);
+		if (!bIsMultidimensional || RangeIndex == INDEX_NONE)
 		{
-			return ProjectorParameters[Index].Value;
+			return ProjectorParameters[ParameterIndexInInstance].Value;
 		}
-		else if (ProjectorParameters[Index].RangeValues.IsValidIndex(RangeIndex))
+		else if (ProjectorParameters[ParameterIndexInInstance].RangeValues.IsValidIndex(RangeIndex))
 		{
-			return ProjectorParameters[Index].RangeValues[RangeIndex];
+			return ProjectorParameters[ParameterIndexInInstance].RangeValues[RangeIndex];
 		}
 	}
 
-	LogParameterNotFoundWarning(ParamName, Index, Index, CustomizableObject, __FUNCTION__);
+	LogParameterNotFoundWarning(ParamName, ParameterIndexInObject, ParameterIndexInInstance, CustomizableObject, __FUNCTION__);
 
 	return FCustomizableObjectProjectorParameterValue::DEFAULT_PARAMETER_VALUE;
 }
