[Mutable] Fix mesh glitches with HWRT enabled.
Add CVar to turn the fix on/off.

#rnx
#tests Cooked package Frontend

diff --git a/Engine/Plugins/Mutable/Source/CustomizableObject/Private/MuCO/CustomizableObjectInstance.cpp b/Engine/Plugins/Mutable/Source/CustomizableObject/Private/MuCO/CustomizableObjectInstance.cpp
--- a/Engine/Plugins/Mutable/Source/CustomizableObject/Private/MuCO/CustomizableObjectInstance.cpp
+++ b/Engine/Plugins/Mutable/Source/CustomizableObject/Private/MuCO/CustomizableObjectInstance.cpp
@@ -3546,6 +3546,11 @@ void ConvertImage(UTexture2D* Texture, TSharedPtr<const mu::FImage> MutableImage
 	Texture->SetPlatformData(MutableCreateImagePlatformData(MutableImage,OnlyLOD,0,0) );
 }
 
+static bool bEnableRayTracingFix = false;
+FAutoConsoleVariableRef CVarMutableEnableRayTracingFix(
+	TEXT("mutable.EnableRayTracingFix"),
+	bEnableRayTracingFix,
+	TEXT("If true, Mutable generated meshes will have ray tracing disabled."));
 
 void UCustomizableInstancePrivate::InitSkeletalMeshData(const TSharedRef<FUpdateContextPrivate>& OperationData, USkeletalMesh* SkeletalMesh, const FMutableRefSkeletalMeshData& RefSkeletalMeshData, const UCustomizableObject& CustomizableObject, FCustomizableObjectComponentIndex ObjectComponentIndex)
 {
@@ -3559,6 +3564,11 @@ void UCustomizableInstancePrivate::InitSkeletalMeshData(const TSharedRef<FUpdate
 	SkeletalMesh->SetPostProcessAnimBlueprint(RefSkeletalMeshData.PostProcessAnimInst.Get());
 	SkeletalMesh->SetShadowPhysicsAsset(RefSkeletalMeshData.ShadowPhysicsAsset.Get());
 	
+	if (bEnableRayTracingFix)
+	{
+		SkeletalMesh->SetSupportRayTracing(false);
+	}
+
 	// Set Min LOD
 	SkeletalMesh->SetMinLod(FMath::Max(CustomizableObject.LODSettings.MinLOD.GetDefault(), (int32)FirstLODAvailable));
 	SkeletalMesh->SetQualityLevelMinLod(CustomizableObject.LODSettings.MinQualityLevelLOD);
@@ -3623,6 +3633,11 @@ void UCustomizableInstancePrivate::InitSkeletalMeshData(const TSharedRef<FUpdate
 		RenderData->NumNonOptionalLODs = NumLODsAvailablePerComponent - OperationData->FirstLODAvailable;
 		RenderData->CurrentFirstLODIdx = OperationData->FirstResidentLOD;
 		RenderData->LODBiasModifier = 0;
+		
+		if (bEnableRayTracingFix)
+		{
+			RenderData->bSupportRayTracing = false;
+		}
 
 		for (int32 LODIndex = 0; LODIndex < NumLODsAvailablePerComponent; ++LODIndex)
 		{
@@ -3639,6 +3654,11 @@ void UCustomizableInstancePrivate::InitSkeletalMeshData(const TSharedRef<FUpdate
 			LODInfo.bSupportUniformlyDistributedSampling = LODData.LODInfo.bSupportUniformlyDistributedSampling;
 			LODInfo.bAllowCPUAccess = LODData.LODInfo.bAllowCPUAccess;
 
+			if (bEnableRayTracingFix)
+			{
+				LODInfo.SkinCacheUsage = ESkinCacheUsage::Disabled;
+			}
+
 			// Disable LOD simplification when baking instances
 			LODInfo.ReductionSettings.NumOfTrianglesPercentage = 1.f;
 			LODInfo.ReductionSettings.NumOfVertPercentage = 1.f;
