[mutable] Fix problem with mipmaps not being fully generated for some combinations of texture sizes and layout block sizes.
#rnx

diff --git a/Engine/Plugins/Mutable/Source/MutableTools/Private/MuT/ASTOpImageMipmap.cpp b/Engine/Plugins/Mutable/Source/MutableTools/Private/MuT/ASTOpImageMipmap.cpp
--- a/Engine/Plugins/Mutable/Source/MutableTools/Private/MuT/ASTOpImageMipmap.cpp
+++ b/Engine/Plugins/Mutable/Source/MutableTools/Private/MuT/ASTOpImageMipmap.cpp
@@ -298,22 +298,20 @@ namespace mu
 		InitialSource = Root->Source.child();
 
 		// Before sinking, see if it is worth splitting into miptail and mip.
-		if (!Root->bPreventSplitTail
-			// \TODO: Review this: It seems it fails for a BC1 8x4 texture having blockSize=4, and level=0
-			&&
-			// No need to split if there is no difference between the levels considering blocks
-			Root->Levels != Root->BlockLevels 
-			)
+		// We always need to leave the root mipmap operation in case any intermediate operation fails to generate all mips.
+		if (!Root->bPreventSplitTail)
 		{
 			// the block mipmaps can be done before composition is done.
 			mu::Ptr<ASTOpImageMipmap> newMip = mu::Clone<ASTOpImageMipmap>(Root);
 			newMip->Levels = Root->BlockLevels;
 			newMip->BlockLevels = Root->BlockLevels;
 			newMip->bOnlyTail = false;
+			newMip->bPreventSplitTail = true;
 
 			// the smallest mipmaps after the composition is done.
 			mu::Ptr<ASTOpImageMipmap> topMipOp = mu::Clone<ASTOpImageMipmap>(Root);
 			topMipOp->bOnlyTail = true;
+			topMipOp->bPreventSplitTail = true;
 
 			// Proceed
 			topMipOp->Source = Visit(InitialSource, newMip.get());
