[Mutable] Enable SkeletalMesh LOD streaming and HWRT fix.

#rnx

diff --git a/Engine/Plugins/Mutable/Source/CustomizableObject/Private/MuCO/CustomizableObjectInstance.cpp b/Engine/Plugins/Mutable/Source/CustomizableObject/Private/MuCO/CustomizableObjectInstance.cpp
--- a/Engine/Plugins/Mutable/Source/CustomizableObject/Private/MuCO/CustomizableObjectInstance.cpp
+++ b/Engine/Plugins/Mutable/Source/CustomizableObject/Private/MuCO/CustomizableObjectInstance.cpp
@@ -3546,11 +3546,14 @@ void ConvertImage(UTexture2D* Texture, TSharedPtr<const mu::FImage> MutableImage
 	Texture->SetPlatformData(MutableCreateImagePlatformData(MutableImage,OnlyLOD,0,0) );
 }
 
-static bool bEnableRayTracingFix = false;
+static int32 EnableRayTracingFix = 1;
 FAutoConsoleVariableRef CVarMutableEnableRayTracingFix(
 	TEXT("mutable.EnableRayTracingFix"),
-	bEnableRayTracingFix,
-	TEXT("If true, Mutable generated meshes will have ray tracing disabled."));
+	EnableRayTracingFix,
+	TEXT("If 0, Disabled. Generated meshes will have ray tracing enabled.")
+	TEXT("If 1, Enable fix for meshes with mesh LOD streaming. Meshes will have ray tracing disabled.")
+	TEXT("If 2, Enable fix for all generated meshes. Meshes will have ray tracing disabled.")
+	);
 
 void UCustomizableInstancePrivate::InitSkeletalMeshData(const TSharedRef<FUpdateContextPrivate>& OperationData, USkeletalMesh* SkeletalMesh, const FMutableRefSkeletalMeshData& RefSkeletalMeshData, const UCustomizableObject& CustomizableObject, FCustomizableObjectComponentIndex ObjectComponentIndex)
 {
@@ -3564,6 +3567,7 @@ void UCustomizableInstancePrivate::InitSkeletalMeshData(const TSharedRef<FUpdate
 	SkeletalMesh->SetPostProcessAnimBlueprint(RefSkeletalMeshData.PostProcessAnimInst.Get());
 	SkeletalMesh->SetShadowPhysicsAsset(RefSkeletalMeshData.ShadowPhysicsAsset.Get());
 	
+	const bool bEnableRayTracingFix = EnableRayTracingFix == 2 || (EnableRayTracingFix == 1 && OperationData->bStreamMeshLODs);
 	if (bEnableRayTracingFix)
 	{
 		SkeletalMesh->SetSupportRayTracing(false);
diff --git a/Engine/Plugins/Mutable/Source/CustomizableObject/Private/MuCO/CustomizableObjectSystem.cpp b/Engine/Plugins/Mutable/Source/CustomizableObject/Private/MuCO/CustomizableObjectSystem.cpp
--- a/Engine/Plugins/Mutable/Source/CustomizableObject/Private/MuCO/CustomizableObjectSystem.cpp
+++ b/Engine/Plugins/Mutable/Source/CustomizableObject/Private/MuCO/CustomizableObjectSystem.cpp
@@ -729,7 +729,7 @@ static FAutoConsoleVariableRef CVarMutableForceStreamMeshLODs(
 	ECVF_Default);
 
 
-bool bStreamMeshLODs = false;
+bool bStreamMeshLODs = true;
 
 static FAutoConsoleVariableRef CVarMutableStreamMeshLODsEnabled(
 	TEXT("Mutable.StreamMeshLODsEnabled"),
