[Mutable] Improved Components API.

- Made FCustomizableObjectComponentIndex private. User does not need to know about ObjectIndex and InstanceIndex.
- Reverted GetComponentName to use int32. Created a private FCustomizableObjectComponentIndex version.
- Replaced TArray with TSet and exposed UCustomizableObjectInstance::GetComponents. TSet is not BP friendly.

#rnx

diff --git a/Engine/Plugins/Mutable/Source/CustomizableObject/Internal/MuCO/CustomizableObjectInstanceDescriptor.h b/Engine/Plugins/Mutable/Source/CustomizableObject/Internal/MuCO/CustomizableObjectInstanceDescriptor.h
--- a/Engine/Plugins/Mutable/Source/CustomizableObject/Internal/MuCO/CustomizableObjectInstanceDescriptor.h
+++ b/Engine/Plugins/Mutable/Source/CustomizableObject/Internal/MuCO/CustomizableObjectInstanceDescriptor.h
@@ -18,6 +18,7 @@ class UCustomizableObject;
 class UCustomizableObjectInstance;
 class FDescriptorHash;
 class FMutableUpdateCandidate;
+struct FCustomizableObjectComponentIndex;
 
 typedef TMap<const UCustomizableObjectInstance*, FMutableUpdateCandidate> FMutableInstanceUpdateMap;
 
diff --git a/Engine/Plugins/Mutable/Source/CustomizableObject/Internal/MuCO/CustomizableObjectPrivate.h b/Engine/Plugins/Mutable/Source/CustomizableObject/Internal/MuCO/CustomizableObjectPrivate.h
--- a/Engine/Plugins/Mutable/Source/CustomizableObject/Internal/MuCO/CustomizableObjectPrivate.h
+++ b/Engine/Plugins/Mutable/Source/CustomizableObject/Internal/MuCO/CustomizableObjectPrivate.h
@@ -1236,6 +1236,43 @@ struct CUSTOMIZABLEOBJECT_API FMutableMeshComponentData
 };
 
 
+
+/** Strongly typed index for the index of a component in a UCustomizableObject. */
+USTRUCT()
+struct FCustomizableObjectComponentIndex
+{
+	GENERATED_BODY()
+
+private:
+
+	int32 Index = 0;
+
+public:
+
+	explicit inline FCustomizableObjectComponentIndex(int32 InIndex=0)
+		: Index(InIndex)
+	{
+	}
+
+	inline void Invalidate()
+	{
+		Index = INDEX_NONE;
+	}
+
+	inline bool IsValid() const
+	{
+		return Index != INDEX_NONE;
+	}
+
+	inline int32 GetValue() const
+	{
+		return Index;
+	}
+
+	inline bool operator==(const FCustomizableObjectComponentIndex&) const = default;
+};
+
+
 UCLASS()
 class CUSTOMIZABLEOBJECT_API UCustomizableObjectPrivate : public UObject
 {
@@ -1400,6 +1437,8 @@ public:
 #endif
 
 	void BackwardsCompatibleFixup(int32 CustomizableObjectCustomVersion);
+
+	FName GetComponentName(FCustomizableObjectComponentIndex ObjectComponentIndex) const;
 	
 	/** Cache of generated SkeletalMeshes */
 	FMeshCache MeshCache;
diff --git a/Engine/Plugins/Mutable/Source/CustomizableObject/Private/MuCO/CustomizableObject.cpp b/Engine/Plugins/Mutable/Source/CustomizableObject/Private/MuCO/CustomizableObject.cpp
--- a/Engine/Plugins/Mutable/Source/CustomizableObject/Private/MuCO/CustomizableObject.cpp
+++ b/Engine/Plugins/Mutable/Source/CustomizableObject/Private/MuCO/CustomizableObject.cpp
@@ -1461,11 +1461,11 @@ int32 UCustomizableObject::GetComponentCount() const
 }
 
 
-FName UCustomizableObject::GetComponentName(FCustomizableObjectComponentIndex ObjectComponentIndex) const
+FName UCustomizableObjectPrivate::GetComponentName(FCustomizableObjectComponentIndex ObjectComponentIndex) const
 {
-	if (IsCompiled())
+	if (GetPublic()->IsCompiled())
 	{
-		const TArray<FName>& ComponentNames = GetPrivate()->GetModelResources().ComponentNamesPerObjectComponent;
+		const TArray<FName>& ComponentNames = GetModelResources().ComponentNamesPerObjectComponent;
 		if (ComponentNames.IsValidIndex(ObjectComponentIndex.GetValue()))
 		{
 			return ComponentNames[ObjectComponentIndex.GetValue()];
@@ -1476,6 +1476,12 @@ FName UCustomizableObject::GetComponentName(FCustomizableObjectComponentIndex Ob
 }
 
 
+FName UCustomizableObject::GetComponentName(int32 ComponentIndex) const
+{
+	return GetPrivate()->GetComponentName(FCustomizableObjectComponentIndex(ComponentIndex));
+}
+
+
 int32 UCustomizableObject::GetParameterCount() const
 {
 	return GetPrivate()->ParameterProperties.Num();
diff --git a/Engine/Plugins/Mutable/Source/CustomizableObject/Private/MuCO/CustomizableObjectInstance.cpp b/Engine/Plugins/Mutable/Source/CustomizableObject/Private/MuCO/CustomizableObjectInstance.cpp
--- a/Engine/Plugins/Mutable/Source/CustomizableObject/Private/MuCO/CustomizableObjectInstance.cpp
+++ b/Engine/Plugins/Mutable/Source/CustomizableObject/Private/MuCO/CustomizableObjectInstance.cpp
@@ -7273,9 +7273,9 @@ bool UCustomizableObjectInstance::HasAnyParameters() const
 }
 
 
-TSet<FName> UCustomizableObjectInstance::GetComponents() const
+TArray<FName> UCustomizableObjectInstance::GetComponentNames() const
 {
-	TSet<FName> GeneratedComponents;
+	TArray<FName> GeneratedComponents;
 
 	// For now, the instances don't really hold a direct array of generated components FNames. 
 	// They can be identified with the ones having a valid SkeletalMesh in the SkeletalMeshes array, but this will 
@@ -7403,7 +7403,7 @@ void UCustomizableObjectInstance::ForEachAnimInstance(int32 ObjectComponentIndex
 	UCustomizableObject* CO = GetCustomizableObject();
 	if (CO)
 	{
-		FName ComponentName = CO->GetComponentName(FCustomizableObjectComponentIndex(ObjectComponentIndex));
+		FName ComponentName = CO->GetPrivate()->GetComponentName(FCustomizableObjectComponentIndex(ObjectComponentIndex));
 		MutablePrivate::InternalForEachAnimInstance<>(GetPrivate(), ComponentName, Delegate);
 	}
 }
@@ -7415,7 +7415,7 @@ void UCustomizableObjectInstance::ForEachAnimInstance(int32 ObjectComponentIndex
 	UCustomizableObject* CO = GetCustomizableObject();
 	if (CO)
 	{
-		FName ComponentName = CO->GetComponentName(FCustomizableObjectComponentIndex(ObjectComponentIndex));
+		FName ComponentName = CO->GetPrivate()->GetComponentName(FCustomizableObjectComponentIndex(ObjectComponentIndex));
 		MutablePrivate::InternalForEachAnimInstance<>(GetPrivate(), ComponentName, Delegate);
 	}
 }
diff --git a/Engine/Plugins/Mutable/Source/CustomizableObject/Private/MuCO/CustomizableObjectInstanceUsage.cpp b/Engine/Plugins/Mutable/Source/CustomizableObject/Private/MuCO/CustomizableObjectInstanceUsage.cpp
--- a/Engine/Plugins/Mutable/Source/CustomizableObject/Private/MuCO/CustomizableObjectInstanceUsage.cpp
+++ b/Engine/Plugins/Mutable/Source/CustomizableObject/Private/MuCO/CustomizableObjectInstanceUsage.cpp
@@ -92,7 +92,7 @@ void UCustomizableObjectInstanceUsage::SetComponentIndex(int32 ComponentIndex)
 	}
 	else
 	{
-		UsedComponentIndex_DEPRECATED = FCustomizableObjectComponentIndex(ComponentIndex);
+		UsedComponentIndex = ComponentIndex;
 	}
 }
 
@@ -105,7 +105,7 @@ int32 UCustomizableObjectInstanceUsage::GetComponentIndex() const
 	}
 	else
 	{
-		return UsedComponentIndex_DEPRECATED.GetValue();
+		return UsedComponentIndex;
 	}
 }
 
@@ -118,7 +118,7 @@ void UCustomizableObjectInstanceUsage::SetComponentName(const FName& Name)
 	}
 	else
 	{
-		UsedComponentIndex_DEPRECATED.Invalidate();
+		UsedComponentIndex = INDEX_NONE;
 		UsedComponentName = Name;
 	}
 }
@@ -132,13 +132,13 @@ FName UCustomizableObjectInstanceUsage::GetComponentName() const
 	}
 	else
 	{
-		if (!UsedComponentIndex_DEPRECATED.IsValid())
+		if (UsedComponentIndex == INDEX_NONE)
 		{
 			return UsedComponentName;	
 		}
 		else
 		{
-			return FName(FString::FromInt(UsedComponentIndex_DEPRECATED.GetValue()));
+			return FName(FString::FromInt(UsedComponentIndex));
 		}
 	}
 }
diff --git a/Engine/Plugins/Mutable/Source/CustomizableObject/Private/MuCO/LogBenchmarkUtil.cpp b/Engine/Plugins/Mutable/Source/CustomizableObject/Private/MuCO/LogBenchmarkUtil.cpp
--- a/Engine/Plugins/Mutable/Source/CustomizableObject/Private/MuCO/LogBenchmarkUtil.cpp
+++ b/Engine/Plugins/Mutable/Source/CustomizableObject/Private/MuCO/LogBenchmarkUtil.cpp
@@ -108,7 +108,7 @@ void FLogBenchmarkUtil::GetInstancesStats(int32& OutNumInstances, int32& OutNumB
 		UCustomizableObject* CO = Instance->GetCustomizableObject();
 		if (CO)
 		{
-			TSet<FName> InstanceComponents = Instance->GetComponents();
+			TArray<FName> InstanceComponents = Instance->GetComponentNames();
 			for (FName ComponentName: InstanceComponents)
 			{
 				if (Instance->GetComponentMeshSkeletalMesh(ComponentName))
@@ -321,7 +321,7 @@ void FLogBenchmarkUtil::FinishUpdateMesh(const TSharedRef<FUpdateContextPrivate>
 	uint32 TriangleCount = 0;	
 	for (int32 ComponentIndex = 0; ComponentIndex < Object->GetComponentCount(); ComponentIndex++)
 	{
-		const FName ComponentName = Object->GetComponentName(FCustomizableObjectComponentIndex(ComponentIndex));
+		const FName ComponentName = Object->GetPrivate()->GetComponentName(FCustomizableObjectComponentIndex(ComponentIndex));
 		
 		// Process the generated components (not null)
 		if (const USkeletalMesh* InstanceSkeletalMesh = Instance->GetComponentMeshSkeletalMesh(ComponentName))
diff --git a/Engine/Plugins/Mutable/Source/CustomizableObject/Public/MuCO/CustomizableObject.h b/Engine/Plugins/Mutable/Source/CustomizableObject/Public/MuCO/CustomizableObject.h
--- a/Engine/Plugins/Mutable/Source/CustomizableObject/Public/MuCO/CustomizableObject.h
+++ b/Engine/Plugins/Mutable/Source/CustomizableObject/Public/MuCO/CustomizableObject.h
@@ -238,42 +238,6 @@ struct FMutableLODSettings
 };
 
 
-/** Strongly typed index for the index of a component in a UCustomizableObject. */
-USTRUCT(BlueprintType)
-struct FCustomizableObjectComponentIndex
-{
-	GENERATED_BODY()
-
-private:
-
-	int32 Index = 0;
-
-public:
-
-	explicit inline FCustomizableObjectComponentIndex(int32 InIndex=0)
-		: Index(InIndex)
-	{
-	}
-
-	inline void Invalidate()
-	{
-		Index = INDEX_NONE;
-	}
-
-	inline bool IsValid() const
-	{
-		return Index != INDEX_NONE;
-	}
-
-	inline int32 GetValue() const
-	{
-		return Index;
-	}
-
-	inline bool operator==(const FCustomizableObjectComponentIndex&) const = default;
-};
-
-
 UCLASS( BlueprintType, config=Engine )
 class CUSTOMIZABLEOBJECT_API UCustomizableObject : public UObject
 {
@@ -406,9 +370,10 @@ public:
 	int32 GetComponentCount() const;
 
 	/** Return the name of the component.
+	 *  @param ComponentIndex [0 - GetComponentCount). Index may not represent the same component between runs. To identify them use the name.
 	 *  @return NAME_None if the component does not exist. */
 	UFUNCTION(BlueprintCallable, Category = CustomizableObject)
-	FName GetComponentName(FCustomizableObjectComponentIndex ComponentIndex) const;
+	FName GetComponentName(int32 ComponentIndex) const;
 	
 	/** Get the number of parameters available in instances of this object. */
 	UFUNCTION(BlueprintCallable, Category = CustomizableObject)
diff --git a/Engine/Plugins/Mutable/Source/CustomizableObject/Public/MuCO/CustomizableObjectInstance.h b/Engine/Plugins/Mutable/Source/CustomizableObject/Public/MuCO/CustomizableObjectInstance.h
--- a/Engine/Plugins/Mutable/Source/CustomizableObject/Public/MuCO/CustomizableObjectInstance.h
+++ b/Engine/Plugins/Mutable/Source/CustomizableObject/Public/MuCO/CustomizableObjectInstance.h
@@ -30,7 +30,7 @@ struct FGameplayTagContainer;
 struct FPropertyChangedEvent;
 struct FTexturePlatformData;
 struct FMutableModelImageProperties;
-
+struct FCustomizableObjectComponentIndex;
 
 /**
  * Represents what kind of saving procedure was performed to save the package
@@ -587,9 +587,9 @@ public:
 	// ------------------------------------------------------------
 
 	/** Return the list of names of components generated for this instance. 
-	* This only has values when the instaance has been completely generated. 
-	*/
-	TSet<FName> GetComponents() const;
+	* This only has values when the instance has been completely generated. */
+	UFUNCTION(BlueprintCallable, Category = CustomizableObjectInstance)
+	TArray<FName> GetComponentNames() const;
 
 	// ------------------------------------------------------------
 
diff --git a/Engine/Plugins/Mutable/Source/CustomizableObject/Public/MuCO/CustomizableObjectInstanceUsage.h b/Engine/Plugins/Mutable/Source/CustomizableObject/Public/MuCO/CustomizableObjectInstanceUsage.h
--- a/Engine/Plugins/Mutable/Source/CustomizableObject/Public/MuCO/CustomizableObjectInstanceUsage.h
+++ b/Engine/Plugins/Mutable/Source/CustomizableObject/Public/MuCO/CustomizableObjectInstanceUsage.h
@@ -104,8 +104,9 @@ private:
 	TObjectPtr<UCustomizableObjectInstance> UsedCustomizableObjectInstance;
 
 	// If no CustomizableSkeletalComponent is associated, this Index will be used
+	// DEPRECATED
 	UPROPERTY()
-	FCustomizableObjectComponentIndex UsedComponentIndex_DEPRECATED;
+	int32 UsedComponentIndex;
 
 	/** Only used if the UsedComponentIndex is INDEX_NONE. */
 	UPROPERTY()
diff --git a/Engine/Plugins/Mutable/Source/CustomizableObjectEditor/Private/MuCOE/CustomizableObjectEditorViewportClient.cpp b/Engine/Plugins/Mutable/Source/CustomizableObjectEditor/Private/MuCOE/CustomizableObjectEditorViewportClient.cpp
--- a/Engine/Plugins/Mutable/Source/CustomizableObjectEditor/Private/MuCOE/CustomizableObjectEditorViewportClient.cpp
+++ b/Engine/Plugins/Mutable/Source/CustomizableObjectEditor/Private/MuCOE/CustomizableObjectEditorViewportClient.cpp
@@ -824,7 +824,7 @@ void FCustomizableObjectEditorViewportClient::OnInstanceUpdate(UCustomizableObje
 	// Add new components
 	for (int32 ObjectComponentIndex = 0; ObjectComponentIndex < CO->GetComponentCount(); ++ObjectComponentIndex)
 	{
-		FName Name = CO->GetComponentName(FCustomizableObjectComponentIndex(ObjectComponentIndex));
+		FName Name = CO->GetPrivate()->GetComponentName(FCustomizableObjectComponentIndex(ObjectComponentIndex));
 		if (Instance->GetComponentMeshSkeletalMesh(Name))
 		{
 			if (!SkeletalMeshComponents.Contains(Name))
diff --git a/Engine/Plugins/Mutable/Source/CustomizableObjectEditor/Private/MuCOE/CustomizableObjectInstanceBakingUtils.cpp b/Engine/Plugins/Mutable/Source/CustomizableObjectEditor/Private/MuCOE/CustomizableObjectInstanceBakingUtils.cpp
--- a/Engine/Plugins/Mutable/Source/CustomizableObjectEditor/Private/MuCOE/CustomizableObjectInstanceBakingUtils.cpp
+++ b/Engine/Plugins/Mutable/Source/CustomizableObjectEditor/Private/MuCOE/CustomizableObjectInstanceBakingUtils.cpp
@@ -437,7 +437,7 @@ bool BakeCustomizableObjectInstance(
 	const int32 NumObjectComponents = InstanceCO->GetComponentCount();
 	for (int32 ObjectComponentIndex = 0; ObjectComponentIndex < NumObjectComponents; ++ObjectComponentIndex)
 	{
-		FName ComponentName = InstanceCO->GetComponentName(FCustomizableObjectComponentIndex(ObjectComponentIndex));
+		FName ComponentName = InstanceCO->GetPrivate()->GetComponentName(FCustomizableObjectComponentIndex(ObjectComponentIndex));
 		USkeletalMesh* Mesh = InInstance.GetComponentMeshSkeletalMesh(ComponentName);
 
 		if (!Mesh)
diff --git a/Engine/Plugins/Mutable/Source/CustomizableObjectEditor/Private/MuCOE/CustomizableObjectInstanceFactory.cpp b/Engine/Plugins/Mutable/Source/CustomizableObjectEditor/Private/MuCOE/CustomizableObjectInstanceFactory.cpp
--- a/Engine/Plugins/Mutable/Source/CustomizableObjectEditor/Private/MuCOE/CustomizableObjectInstanceFactory.cpp
+++ b/Engine/Plugins/Mutable/Source/CustomizableObjectEditor/Private/MuCOE/CustomizableObjectInstanceFactory.cpp
@@ -52,7 +52,7 @@ void UCustomizableObjectInstanceFactory::PostSpawnActor(UObject* Asset, AActor*
 
 	for (int32 ObjectComponentIndex = 0; ObjectComponentIndex < Object->GetComponentCount(); ++ObjectComponentIndex)
 	{
-		const FName& ComponentName = Object->GetComponentName(FCustomizableObjectComponentIndex(ObjectComponentIndex));
+		const FName& ComponentName = Object->GetPrivate()->GetComponentName(FCustomizableObjectComponentIndex(ObjectComponentIndex));
 		
 		USkeletalMesh* SkeletalMesh = Instance->GetComponentMeshSkeletalMesh(ComponentName);
 		if (!SkeletalMesh)
@@ -121,7 +121,7 @@ bool UCustomizableObjectInstanceFactory::CanCreateActorFrom(const FAssetData& As
     {
         if (UCustomizableObjectInstance* CustomizableObjectInstance = Cast<UCustomizableObjectInstance>(AssetData.GetAsset()))
         {
-			TSet<FName> ComponentNames = CustomizableObjectInstance->GetComponents();
+			TArray<FName> ComponentNames = CustomizableObjectInstance->GetComponentNames();
 
 			// If there is a valid component with skeletal mesh
 			// TODO: This is not accurate in the future, if we have components that have grooms or clothing instead.
diff --git a/Engine/Plugins/Mutable/Source/CustomizableObjectEditor/Private/MuCOE/SCustomizableObjectEditorTextureAnalyzer.cpp b/Engine/Plugins/Mutable/Source/CustomizableObjectEditor/Private/MuCOE/SCustomizableObjectEditorTextureAnalyzer.cpp
--- a/Engine/Plugins/Mutable/Source/CustomizableObjectEditor/Private/MuCOE/SCustomizableObjectEditorTextureAnalyzer.cpp
+++ b/Engine/Plugins/Mutable/Source/CustomizableObjectEditor/Private/MuCOE/SCustomizableObjectEditorTextureAnalyzer.cpp
@@ -135,7 +135,7 @@ void SCustomizableObjecEditorTextureAnalyzer::FillTextureAnalyzerTable(UCustomiz
 
 		for (int32 ObjectComponentIndex = 0; ObjectComponentIndex < CustomizableObject->GetComponentCount(); ++ObjectComponentIndex)
 		{
-			const FName ComponentName = CustomizableObject->GetComponentName(FCustomizableObjectComponentIndex(ObjectComponentIndex));
+			const FName ComponentName = CustomizableObject->GetPrivate()->GetComponentName(FCustomizableObjectComponentIndex(ObjectComponentIndex));
 			
 			if (const USkeletalMesh* SkeletalMesh = PreviewInstance->GetComponentMeshSkeletalMesh(ComponentName))
 			{
diff --git a/Engine/Plugins/Mutable/Source/MutableValidation/Private/MuV/CustomizableObjectInstanceUpdateUtility.cpp b/Engine/Plugins/Mutable/Source/MutableValidation/Private/MuV/CustomizableObjectInstanceUpdateUtility.cpp
--- a/Engine/Plugins/Mutable/Source/MutableValidation/Private/MuV/CustomizableObjectInstanceUpdateUtility.cpp
+++ b/Engine/Plugins/Mutable/Source/MutableValidation/Private/MuV/CustomizableObjectInstanceUpdateUtility.cpp
@@ -109,7 +109,7 @@ void FCustomizableObjectInstanceUpdateUtility::OnInstanceUpdateResult(const FUpd
 
 		UCustomizableObject* CO = Instance->GetCustomizableObject();
 
-		TSet<FName> InstanceComponents = Instance->GetComponents();
+		TArray<FName> InstanceComponents = Instance->GetComponentNames();
 		for (FName ComponentName : InstanceComponents)
 		{
 			TStrongObjectPtr<USkeletalMeshComponent> SkeletalComponent = TStrongObjectPtr(NewObject<USkeletalMeshComponent>());
