[Mutable] Fix use of undefined type 'FCustomizableObjectComponentIndex'.

#rnx

diff --git a/Engine/Plugins/Mutable/Source/CustomizableObject/Private/MuCO/CustomizableSkeletalMeshActor.cpp b/Engine/Plugins/Mutable/Source/CustomizableObject/Private/MuCO/CustomizableSkeletalMeshActor.cpp
--- a/Engine/Plugins/Mutable/Source/CustomizableObject/Private/MuCO/CustomizableSkeletalMeshActor.cpp
+++ b/Engine/Plugins/Mutable/Source/CustomizableObject/Private/MuCO/CustomizableSkeletalMeshActor.cpp
@@ -5,6 +5,7 @@
 #include "Components/SkeletalMeshComponent.h"
 #include "MuCO/CustomizableSkeletalComponent.h"
 #include "MuCO/CustomizableObject.h"
+#include "MuCO/CustomizableObjectPrivate.h"
 #include "MuCO/CustomizableObjectSystem.h"
 
 #include UE_INLINE_GENERATED_CPP_BY_NAME(CustomizableSkeletalMeshActor)
