[Mutable] Fix crash when running editor in standalone mode. Fix standalone mode discarding model compiled data.

#rnx

diff --git a/Engine/Plugins/Mutable/Source/CustomizableObject/Private/MuCO/CustomizableObject.cpp b/Engine/Plugins/Mutable/Source/CustomizableObject/Private/MuCO/CustomizableObject.cpp
--- a/Engine/Plugins/Mutable/Source/CustomizableObject/Private/MuCO/CustomizableObject.cpp
+++ b/Engine/Plugins/Mutable/Source/CustomizableObject/Private/MuCO/CustomizableObject.cpp
@@ -1275,15 +1275,20 @@ void UCustomizableObjectPrivate::AddUncompiledCOWarning(const FString& Additiona
 	System->AddUncompiledCOWarning(*GetPublic(), &AdditionalLoggingInfo);
 }
 
-
 USkeletalMesh* UCustomizableObject::GetComponentMeshReferenceSkeletalMesh(const FName& ComponentName) const
 {
 #if WITH_EDITORONLY_DATA
-	if (const ICustomizableObjectEditorModule* Module = ICustomizableObjectEditorModule::Get())
+	if (!IsRunningGame())
 	{
-		return Module->GetReferenceSkeletalMesh(*this, ComponentName);
+		if (const ICustomizableObjectEditorModule* Module = ICustomizableObjectEditorModule::Get())
+		{
+			return Module->GetReferenceSkeletalMesh(*this, ComponentName);
+		}
+		
+		return nullptr;
 	}
-#else
+#endif
+
 	const FModelResources& ModelResources = Private->GetModelResources();
 	int32 ObjectComponentIndex = ModelResources.ComponentNamesPerObjectComponent.IndexOfByKey(ComponentName);
 	if (ModelResources.ReferenceSkeletalMeshesData.IsValidIndex(ObjectComponentIndex))
@@ -1291,11 +1296,10 @@ USkeletalMesh* UCustomizableObject::GetComponentMeshReferenceSkeletalMesh(const
 		// Can be nullptr if RefSkeletalMeshes are not loaded yet.
 		return ModelResources.ReferenceSkeletalMeshesData[ObjectComponentIndex].SkeletalMesh;
 	}
-#endif
+	
 	return nullptr;
 }
 
-
 int32 UCustomizableObject::FindState( const FString& Name ) const
 {
 	int32 Result = -1;
@@ -2359,6 +2363,11 @@ const FModelResources& UCustomizableObjectPrivate::GetModelResources(bool bIsCoo
 #if WITH_EDITOR
 bool UCustomizableObjectPrivate::IsCompilationOutOfDate(bool bSkipIndirectReferences, TArray<FName>& OutOfDatePackages, TArray<FName>& AddedPackages, TArray<FName>& RemovedPackages, bool& bReleaseVersionDiff) const
 {
+	if (IsRunningGame())
+	{
+		return false;
+	}
+
 	if (const ICustomizableObjectEditorModule* Module = ICustomizableObjectEditorModule::Get())
 	{
 		return Module->IsCompilationOutOfDate(*GetPublic(), bSkipIndirectReferences, OutOfDatePackages, AddedPackages, RemovedPackages, bReleaseVersionDiff);
