[mutable] Make instance generation more resilient to missing data files.
#rnx

diff --git a/Engine/Plugins/Mutable/Source/CustomizableObject/Private/MuCO/CustomizableObject.cpp b/Engine/Plugins/Mutable/Source/CustomizableObject/Private/MuCO/CustomizableObject.cpp
--- a/Engine/Plugins/Mutable/Source/CustomizableObject/Private/MuCO/CustomizableObject.cpp
+++ b/Engine/Plugins/Mutable/Source/CustomizableObject/Private/MuCO/CustomizableObject.cpp
@@ -1188,7 +1188,7 @@ FString UCustomizableObject::GetDesc()
 {
 	int32 States = GetStateCount();
 	int32 Params = GetParameterCount();
-	return FString::Printf(TEXT("%d States, %d FParameters"), States, Params);
+	return FString::Printf(TEXT("%d States, %d Parameters"), States, Params);
 }
 
 
diff --git a/Engine/Plugins/Mutable/Source/CustomizableObject/Private/MuCO/CustomizableObjectInstance.cpp b/Engine/Plugins/Mutable/Source/CustomizableObject/Private/MuCO/CustomizableObjectInstance.cpp
--- a/Engine/Plugins/Mutable/Source/CustomizableObject/Private/MuCO/CustomizableObjectInstance.cpp
+++ b/Engine/Plugins/Mutable/Source/CustomizableObject/Private/MuCO/CustomizableObjectInstance.cpp
@@ -4391,9 +4391,21 @@ void UCustomizableInstancePrivate::BuildOrCopyClothingData(const TSharedRef<FUpd
 					{
 						const uint16* Indices = SectionClothMappingData[Idx].SourceMeshVertIndices;
 
-						VertexUtilizationBuffer[Indices[0]] = 1;
-						VertexUtilizationBuffer[Indices[1]] = 1;
-						VertexUtilizationBuffer[Indices[2]] = 1;
+						// It may happen with missing data files
+						if (VertexUtilizationBuffer.IsValidIndex(Indices[0]))
+						{
+							VertexUtilizationBuffer[Indices[0]] = 1;
+						}
+
+						if (VertexUtilizationBuffer.IsValidIndex(Indices[1]))
+						{
+							VertexUtilizationBuffer[Indices[1]] = 1;
+						}
+
+						if (VertexUtilizationBuffer.IsValidIndex(Indices[2]))
+						{
+							VertexUtilizationBuffer[Indices[2]] = 1;
+						}
 					}
 				}
 
@@ -4486,9 +4498,20 @@ void UCustomizableInstancePrivate::BuildOrCopyClothingData(const TSharedRef<FUpd
 					for (FMeshToMeshVertData& VertClothData : SectionWithCloth.MappingData)
 					{
 						uint16* Indices = VertClothData.SourceMeshVertIndices;
-						Indices[0] = (uint16)IndexMap[Indices[0]];
-						Indices[1] = (uint16)IndexMap[Indices[1]];
-						Indices[2] = (uint16)IndexMap[Indices[2]];
+
+						// It could happen with missing data.
+						if (IndexMap.IsValidIndex(Indices[0]))
+						{
+							Indices[0] = (uint16)IndexMap[Indices[0]];
+						}
+						if (IndexMap.IsValidIndex(Indices[1]))
+						{
+							Indices[1] = (uint16)IndexMap[Indices[1]];
+						}
+						if (IndexMap.IsValidIndex(Indices[2]))
+						{
+							Indices[2] = (uint16)IndexMap[Indices[2]];
+						}
 					}
 				}
 
diff --git a/Engine/Plugins/Mutable/Source/CustomizableObjectEditor/Private/MuCOE/CustomizableInstanceDetails.cpp b/Engine/Plugins/Mutable/Source/CustomizableObjectEditor/Private/MuCOE/CustomizableInstanceDetails.cpp
--- a/Engine/Plugins/Mutable/Source/CustomizableObjectEditor/Private/MuCOE/CustomizableInstanceDetails.cpp
+++ b/Engine/Plugins/Mutable/Source/CustomizableObjectEditor/Private/MuCOE/CustomizableInstanceDetails.cpp
@@ -90,8 +90,8 @@ void FCustomizableInstanceDetails::CustomizeDetails(const TSharedPtr<IDetailLayo
 	// New Category that will store all properties widgets
 	IDetailCategoryBuilder& ResourcesCategory = DetailBuilder->EditCategory("Generated Resources");
 	IDetailCategoryBuilder& VisibilitySettingsCategory = DetailBuilder->EditCategory("ParametersVisibility");
-	IDetailCategoryBuilder& ParametersCategory = DetailBuilder->EditCategory("Instance FParameters");
-	IDetailCategoryBuilder& OldParametersCategory = DetailBuilder->EditCategory("Old Instance FParameters");
+	IDetailCategoryBuilder& ParametersCategory = DetailBuilder->EditCategory("Instance Parameters");
+	IDetailCategoryBuilder& OldParametersCategory = DetailBuilder->EditCategory("Old Instance Parameters");
 	IDetailCategoryBuilder& SkeletalMeshCategory = DetailBuilder->EditCategory("CustomizableSkeletalMesh");
 	IDetailCategoryBuilder& TextureParametersCategory = DetailBuilder->EditCategory("TextureParameter");
 
@@ -99,7 +99,7 @@ void FCustomizableInstanceDetails::CustomizeDetails(const TSharedPtr<IDetailLayo
 	UCustomizableObject* CustomizableObject = CustomInstance->GetCustomizableObject();
 	if (!CustomizableObject)
 	{
-		VisibilitySettingsCategory.AddCustomRow(LOCTEXT("CustomizableInstanceDetails_NoCOMessage", "Instance FParameters"))
+		VisibilitySettingsCategory.AddCustomRow(LOCTEXT("CustomizableInstanceDetails_NoCOMessage", "Instance Parameters"))
 		[
 			SNew(STextBlock).Text(LOCTEXT("Model not compiled", "Model not compiled"))
 		];
@@ -231,7 +231,7 @@ void FCustomizableInstanceDetails::CustomizeDetails(const TSharedPtr<IDetailLayo
 		[
 			SNew(SButton)
 			.HAlign(HAlign_Center)
-			.Text(LOCTEXT("Copy_Parameters", "Copy FParameters"))
+			.Text(LOCTEXT("Copy_Parameters", "Copy Parameters"))
 			.OnClicked(this, &FCustomizableInstanceDetails::OnCopyAllParameters)
 			.IsEnabled(CustomInstance->HasAnyParameters())
 			.ToolTipText(FText::FromString(FString("Copy the preview Instance parameters")))
@@ -242,7 +242,7 @@ void FCustomizableInstanceDetails::CustomizeDetails(const TSharedPtr<IDetailLayo
 		[
 			SNew(SButton)
 			.HAlign(HAlign_Center)
-			.Text(LOCTEXT("Paste_Parameters", "Paste FParameters"))
+			.Text(LOCTEXT("Paste_Parameters", "Paste Parameters"))
 			.OnClicked(this, &FCustomizableInstanceDetails::OnPasteAllParameters)
 			.IsEnabled(CustomInstance->HasAnyParameters())
 			.ToolTipText(FText::FromString(FString("Paste the preview Instance parameters")))
@@ -265,9 +265,9 @@ void FCustomizableInstanceDetails::CustomizeDetails(const TSharedPtr<IDetailLayo
 
 	if (bHiddenParamsRuntime)
 	{
-		FText HiddenParamsRuntimeMessage = LOCTEXT("CustomizableInstanceDetails_HiddemParamsRuntime", "FParameters are hidden due to their Runtime type. \nUncheck the Only Runtime checkbox to see them.");
+		FText HiddenParamsRuntimeMessage = LOCTEXT("CustomizableInstanceDetails_HiddemParamsRuntime", "Parameters are hidden due to their Runtime type. \nUncheck the Only Runtime checkbox to see them.");
 
-		ParametersCategory.AddCustomRow(LOCTEXT("CustomizableInstanceDetails_HiddemParamsRuntimeRow", "FParameters are hidden"))
+		ParametersCategory.AddCustomRow(LOCTEXT("CustomizableInstanceDetails_HiddemParamsRuntimeRow", "Parameters are hidden"))
 		[
 			SNew(SHorizontalBox)
 				+ SHorizontalBox::Slot()
@@ -2258,7 +2258,7 @@ FReply FCustomizableInstanceDetails::OnPasteAllParameters()
 	const int32 MinLOD = Descriptor.GetMinLod();
 	const TArray<uint16> RequestedLODLevels = Descriptor.GetRequestedLODLevels();
 	
-	FScopedTransaction LocalTransaction(LOCTEXT("OnPasteAllParameters", "Paste All FParameters"));
+	FScopedTransaction LocalTransaction(LOCTEXT("OnPasteAllParameters", "Paste All Parameters"));
 	CustomInstance->Modify();
 
 	if (Struct->ImportText(*ClipText, &Descriptor, nullptr, 0, GLog, GetPathNameSafe(Struct)))
@@ -2278,7 +2278,7 @@ FReply FCustomizableInstanceDetails::OnPasteAllParameters()
 
 FReply FCustomizableInstanceDetails::OnResetAllParameters()
 {
-	FScopedTransaction LocalTransaction(LOCTEXT("OnResetAllParameters", "Reset All FParameters"));
+	FScopedTransaction LocalTransaction(LOCTEXT("OnResetAllParameters", "Reset All Parameters"));
 	CustomInstance->Modify();
 
 	CustomInstance->SetDefaultValues();
diff --git a/Engine/Plugins/Mutable/Source/MutableRuntime/Internal/MuR/ModelPrivate.h b/Engine/Plugins/Mutable/Source/MutableRuntime/Internal/MuR/ModelPrivate.h
--- a/Engine/Plugins/Mutable/Source/MutableRuntime/Internal/MuR/ModelPrivate.h
+++ b/Engine/Plugins/Mutable/Source/MutableRuntime/Internal/MuR/ModelPrivate.h
@@ -621,6 +621,13 @@ namespace mu
 					for (int32 LOD = 0; LOD < FinalLODs; ++LOD)
 					{
 						int32 LODIndex = ConstantImageLODIndices[ResultLODIndexIndex + LOD];
+
+						// This could happen in the case of missing data files
+						if (!ConstantImageLODs[LODIndex].Value)
+						{
+							break;
+						}
+
 						int32 MipSizeBytes = ConstantImageLODs[LODIndex].Value->GetLODDataSize(0);
 						Result->DataStorage.ResizeLOD(LOD, MipSizeBytes);
 					}
@@ -642,7 +649,12 @@ namespace mu
 					{
 						ResultLODIndex = ConstantImageLODIndices[ResultLODIndexIndex + LOD + 1];
 						CurrentMip = ConstantImageLODs[ResultLODIndex].Value;
-						check(CurrentMip);
+
+						// This could only happen if missing or corrupted data.
+						if (!CurrentMip)
+						{
+							break;
+						}
 					}
 				}
 
diff --git a/Engine/Plugins/Mutable/Source/MutableRuntime/Private/MuR/CodeRunnerTasked.cpp b/Engine/Plugins/Mutable/Source/MutableRuntime/Private/MuR/CodeRunnerTasked.cpp
--- a/Engine/Plugins/Mutable/Source/MutableRuntime/Private/MuR/CodeRunnerTasked.cpp
+++ b/Engine/Plugins/Mutable/Source/MutableRuntime/Private/MuR/CodeRunnerTasked.cpp
@@ -744,7 +744,6 @@ namespace mu
 
 			}
 
-			// emergy fix c36adf47-e40d-490f-b709-41142bafad78
 			if (Mask->GetLODCount() < Base->GetLODCount())
 			{
 				MUTABLE_CPUPROFILER_SCOPE(ImageLayer_EmergencyFix);
@@ -1021,7 +1020,6 @@ namespace mu
 				Mask = Resized;
 			}
 
-			// emergy fix c36adf47-e40d-490f-b709-41142bafad78
 			if (Mask->GetLODCount() < Base->GetLODCount())
 			{
 				MUTABLE_CPUPROFILER_SCOPE(ImageResize_EmergencyFixMips);
@@ -2183,7 +2181,7 @@ namespace mu
 		// Since task could be reordered, we need to make sure we end the rom read before continuing
 		if (FRomLoadOp* RomLoadOp = Runner->RomLoadOps.Find(RomIndex))
 		{
-			Runner->System->StreamInterface->EndRead(RomLoadOp->m_streamID);									
+			bool bSuccess = Runner->System->StreamInterface->EndRead(RomLoadOp->m_streamID);
 
 			const int32 ResIndex = Program.Roms[RomIndex].ResourceIndex;
 			check(!Program.ConstantMeshes[ResIndex].Value)
@@ -2193,7 +2191,17 @@ namespace mu
 			FInputArchive Arch(&Stream);
 
 			check(!Program.ConstantMeshes[ResIndex].Value);
-			TSharedPtr<FMesh> Value = FMesh::StaticUnserialise(Arch);
+			TSharedPtr<FMesh> Value;
+			if (bSuccess)
+			{
+				Value = FMesh::StaticUnserialise(Arch);
+			}
+			else
+			{
+				UE_LOG(LogMutableCore, Error, TEXT("Failed to load rom data for rom index [%d]. Missing files?"), RomIndex);
+				Value = MakeShared<FMesh>();
+			}
+
 			Program.SetMeshRomValue(RomIndex, Value);
 
 			check(Program.ConstantMeshes[ResIndex].Value);
diff --git a/Engine/Plugins/Mutable/Source/MutableRuntime/Private/MuR/ImageRLE.cpp b/Engine/Plugins/Mutable/Source/MutableRuntime/Private/MuR/ImageRLE.cpp
--- a/Engine/Plugins/Mutable/Source/MutableRuntime/Private/MuR/ImageRLE.cpp
+++ b/Engine/Plugins/Mutable/Source/MutableRuntime/Private/MuR/ImageRLE.cpp
@@ -123,6 +123,13 @@ namespace mu
 
     uint32 UncompressRLE_L( int32 width, int32 rows, const uint8* pStartBaseData, uint8* pStartDestData )
     {
+		// Empty data. 
+		if (!pStartBaseData)
+		{
+			FMemory::Memset(pStartDestData, 0, width * rows);
+			return 0;
+		}
+
 		const uint8* pBaseData = pStartBaseData;
 		uint8* pDestData = pStartDestData;
 		pBaseData += sizeof(uint32); // Total mip size
diff --git a/Engine/Plugins/Mutable/Source/MutableRuntime/Private/MuR/OpImageLayer.h b/Engine/Plugins/Mutable/Source/MutableRuntime/Private/MuR/OpImageLayer.h
--- a/Engine/Plugins/Mutable/Source/MutableRuntime/Private/MuR/OpImageLayer.h
+++ b/Engine/Plugins/Mutable/Source/MutableRuntime/Private/MuR/OpImageLayer.h
@@ -314,6 +314,12 @@ namespace mu
 				const uint8* MaskBuf = MaskBatchView.GetData();
 				uint8* DestBuf = DestBatchView.GetData() + BaseChannelOffset;
 
+				// This could happen in case of missing data files.
+				if (!BaseBatchView.GetData() || !MaskBuf || !DestBatchView.GetData())
+				{
+					return;
+				}
+
 				const int32 NumElems = BaseBatchView.Num() / BytesPerElem;
 				check(NumElems == MaskBatchView.Num());
 				check(NumElems == DestBatchView.Num() / BytesPerElem);
@@ -363,6 +369,12 @@ namespace mu
                 const uint8* MaskBuf = MaskImage->GetLODData(Lod);
 				uint8* DestBuf = DestImage->GetLODData(Lod);
 
+				// This could happen in case of missing data files.
+				if (!BaseBuf || !MaskBuf || !DestBuf)
+				{
+					continue;
+				}
+
 				// Remove RLE header, mip size and row sizes.
 				MaskBuf += sizeof(uint32);
 				MaskBuf += Rows * sizeof(uint32);
@@ -1113,6 +1125,12 @@ namespace mu
 				const uint8* MaskBuf = MaskBatchView.GetData();
 				uint8* DestBuf = DestBatchView.GetData();
 
+				// This could happen in case of missing data files.
+				if (!BaseBuf || !BlendedBuf || !MaskBuf || !DestBuf)
+				{
+					return;
+				}
+
 				const int32 NumElems = BaseBatchView.Num() / BaseBytesPerElem;
 				check(NumElems == BlendBatchView.Num() / BlendBytesPerElem);
 				check(NumElems == MaskBatchView.Num() / 1);
@@ -1168,6 +1186,12 @@ namespace mu
 				const uint8* BlendedBuf = BlendImage->GetLODData(LOD);
 				const uint8* MaskBuf = MaskImage->GetLODData(LOD);
 
+				// This could happen in case of missing data files.
+				if (!BaseBuf || !BlendedBuf || !MaskBuf || !DestBuf)
+				{
+					continue;
+				}
+
 				// Remove RLE header, mip size and row sizes.
 				MaskBuf += sizeof(uint32) + Rows * sizeof(uint32);
 
@@ -1364,6 +1388,12 @@ namespace mu
 			uint8 const * const BlendedBuf = BlendedView.GetData();
 			uint8 const * const BaseBuf = BaseView.GetData(); 
 
+			// This could happen in case of missing data files.
+			if (!BaseBuf || !BlendedBuf)
+			{
+				return;
+			}
+
 			const int32 NumElems = BaseView.Num() / BytesPerElem;
 			check(NumElems == DestView.Num() / BytesPerElem);
 			check(NumElems == BlendedView.Num() / BytesPerElem);
@@ -1456,6 +1486,12 @@ namespace mu
         const uint8* BlendedBuf = BlendImage->GetLODData(0);
 		uint8* DestBuf = DestImage->GetLODData(0) + DestOffset;
 
+		// This could happen in case of missing data files.
+		if (!BlendedBuf || !MaskBuf || !DestImage->GetLODData(0))
+		{
+			return;
+		}
+
 		EImageFormat MaskFormat = MaskImage->GetFormat();
         bool bIsUncompressed = (MaskFormat == EImageFormat::L_UByte);
 
@@ -2186,6 +2222,12 @@ namespace mu
 			const uint8* MaskBuf = MaskBatchView.GetData();
 			uint8* ResultBuf = ResultBatchView.GetData();
 			
+			// This could happen in case of missing data files.
+			if (!BaseBuf || !BlendBuf || !MaskBuf)
+			{
+				return;
+			}
+
 			for (int32 I = 0; I < NumElems; ++I)
 			{
 				const uint32 Base = PackPixel<NC>(&BaseBuf[NC * I]);
@@ -2242,7 +2284,13 @@ namespace mu
 			const uint8* BaseBuf = BaseBatchView.GetData();
 			const uint8* MaskBuf = MaskBatchView.GetData();
 			uint8* ResultBuf = ResultBatchView.GetData();
-			
+
+			// This could happen in case of missing data files.
+			if (!BaseBuf || !MaskBuf)
+			{
+				return;
+			}
+
 			for (int32 I = 0; I < NumElems; ++I)
 			{
 				const uint32 Base = PackPixel<NC>(&BaseBuf[NC * I]);
diff --git a/Engine/Plugins/Mutable/Source/MutableRuntime/Private/MuR/OpImageMipmap.cpp b/Engine/Plugins/Mutable/Source/MutableRuntime/Private/MuR/OpImageMipmap.cpp
--- a/Engine/Plugins/Mutable/Source/MutableRuntime/Private/MuR/OpImageMipmap.cpp
+++ b/Engine/Plugins/Mutable/Source/MutableRuntime/Private/MuR/OpImageMipmap.cpp
@@ -594,6 +594,12 @@ namespace OpImageMipmapInternal
 				int32 UncompressedNumMips = Scratch.UncompressedMips->GetLODCount();
 				for (int32 L = 0; L < UncompressedNumMips; ++L)
 				{	
+					// It could happen with missing data.
+					if (!Scratch.CompressedMips)
+					{
+						break;
+					}
+
 					const FIntVector2 LODSize = Scratch.Uncompressed->CalculateMipSize(L);
 					Scratch.CompressedMips->DataStorage.ResizeLOD(L, LODSize.X*LODSize.Y);
 				}
diff --git a/Engine/Plugins/Mutable/Source/MutableRuntime/Private/MuR/OpImagePixelFormat.cpp b/Engine/Plugins/Mutable/Source/MutableRuntime/Private/MuR/OpImagePixelFormat.cpp
--- a/Engine/Plugins/Mutable/Source/MutableRuntime/Private/MuR/OpImagePixelFormat.cpp
+++ b/Engine/Plugins/Mutable/Source/MutableRuntime/Private/MuR/OpImagePixelFormat.cpp
@@ -366,8 +366,6 @@ void FImageOperator::ImagePixelFormat(bool& bOutSuccess, int32 CompressionQualit
 
 			case EImageFormat::L_UByte:
 			{
-				check(bFormatAllLODs);
-
 				// Try to compress
 				if (Result->DataStorage.IsEmpty())
 				{
@@ -382,7 +380,32 @@ void FImageOperator::ImagePixelFormat(bool& bOutSuccess, int32 CompressionQualit
 					Result->DataStorage.CompactedTailOffsets = Base->DataStorage.CompactedTailOffsets;
 				}
 				
-				CompressRLE_L(Base, Result);
+				int32 SizeX = Base->GetSizeX();
+				int32 SizeY = Base->GetSizeY();
+
+				for (int32 L = 0; L < NumLODs; ++L)
+				{
+					bool bSuccess = false;
+					while (!bSuccess)
+					{
+						uint32 CompressedSize = 0;
+						CompressRLE_L(CompressedSize, SizeX, SizeY, Base->GetLODData(BaseLODBegin + L), Result->GetLODData(ResultLODBegin + L), Result->GetLODDataSize(ResultLODBegin + L));
+
+						bSuccess = CompressedSize > 0;
+						if (bSuccess)
+						{
+							Result->DataStorage.ResizeLOD(ResultLODBegin + L, CompressedSize);
+						}
+						else
+						{
+							const int32 MipMemorySize = Result->DataStorage.GetLOD(ResultLODBegin + L).Num();
+							Result->DataStorage.ResizeLOD(ResultLODBegin + L, FMath::Max(4, MipMemorySize * 2));
+						}
+					}
+
+					SizeX = FMath::DivideAndRoundUp(SizeX, 2);
+					SizeY = FMath::DivideAndRoundUp(SizeY, 2);
+				}
 
 //#ifdef MUTABLE_DEBUG_RLE		
 //				if (bOutSuccess)
diff --git a/Engine/Plugins/Mutable/Source/MutableTools/Private/MuT/CodeOptimiser_Parameters.cpp b/Engine/Plugins/Mutable/Source/MutableTools/Private/MuT/CodeOptimiser_Parameters.cpp
--- a/Engine/Plugins/Mutable/Source/MutableTools/Private/MuT/CodeOptimiser_Parameters.cpp
+++ b/Engine/Plugins/Mutable/Source/MutableTools/Private/MuT/CodeOptimiser_Parameters.cpp
@@ -910,9 +910,7 @@ namespace mu
                         Ptr<ASTOpImageCompose> maskOp;
                         if (baseMask || blockMask)
                         {
-                            // \TODO: BLEH! This may create a discrepancy of number of mips between
-                            // the base image and the mask This is for now solved with emergy fix
-                            // c36adf47-e40d-490f-b709-41142bafad78
+                            // This may create a discrepancy of number of mips between the base image and the mask This is for now solved with emergy fix
 							Ptr<ASTOp> newBaseMask = EnsureValidMask(baseMask, baseImage);
 							Ptr<ASTOp> newBlockMask = EnsureValidMask(blockMask, blockImage);
 
@@ -954,9 +952,7 @@ namespace mu
                         Ptr<ASTOpImageCompose> maskOp;
                         if (baseMask || blockMask)
                         {
-                            // \TODO: BLEH! This may create a discrepancy of number of mips between
-                            // the base image and the mask This is for now solved with emergy fix
-                            // c36adf47-e40d-490f-b709-41142bafad78
+                            // This may create a discrepancy of number of mips between the base image and the mask This is for now solved with emergy fix
 							Ptr<ASTOp> newBaseMask = EnsureValidMask(baseMask, baseImage);
 							Ptr<ASTOp> newBlockMask = EnsureValidMask(blockMask, blockImage);
 
