[Mutable] Added support to table nodes with more than one material pin.

#rnx

diff --git a/Engine/Plugins/Mutable/Source/CustomizableObject/Private/MuCO/CustomizableObjectSystem.cpp b/Engine/Plugins/Mutable/Source/CustomizableObject/Private/MuCO/CustomizableObjectSystem.cpp
--- a/Engine/Plugins/Mutable/Source/CustomizableObject/Private/MuCO/CustomizableObjectSystem.cpp
+++ b/Engine/Plugins/Mutable/Source/CustomizableObject/Private/MuCO/CustomizableObjectSystem.cpp
@@ -2296,7 +2296,7 @@ namespace impl
 							FString EncodingMaterialIdString = "__MutableMaterialId";
 
 							// Decoding Material Switch from Mutable parameter name
-							if (ScalarName.ToString().Equals(EncodingMaterialIdString))
+							if (ScalarName.ToString().Contains(EncodingMaterialIdString))
 							{
 								Surface.MaterialIndex = static_cast<uint32>(ScalarValue);
 
diff --git a/Engine/Plugins/Mutable/Source/CustomizableObjectEditor/Private/MuCOE/GenerateMutableSource/GenerateMutableSourceSurface.cpp b/Engine/Plugins/Mutable/Source/CustomizableObjectEditor/Private/MuCOE/GenerateMutableSource/GenerateMutableSourceSurface.cpp
--- a/Engine/Plugins/Mutable/Source/CustomizableObjectEditor/Private/MuCOE/GenerateMutableSource/GenerateMutableSourceSurface.cpp
+++ b/Engine/Plugins/Mutable/Source/CustomizableObjectEditor/Private/MuCOE/GenerateMutableSource/GenerateMutableSourceSurface.cpp
@@ -382,7 +382,7 @@ mu::Ptr<mu::NodeSurface> GenerateMutableSourceSurface(const UEdGraphPin * Pin, F
 			{
 				TableColumnName = MaterialAssetConnectedPin->PinFriendlyName.ToString();
 
-				if (UMaterialInstance * TableMaterial = TypedNodeTable->GetColumnDefaultAssetByType<UMaterialInstance>(MaterialAssetConnectedPin))
+				if (UMaterialInstance* TableMaterial = TypedNodeTable->GetColumnDefaultAssetByType<UMaterialInstance>(MaterialAssetConnectedPin))
 				{
 					// Checking if the reference material of the Table Node has the same parent as the material of the Material Node 
 					if (!TypedNodeMat->GetMaterial() || TableMaterial->GetMaterial() != TypedNodeMat->GetMaterial()->GetMaterial())
@@ -568,6 +568,8 @@ mu::Ptr<mu::NodeSurface> GenerateMutableSourceSurface(const UEdGraphPin * Pin, F
 
 								if (bTableMaterialPinLinked)
 								{
+									// If the table material pin is linked to a table node, get all the textures of the current material parameter (CurrentMaterialTableParameter) from the Material Instances of the specified data table column.
+									// Then Generate a mutable table column with all these textures.
 									if (const UEdGraphPin* ConnectedPin = FollowInputPin(*TypedNodeMat->GetMaterialAssetPin()))
 									{
 										return GenerateMutableSourceImage(ConnectedPin, GenerationContext, Props.TextureSize);
@@ -919,6 +921,7 @@ mu::Ptr<mu::NodeSurface> GenerateMutableSourceSurface(const UEdGraphPin * Pin, F
 			SurfNode->Scalars.SetNum(NumScalar + 1);
 
 			const UEdGraphPin* MaterialPin = TypedNodeMat->GetMaterialAssetPin();
+			check(MaterialPin);
 
 			//Encoding name for material material id parameter
 			FString MaterialName = "__MutableMaterialId";
@@ -927,6 +930,9 @@ mu::Ptr<mu::NodeSurface> GenerateMutableSourceSurface(const UEdGraphPin * Pin, F
 			{
 				if (const UEdGraphPin* ConnectedPin = FollowInputPin(*MaterialPin))
 				{
+					// Adding pin ID to support multiple material columns
+					MaterialName += ConnectedPin->PinId.ToString();
+
 					GenerationContext.CurrentMaterialTableParameterId = MaterialName;
 					mu::NodeScalarPtr ScalarNode = GenerateMutableSourceFloat(ConnectedPin, GenerationContext);
 
diff --git a/Engine/Plugins/Mutable/Source/CustomizableObjectEditor/Private/MuCOE/GenerateMutableSource/GenerateMutableSourceTable.cpp b/Engine/Plugins/Mutable/Source/CustomizableObjectEditor/Private/MuCOE/GenerateMutableSource/GenerateMutableSourceTable.cpp
--- a/Engine/Plugins/Mutable/Source/CustomizableObjectEditor/Private/MuCOE/GenerateMutableSource/GenerateMutableSourceTable.cpp
+++ b/Engine/Plugins/Mutable/Source/CustomizableObjectEditor/Private/MuCOE/GenerateMutableSource/GenerateMutableSourceTable.cpp
@@ -514,6 +514,10 @@ bool FillTableColumn(const UCustomizableObjectNodeTable* TableNode, mu::Ptr<mu::
 				return true;
 			}
 
+
+			// We get here if a mesh section node has the Table Material pin linked and a Texture pin set to Mutable but nothing linked to it.
+			// This part of the code will generate a new Mutable Texture column with all the material instances textures specified in the current texture parameter (CurrentMaterialParameterId)
+
 			int32 ColumnIndex;
 
 			// Getting parameter value
