[Mutable] Add Metadata to the VertexBuffer when initializing streamable LODs.

#rnx
#tests Cooked builld frontend

diff --git a/Engine/Plugins/Mutable/Source/CustomizableObject/Private/MuCO/UnrealConversionUtils.cpp b/Engine/Plugins/Mutable/Source/CustomizableObject/Private/MuCO/UnrealConversionUtils.cpp
--- a/Engine/Plugins/Mutable/Source/CustomizableObject/Private/MuCO/UnrealConversionUtils.cpp
+++ b/Engine/Plugins/Mutable/Source/CustomizableObject/Private/MuCO/UnrealConversionUtils.cpp
@@ -248,20 +248,24 @@ namespace UnrealConversionUtils
 		check(MutableMeshVertexBuffers.GetElementCount() > 0);
 
 		const bool bUseFullPrecisionUVs = true;
+		const bool bUseFullPrecisionTangentBasis = false;
 
-		const int32 NumVertices = 1;
-		const int32 NumTexCoords = MutableMeshVertexBuffers.GetBufferChannelCount(MUTABLE_VERTEXBUFFER_TEXCOORDS);
+		const uint32 NumVertices = MutableMeshVertexBuffers.GetElementCount();
+		const uint32 NumTexCoords = MutableMeshVertexBuffers.GetBufferChannelCount(MUTABLE_VERTEXBUFFER_TEXCOORDS);
 
-		FStaticMeshVertexBuffers_InitWithMutableData(
-			LODResource.StaticVertexBuffers,
-			NumVertices,
-			NumTexCoords,
-			bUseFullPrecisionUVs,
-			bAllowCPUAccess,
-			MutableMeshVertexBuffers.GetBufferData(MUTABLE_VERTEXBUFFER_POSITION),
-			MutableMeshVertexBuffers.GetBufferData(MUTABLE_VERTEXBUFFER_TANGENT),
-			MutableMeshVertexBuffers.GetBufferData(MUTABLE_VERTEXBUFFER_TEXCOORDS)
-		);
+		const uint32 DummyNumVertices = 1;
+
+		// Static Vertex buffers
+		{
+			LODResource.StaticVertexBuffers.PositionVertexBuffer.Init(DummyNumVertices, bAllowCPUAccess);
+			LODResource.StaticVertexBuffers.PositionVertexBuffer.SetMetaData(LODResource.StaticVertexBuffers.PositionVertexBuffer.GetStride(), NumVertices);
+
+			// tangent and texture coords
+			LODResource.StaticVertexBuffers.StaticMeshVertexBuffer.SetUseFullPrecisionUVs(bUseFullPrecisionUVs);
+			LODResource.StaticVertexBuffers.StaticMeshVertexBuffer.SetUseHighPrecisionTangentBasis(bUseFullPrecisionTangentBasis);
+			LODResource.StaticVertexBuffers.StaticMeshVertexBuffer.Init(DummyNumVertices, NumTexCoords, bAllowCPUAccess);
+			LODResource.StaticVertexBuffers.StaticMeshVertexBuffer.SetMetaData(NumTexCoords, NumVertices, bUseFullPrecisionUVs, bUseFullPrecisionTangentBasis);
+		}
 
 		mu::EMeshBufferFormat BoneIndexFormat = mu::EMeshBufferFormat::None;
 		int32 NumBoneInfluences = 0;
@@ -284,26 +288,25 @@ namespace UnrealConversionUtils
 				nullptr, nullptr, &BoneWeightFormat, nullptr, nullptr);
 		}
 
-		if (BoneIndexFormat == mu::EMeshBufferFormat::UInt16)
+		// Skin Weights
 		{
-			LODResource.SkinWeightVertexBuffer.SetUse16BitBoneIndex(true);
-		}
+			const bool bUse16BitBoneIndex = BoneIndexFormat == mu::EMeshBufferFormat::UInt16;
+			const bool bUse16BitBoneWeights = BoneWeightFormat == mu::EMeshBufferFormat::NUInt16;
 
-		if (BoneWeightFormat == mu::EMeshBufferFormat::NUInt16)
-		{
-			LODResource.SkinWeightVertexBuffer.SetUse16BitBoneWeight(true);
-		}
+			LODResource.SkinWeightVertexBuffer.SetUse16BitBoneIndex(bUse16BitBoneIndex);
+			LODResource.SkinWeightVertexBuffer.SetUse16BitBoneWeight(bUse16BitBoneWeights);
+			LODResource.SkinWeightVertexBuffer.SetNeedsCPUAccess(bAllowCPUAccess);
 
-		// Init skin weight buffer
-		FSkinWeightVertexBuffer_InitWithMutableData(
-			LODResource.SkinWeightVertexBuffer,
-			NumVertices,
-			NumBoneInfluences * NumVertices,
-			NumBoneInfluences,
-			bAllowCPUAccess,
-			MutableMeshVertexBuffers.GetBufferData(BoneIndexBuffer),
-			MutableMeshVertexBuffers.GetElementSize(BoneIndexBuffer)
-		);
+			FSkinWeightDataVertexBuffer* VertexBuffer = LODResource.SkinWeightVertexBuffer.GetDataVertexBuffer();
+			VertexBuffer->Init(NumBoneInfluences, DummyNumVertices);
+			VertexBuffer->SetMetaData(NumVertices, NumBoneInfluences, bUse16BitBoneIndex, bUse16BitBoneWeights);
+
+			if (VertexBuffer->GetVariableBonesPerVertex())
+			{
+				FSkinWeightLookupVertexBuffer* LookUpVertexBuffer = const_cast<FSkinWeightLookupVertexBuffer*>(LODResource.SkinWeightVertexBuffer.GetLookupVertexBuffer());
+				LookUpVertexBuffer->SetMetaData(NumVertices);
+			}
+		}
 
 		// Optional buffers
 		for (int32 Buffer = MUTABLE_VERTEXBUFFER_TEXCOORDS + 1; Buffer < MutableMeshVertexBuffers.GetBufferCount(); ++Buffer)
@@ -320,8 +323,9 @@ namespace UnrealConversionUtils
 				// colour buffer?
 				if (Semantic == mu::EMeshBufferSemantic::Color)
 				{
-					const void* DataPtr = MutableMeshVertexBuffers.GetBufferData(Buffer);
-					FColorVertexBuffers_InitWithMutableData(LODResource.StaticVertexBuffers, NumVertices, DataPtr);
+					LODResource.StaticVertexBuffers.ColorVertexBuffer.Init(1, bAllowCPUAccess);
+					LODResource.StaticVertexBuffers.ColorVertexBuffer.SetMetaData(LODResource.StaticVertexBuffers.ColorVertexBuffer.GetStride(), NumVertices);
+					
 					check(LODResource.StaticVertexBuffers.ColorVertexBuffer.GetStride() == MutableMeshVertexBuffers.GetElementSize(Buffer));
 				}
 			}
