[Mutable] Disable workaround code that fixes an artifact when HWRT and StreamMeshLODs are enabled after applying a proper fix.
Enable SkinWeightProfileData initialization for all generated skeletal meshes. Previously, only streamable meshes would have it initialized by default.

#rnx

diff --git a/Engine/Plugins/Mutable/Source/CustomizableObject/Private/MuCO/CustomizableObjectInstance.cpp b/Engine/Plugins/Mutable/Source/CustomizableObject/Private/MuCO/CustomizableObjectInstance.cpp
--- a/Engine/Plugins/Mutable/Source/CustomizableObject/Private/MuCO/CustomizableObjectInstance.cpp
+++ b/Engine/Plugins/Mutable/Source/CustomizableObject/Private/MuCO/CustomizableObjectInstance.cpp
@@ -3546,7 +3546,7 @@ void ConvertImage(UTexture2D* Texture, TSharedPtr<const mu::FImage> MutableImage
 	Texture->SetPlatformData(MutableCreateImagePlatformData(MutableImage,OnlyLOD,0,0) );
 }
 
-static int32 EnableRayTracingFix = 1;
+static int32 EnableRayTracingFix = 0;
 FAutoConsoleVariableRef CVarMutableEnableRayTracingFix(
 	TEXT("mutable.EnableRayTracingFix"),
 	EnableRayTracingFix,
diff --git a/Engine/Plugins/Mutable/Source/CustomizableObject/Private/MuCO/UnrealConversionUtils.cpp b/Engine/Plugins/Mutable/Source/CustomizableObject/Private/MuCO/UnrealConversionUtils.cpp
--- a/Engine/Plugins/Mutable/Source/CustomizableObject/Private/MuCO/UnrealConversionUtils.cpp
+++ b/Engine/Plugins/Mutable/Source/CustomizableObject/Private/MuCO/UnrealConversionUtils.cpp
@@ -17,7 +17,7 @@ class USkeleton;
 
 TAutoConsoleVariable<int32> CVarMutableInitSkinWightProfileData(
 	TEXT("mutable.InitSkinWeightProfileData"),
-	1,
+	2,
 	TEXT("0 - Do not initialize SkinWeightProfilesData buffer by default")
 	TEXT("1 - Initialize on streamable SkeletalMeshes")
 	TEXT("2 - Initialize on all generated SkeletalMeshes"));
