[mutable] Fix channel swizzling bug when using ASTC compression.
- fix astc actually using luminance-alpha instead of red-green
- fix pixel conversion for a BGRA case.
- add debug astc texture output.
#rnx

diff --git a/Engine/Plugins/Mutable/Source/CustomizableObjectEditor/Private/MuCOE/SMutableImageViewer.cpp b/Engine/Plugins/Mutable/Source/CustomizableObjectEditor/Private/MuCOE/SMutableImageViewer.cpp
--- a/Engine/Plugins/Mutable/Source/CustomizableObjectEditor/Private/MuCOE/SMutableImageViewer.cpp
+++ b/Engine/Plugins/Mutable/Source/CustomizableObjectEditor/Private/MuCOE/SMutableImageViewer.cpp
@@ -4,6 +4,7 @@
 
 #include "Brushes/SlateImageBrush.h"
 #include "MuT/TypeInfo.h"
+#include "MuT/UnrealPixelFormatOverride.h"
 #include "TextureResource.h"
 #include "MuCO/CustomizableObject.h"
 #include "MuCO/CustomizableObjectPrivate.h"
@@ -305,6 +306,16 @@ void SMutableImageViewer::SetImage(const TSharedPtr<const mu::FImage>& InMutable
 	MutableImage = InMutableImage;
 	CurrentVisibleLOD = CurrentLOD;
 	bIsPendingUpdate = true;
+
+	constexpr bool bDebugDump = false;
+	if (bDebugDump && InMutableImage)
+	{
+		static int32 SImageCount = 0;
+		FString ImagePath = FPaths::ConvertRelativePathToFull(FPaths::Combine(FPaths::ProjectSavedDir(), FString::Printf( TEXT("MutableDebug_%d.astc"), SImageCount)));
+		++SImageCount;
+
+		DebugImageDump( InMutableImage.Get(), ImagePath);
+	}
 }
 
 
diff --git a/Engine/Plugins/Mutable/Source/MutableRuntime/Private/MuR/BlockCompression/Miro/Miro.cpp b/Engine/Plugins/Mutable/Source/MutableRuntime/Private/MuR/BlockCompression/Miro/Miro.cpp
--- a/Engine/Plugins/Mutable/Source/MutableRuntime/Private/MuR/BlockCompression/Miro/Miro.cpp
+++ b/Engine/Plugins/Mutable/Source/MutableRuntime/Private/MuR/BlockCompression/Miro/Miro.cpp
@@ -5140,7 +5140,7 @@ namespace impl
 
 
 		template<int32 BLOCK_WIDTH, int32 BLOCK_HEIGHT>
-		void encode_rg_single_partition(const unorm8_t texels[BLOCK_WIDTH * BLOCK_HEIGHT],
+		void encode_la_single_partition(const unorm8_t texels[BLOCK_WIDTH * BLOCK_HEIGHT],
 			vec3f_t e0,
 			vec3f_t e1,
 			PhysicalBlock* physical_block)
@@ -5200,7 +5200,7 @@ namespace impl
 
 
 		template<int32 BLOCK_WIDTH, int32 BLOCK_HEIGHT>
-		void encode_rg_single_partition_dual_plane(const unorm8_t texels[BLOCK_WIDTH * BLOCK_HEIGHT],
+		void encode_la_single_partition_dual_plane(const unorm8_t texels[BLOCK_WIDTH * BLOCK_HEIGHT],
 			vec4f_t e0,
 			vec4f_t e1,
 			PhysicalBlock* physical_block)
@@ -5713,7 +5713,7 @@ namespace impl
 
 
 		template<int32 BLOCK_WIDTH, int32 BLOCK_HEIGHT>
-		void compress_block_rg(const unorm8_t texels[BLOCK_WIDTH * BLOCK_HEIGHT], PhysicalBlock* physical_block)
+		void compress_block_la(const unorm8_t texels[BLOCK_WIDTH * BLOCK_HEIGHT], PhysicalBlock* physical_block)
 		{
 			unorm8_t color;
 			int32 Epsilon = APPROX_COLOR_EPSILON;
@@ -5742,13 +5742,13 @@ namespace impl
 			{
 				vec4f_t re0(r0, g0, r0, g0);
 				vec4f_t re1(r1, g1, r1, g1);
-				encode_rg_single_partition_dual_plane<BLOCK_WIDTH, BLOCK_HEIGHT>(texels, re0, re1, physical_block);
+				encode_la_single_partition_dual_plane<BLOCK_WIDTH, BLOCK_HEIGHT>(texels, re0, re1, physical_block);
 			}
 			else
 			{
 				vec3f_t re0(r0, g0, 0);
 				vec3f_t re1(r1, g1, 0);
-				encode_rg_single_partition<BLOCK_WIDTH, BLOCK_HEIGHT>(texels, re0, re1, physical_block);
+				encode_la_single_partition<BLOCK_WIDTH, BLOCK_HEIGHT>(texels, re0, re1, physical_block);
 			}
 		}
 
@@ -8810,7 +8810,8 @@ namespace miro
 					}
 
 					*rowTo = physical_block_zero;
-					astcrt::compress_block_rg<BLOCK_SIZE, BLOCK_SIZE>((astcrt::unorm8_t*)block, rowTo);
+					//astcrt::compress_block_la<BLOCK_SIZE, BLOCK_SIZE>((astcrt::unorm8_t*)block, rowTo);
+					astcrt::compress_block<BLOCK_SIZE, BLOCK_SIZE>((astcrt::unorm8_t*)block, rowTo);
 					++rowTo;
 				}
 			});
@@ -8861,8 +8862,8 @@ namespace miro
 					}
 
 					*rowTo = physical_block_zero;
-					astcrt::compress_block_rg<BLOCK_SIZE, BLOCK_SIZE>((astcrt::unorm8_t*)block, rowTo);
-
+					//astcrt::compress_block_la<BLOCK_SIZE, BLOCK_SIZE>((astcrt::unorm8_t*)block, rowTo);
+					astcrt::compress_block<BLOCK_SIZE, BLOCK_SIZE>((astcrt::unorm8_t*)block, rowTo);
 					++rowTo;
 				}
 			});
@@ -8900,7 +8901,7 @@ namespace miro
 							{
 								uint8* ToPixel = to + (yi * sx + xi) * 3;
 								ToPixel[0] = Block[BlockY * BlockSize * 4 + BlockX * 4 + 0];
-								ToPixel[1] = Block[BlockY * BlockSize * 4 + BlockX * 4 + 3];
+								ToPixel[1] = Block[BlockY * BlockSize * 4 + BlockX * 4 + 1];
 								ToPixel[2] = 255;
 							}
 						}
@@ -8941,7 +8942,7 @@ namespace miro
 							{
 								uint8* ToPixel = to + (yi * sx + xi) * 4;
 								ToPixel[0] = Block[BlockY * BlockSize * 4 + BlockX * 4 + 0];
-								ToPixel[1] = Block[BlockY * BlockSize * 4 + BlockX * 4 + 3];
+								ToPixel[1] = Block[BlockY * BlockSize * 4 + BlockX * 4 + 1];
 								ToPixel[2] = 255;
 								ToPixel[3] = 255;
 							}
@@ -9435,7 +9436,7 @@ namespace miro::SubImageDecompression
 						{
 							uint8* ToPixel = To + (yi * ToSize.X + xi) * 3;
 							ToPixel[0] = Block[BlockY * BlockSize * 4 + BlockX * 4 + 0];
-							ToPixel[1] = Block[BlockY * BlockSize * 4 + BlockX * 4 + 3];
+							ToPixel[1] = Block[BlockY * BlockSize * 4 + BlockX * 4 + 1];
 							ToPixel[2] = 255;
 						}
 					}
@@ -9474,7 +9475,7 @@ namespace miro::SubImageDecompression
 						{
 							uint8* ToPixel = To + (yi * ToSize.X + xi) * 4;
 							ToPixel[0] = Block[BlockY * BlockSize * 4 + BlockX * 4 + 0];
-							ToPixel[1] = Block[BlockY * BlockSize * 4 + BlockX * 4 + 3];
+							ToPixel[1] = Block[BlockY * BlockSize * 4 + BlockX * 4 + 1];
 							ToPixel[2] = 255;
 							ToPixel[3] = 255;
 						}
diff --git a/Engine/Plugins/Mutable/Source/MutableRuntime/Private/MuR/OpImagePixelFormat.cpp b/Engine/Plugins/Mutable/Source/MutableRuntime/Private/MuR/OpImagePixelFormat.cpp
--- a/Engine/Plugins/Mutable/Source/MutableRuntime/Private/MuR/OpImagePixelFormat.cpp
+++ b/Engine/Plugins/Mutable/Source/MutableRuntime/Private/MuR/OpImagePixelFormat.cpp
@@ -825,10 +825,10 @@ void FImageOperator::ImagePixelFormat(bool& bOutSuccess, int32 CompressionQualit
 				{
 					for (int32 I = 0; I < NumElems; ++I)
 					{
-						DestBuf[I*4 + 0] = BaseBuf[I*3 + 2];
-						DestBuf[I*4 + 1] = BaseBuf[I*3 + 1];
-						DestBuf[I*4 + 2] = BaseBuf[I*3 + 0];
-						DestBuf[I*4 + 3] = BaseBuf[I*3 + 3];
+						DestBuf[I*4 + 0] = BaseBuf[I*4 + 2];
+						DestBuf[I*4 + 1] = BaseBuf[I*4 + 1];
+						DestBuf[I*4 + 2] = BaseBuf[I*4 + 0];
+						DestBuf[I*4 + 3] = BaseBuf[I*4 + 3];
 					}
 				});
 
diff --git a/Engine/Plugins/Mutable/Source/MutableTools/Private/MuT/UnrealPixelFormatOverride.cpp b/Engine/Plugins/Mutable/Source/MutableTools/Private/MuT/UnrealPixelFormatOverride.cpp
--- a/Engine/Plugins/Mutable/Source/MutableTools/Private/MuT/UnrealPixelFormatOverride.cpp
+++ b/Engine/Plugins/Mutable/Source/MutableTools/Private/MuT/UnrealPixelFormatOverride.cpp
@@ -7,6 +7,8 @@
 #include "Interfaces/ITextureFormatModule.h"
 #include "Modules/ModuleManager.h"
 #include "Engine/TextureDefines.h"
+#include "HAL/PlatformFileManager.h"
+#include "HAL/PlatformFile.h"
 #include "TextureCompressorModule.h"
 #include "TextureBuildUtilities.h"
 #include "ImageCore.h"
@@ -140,11 +142,11 @@ void FillBuildSettingsFromMutableFormat(FTextureBuildSettings& Settings, bool& b
 }
 
 
-void MutableToImageCore(const mu::FImage* InMutable, FImage& CoreImage, int32 LOD)
+void MutableToImageCore(const mu::FImage* InMutable, FImage& CoreImage, int32 LOD, bool bSwizzleRGBHack)
 {
 	MUTABLE_CPUPROFILER_SCOPE(MutableToImageCore);
 
-	TSharedPtr<const mu::FImage> TempImage;
+	TSharedPtr<mu::FImage> TempMutable;
 
 	ERawImageFormat::Type CoreImageFormat;
 	switch (InMutable->GetFormat())
@@ -159,10 +161,25 @@ void MutableToImageCore(const mu::FImage* InMutable, FImage& CoreImage, int32 LO
 	{
 		// Unsupported format: force conversion
 		mu::FImageOperator ImOp = mu::FImageOperator::GetDefault(mu::FImageOperator::FImagePixelFormatFunc());
-		TempImage = ImOp.ImagePixelFormat(4, InMutable, mu::EImageFormat::BGRA_UByte,LOD);
-		InMutable = TempImage.Get();
+
+		TempMutable = ImOp.ImagePixelFormat(4, InMutable, mu::EImageFormat::BGRA_UByte, LOD);
+		InMutable = TempMutable.Get();
+
+		if (bSwizzleRGBHack)
+		{
+			// Editor's ASTC compressor doesn't handle "number of channels". Blank out the unused channels to improve quality.
+			int32 NumPixels = TempMutable->GetLODDataSize(0) / 4;
+			uint8* BlueChannel = TempMutable->GetLODData(0);
+			for (int32 Pixel=0;Pixel<NumPixels;++Pixel)
+			{
+				*BlueChannel = 0;
+				BlueChannel += 4;
+			}
+		}
+
 		// We are extracting one LOD, so always access LOD 0 of the resulting mutable image 
 		LOD = 0;
+
 		CoreImageFormat = ERawImageFormat::BGRA8;
 		break;
 	}
@@ -284,10 +301,18 @@ void UnrealPixelFormatFunc(bool& bOutSuccess, int32 Quality, mu::FImage* Target,
 		LODCount = 1;
 	}
 
+	// This seems to be necessary because of a probable double swizzling that happens during conversions.
+	bool bSwizzleRGBHack =
+		Target->GetFormat() == mu::EImageFormat::ASTC_4x4_RG_LDR ||
+		Target->GetFormat() == mu::EImageFormat::ASTC_6x6_RG_LDR ||
+		Target->GetFormat() == mu::EImageFormat::ASTC_8x8_RG_LDR ||
+		Target->GetFormat() == mu::EImageFormat::ASTC_10x10_RG_LDR ||
+		Target->GetFormat() == mu::EImageFormat::ASTC_12x12_RG_LDR;
+
 	for (int32 LOD = FirstLOD; bOutSuccess && (LOD < LODCount); ++LOD)
 	{
 		FImage SourceUnreal;
-		MutableToImageCore(Source, SourceUnreal, LOD);
+		MutableToImageCore(Source, SourceUnreal, LOD, bSwizzleRGBHack);
 
 		FCompressedImage2D CompressedUnreal;
 		bOutSuccess = TextureFormat->CompressImage(SourceUnreal, Settings,
@@ -303,3 +328,57 @@ void UnrealPixelFormatFunc(bool& bOutSuccess, int32 Quality, mu::FImage* Target,
 	}
 }
 
+
+void DebugImageDump(const mu::FImage* Image, const FString& FileName)
+{
+	int32 LOD = 0;
+
+	if (!Image || !Image->GetLODData(LOD))
+	{
+		return;
+	}
+
+	struct FAstcHeader
+	{
+		uint8 magic[4];
+		uint8 block_x;
+		uint8 block_y;
+		uint8 block_z;
+		uint8 dim_x[3];
+		uint8 dim_y[3];
+		uint8 dim_z[3];
+	};
+	static_assert(sizeof(FAstcHeader)==16);
+
+	FAstcHeader Header;
+	FMemory::Memzero(Header);
+
+	Header.magic[0] = 0x13;
+	Header.magic[1] = 0xAB;
+	Header.magic[2] = 0xA1;
+	Header.magic[3] = 0x5C;
+
+	int32 SizeX = Image->GetSize()[0];
+	int32 SizeY = Image->GetSize()[1];
+	Header.dim_x[0] = (SizeX >>  0) & 0xff;
+	Header.dim_x[1] = (SizeX >>  8) & 0xff;
+	Header.dim_x[2] = (SizeX >> 16) & 0xff;
+	Header.dim_y[0] = (SizeY >>  0) & 0xff;
+	Header.dim_y[1] = (SizeY >>  8) & 0xff;
+	Header.dim_y[2] = (SizeY >> 16) & 0xff;
+	Header.dim_z[0] = 1;
+
+	Header.block_x = mu::GetImageFormatData(Image->GetFormat()).PixelsPerBlockX;
+	Header.block_y = mu::GetImageFormatData(Image->GetFormat()).PixelsPerBlockY;
+	Header.block_z = 1;
+
+	IPlatformFile& PlatformFile = FPlatformFileManager::Get().GetPlatformFile();
+	IFileHandle* File = PlatformFile.OpenWrite(*FileName);
+	if (File)
+	{
+		File->Write(reinterpret_cast<const uint8*>(&Header), sizeof(FAstcHeader));
+		File->Write(Image->GetLODData(LOD),Image->GetLODDataSize(LOD));
+		delete File;
+	}
+}
+
diff --git a/Engine/Plugins/Mutable/Source/MutableTools/Public/MuT/UnrealPixelFormatOverride.h b/Engine/Plugins/Mutable/Source/MutableTools/Public/MuT/UnrealPixelFormatOverride.h
--- a/Engine/Plugins/Mutable/Source/MutableTools/Public/MuT/UnrealPixelFormatOverride.h
+++ b/Engine/Plugins/Mutable/Source/MutableTools/Public/MuT/UnrealPixelFormatOverride.h
@@ -17,7 +17,7 @@ MUTABLETOOLS_API extern void PrepareUnrealCompression();
 void FillBuildSettingsFromMutableFormat(FTextureBuildSettings& Settings, bool& bOutHasAlpha, mu::EImageFormat Format);
 
 /** */
-void MutableToImageCore(const mu::FImage* InMutable, FImage& CoreImage, int32 LOD);
+void MutableToImageCore(const mu::FImage* InMutable, FImage& CoreImage, int32 LOD, bool bSwizzleRGBHack);
 
 /** */
 bool ImageCoreToMutable(const FCompressedImage2D& Compressed, mu::FImage* Mutable, int32 LOD);
@@ -33,3 +33,6 @@ MUTABLETOOLS_API extern mu::EImageFormat QualityAndPerformanceFix(mu::EImageForm
 * PrepareUnrealCompression needs to be called from the game thread once prior to using this function. 
 */
 MUTABLETOOLS_API extern void UnrealPixelFormatFunc(bool& bOutSuccess, int32 Quality, mu::FImage* Target, const mu::FImage* Source, int32 OnlyLOD);
+
+/** Debug export the given image to a temp file in the saved folder. */
+MUTABLETOOLS_API void DebugImageDump(const mu::FImage*, const FString& FilePath);
