[Mutable] Load Runtime Referenced Textures async.

- Added CacheRuntimeTexturesEvent as a new precondition to Task_Mutable_GetMeshes.

#rnx

diff --git a/Engine/Plugins/Mutable/Source/CustomizableObject/Internal/MuCO/FMutableTaskGraph.h b/Engine/Plugins/Mutable/Source/CustomizableObject/Internal/MuCO/FMutableTaskGraph.h
--- a/Engine/Plugins/Mutable/Source/CustomizableObject/Internal/MuCO/FMutableTaskGraph.h
+++ b/Engine/Plugins/Mutable/Source/CustomizableObject/Internal/MuCO/FMutableTaskGraph.h
@@ -42,7 +42,7 @@ class FMutableTaskGraph
 public:
 
 	/** Create and launch a task on the Mutable Thread with Normal priority. */
-	UE::Tasks::FTask AddMutableThreadTask(const TCHAR* DebugName, TUniqueFunction<void()>&& TaskBody);
+	UE::Tasks::FTask AddMutableThreadTask(const TCHAR* DebugName, TUniqueFunction<void()>&& TaskBody, const TArray<UE::Tasks::FTask>& Prerequisites = {});
 
 	/** Create and launch a task on the Mutable Thread with Low priority. */
 	uint32 AddMutableThreadTaskLowPriority(const TCHAR* DebugName, TFunction<void()>&& TaskBody);
diff --git a/Engine/Plugins/Mutable/Source/CustomizableObject/Private/MuCO/CustomizableObjectSystem.cpp b/Engine/Plugins/Mutable/Source/CustomizableObject/Private/MuCO/CustomizableObjectSystem.cpp
--- a/Engine/Plugins/Mutable/Source/CustomizableObject/Private/MuCO/CustomizableObjectSystem.cpp
+++ b/Engine/Plugins/Mutable/Source/CustomizableObject/Private/MuCO/CustomizableObjectSystem.cpp
@@ -3197,7 +3197,7 @@ namespace impl
 	}
 
 
-	void Task_Game_LockMeshCache(const TSharedRef<FUpdateContextPrivate>& Operation)
+	void Task_Game_LockMeshCache(const TSharedRef<FUpdateContextPrivate>& Operation, UE::Tasks::FTask CacheRuntimeTexturesEvent)
 	{
 		MUTABLE_CPUPROFILER_SCOPE(Task_Game_LockMeshCache);
 
@@ -3220,7 +3220,7 @@ namespace impl
 			[Operation]()
 			{
 				Task_Mutable_GetMeshes(Operation);
-			});
+			}, { CacheRuntimeTexturesEvent });
 
 		SystemPrivate->AddGameThreadTask(
 			FMutableTask
@@ -3413,10 +3413,46 @@ namespace impl
 			CandidateInstancePrivateData->LiveUpdateModeInstanceID = 0;
 		}
 
+
 		Operation->Model = CustomizableObject->GetPrivate()->GetModel().ToSharedRef();
 
+		UE::Tasks::FTask CacheRuntimeTexturesEvent = UE::Tasks::MakeCompletedTask<void>();
+
 #if WITH_EDITOR
-		SystemPrivateData->GetResourceProviderChecked()->CacheRuntimeReferencedImages(Operation->Model.ToSharedRef(), CustomizableObject->GetPrivate()->GetModelResources().RuntimeReferencedTextures);
+		// Async load all Runtime Referenced Textures.
+		const TArray<TSoftObjectPtr<const UTexture>>& RuntimeReferencedTextures = CustomizableObject->GetPrivate()->GetModelResources().RuntimeReferencedTextures;
+		if (!RuntimeReferencedTextures.IsEmpty())
+		{
+			UE::Tasks::FTaskEvent Event = UE::Tasks::FTaskEvent(TEXT("Texture"));
+			CacheRuntimeTexturesEvent = Event;
+
+			TArray<FSoftObjectPath> Textures;
+			Textures.Reserve(RuntimeReferencedTextures.Num());
+			for (const TSoftObjectPtr<const UTexture>& Texture : RuntimeReferencedTextures)
+			{
+				Textures.Add(Texture.ToSoftObjectPath());
+			}
+		
+			SystemPrivateData->StreamableManager.RequestAsyncLoad(Textures, FStreamableDelegate::CreateLambda([Operation, Event]() mutable
+			{
+				UCustomizableObjectSystem* System = UCustomizableObjectSystem::GetInstance();
+				if (!System)
+				{
+					Event.Trigger();
+					return;
+				}
+				
+				UCustomizableObject* CustomizableObject = Operation->GetCapturedDescriptor().GetCustomizableObject();
+				if (!CustomizableObject)
+				{
+					Event.Trigger();
+					return;
+				}
+			
+				System->GetPrivate()->GetResourceProviderChecked()->CacheRuntimeReferencedImages(Operation->Model.ToSharedRef(), CustomizableObject->GetPrivate()->GetModelResources().RuntimeReferencedTextures);
+				Event.Trigger();
+			}));
+		}
 #endif
 		
 		// Task: Mutable Update and GetMesh
@@ -3496,9 +3532,9 @@ namespace impl
 				FMutableTask 
 				{
 					FMutableTaskDelegate::CreateLambda(
-					[Operation]()
+					[Operation, CacheRuntimeTexturesEvent]()
 					{
-						impl::Task_Game_LockMeshCache(Operation);
+						impl::Task_Game_LockMeshCache(Operation, CacheRuntimeTexturesEvent);
 					}),
 					{ Mutable_GetMeshTask }
 				});
@@ -3510,7 +3546,7 @@ namespace impl
 				[Operation]()
 				{
 					Task_Mutable_GetMeshes(Operation);
-				});
+				}, { CacheRuntimeTexturesEvent });
 
 			// Task: Lock cache
 			//-------------------------------------------------------------
diff --git a/Engine/Plugins/Mutable/Source/CustomizableObject/Private/MuCO/FMutableTaskGraph.cpp b/Engine/Plugins/Mutable/Source/CustomizableObject/Private/MuCO/FMutableTaskGraph.cpp
--- a/Engine/Plugins/Mutable/Source/CustomizableObject/Private/MuCO/FMutableTaskGraph.cpp
+++ b/Engine/Plugins/Mutable/Source/CustomizableObject/Private/MuCO/FMutableTaskGraph.cpp
@@ -21,14 +21,19 @@ static TAutoConsoleVariable<bool> CVarEnableMutableTaskLowPriority(
 	ECVF_Scalability);
 
 
-UE::Tasks::FTask FMutableTaskGraph::AddMutableThreadTask(const TCHAR* DebugName, TUniqueFunction<void()>&& TaskBody)
+
+UE::Tasks::FTask FMutableTaskGraph::AddMutableThreadTask(const TCHAR* DebugName, TUniqueFunction<void()>&& TaskBody, const TArray<UE::Tasks::FTask>& InPrerequisites)
 {
 	FScopeLock Lock(&MutableTaskLock);
 
-	TArray<UE::Tasks::FTask, TFixedAllocator<1>> Prerequisites = LastMutableTask.IsValid() 
-			? TArray<UE::Tasks::FTask, TFixedAllocator<1>>{LastMutableTask} 
-			: TArray<UE::Tasks::FTask, TFixedAllocator<1>>{};
+	TArray<UE::Tasks::FTask> Prerequisites;
+	Prerequisites.Append(InPrerequisites);
 
+	if (LastMutableTask.IsValid())
+	{
+		Prerequisites.Add(LastMutableTask);
+	}
+	
 	LastMutableTask = UE::Tasks::Launch(DebugName, MoveTemp(TaskBody), Prerequisites, TASKGRAPH_PRIORITY);	
 	
 	return LastMutableTask;
diff --git a/Engine/Plugins/Mutable/Source/CustomizableObject/Private/MuCO/UnrealMutableImageProvider.cpp b/Engine/Plugins/Mutable/Source/CustomizableObject/Private/MuCO/UnrealMutableImageProvider.cpp
--- a/Engine/Plugins/Mutable/Source/CustomizableObject/Private/MuCO/UnrealMutableImageProvider.cpp
+++ b/Engine/Plugins/Mutable/Source/CustomizableObject/Private/MuCO/UnrealMutableImageProvider.cpp
@@ -747,11 +747,17 @@ void FUnrealMutableResourceProvider::CacheRuntimeReferencedImages(const TSharedR
 	ModelImages.SourceTextures.Reset();
 	for (const TSoftObjectPtr<const UTexture>& RuntimeReferencedTexture : RuntimeReferencedTextures)
 	{
-		const UTexture* Texture = RuntimeReferencedTexture.LoadSynchronous();
+		const UTexture* Texture = RuntimeReferencedTexture.Get(); // Is already loaded.
 		if (!Texture)
 		{
-			UE_LOG(LogMutable, Warning, TEXT("Failed to load texture [%s]."), *RuntimeReferencedTexture->GetPathName());
-			continue;
+			UE_LOG(LogMutable, Warning, TEXT("Runtime Referenced Texture [%s] was not async loaded. Forcing load sync."), *RuntimeReferencedTexture->GetPathName());
+			
+			Texture = RuntimeReferencedTexture.LoadSynchronous();
+			if (!Texture)
+			{
+				UE_LOG(LogMutable, Warning, TEXT("Failed to force load sync [%s]."), *RuntimeReferencedTexture->GetPathName());
+				continue;
+			}
 		}
 
 		ModelImages.SourceTextures.Emplace(*Texture); // Perform a CopyTornOff. Once done, we no longer need the texture loaded.
diff --git a/Engine/Plugins/Mutable/Source/CustomizableObject/Private/MuCO/UnrealMutableImageProvider.h b/Engine/Plugins/Mutable/Source/CustomizableObject/Private/MuCO/UnrealMutableImageProvider.h
--- a/Engine/Plugins/Mutable/Source/CustomizableObject/Private/MuCO/UnrealMutableImageProvider.h
+++ b/Engine/Plugins/Mutable/Source/CustomizableObject/Private/MuCO/UnrealMutableImageProvider.h
@@ -18,7 +18,7 @@ class UTexture2D;
 
 
 /** Implementation of a mutable core provider for image parameters that are application-specific. */
-class FUnrealMutableResourceProvider : public mu::FExternalResourceProvider, public FGCObject
+class FUnrealMutableResourceProvider : public TSharedFromThis<FUnrealMutableResourceProvider>, public mu::FExternalResourceProvider, public FGCObject
 {
 
 public:
