[Mutable] Respect AternateTextureCompression when selecting high quality compression module.

#rnx

diff --git a/Engine/Plugins/Mutable/Source/MutableTools/Private/MuT/UnrealPixelFormatOverride.cpp b/Engine/Plugins/Mutable/Source/MutableTools/Private/MuT/UnrealPixelFormatOverride.cpp
--- a/Engine/Plugins/Mutable/Source/MutableTools/Private/MuT/UnrealPixelFormatOverride.cpp
+++ b/Engine/Plugins/Mutable/Source/MutableTools/Private/MuT/UnrealPixelFormatOverride.cpp
@@ -5,7 +5,10 @@
 #include "MuR/MutableRuntimeModule.h"
 #include "Interfaces/ITextureFormatManagerModule.h"
 #include "Interfaces/ITextureFormatModule.h"
+#include "Interfaces/ITargetPlatformManagerModule.h"
+#include "Interfaces/ITargetPlatformModule.h"
 #include "Modules/ModuleManager.h"
+#include "Misc/ConfigCacheIni.h"
 #include "Engine/TextureDefines.h"
 #include "HAL/PlatformFileManager.h"
 #include "HAL/PlatformFile.h"
@@ -13,9 +16,44 @@
 #include "TextureBuildUtilities.h"
 #include "ImageCore.h"
 
-
 static ITextureFormatManagerModule* STextureFormatManager = nullptr;
 
+static FName SPrefixedMutableTextureFormatNameTable[(uint32)mu::EImageFormat::Count];
+static bool SbPrefixedMutableTextureFormatNameTableInitialized = false;
+
+namespace
+{
+
+FName GetMutableFormatTextureFormatName(mu::EImageFormat MutableFormat)
+{
+	switch (MutableFormat)
+	{
+		case mu::EImageFormat::BC1                 : return FName(TEXT("DXT1"));
+		case mu::EImageFormat::BC2                 : return FName(TEXT("DXT3"));
+		case mu::EImageFormat::BC3                 : return FName(TEXT("DXT5"));
+		case mu::EImageFormat::BC4                 : return FName(TEXT("BC4"));
+		case mu::EImageFormat::BC5                 : return FName(TEXT("BC5"));
+		case mu::EImageFormat::ASTC_4x4_RGB_LDR    : return FName(TEXT("ASTC_RGBA_HQ"));
+    	case mu::EImageFormat::ASTC_4x4_RGBA_LDR   : return FName(TEXT("ASTC_RGBA_HQ"));
+    	case mu::EImageFormat::ASTC_4x4_RG_LDR     : return FName(TEXT("ASTC_RGB"));
+    	case mu::EImageFormat::ASTC_8x8_RGB_LDR    : return FName(TEXT("ASTC_RGBA"));   
+    	case mu::EImageFormat::ASTC_8x8_RGBA_LDR   : return FName(TEXT("ASTC_RGBA"));  
+    	case mu::EImageFormat::ASTC_8x8_RG_LDR     : return FName(TEXT("ASTC_NormalLA"));
+    	case mu::EImageFormat::ASTC_12x12_RGB_LDR  : return FName(TEXT("ASTC_RGBA"));
+    	case mu::EImageFormat::ASTC_12x12_RGBA_LDR : return FName(TEXT("ASTC_RGBA"));	 
+    	case mu::EImageFormat::ASTC_12x12_RG_LDR   : return FName(TEXT("ASTC_NormalRG"));
+    	case mu::EImageFormat::ASTC_6x6_RGB_LDR    : return FName(TEXT("ASTC_RGBA"));
+    	case mu::EImageFormat::ASTC_6x6_RGBA_LDR   : return FName(TEXT("ASTC_RGBA"));   
+    	case mu::EImageFormat::ASTC_6x6_RG_LDR     : return FName(TEXT("ASTC_NormalRG"));
+    	case mu::EImageFormat::ASTC_10x10_RGB_LDR  : return FName(TEXT("ASTC_RGBA"));
+    	case mu::EImageFormat::ASTC_10x10_RGBA_LDR : return FName(TEXT("ASTC_RGBA"));  
+    	case mu::EImageFormat::ASTC_10x10_RG_LDR   : return FName(TEXT("ASTC_NormalRG"));
+
+		default: return NAME_None;
+	}
+}
+
+}
 
 void PrepareUnrealCompression()
 {
@@ -26,114 +64,138 @@ void PrepareUnrealCompression()
 		STextureFormatManager = &FModuleManager::LoadModuleChecked<ITextureFormatManagerModule>("TextureFormat");
 		check(STextureFormatManager);
 	}
-}
 
+	if (!SbPrefixedMutableTextureFormatNameTableInitialized)
+	{
+		const ITargetPlatformSettings* TargetPlatformSettings = GetTargetPlatformManagerRef().GetRunningTargetPlatform()->GetTargetPlatformSettings();
+		check(TargetPlatformSettings);
+
+		for (uint32 FormatEnumIndex = 0; FormatEnumIndex < (uint32)mu::EImageFormat::Count; ++FormatEnumIndex)
+		{
+			FName TextureFormatName = GetMutableFormatTextureFormatName(mu::EImageFormat(FormatEnumIndex));
+			SPrefixedMutableTextureFormatNameTable[FormatEnumIndex] = TextureFormatName;
+				
+			// Based on ConditionalGetPrefixedFormat in Texture.cpp
+			FString TextureCompressionFormat;
+			bool bHasFormat = TargetPlatformSettings->GetConfigSystem()->GetString(TEXT("AlternateTextureCompression"), TEXT("TextureCompressionFormat"), TextureCompressionFormat, GEngineIni);
+			bHasFormat = bHasFormat && !TextureCompressionFormat.IsEmpty();
+			
+			if (bHasFormat)
+			{
+				ITextureFormatModule* TextureFormatModule = FModuleManager::LoadModulePtr<ITextureFormatModule>(*TextureCompressionFormat);
+			
+				ITextureFormat* TextureFormat = TextureFormatModule->GetTextureFormat();
+				if (TextureFormat)
+				{
+					FString FormatPrefix = TextureFormat->GetAlternateTextureFormatPrefix();
+					check(!FormatPrefix.IsEmpty());
+				
+					FName NewFormatName(FormatPrefix + TextureFormatName.ToString());
+
+					// check that prefixed name is one we support
+					// only apply prefix if it is in list
+					TArray<FName> SupportedFormats;
+					TextureFormat->GetSupportedFormats(SupportedFormats);
+
+					if (SupportedFormats.Contains(NewFormatName))
+					{
+						SPrefixedMutableTextureFormatNameTable[FormatEnumIndex] = NewFormatName;
+					}
+				}
+			}
+		}
+
+		SbPrefixedMutableTextureFormatNameTableInitialized = true;
+	}
+}
 
 void FillBuildSettingsFromMutableFormat(FTextureBuildSettings& Settings, bool& bOutHasAlpha, mu::EImageFormat Format)
 {
 	Settings.MipGenSettings = TMGS_NoMipmaps;
 
+	check(SbPrefixedMutableTextureFormatNameTableInitialized);
+	Settings.TextureFormatName = SPrefixedMutableTextureFormatNameTable[(uint32)Format];
+	Settings.BaseTextureFormatName = SPrefixedMutableTextureFormatNameTable[(uint32)Format];
+
 	switch (Format)
 	{
 	case mu::EImageFormat::ASTC_4x4_RGBA_LDR:
-		Settings.TextureFormatName = TEXT("ASTC_RGBA_HQ");
 		Settings.CompressionQuality = 4; // See GetQualityFormat in TextureFormatASTC.cpp
 		bOutHasAlpha = true;
 		break;
 	case mu::EImageFormat::ASTC_6x6_RGBA_LDR:
-		Settings.TextureFormatName = TEXT("ASTC_RGBA");
 		Settings.CompressionQuality = 3; // See GetQualityFormat in TextureFormatASTC.cpp
 		bOutHasAlpha = true;
 		break;
 	case mu::EImageFormat::ASTC_8x8_RGBA_LDR:
-		Settings.TextureFormatName = TEXT("ASTC_RGBA");
 		Settings.CompressionQuality = 2; // See GetQualityFormat in TextureFormatASTC.cpp
 		bOutHasAlpha = true;
 		break;
 	case mu::EImageFormat::ASTC_10x10_RGBA_LDR:
-		Settings.TextureFormatName = TEXT("ASTC_RGBA");
 		Settings.CompressionQuality = 1; // See GetQualityFormat in TextureFormatASTC.cpp
 		bOutHasAlpha = true;
 		break;
 	case mu::EImageFormat::ASTC_12x12_RGBA_LDR:
-		Settings.TextureFormatName = TEXT("ASTC_RGBA");
 		Settings.CompressionQuality = 0; // See GetQualityFormat in TextureFormatASTC.cpp
 		bOutHasAlpha = true;
 		break;
 
 	case mu::EImageFormat::ASTC_4x4_RGB_LDR:
-		Settings.TextureFormatName = TEXT("ASTC_RGBA_HQ");
 		Settings.CompressionQuality = 4; // See GetQualityFormat in TextureFormatASTC.cpp
 		bOutHasAlpha = false;
 		break;
 	case mu::EImageFormat::ASTC_6x6_RGB_LDR:
-		Settings.TextureFormatName = TEXT("ASTC_RGBA");
 		Settings.CompressionQuality = 3; // See GetQualityFormat in TextureFormatASTC.cpp
 		bOutHasAlpha = false;
 		break;
 	case mu::EImageFormat::ASTC_8x8_RGB_LDR:
-		Settings.TextureFormatName = TEXT("ASTC_RGBA");
 		Settings.CompressionQuality = 2; // See GetQualityFormat in TextureFormatASTC.cpp
 		bOutHasAlpha = false;
 		break;
 	case mu::EImageFormat::ASTC_10x10_RGB_LDR:
-		Settings.TextureFormatName = TEXT("ASTC_RGBA");
 		Settings.CompressionQuality = 1; // See GetQualityFormat in TextureFormatASTC.cpp
 		bOutHasAlpha = false;
 		break;
 	case mu::EImageFormat::ASTC_12x12_RGB_LDR:
-		Settings.TextureFormatName = TEXT("ASTC_RGBA");
 		Settings.CompressionQuality = 0; // See GetQualityFormat in TextureFormatASTC.cpp
 		bOutHasAlpha = false;
 		break;
 
 	case mu::EImageFormat::ASTC_4x4_RG_LDR:
-		Settings.TextureFormatName = TEXT("ASTC_RGB"); // There is no way to get a 4x4 RG in the ASTC compressor from a TextureFormatName
 		Settings.CompressionQuality = 4; // See GetQualityFormat in TextureFormatASTC.cpp
 		bOutHasAlpha = false;
 		break;
 	case mu::EImageFormat::ASTC_6x6_RG_LDR:
-		Settings.TextureFormatName = TEXT("ASTC_NormalRG");
 		Settings.CompressionQuality = 3; // See GetQualityFormat in TextureFormatASTC.cpp
 		bOutHasAlpha = false;
 		break;
 	case mu::EImageFormat::ASTC_8x8_RG_LDR:
-		Settings.TextureFormatName = TEXT("ASTC_NormalLA");
 		Settings.CompressionQuality = 2; // See GetQualityFormat in TextureFormatASTC.cpp
 		bOutHasAlpha = false;
 		break;
 	case mu::EImageFormat::ASTC_10x10_RG_LDR:
-		Settings.TextureFormatName = TEXT("ASTC_NormalRG");
 		Settings.CompressionQuality = 1; // See GetQualityFormat in TextureFormatASTC.cpp
 		bOutHasAlpha = false;
 		break;
 	case mu::EImageFormat::ASTC_12x12_RG_LDR:
-		Settings.TextureFormatName = TEXT("ASTC_NormalRG");
 		Settings.CompressionQuality = 0; // See GetQualityFormat in TextureFormatASTC.cpp
 		bOutHasAlpha = false;
 		break;
-
 	case mu::EImageFormat::BC1:
-		Settings.TextureFormatName = TEXT("DXT1");
 		bOutHasAlpha = false;
 		break;
 	case mu::EImageFormat::BC2:
-		Settings.TextureFormatName = TEXT("DXT3");
 		bOutHasAlpha = true;
 		break;
 	case mu::EImageFormat::BC3:
-		Settings.TextureFormatName = TEXT("DXT5");
 		bOutHasAlpha = true;
 		break;
 	case mu::EImageFormat::BC4:
-		Settings.TextureFormatName = TEXT("BC4");
 		bOutHasAlpha = false;
 		break;
 	case mu::EImageFormat::BC5:
-		Settings.TextureFormatName = TEXT("BC5");
 		bOutHasAlpha = true;
 		break;
-
 	default:
 		Settings.TextureFormatName = NAME_None;
 		bOutHasAlpha = false;
@@ -248,13 +310,13 @@ mu::EImageFormat QualityAndPerformanceFix(mu::EImageFormat Format)
 {
 	switch (Format)
 	{
-	case mu::EImageFormat::ASTC_8x8_RGB_LDR:		return mu::EImageFormat::ASTC_4x4_RGB_LDR; break;
+	case mu::EImageFormat::ASTC_8x8_RGB_LDR:	return mu::EImageFormat::ASTC_4x4_RGB_LDR; break;
 	case mu::EImageFormat::ASTC_8x8_RGBA_LDR:	return mu::EImageFormat::ASTC_4x4_RGBA_LDR; break;
 	case mu::EImageFormat::ASTC_8x8_RG_LDR:		return mu::EImageFormat::ASTC_4x4_RG_LDR; break;
 	case mu::EImageFormat::ASTC_12x12_RGB_LDR:	return mu::EImageFormat::ASTC_4x4_RGB_LDR; break;
 	case mu::EImageFormat::ASTC_12x12_RGBA_LDR:	return mu::EImageFormat::ASTC_4x4_RGBA_LDR; break;
 	case mu::EImageFormat::ASTC_12x12_RG_LDR:	return mu::EImageFormat::ASTC_4x4_RG_LDR; break;
-	case mu::EImageFormat::ASTC_6x6_RGB_LDR:		return mu::EImageFormat::ASTC_4x4_RGB_LDR; break;
+	case mu::EImageFormat::ASTC_6x6_RGB_LDR:	return mu::EImageFormat::ASTC_4x4_RGB_LDR; break;
 	case mu::EImageFormat::ASTC_6x6_RGBA_LDR:	return mu::EImageFormat::ASTC_4x4_RGBA_LDR; break;
 	case mu::EImageFormat::ASTC_6x6_RG_LDR:		return mu::EImageFormat::ASTC_4x4_RG_LDR; break;
 	case mu::EImageFormat::ASTC_10x10_RGB_LDR:	return mu::EImageFormat::ASTC_4x4_RGB_LDR; break;
