[Mutable] Fixed Normals quantizaiton issues with transforms.

#rnx

diff --git a/Engine/Plugins/Mutable/Source/MutableRuntime/Private/MuR/OpMeshTransform.h b/Engine/Plugins/Mutable/Source/MutableRuntime/Private/MuR/OpMeshTransform.h
--- a/Engine/Plugins/Mutable/Source/MutableRuntime/Private/MuR/OpMeshTransform.h
+++ b/Engine/Plugins/Mutable/Source/MutableRuntime/Private/MuR/OpMeshTransform.h
@@ -66,24 +66,30 @@ namespace mu
                 case EMeshBufferSemantic::Normal:
                 case EMeshBufferSemantic::Tangent:
                 case EMeshBufferSemantic::Binormal:
-                    for ( uint32_t v=0; v<vcount; ++v )
-                    {
-						FVector4f value( 0.0f, 0.0f, 0.0f, 1.0f );
-                        for( int i=0; i<it.GetComponents(); ++i )
-                        {
-                            ConvertData( i, &value[0], EMeshBufferFormat::Float32, it.ptr(), it.GetFormat() );
-                        }
-
-                        value = transformIT.TransformFVector4(value);
-
-                        for( int i=0; i<it.GetComponents(); ++i )
-                        {
-                            ConvertData( i, it.ptr(), it.GetFormat(), &value[0], EMeshBufferFormat::Float32 );
-                        }
-
-                        ++it;
-                    }
-                    break;
+	                {
+                		const uint8 NumComponents = FMath::Min(it.GetComponents(), 3); // Due to quantization, the serialized component W may not be zero. Must be zero to avoid being affected by the transform position.
+
+                		for ( uint32_t v=0; v<vcount; ++v )
+                		{
+                			FVector4f value( 0.0f, 0.0f, 0.0f, 0.0f );
+                        
+                			for (uint8 i = 0; i < NumComponents; ++i)
+                			{
+                				ConvertData( i, &value[0],  EMeshBufferFormat::Float32, it.ptr(), it.GetFormat() );
+                			}
+
+                			value = transformIT.TransformFVector4(value);
+
+                			// Notice that 4th component is not modified.
+                			for (uint8 i = 0; i < NumComponents; ++i)
+                			{
+                				ConvertData( i, it.ptr(), it.GetFormat(), &value[0],  EMeshBufferFormat::Float32 );
+                			}
+
+                			++it;
+                		}
+	                }
+                	break;
 
                 default:
                     break;
diff --git a/Engine/Plugins/Mutable/Source/MutableRuntime/Private/MuR/OpMeshTransformWithMesh.cpp b/Engine/Plugins/Mutable/Source/MutableRuntime/Private/MuR/OpMeshTransformWithMesh.cpp
--- a/Engine/Plugins/Mutable/Source/MutableRuntime/Private/MuR/OpMeshTransformWithMesh.cpp
+++ b/Engine/Plugins/Mutable/Source/MutableRuntime/Private/MuR/OpMeshTransformWithMesh.cpp
@@ -71,25 +71,31 @@ void MeshTransformWithMesh(FMesh* Result, const FMesh* SourceMesh, const FMesh*
             case EMeshBufferSemantic::Normal:
             case EMeshBufferSemantic::Tangent:
             case EMeshBufferSemantic::Binormal:
-                for ( uint32_t v=0; v<VCount; ++v )
-                {
-                	if (!BoundingMesh || VertexInBoundaryMesh[v])
-                	{
-                		FVector4f value( 0.0f, 0.0f, 0.0f, 1.0f );
-	                    for( int i=0; i<it.GetComponents(); ++i )
-	                    {
-	                        ConvertData( i, &value[0], EMeshBufferFormat::Float32, it.ptr(), it.GetFormat() );
-	                    }
-
-	                    value = TransformInvT.TransformFVector4(value);
-
-	                    for( int i=0; i<it.GetComponents(); ++i )
-	                    {
-	                        ConvertData( i, it.ptr(), it.GetFormat(), &value[0], EMeshBufferFormat::Float32 );
-	                    }
-                    }
-					++it;
-                }
+	            {
+            		const uint8 NumComponents = FMath::Min(it.GetComponents(), 3); // Due to quantization, the serialized component W may not be zero. Must be zero to avoid being affected by the transform position.
+
+            		for ( uint32_t v=0; v<VCount; ++v )
+            		{
+            			if (!BoundingMesh || VertexInBoundaryMesh[v])
+            			{
+            				FVector4f value( 0.0f, 0.0f, 0.0f, 0.0f );
+                		
+            				for (uint8 i = 0; i < NumComponents; ++i )
+            				{
+            					ConvertData( i, &value[0],  EMeshBufferFormat::Float32, it.ptr(), it.GetFormat() );
+            				}
+
+            				value = TransformInvT.TransformFVector4(value);
+
+            				// Notice that 4th component is not modified.
+            				for (uint8 i = 0; i < NumComponents; ++i )
+            				{
+            					ConvertData( i, it.ptr(), it.GetFormat(), &value[0],  EMeshBufferFormat::Float32 );
+            				}
+            			}
+            			++it;
+            		}
+	            }
                 break;
 
             default:
