[Mutable] Fix crash when updating + compiling a CO.

- Fixed forced parameters not being considered in the Descriptor Hash.

#rnx

diff --git a/Engine/Plugins/Mutable/Source/CustomizableObject/Internal/MuCO/CustomizableObjectInstanceDescriptor.h b/Engine/Plugins/Mutable/Source/CustomizableObject/Internal/MuCO/CustomizableObjectInstanceDescriptor.h
--- a/Engine/Plugins/Mutable/Source/CustomizableObject/Internal/MuCO/CustomizableObjectInstanceDescriptor.h
+++ b/Engine/Plugins/Mutable/Source/CustomizableObject/Internal/MuCO/CustomizableObjectInstanceDescriptor.h
@@ -307,8 +307,6 @@ struct CUSTOMIZABLEOBJECT_API FCustomizableObjectInstanceDescriptor
 
 	FString ToString() const;
 	
-private:
-
 	UPROPERTY()
 	TObjectPtr<UCustomizableObject> CustomizableObject = nullptr;
 
diff --git a/Engine/Plugins/Mutable/Source/CustomizableObject/Internal/MuCO/CustomizableObjectPrivate.h b/Engine/Plugins/Mutable/Source/CustomizableObject/Internal/MuCO/CustomizableObjectPrivate.h
--- a/Engine/Plugins/Mutable/Source/CustomizableObject/Internal/MuCO/CustomizableObjectPrivate.h
+++ b/Engine/Plugins/Mutable/Source/CustomizableObject/Internal/MuCO/CustomizableObjectPrivate.h
@@ -52,6 +52,8 @@ class USkeletalMeshLODSettings;
 struct FModelResources;
 struct FModelStreamableBulkData;
 struct FObjectAndNameAsStringProxyArchive;
+struct FCustomizableObjectInstanceDescriptor;
+
 
 FGuid CUSTOMIZABLEOBJECT_API GenerateIdentifier(const UCustomizableObject& CustomizableObject);
 
@@ -1323,7 +1325,7 @@ public:
 	bool IsLocked() const;
 
 	/** Modify the provided mutable parameters so that the forced values for the given customizable object state are applied. */
-	void ApplyStateForcedValuesToParameters(int32 State, mu::FParameters* Parameters);
+	void ApplyStateForcedValuesToParameters(FCustomizableObjectInstanceDescriptor& Descriptor);
 
 	int32 FindParameter(const FString& Name) const;
 	int32 FindParameterTyped(const FString& Name, EMutableParameterType Type) const;
diff --git a/Engine/Plugins/Mutable/Source/CustomizableObject/Private/MuCO/CustomizableObject.cpp b/Engine/Plugins/Mutable/Source/CustomizableObject/Private/MuCO/CustomizableObject.cpp
--- a/Engine/Plugins/Mutable/Source/CustomizableObject/Private/MuCO/CustomizableObject.cpp
+++ b/Engine/Plugins/Mutable/Source/CustomizableObject/Private/MuCO/CustomizableObject.cpp
@@ -2036,55 +2036,35 @@ bool UCustomizableObject::IsParameterMultidimensional(const int32& InParamIndex)
 }
 
 
-void UCustomizableObjectPrivate::ApplyStateForcedValuesToParameters(int32 State, mu::FParameters* Parameters)
+void UCustomizableObjectPrivate::ApplyStateForcedValuesToParameters(FCustomizableObjectInstanceDescriptor& Descriptor)
 {
-	const FString& StateName = GetPublic()->GetStateName(State);
+	const FString& StateName = Descriptor.GetCurrentState();
 	const FMutableStateData* StateData = GetModelResources().StateUIDataMap.Find(StateName);
 	if (!StateData)
 	{
 		return;
 	}
-	
-	for (const TPair<FString, FString>& ForcedParameter : StateData->ForcedParameterValues)
+
+	for (FCustomizableObjectIntParameterValue& IntParameter : Descriptor.IntParameters)
 	{
-		int32 ForcedParameterIndex = FindParameter(ForcedParameter.Key);
-		if (ForcedParameterIndex == INDEX_NONE)
+		if (const FString* Result = StateData->ForcedParameterValues.Find(IntParameter.ParameterName))
 		{
-			continue;
-		}
+			IntParameter.ParameterValueName = *Result;
 
-		bool bIsMultidimensional = Parameters->NewRangeIndex(ForcedParameterIndex) != nullptr;
-		if (!bIsMultidimensional)
-		{
-			switch (GetParameterType(ForcedParameterIndex))
-			{
-			case EMutableParameterType::Int:
-			{
-				FString StringValue = ForcedParameter.Value;
-				if (StringValue.IsNumeric())
-				{
-					Parameters->SetIntValue(ForcedParameterIndex, FCString::Atoi(*StringValue));
-				}
-				else
-				{
-					int32 IntParameterIndex = FindIntParameterValue(ForcedParameterIndex, StringValue);
-					Parameters->SetIntValue(ForcedParameterIndex, IntParameterIndex);
-				}
-				break;
-			}
-			case EMutableParameterType::Bool:
-			{
-				Parameters->SetBoolValue(ForcedParameterIndex, ForcedParameter.Value.ToBool());
-				break;
-			}
-			default:
+			for (FString& Range : IntParameter.ParameterRangeValueNames)
 			{
-				UE_LOG(LogMutable, Warning, TEXT("Forced parameter type not supported."));
-				break;
-			}
+				Range = *Result;
 			}
 		}
 	}
+
+	for (FCustomizableObjectBoolParameterValue& BoolParameter : Descriptor.BoolParameters)
+	{
+		if (const FString* Result = StateData->ForcedParameterValues.Find(BoolParameter.ParameterName))
+		{
+			BoolParameter.ParameterValue = Result->ToBool();
+		}
+	}
 }
 
 
diff --git a/Engine/Plugins/Mutable/Source/CustomizableObject/Private/MuCO/CustomizableObjectSystem.cpp b/Engine/Plugins/Mutable/Source/CustomizableObject/Private/MuCO/CustomizableObjectSystem.cpp
--- a/Engine/Plugins/Mutable/Source/CustomizableObject/Private/MuCO/CustomizableObjectSystem.cpp
+++ b/Engine/Plugins/Mutable/Source/CustomizableObject/Private/MuCO/CustomizableObjectSystem.cpp
@@ -155,21 +155,20 @@ static FAutoConsoleVariableSink CVarMutableSink(FConsoleCommandDelegate::CreateS
 FUpdateContextPrivate::FUpdateContextPrivate(UCustomizableObjectInstance& InInstance, const FCustomizableObjectInstanceDescriptor& Descriptor)
 {
 	check(IsInGameThread());
-	check(InInstance.GetPrivate());
 	check(InInstance.GetCustomizableObject());
 
 	Instance = &InInstance;
+
 	CapturedDescriptor = Descriptor;
+	InInstance.GetCustomizableObject()->GetPrivate()->ApplyStateForcedValuesToParameters(CapturedDescriptor);
+
 	CapturedDescriptorHash = FDescriptorHash(Descriptor);
-	Parameters = Descriptor.GetParameters();
 	NumObjectComponents = InInstance.GetCustomizableObject()->GetComponentCount();
 	FirstLODAvailable = InInstance.GetCustomizableObject()->GetPrivate()->GetMinLODIndex();
 	FirstResidentLOD = InInstance.GetPrivate()->FirstResidentLOD;
 
 	MutableSystem = UCustomizableObjectSystem::GetInstance()->GetPrivate()->MutableSystem;	
 	check(MutableSystem);
-	
-	InInstance.GetCustomizableObject()->GetPrivate()->ApplyStateForcedValuesToParameters(CapturedDescriptor.GetState(), Parameters.Get());
 
 	UCustomizableObjectSystem* System = UCustomizableObjectSystem::GetInstance();
 	System->GetPrivate()->CacheTextureParameters(CapturedDescriptor.GetTextureParameters());
@@ -3348,8 +3347,8 @@ namespace impl
 			bCancel = true;
 		}
 
-		TSharedPtr<const mu::FParameters> Parameters = Operation->Parameters;
-		if (!Parameters)
+		Operation->Parameters = Operation->GetCapturedDescriptor().GetParameters();
+		if (!Operation->Parameters)
 		{
 			bCancel = true;
 		}
