[Mutable] Fix LOD update using the wrong update parameters.

- Moved ESkeletalMeshStatus to private. Was already not exposed to users.

#rnx

diff --git a/Engine/Plugins/Mutable/Source/CustomizableObject/Internal/MuCO/CustomizableObjectInstancePrivate.h b/Engine/Plugins/Mutable/Source/CustomizableObject/Internal/MuCO/CustomizableObjectInstancePrivate.h
--- a/Engine/Plugins/Mutable/Source/CustomizableObject/Internal/MuCO/CustomizableObjectInstancePrivate.h
+++ b/Engine/Plugins/Mutable/Source/CustomizableObject/Internal/MuCO/CustomizableObjectInstancePrivate.h
@@ -194,6 +194,15 @@ DECLARE_MULTICAST_DELEGATE_OneParam(FObjectInstanceTransactedDelegate, const FTr
 #endif //WITH_EDITOR
 
 
+/** Indicates the status of the generated Skeletal Mesh. */
+enum class ESkeletalMeshStatus : uint8
+{
+	NotGenerated, // Set only when loading the Instance for the first time or after compiling. Any generation, successful or not, can not end up in this state.
+	Success, // Generated successfully.
+	Error // Not generated. Set only after a failed update.
+};
+
+
 UCLASS()
 class CUSTOMIZABLEOBJECT_API UCustomizableInstancePrivate : public UObject
 {
diff --git a/Engine/Plugins/Mutable/Source/CustomizableObject/Private/MuCO/CustomizableObjectSystem.cpp b/Engine/Plugins/Mutable/Source/CustomizableObject/Private/MuCO/CustomizableObjectSystem.cpp
--- a/Engine/Plugins/Mutable/Source/CustomizableObject/Private/MuCO/CustomizableObjectSystem.cpp
+++ b/Engine/Plugins/Mutable/Source/CustomizableObject/Private/MuCO/CustomizableObjectSystem.cpp
@@ -3802,13 +3802,14 @@ int32 UCustomizableObjectSystem::TickInternal(const bool bBlocking)
 					LODUpdateWithSameInstance->ApplyLODUpdateParamsToInstance(PendingInstanceUpdateFound->Context.Get());
 				}
 
-				Private->StartUpdateSkeletalMesh(PendingInstanceUpdateFound->Context);				
+				Private->StartUpdateSkeletalMesh(PendingInstanceUpdateFound->Context);
 				Private->MutablePendingInstanceWork.RemoveUpdate(PendingInstanceUpdateFound->Context->Instance);
 			}
 			else if (LODUpdateCandidateFound)
 			{
 				UCustomizableObjectInstance* Instance = LODUpdateCandidateFound->CustomizableObjectInstance;
-				const FCustomizableObjectInstanceDescriptor& Descriptor = Instance->GetPrivate()->SkeletalMeshStatus == ESkeletalMeshStatus::NotGenerated ? Instance->GetPrivate()->GetDescriptor() : Instance->GetPrivate()->CommittedDescriptor;
+				const bool bGenerated = Instance->GetPrivate()->SkeletalMeshStatus == ESkeletalMeshStatus::Success;
+				const FCustomizableObjectInstanceDescriptor& Descriptor = bGenerated ? Instance->GetPrivate()->CommittedDescriptor : Instance->GetPrivate()->GetDescriptor();
 
 				const TSharedRef<FUpdateContextPrivate> Context = MakeShared<FUpdateContextPrivate>(*Instance, Descriptor);
 
diff --git a/Engine/Plugins/Mutable/Source/CustomizableObject/Public/MuCO/CustomizableObjectInstance.h b/Engine/Plugins/Mutable/Source/CustomizableObject/Public/MuCO/CustomizableObjectInstance.h
--- a/Engine/Plugins/Mutable/Source/CustomizableObject/Public/MuCO/CustomizableObjectInstance.h
+++ b/Engine/Plugins/Mutable/Source/CustomizableObject/Public/MuCO/CustomizableObjectInstance.h
@@ -172,15 +172,6 @@ enum class EUpdateResult : uint8
 };
 
 
-/** Indicates the status of the generated Skeletal Mesh. */
-enum class ESkeletalMeshStatus : uint8
-{
-	NotGenerated, // Set only when loading the Instance for the first time or after compiling. Any generation, successful or not, can not end up in this state.
-	Success, // Generated successfully.
-	Error // Generated with errors.
-};
-
-
 /** Instance Update Context.
  * Used to avoid changing the delegate signature in the future.  */
 USTRUCT(BlueprintType)
