[Mutable] Fixed some blueprint support errors:
+Removed Minimal optimization from CompileCustomizableObjectSynchronously function node.
+Fixed crash when CompileCustomizableObjectSynchronously and the CO was null.
+Fixed bug where a COI could not be baked thorugh BP becouse it was open in a COI Editor.
+Fixed bug where a COI could not be baked if the update "failed" with a ErrorOptimization. It should bake the instance anyways.
+Added new variable to store the value of the TextureCompression variable of the last compilation. This avoid the pop up of a warning that is not true.

#rnx

diff --git a/Engine/Plugins/Mutable/Source/CustomizableObject/Internal/MuCO/CustomizableObjectPrivate.h b/Engine/Plugins/Mutable/Source/CustomizableObject/Internal/MuCO/CustomizableObjectPrivate.h
--- a/Engine/Plugins/Mutable/Source/CustomizableObject/Internal/MuCO/CustomizableObjectPrivate.h
+++ b/Engine/Plugins/Mutable/Source/CustomizableObject/Internal/MuCO/CustomizableObjectPrivate.h
@@ -1177,7 +1177,15 @@ struct FModelResources
 
 	UPROPERTY()
 	FString ReleaseVersion;
-	
+
+#if WITH_EDITORONLY_DATA
+	/** Value of the variable TextureCompression in the last compilation of this CO.
+		* this is needed since we can compile a CO through blueprints with a different 
+		* compilation setting than the stored in the COE.*/
+	UPROPERTY()
+	bool bCompiledWithHDTextureCompression = false;
+#endif
+
 #if WITH_EDITORONLY_DATA
 	void CUSTOMIZABLEOBJECT_API Serialize(FObjectAndNameAsStringProxyArchive& Ar, bool bIsCooking);
 	bool CUSTOMIZABLEOBJECT_API Unserialize(FObjectAndNameAsStringProxyArchive& Ar, UCustomizableObject& Outer, const ITargetPlatform* InTargetPlatform, bool bIsCooking);
@@ -1614,6 +1622,8 @@ public:
 
 		RemovedUnnecessarySerializationVersioning,
 
+		AddTextureCompressionSettingCompilationInfo,
+
 		// -----<new versions can be added above this line>--------
 		LastCustomizableObjectVersion
 	};
diff --git a/Engine/Plugins/Mutable/Source/CustomizableObject/Private/MuCO/CustomizableObject.cpp b/Engine/Plugins/Mutable/Source/CustomizableObject/Private/MuCO/CustomizableObject.cpp
--- a/Engine/Plugins/Mutable/Source/CustomizableObject/Private/MuCO/CustomizableObject.cpp
+++ b/Engine/Plugins/Mutable/Source/CustomizableObject/Private/MuCO/CustomizableObject.cpp
@@ -789,6 +789,10 @@ void FModelResources::Serialize(FObjectAndNameAsStringProxyArchive& MemoryWriter
 		MemoryWriter << EditorOnlyMorphTargetReconstructionData;
 		MemoryWriter << EditorOnlyClothingMeshToMeshVertData;
 	}
+
+#if WITH_EDITORONLY_DATA
+	MemoryWriter << bCompiledWithHDTextureCompression;
+#endif
 }
 
 
@@ -944,6 +948,10 @@ bool FModelResources::Unserialize(FObjectAndNameAsStringProxyArchive& MemoryRead
 		MemoryReader << EditorOnlyClothingMeshToMeshVertData;
 	}
 
+#if WITH_EDITORONLY_DATA
+	MemoryReader << bCompiledWithHDTextureCompression;
+#endif
+
 	return true;
 }
 
diff --git a/Engine/Plugins/Mutable/Source/CustomizableObjectEditor/Private/MuCOE/CustomizableObjectCompiler.cpp b/Engine/Plugins/Mutable/Source/CustomizableObjectEditor/Private/MuCOE/CustomizableObjectCompiler.cpp
--- a/Engine/Plugins/Mutable/Source/CustomizableObjectEditor/Private/MuCOE/CustomizableObjectCompiler.cpp
+++ b/Engine/Plugins/Mutable/Source/CustomizableObjectEditor/Private/MuCOE/CustomizableObjectCompiler.cpp
@@ -959,6 +959,8 @@ void FCustomizableObjectCompiler::CompileInternal(bool bAsync)
 		
 		ModelResources.bIsCompiledWithOptimization = GenerationContext.Options.OptimizationLevel == UE_MUTABLE_MAX_OPTIMIZATION;
 
+		ModelResources.bCompiledWithHDTextureCompression = GenerationContext.Options.TextureCompression == ECustomizableObjectTextureCompression::HighQuality;
+
 		CurrentObject->GetPrivate()->GetAlwaysLoadedExtensionData() = MoveTemp(GenerationContext.AlwaysLoadedExtensionData);
 
 		CurrentObject->GetPrivate()->GetStreamedExtensionData().Empty(GenerationContext.StreamedExtensionData.Num());
diff --git a/Engine/Plugins/Mutable/Source/CustomizableObjectEditor/Private/MuCOE/CustomizableObjectEditorFunctionLibrary.cpp b/Engine/Plugins/Mutable/Source/CustomizableObjectEditor/Private/MuCOE/CustomizableObjectEditorFunctionLibrary.cpp
--- a/Engine/Plugins/Mutable/Source/CustomizableObjectEditor/Private/MuCOE/CustomizableObjectEditorFunctionLibrary.cpp
+++ b/Engine/Plugins/Mutable/Source/CustomizableObjectEditor/Private/MuCOE/CustomizableObjectEditorFunctionLibrary.cpp
@@ -19,6 +19,12 @@ ECustomizableObjectCompilationState UCustomizableObjectEditorFunctionLibrary::Co
 	ECustomizableObjectTextureCompression InTextureCompression,
 	bool bGatherReferences)
 {
+	if (!CustomizableObject)
+	{
+		UE_LOG(LogMutable, Warning, TEXT("Compilation Failed: Trying to compile a null CO!"));
+		return ECustomizableObjectCompilationState::Failed;
+	}
+
 	// store package dirty state so that we can restore it - compile is not an edit:
 	const bool bPackageWasDirty = CustomizableObject->GetOutermost()->IsDirty();
 
@@ -26,7 +32,24 @@ ECustomizableObjectCompilationState UCustomizableObjectEditorFunctionLibrary::Co
 
 	TSharedRef<FCompilationRequest> CompileRequest = MakeShared<FCompilationRequest>(*CustomizableObject, false);
 	FCompilationOptions& Options = CompileRequest->GetCompileOptions();
-	Options.OptimizationLevel = static_cast<int32>(InOptimizationLevel);
+
+	switch (InOptimizationLevel)
+	{
+	case ECustomizableObjectOptimizationLevel::None:
+PRAGMA_DISABLE_DEPRECATION_WARNINGS
+	case ECustomizableObjectOptimizationLevel::Minimal:
+PRAGMA_ENABLE_DEPRECATION_WARNINGS
+		Options.OptimizationLevel = 0;
+		break;
+		
+	case ECustomizableObjectOptimizationLevel::Maximum:
+		Options.OptimizationLevel = UE_MUTABLE_MAX_OPTIMIZATION;
+		break;
+
+	default:
+		unimplemented();
+	}
+	
 	Options.TextureCompression = InTextureCompression;
 	Options.bSilentCompilation = false;
 	Options.bGatherReferences = bGatherReferences;
diff --git a/Engine/Plugins/Mutable/Source/CustomizableObjectEditor/Private/MuCOE/CustomizableObjectInstanceBaker.cpp b/Engine/Plugins/Mutable/Source/CustomizableObjectEditor/Private/MuCOE/CustomizableObjectInstanceBaker.cpp
--- a/Engine/Plugins/Mutable/Source/CustomizableObjectEditor/Private/MuCOE/CustomizableObjectInstanceBaker.cpp
+++ b/Engine/Plugins/Mutable/Source/CustomizableObjectEditor/Private/MuCOE/CustomizableObjectInstanceBaker.cpp
@@ -4,6 +4,7 @@
 
 #include "MuCO/CustomizableObject.h"
 #include "MuCO/CustomizableObjectInstance.h"
+#include "MuCO/CustomizableObjectInstancePrivate.h"
 #include "MuCO/CustomizableObjectSystem.h"
 #include "MuCOE/CustomizableObjectInstanceBakingUtils.h"
 
@@ -25,7 +26,9 @@ void UCustomizableObjectInstanceBaker::BakeInstance(UCustomizableObjectInstance*
 		FinishBakingOperation(false, nullptr);
 		return;
 	}
-	InstanceToBake = InTargetInstance;
+
+	// We clone the InTargetInstance to ensure that the baked it is not open in a COI edior.
+	InstanceToBake = InTargetInstance->Clone();
 	
 	if (Configuration.OutputPath.IsEmpty())
 	{
@@ -58,7 +61,7 @@ void UCustomizableObjectInstanceBaker::OnInstanceUpdate(const FUpdateContext& Re
 	const EUpdateResult InstanceUpdateResult = Result.UpdateResult;
 	UE_LOG(LogMutable, Display, TEXT("Instance finished update with state : %s."), *UEnum::GetValueAsString(InstanceUpdateResult));
 	
-	if (UCustomizableObjectSystem::IsUpdateResultValid(InstanceUpdateResult))
+	if (InstanceToBake->GetPrivate()->SkeletalMeshStatus == ESkeletalMeshStatus::Success)
 	{
 		Bake();
 	}
diff --git a/Engine/Plugins/Mutable/Source/CustomizableObjectEditor/Private/MuCOE/CustomizableObjectInstanceBakingUtils.cpp b/Engine/Plugins/Mutable/Source/CustomizableObjectEditor/Private/MuCOE/CustomizableObjectInstanceBakingUtils.cpp
--- a/Engine/Plugins/Mutable/Source/CustomizableObjectEditor/Private/MuCOE/CustomizableObjectInstanceBakingUtils.cpp
+++ b/Engine/Plugins/Mutable/Source/CustomizableObjectEditor/Private/MuCOE/CustomizableObjectInstanceBakingUtils.cpp
@@ -414,7 +414,7 @@ bool BakeCustomizableObjectInstance(
 	// COI Validation completed : Proceed with the baking operation ----------------------------------------------------------------------------------
 	
 	// Notify of better configuration -> Continue operation normally
-	if (InstanceCO->GetPrivate()->GetCompileOptions().TextureCompression != ECustomizableObjectTextureCompression::HighQuality)
+	if (InstanceCO->GetPrivate()->GetModelResources().bCompiledWithHDTextureCompression == false)
 	{
 		FCustomizableObjectEditorLogger::CreateLog(
 		LOCTEXT("CustomizableObjecInstanceBakingUtils_LowQualityTextures", "The Customizable Object wasn't compiled with high quality textures. For the best baking results, change the Texture Compression setting and recompile it."))
diff --git a/Engine/Plugins/Mutable/Source/CustomizableObjectEditor/Public/MuCOE/CustomizableObjectEditorFunctionLibrary.h b/Engine/Plugins/Mutable/Source/CustomizableObjectEditor/Public/MuCOE/CustomizableObjectEditorFunctionLibrary.h
--- a/Engine/Plugins/Mutable/Source/CustomizableObjectEditor/Public/MuCOE/CustomizableObjectEditorFunctionLibrary.h
+++ b/Engine/Plugins/Mutable/Source/CustomizableObjectEditor/Public/MuCOE/CustomizableObjectEditorFunctionLibrary.h
@@ -10,9 +10,9 @@
 UENUM(BlueprintType)
 enum class ECustomizableObjectOptimizationLevel : uint8
 {
-	None,
-	Minimal,
-	Maximum
+	None = 0, 
+	Minimal UE_DEPRECATED(5.6, "Converted internally to None. Use None instead.") = 1,
+	Maximum = 2
 };
 
 
@@ -70,7 +70,7 @@ public:
 	UFUNCTION(BlueprintCallable, Category = "CustomizableObject")
 	static ECustomizableObjectCompilationState CompileCustomizableObjectSynchronously(
 		UCustomizableObject* CustomizableObject, 
-		ECustomizableObjectOptimizationLevel OptimizationLevel = ECustomizableObjectOptimizationLevel::Minimal, 
+		ECustomizableObjectOptimizationLevel OptimizationLevel = ECustomizableObjectOptimizationLevel::None, 
 		ECustomizableObjectTextureCompression TextureCompression = ECustomizableObjectTextureCompression::Fast,
 		bool bGatherReferences = false);
 
@@ -78,3 +78,4 @@ public:
 	UFUNCTION(BlueprintCallable, Category = "CustomizableObject")
 	static UCustomizableObject* NewCustomizableObject(const FNewCustomizableObjectParameters& Parameters);
 };
+
