[mutable] Moved experimental nodes to separate grouping below the standard nodes.
- We want to emphasize they are experimental and that is why they are a bit "out of the way".

diff --git a/Engine/Plugins/Mutable/Source/CustomizableObjectEditor/Private/MuCOE/EdGraphSchema_CustomizableObject.cpp b/Engine/Plugins/Mutable/Source/CustomizableObjectEditor/Private/MuCOE/EdGraphSchema_CustomizableObject.cpp
--- a/Engine/Plugins/Mutable/Source/CustomizableObjectEditor/Private/MuCOE/EdGraphSchema_CustomizableObject.cpp
+++ b/Engine/Plugins/Mutable/Source/CustomizableObjectEditor/Private/MuCOE/EdGraphSchema_CustomizableObject.cpp
@@ -342,7 +342,7 @@ void UEdGraphSchema_CustomizableObject::GetGraphContextActions(FGraphContextMenu
 {
 	//const UCustomizableObjectGraph* CustomizableObjectGraph = CastChecked<UCustomizableObjectGraph>(ContextMenuBuilder.CurrentGraph);
 	//TSharedPtr<FCustomizableObjectSchemaAction_NewNode> Action;
-	int32 GeneralGrouping = 2;
+	const int32 GeneralGrouping = 3;
 
 	const bool bDisableFilter = false;
 
@@ -559,7 +559,7 @@ void UEdGraphSchema_CustomizableObject::GetGraphContextActions(FGraphContextMenu
 			ContextMenuBuilder.CreateTemplateNode<UCustomizableObjectNodeTextureParameter>(),
 		};
 
-		AddNewNodeCategoryActionsFiltered(ExperimentalTemplateNodes, ContextMenuBuilder, NC_Experimental.ToString(), GeneralGrouping, Filter);
+		AddNewNodeCategoryActionsFiltered(ExperimentalTemplateNodes, ContextMenuBuilder, NC_Experimental.ToString(), 2, Filter);
 	}
 
 	{
