[mutable] Reduced fixed CustomizableObject memory usage when using many textures.
- Change data structures so that there are less per-image data permanently allocated.
- Use rom indices as resource indices for textures, since each resources goes to its own rom.
#rnx

diff --git a/Engine/Plugins/Mutable/Source/CustomizableObject/Internal/MuCO/CustomizableObjectPrivate.h b/Engine/Plugins/Mutable/Source/CustomizableObject/Internal/MuCO/CustomizableObjectPrivate.h
--- a/Engine/Plugins/Mutable/Source/CustomizableObject/Internal/MuCO/CustomizableObjectPrivate.h
+++ b/Engine/Plugins/Mutable/Source/CustomizableObject/Internal/MuCO/CustomizableObjectPrivate.h
@@ -1364,6 +1364,8 @@ public:
 	/** Attempts to load the compiled data from DDC. Builds key if not supplied. */
 	void LoadCompiledDataFromDDC(FCompilationOptions Options, UE::DerivedData::ECachePolicy DefaultPolicy, UE::DerivedData::FCacheKey* DDCKey);
 
+	/** Log data for debugging purposes */
+	void LogMemory();
 #endif
 	
 	/** Rebuild ParameterProperties from the current compiled model. */
@@ -1624,6 +1626,8 @@ public:
 
 		AddTextureCompressionSettingCompilationInfo,
 
+		RestructureConstantImageData,
+
 		// -----<new versions can be added above this line>--------
 		LastCustomizableObjectVersion
 	};
diff --git a/Engine/Plugins/Mutable/Source/CustomizableObject/Private/MuCO/CustomizableObject.cpp b/Engine/Plugins/Mutable/Source/CustomizableObject/Private/MuCO/CustomizableObject.cpp
--- a/Engine/Plugins/Mutable/Source/CustomizableObject/Private/MuCO/CustomizableObject.cpp
+++ b/Engine/Plugins/Mutable/Source/CustomizableObject/Private/MuCO/CustomizableObject.cpp
@@ -31,6 +31,7 @@
 #include "MuCO/UnrealMutableModelDiskStreamer.h"
 #include "MuCO/UnrealPortabilityHelpers.h"
 #include "MuR/Model.h"
+#include "MuR/ModelPrivate.h"
 #include "MuR/Operations.h"
 #include "PhysicsEngine/PhysicsAsset.h"
 #include "Serialization/MemoryReader.h"
@@ -3374,6 +3375,169 @@ void UCustomizableObjectPrivate::LoadCompiledDataFromDDC(FCompilationOptions Opt
 }
 
 
+class FMutableMemoryCounterArchive final : public FArchive
+{
+public:
+	FMutableMemoryCounterArchive()
+	{
+		SetIsSaving(true);
+		SetIsPersistent(true);
+		ArIsCountingMemory = true;
+	}
+
+	virtual void Serialize(void*, int64 Length) override { Size += Length; }
+	virtual int64 TotalSize() override { return Size; }
+
+private:
+	int64 Size = 0;
+};
+
+
+void UCustomizableObjectPrivate::LogMemory()
+{
+	// Log in-memory data for the compiled CustomizableObject
+	UE_LOG(LogMutable, Log, TEXT("CustomizableObject [%s] memory report:"), *GetPublic()->GetName());
+	const mu::FModel* MuModel = GetModel().Get();
+	if (MuModel)
+	{
+		const mu::FProgram& Program = MuModel->GetPrivate()->Program;
+		int32 ByteCode = Program.ByteCode.GetAllocatedSize();
+		int32 OpAddress = Program.OpAddress.GetAllocatedSize();
+		int32 Roms = Program.Roms.GetAllocatedSize();
+		int32 ConstantImages = Program.ConstantImages.GetAllocatedSize();
+		int32 ConstantImageLODIndices = Program.ConstantImageLODIndices.GetAllocatedSize();
+		int32 ConstantImageLODPermanent = Program.ConstantImageLODsPermanent.GetAllocatedSize();
+		int32 ConstantImageLODStreamed = Program.ConstantImageLODsStreamed.GetAllocatedSize();
+
+		int32 PermanentImages = 0;
+		for (const TSharedPtr<const mu::FImage>& Entry : Program.ConstantImageLODsPermanent)
+		{
+			if (Entry)
+			{
+				PermanentImages += Entry->GetDataSize();
+			}
+		}
+		int32 ImageTotal = ConstantImages + ConstantImageLODIndices + ConstantImageLODPermanent + PermanentImages + ConstantImageLODStreamed;
+
+		int32 ConstantMeshes = Program.ConstantMeshes.GetAllocatedSize();
+		int32 PermanentMeshes = 0;
+		for (const TPair<int32, TSharedPtr<const mu::FMesh>>& Entry : Program.ConstantMeshes)
+		{
+			if (Entry.Key == -1)
+			{
+				PermanentMeshes += Entry.Value->GetDataSize();
+			}
+		}
+		int32 MeshesTotal = ConstantMeshes + PermanentMeshes;
+
+		int32 StringsTotal = Program.ConstantStrings.GetAllocatedSize();
+		for (const FString& Entry : Program.ConstantStrings)
+		{
+			StringsTotal += Entry.GetAllocatedSize();
+		}
+
+		int32 SkeletonsTotal = Program.ConstantSkeletons.GetAllocatedSize();
+		for (const TSharedPtr<const mu::FSkeleton>& Entry : Program.ConstantSkeletons)
+		{
+			SkeletonsTotal += sizeof(mu::FSkeleton) + Entry->DebugBoneNames.GetAllocatedSize() + Entry->BoneIds.GetAllocatedSize() + Entry->BoneParents.GetAllocatedSize();
+		}
+
+		int32 PhysicsTotal = Program.ConstantSkeletons.GetAllocatedSize();
+		for (const TSharedPtr<const mu::FPhysicsBody>& Entry : Program.ConstantPhysicsBodies)
+		{
+			PhysicsTotal += sizeof(mu::FPhysicsBody) + Entry->BoneIds.GetAllocatedSize() + Entry->Bodies.GetAllocatedSize() + Entry->BodiesCustomIds.GetAllocatedSize();
+			for (const mu::FPhysicsBodyAggregate& Body : Entry->Bodies)
+			{
+				PhysicsTotal += Body.Spheres.GetAllocatedSize() + Body.Boxes.GetAllocatedSize() + Body.Convex.GetAllocatedSize() + Body.Sphyls.GetAllocatedSize() + Body.TaperedCapsules.GetAllocatedSize();
+			}
+		}
+
+		int32 ParametersTotal = Program.Parameters.GetAllocatedSize();
+		for (const mu::FParameterDesc& Entry : Program.Parameters)
+		{
+			ParametersTotal += Entry.Ranges.GetAllocatedSize() + Entry.PossibleValues.GetAllocatedSize();
+			for (const mu::FParameterDesc::FIntValueDesc& ValueDesc : Entry.PossibleValues)
+			{
+				ParametersTotal += ValueDesc.Name.GetAllocatedSize();
+			}
+		}
+
+		int32 ModelTotal = ByteCode + OpAddress + Roms 
+			+ ImageTotal + MeshesTotal 
+			+ StringsTotal + SkeletonsTotal + PhysicsTotal
+			+ ParametersTotal;
+
+
+		int32 StreamableData = 0;
+		if (ModelStreamableData && ModelStreamableData->StreamingData)
+		{
+			const FModelStreamableBulkData& Bulk = *ModelStreamableData->StreamingData;
+			StreamableData += Bulk.ModelStreamables.GetAllocatedSize();
+			StreamableData += Bulk.ClothingStreamables.GetAllocatedSize();
+			StreamableData += Bulk.RealTimeMorphStreamables.GetAllocatedSize();
+			StreamableData += Bulk.StreamableBulkData.GetAllocatedSize();
+		}
+
+		int32 ModeleResData = 0;
+		{
+			ModeleResData += sizeof(FModelResources);
+			
+			FMutableMemoryCounterArchive Arch;
+			Arch << ModelResources.ReferenceSkeletalMeshesData;
+			Arch << ModelResources.Skeletons;
+			Arch << ModelResources.Materials;
+			Arch << ModelResources.PassThroughTextures;
+			Arch << ModelResources.PassThroughMeshes;
+			Arch << ModelResources.PhysicsAssets;
+			Arch << ModelResources.AnimBPs;
+			Arch << ModelResources.AnimBpOverridePhysiscAssetsInfo;
+			Arch << ModelResources.MaterialSlotNames;
+			Arch << ModelResources.BoneNamesMap;
+			Arch << ModelResources.SocketArray;
+			Arch << ModelResources.SkinWeightProfilesInfo;
+			Arch << ModelResources.MeshMetadata;
+			Arch << ModelResources.SurfaceMetadata;
+			Arch << ModelResources.ParameterUIDataMap;
+			Arch << ModelResources.StateUIDataMap;
+			Arch << ModelResources.ClothSharedConfigsData;
+			Arch << ModelResources.ClothingAssetsData;
+			Arch << ModelResources.ComponentNamesPerObjectComponent;
+
+			ModeleResData += ModelResources.StreamedResourceData.GetAllocatedSize();
+			for (FCustomizableObjectStreamedResourceData& Data : ModelResources.StreamedResourceData)
+			{
+				TSoftObjectPtr<UCustomizableObjectResourceDataContainer> PathCopy = Data.GetPath();
+				Arch << PathCopy;
+			}
+
+			ModeleResData += Arch.TotalSize();
+		}
+
+		UE_LOG(LogMutable, Log, TEXT("Total              : %8d"), ModelTotal+StreamableData+ModeleResData);
+		UE_LOG(LogMutable, Log, TEXT("  Streamable       : %8d"), StreamableData);
+		UE_LOG(LogMutable, Log, TEXT("  ModelResources   : %8d"), ModeleResData);
+		UE_LOG(LogMutable, Log, TEXT("  Model            : %8d"), ModelTotal);
+		UE_LOG(LogMutable, Log, TEXT("    ByteCode       : %8d"), ByteCode);
+		UE_LOG(LogMutable, Log, TEXT("    OpAddress      : %8d"), OpAddress);
+		UE_LOG(LogMutable, Log, TEXT("    Roms           : %8d"), Roms);
+		UE_LOG(LogMutable, Log, TEXT("    Strings        : %8d"), StringsTotal);
+		UE_LOG(LogMutable, Log, TEXT("    Images         : %8d"), ImageTotal);
+		UE_LOG(LogMutable, Log, TEXT("      Buffer       : %8d"), ConstantImages);
+		UE_LOG(LogMutable, Log, TEXT("      LODIndices   : %8d"), ConstantImageLODIndices);
+		UE_LOG(LogMutable, Log, TEXT("      PermanentBuf : %8d"), ConstantImageLODPermanent);
+		UE_LOG(LogMutable, Log, TEXT("      Permanent    : %8d"), PermanentImages);
+		UE_LOG(LogMutable, Log, TEXT("      StreamedBuf  : %8d"), ConstantImageLODStreamed);
+		UE_LOG(LogMutable, Log, TEXT("    Meshes         : %8d"), MeshesTotal);
+		UE_LOG(LogMutable, Log, TEXT("      Buffer       : %8d"), ConstantMeshes);
+		UE_LOG(LogMutable, Log, TEXT("      Permanent    : %8d"), PermanentMeshes);
+		UE_LOG(LogMutable, Log, TEXT("    Skeletons      : %8d"), SkeletonsTotal);
+		UE_LOG(LogMutable, Log, TEXT("    PhysicBodies   : %8d"), PhysicsTotal);
+		UE_LOG(LogMutable, Log, TEXT("    Parameters     : %8d"), ParametersTotal);
+	}
+
+}
+
+
 #endif // WITH_EDITOR
 
 //-------------------------------------------------------------------------------------------------
diff --git a/Engine/Plugins/Mutable/Source/CustomizableObject/Private/MuCO/CustomizableObjectSystem.cpp b/Engine/Plugins/Mutable/Source/CustomizableObject/Private/MuCO/CustomizableObjectSystem.cpp
--- a/Engine/Plugins/Mutable/Source/CustomizableObject/Private/MuCO/CustomizableObjectSystem.cpp
+++ b/Engine/Plugins/Mutable/Source/CustomizableObject/Private/MuCO/CustomizableObjectSystem.cpp
@@ -128,6 +128,14 @@ static FAutoConsoleVariableRef CVarMutableEnableLODManagmentInEditor(
 	bEnableLODManagmentInEditor,
 	TEXT("true/false - If true, enables custom LODManagment in the editor. "),
 	ECVF_Default);
+
+
+TAutoConsoleVariable<bool> CVarMutableLogObjectMemoryOnUpdate(
+	TEXT("mutable.LogObjectMemoryOnUpdate"),
+	false,
+	TEXT("Log the memory used for a CO on every update."),
+	ECVF_Scalability);
+
 #endif
 
 TAutoConsoleVariable<bool> CVarEnableReleaseMeshResources(
@@ -3568,6 +3576,14 @@ void UCustomizableObjectSystem::AdvanceCurrentOperation()
 	{
 		MUTABLE_CPUPROFILER_SCOPE(OperationUpdate);
 
+#if WITH_EDITOR
+		if (CVarMutableLogObjectMemoryOnUpdate.GetValueOnAnyThread())
+		{
+			UCustomizableObject* CO = Private->CurrentMutableOperation->Instance->GetCustomizableObject();
+			CO->GetPrivate()->LogMemory();
+		}
+#endif
+
 		// Start the first task of the update process. See namespace impl comments above.
 		impl::Task_Game_StartUpdate(Private->CurrentMutableOperation.ToSharedRef());
 	}
diff --git a/Engine/Plugins/Mutable/Source/MutableRuntime/Internal/MuR/ModelPrivate.h b/Engine/Plugins/Mutable/Source/MutableRuntime/Internal/MuR/ModelPrivate.h
--- a/Engine/Plugins/Mutable/Source/MutableRuntime/Internal/MuR/ModelPrivate.h
+++ b/Engine/Plugins/Mutable/Source/MutableRuntime/Internal/MuR/ModelPrivate.h
@@ -26,6 +26,16 @@ namespace mu
 	constexpr int32 DebugRomIndex = 44;
 	constexpr int32 DebugImageIndex = 9;
 
+	struct FConstantResourceIndex
+	{
+		uint32 Index : 31;
+		/** This may mean that the resource needs to be looked up on a different array. */
+		uint32 Streamable : 1;
+	};
+	static_assert(sizeof(FConstantResourceIndex) == 4);
+	MUTABLE_DEFINE_POD_SERIALISABLE(FConstantResourceIndex);
+	MUTABLE_DEFINE_POD_VECTOR_SERIALISABLE(FConstantResourceIndex);
+
     //! Data stored for a rom even if it is not loaded
     struct FRomData
     {
@@ -38,8 +48,10 @@ namespace mu
 		//! Size of the rom
 		uint32 Size=0;
 
-		//! Index of the resource in its type-specific array
-		uint32 ResourceIndex=0;
+		/** Index of the resource in its type-specific arrays. 
+		* TODO: Remove when streamed meshes also use the rom index as resource index.
+		*/
+		FConstantResourceIndex ResourceIndex = {};
 
 		//! Index of the resource in its type-specific array (one of DATATYPE values)
 		uint16 ResourceType=0;
@@ -63,7 +75,6 @@ namespace mu
 		FMemory::Memcpy (&Code[Pos], &Data, sizeof(DATA));
     }
 
-
 	//!
 	struct FImageLODRange
 	{
@@ -232,11 +243,19 @@ namespace mu
 		/** Loaded roms worth tracking (only images and meshes for now). Stores the rom's data type.*/
 		TSparseArray<uint8> LoadedMemTrackedRoms;
 
-		//! Constant image mip data: the first is the index in Roms for each image mip or -1 if it is always loaded.
-		TArray<TPair<int32, TSharedPtr<const FImage>>> ConstantImageLODs;
+		/** Constant image mip data is split in 2 sets: ConstantImageLODsPermanent constains data that is always loaded. 
+		* Index with FConstantImageLODIndex::Index, when Streamable is 0.
+		*/
+		TArray<TSharedPtr<const FImage>> ConstantImageLODsPermanent;
+
+		/** Constant image mip data is split in 2 sets: ConstantImageLODsStreamed constains data that is streamed in and out. 
+		* Index with FConstantImageLODIndex::Index, when Streamable is 1.
+		* This part is empty for an unused Model, and shouldn't be serialised.
+		*/
+		TMap<uint32, TSharedPtr<const FImage>> ConstantImageLODsStreamed;
 
 		//! Constant image mip chain indices: ranges in this array are defined in FImageLODRange and the indices here refer to ConstantImageLODs.
-		TArray<uint32> ConstantImageLODIndices;
+		TArray<FConstantResourceIndex> ConstantImageLODIndices;
 
 		//! Constant image data.
 		TArray<FImageLODRange> ConstantImages;
@@ -295,7 +314,7 @@ namespace mu
             arch << ByteCode;
             arch << States;
 			arch << Roms;
-			arch << ConstantImageLODs;
+			arch << ConstantImageLODsPermanent;
 			arch << ConstantImageLODIndices;
 			arch << ConstantImages;
 			arch << ConstantMeshes;
@@ -320,7 +339,7 @@ namespace mu
             arch >> ByteCode;
             arch >> States;
 			arch >> Roms;
-			arch >> ConstantImageLODs;
+			arch >> ConstantImageLODsPermanent;
 			arch >> ConstantImageLODIndices;
 			arch >> ConstantImages;
 			arch >> ConstantMeshes;
@@ -347,12 +366,18 @@ namespace mu
 		//! Return true if the given ROM is loaded.
 		FORCEINLINE bool IsRomLoaded(int32 RomIndex) const
 		{
+			FConstantResourceIndex ResourceIndex = Roms[RomIndex].ResourceIndex;
+			if (!ResourceIndex.Streamable)
+			{
+				return true;
+			}
+
 			switch (Roms[RomIndex].ResourceType)
 			{
 				case DT_IMAGE: 
-					return ConstantImageLODs[Roms[RomIndex].ResourceIndex].Value != nullptr;
+					return ConstantImageLODsStreamed.Contains(ResourceIndex.Index);
 				case DT_MESH: 
-					return ConstantMeshes[Roms[RomIndex].ResourceIndex].Value != nullptr;
+					return ConstantMeshes[ResourceIndex.Index].Value != nullptr;
 				default:
 					check(false);
 					break;
@@ -374,20 +399,26 @@ namespace mu
 				LoadedMemTrackedRoms.RemoveAt(RomIndex);
 			}
 
+			FConstantResourceIndex ResourceIndex = Roms[RomIndex].ResourceIndex;
+			check(ResourceIndex.Streamable);
+
 			switch (Roms[RomIndex].ResourceType)
 			{
 			case DT_IMAGE:
-				if (ConstantImageLODs[Roms[RomIndex].ResourceIndex].Value)
+			{
+				TSharedPtr<const FImage> Data;
+				ConstantImageLODsStreamed.RemoveAndCopyValue(ResourceIndex.Index, Data);
+				if (Data)
 				{
-					RomSize = ConstantImageLODs[Roms[RomIndex].ResourceIndex].Value->GetDataSize();
-					ConstantImageLODs[Roms[RomIndex].ResourceIndex].Value = nullptr;
+					RomSize = Data->GetDataSize();
 				}
 				break;
+			}
 			case DT_MESH:
-				if (ConstantMeshes[Roms[RomIndex].ResourceIndex].Value)
+				if (ConstantMeshes[ResourceIndex.Index].Value)
 				{
-					RomSize = ConstantMeshes[Roms[RomIndex].ResourceIndex].Value->GetDataSize();
-					ConstantMeshes[Roms[RomIndex].ResourceIndex].Value = nullptr;
+					RomSize = ConstantMeshes[ResourceIndex.Index].Value->GetDataSize();
+					ConstantMeshes[ResourceIndex.Index].Value = nullptr;
 				}
 				break;
 			default:
@@ -402,16 +433,21 @@ namespace mu
 		{
 			check(Roms[RomIndex].ResourceType == DT_MESH);
 			
+			FConstantResourceIndex ResourceIndex = Roms[RomIndex].ResourceIndex;
+			check(ResourceIndex.Streamable);
+
 			LoadedMemTrackedRoms.EmplaceAt(RomIndex, (uint8)Roms[RomIndex].ResourceType);
-			ConstantMeshes[Roms[RomIndex].ResourceIndex].Value = Value;
+			ConstantMeshes[ResourceIndex.Index].Value = Value;
 		}
 
 		FORCEINLINE void SetImageRomValue(int32 RomIndex, const TSharedPtr<FImage>& Value)
 		{
 			check(Roms[RomIndex].ResourceType == DT_IMAGE);
-			
+
+			check(!ConstantImageLODsStreamed.Contains(RomIndex));
+
 			LoadedMemTrackedRoms.EmplaceAt(RomIndex, (uint8)Roms[RomIndex].ResourceType);
-			ConstantImageLODs[Roms[RomIndex].ResourceIndex].Value = Value;
+			ConstantImageLODsStreamed.Add(RomIndex, Value );
 		}
 
 
@@ -569,6 +605,23 @@ namespace mu
             return index;
         }
 
+
+		inline TSharedPtr<const FImage> GetImageLOD(FConstantResourceIndex Index) const
+		{
+			TSharedPtr<const FImage> Mip;
+			if (!Index.Streamable)
+			{
+				Mip = ConstantImageLODsPermanent[Index.Index];
+			}
+			else
+			{
+				Mip = ConstantImageLODsStreamed[Index.Index];
+			}
+
+			return Mip;
+		}
+
+
         //! Get a constant image, assuming at least some mips are loaded. The image constant will be composed with lodaded mips if necessary.
 		template <typename CreateImageFunc>
         void GetConstant( int32 ConstantIndex, TSharedPtr<const FImage>& res, int32 MipsToSkip, const CreateImageFunc& CreateImage) const
@@ -578,8 +631,8 @@ namespace mu
 
 			// Get the first mip
 			int32 ResultLODIndexIndex = FirstLODIndexIndex + ReallySkippedLODs;
-			int32 ResultLODIndex = ConstantImageLODIndices[ResultLODIndexIndex];
-			TSharedPtr<const FImage> CurrentMip = ConstantImageLODs[ResultLODIndex].Value;
+			FConstantResourceIndex ResultLODIndex = ConstantImageLODIndices[ResultLODIndexIndex];
+			TSharedPtr<const FImage> CurrentMip = GetImageLOD(ResultLODIndex);
 
 			// We may need to skip more LODs if they are not loaded
 			while (!CurrentMip)
@@ -595,7 +648,7 @@ namespace mu
 
 				++ResultLODIndexIndex;
 				ResultLODIndex = ConstantImageLODIndices[ResultLODIndexIndex];
-				CurrentMip = ConstantImageLODs[ResultLODIndex].Value;
+				CurrentMip = GetImageLOD(ResultLODIndex);
 			}
 
 			int32 FinalLODs = ConstantImages[ConstantIndex].LODCount - ReallySkippedLODs;
@@ -620,15 +673,17 @@ namespace mu
 				{
 					for (int32 LOD = 0; LOD < FinalLODs; ++LOD)
 					{
-						int32 LODIndex = ConstantImageLODIndices[ResultLODIndexIndex + LOD];
+						FConstantResourceIndex LODIndex = ConstantImageLODIndices[ResultLODIndexIndex + LOD];
+
+						TSharedPtr<const FImage> Image = GetImageLOD(LODIndex);
 
 						// This could happen in the case of missing data files
-						if (!ConstantImageLODs[LODIndex].Value)
+						if (!Image)
 						{
 							break;
 						}
 
-						int32 MipSizeBytes = ConstantImageLODs[LODIndex].Value->GetLODDataSize(0);
+						int32 MipSizeBytes = Image->GetLODDataSize(0);
 						Result->DataStorage.ResizeLOD(LOD, MipSizeBytes);
 					}
 				}
@@ -648,7 +703,7 @@ namespace mu
 					if (LOD + 1 < FinalLODs)
 					{
 						ResultLODIndex = ConstantImageLODIndices[ResultLODIndexIndex + LOD + 1];
-						CurrentMip = ConstantImageLODs[ResultLODIndex].Value;
+						CurrentMip = GetImageLOD(ResultLODIndex);
 
 						// This could only happen if missing or corrupted data.
 						if (!CurrentMip)
diff --git a/Engine/Plugins/Mutable/Source/MutableRuntime/Internal/MuR/SystemPrivate.h b/Engine/Plugins/Mutable/Source/MutableRuntime/Internal/MuR/SystemPrivate.h
--- a/Engine/Plugins/Mutable/Source/MutableRuntime/Internal/MuR/SystemPrivate.h
+++ b/Engine/Plugins/Mutable/Source/MutableRuntime/Internal/MuR/SystemPrivate.h
@@ -1621,11 +1621,11 @@ namespace mu
 			{
 				// Count streamable and currently-loaded resources
 				const FProgram& Program = Model->GetPrivate()->Program;
-				for (const TPair<int32, TSharedPtr<const FImage>>& Rom : Program.ConstantImageLODs)
+				for (const TPair<uint32,TSharedPtr<const FImage>>& ImageLOD : Program.ConstantImageLODsStreamed)
 				{
-					if (Rom.Value && Rom.Key >= 0)
+					if (ImageLOD.Value)
 					{
-						Result += Rom.Value->GetDataSize();
+						Result += ImageLOD.Value->GetDataSize();
 					}
 				}
 				for (const TPair<int32, TSharedPtr<const FMesh>>& Rom : Program.ConstantMeshes)
diff --git a/Engine/Plugins/Mutable/Source/MutableRuntime/Private/MuR/CodeRunnerTasked.cpp b/Engine/Plugins/Mutable/Source/MutableRuntime/Private/MuR/CodeRunnerTasked.cpp
--- a/Engine/Plugins/Mutable/Source/MutableRuntime/Private/MuR/CodeRunnerTasked.cpp
+++ b/Engine/Plugins/Mutable/Source/MutableRuntime/Private/MuR/CodeRunnerTasked.cpp
@@ -2183,14 +2183,14 @@ namespace mu
 		{
 			bool bSuccess = Runner->System->StreamInterface->EndRead(RomLoadOp->m_streamID);
 
-			const int32 ResIndex = Program.Roms[RomIndex].ResourceIndex;
-			check(!Program.ConstantMeshes[ResIndex].Value)
+			const FConstantResourceIndex ResIndex = Program.Roms[RomIndex].ResourceIndex;
+			check(!Program.ConstantMeshes[ResIndex.Index].Value)
 			MUTABLE_CPUPROFILER_SCOPE(Unserialise);
 
 			FInputMemoryStream Stream(RomLoadOp->m_streamBuffer.GetData(), RomLoadOp->m_streamBuffer.Num());
 			FInputArchive Arch(&Stream);
 
-			check(!Program.ConstantMeshes[ResIndex].Value);
+			check(!Program.ConstantMeshes[ResIndex.Index].Value);
 			TSharedPtr<FMesh> Value;
 			if (bSuccess)
 			{
@@ -2204,7 +2204,7 @@ namespace mu
 
 			Program.SetMeshRomValue(RomIndex, Value);
 
-			check(Program.ConstantMeshes[ResIndex].Value);
+			check(Program.ConstantMeshes[ResIndex.Index].Value);
 
 			Runner->RomLoadOps.Remove(*RomLoadOp);
 		}
@@ -2323,7 +2323,7 @@ namespace mu
 		MUTABLE_CPUPROFILER_SCOPE(FLoadImageRomsTask_Prepare);
 		bOutFailed = false;
 
-		FProgram& program = Runner->Model->GetPrivate()->Program;
+		FProgram& Program = Runner->Model->GetPrivate()->Program;
 
 		FWorkingMemoryManager::FModelCacheEntry* ModelCache = Runner->System->WorkingMemoryManager.FindModelCache(Runner->Model.Get());
 
@@ -2334,23 +2334,23 @@ namespace mu
 		for (int32 LODIndex = 0; LODIndex < LODIndexCount; ++LODIndex)
 		{
 			const int32 CurrentIndexIndex = LODIndexIndex + LODIndex;
-			const int32 CurrentIndex = program.ConstantImageLODIndices[CurrentIndexIndex];
+			const FConstantResourceIndex CurrentIndex = Program.ConstantImageLODIndices[CurrentIndexIndex];
 
-			if (program.ConstantImageLODs[CurrentIndex].Key < 0)
+			if (!CurrentIndex.Streamable)
 			{
 				// This data is always resident.
 				continue;
 			}
 
-			int32 RomIndex = program.ConstantImageLODs[CurrentIndex].Key;
-			check(RomIndex < program.Roms.Num());
+			int32 RomIndex = CurrentIndex.Index;
+			check(RomIndex < Program.Roms.Num());
 
 			++ModelCache->PendingOpsPerRom[RomIndex];
 
 			if (DebugRom && (DebugRomAll || RomIndex == DebugRomIndex))
 				UE_LOG(LogMutableCore, Log, TEXT("Preparing rom %d, now peding ops is %d."), RomIndex, ModelCache->PendingOpsPerRom[RomIndex]);
 
-			if (program.IsRomLoaded(RomIndex))
+			if (Program.IsRomLoaded(RomIndex))
 			{
 				continue;
 			}
@@ -2367,7 +2367,7 @@ namespace mu
 			
 			check(Runner->System->StreamInterface);
 
-			const uint32 RomSize = program.Roms[RomIndex].Size;
+			const uint32 RomSize = Program.Roms[RomIndex].Size;
 			check(RomSize > 0);
 
 			// Free roms if necessary
@@ -2389,7 +2389,7 @@ namespace mu
 				ReadCompletionEvent.Trigger();
 			};
 			
-			const uint32 RomId = program.Roms[RomIndex].Id;
+			const uint32 RomId = Program.Roms[RomIndex].Id;
 			RomLoadOp.m_streamID = Runner->System->StreamInterface->BeginReadBlock(Runner->Model.Get(), RomId, RomLoadOp.m_streamBuffer.GetData(), RomSize, &Callback);
 			if (RomLoadOp.m_streamID < 0)
 			{
@@ -2440,14 +2440,10 @@ namespace mu
 					FInputMemoryStream Stream(RomLoadOp->m_streamBuffer.GetData(), RomLoadOp->m_streamBuffer.Num());
 					FInputArchive Arch(&Stream);
 
-					const int32 ResIndex = Program.Roms[RomIndex].ResourceIndex;
-
 					// TODO: Try to reuse buffer from PooledImages.
-					check(!Program.ConstantImageLODs[ResIndex].Value);
 					TSharedPtr<FImage> Value = FImage::StaticUnserialise(Arch);
 
 					Program.SetImageRomValue(RomIndex, Value);
-					check(Program.ConstantImageLODs[ResIndex].Value);
 				}
 				else
 				{
@@ -2470,15 +2466,15 @@ namespace mu
 		for (int32 LODIndex = 0; bSuccess && (LODIndex < LODIndexCount); ++LODIndex)
 		{
 			int32 CurrentIndexIndex = LODIndexIndex + LODIndex;
-			int32 CurrentIndex = Program.ConstantImageLODIndices[CurrentIndexIndex];
+			FConstantResourceIndex CurrentIndex = Program.ConstantImageLODIndices[CurrentIndexIndex];
 
-			if (Program.ConstantImageLODs[CurrentIndex].Key < 0)
+			if (!CurrentIndex.Streamable)
 			{
 				// This data is always resident.
 				continue;
 			}
 
-			int32 RomIndex = Program.ConstantImageLODs[CurrentIndex].Key;
+			int32 RomIndex = CurrentIndex.Index;
 			check(RomIndex < Program.Roms.Num());
 
 			Runner->System->WorkingMemoryManager.MarkRomUsed(RomIndex, Runner->Model);
@@ -2636,17 +2632,17 @@ namespace mu
 	{
 		TSharedPtr<FIssuedTask> Issued;
 
-		FProgram& program = Model->GetPrivate()->Program;
+		FProgram& Program = Model->GetPrivate()->Program;
 
-		EOpType type = program.GetOpType(item.At);
+		EOpType type = Program.GetOpType(item.At);
 
 		switch (type)
 		{
 		case EOpType::ME_CONSTANT:
 		{
-			OP::MeshConstantArgs args = program.GetOpArgs<OP::MeshConstantArgs>(item.At);
-			int32 RomIndex = program.ConstantMeshes[args.value].Key;
-			if (RomIndex >= 0 && !program.ConstantMeshes[args.value].Value)
+			OP::MeshConstantArgs args = Program.GetOpArgs<OP::MeshConstantArgs>(item.At);
+			int32 RomIndex = Program.ConstantMeshes[args.value].Key;
+			if (RomIndex >= 0 && !Program.ConstantMeshes[args.value].Value)
 			{
 				Issued = MakeShared<FLoadMeshRomTask>(item, RomIndex);
 			}
@@ -2659,9 +2655,9 @@ namespace mu
 		
 		case EOpType::ED_CONSTANT:
 		{
-			OP::ResourceConstantArgs Args = program.GetOpArgs<OP::ResourceConstantArgs>(item.At);
+			OP::ResourceConstantArgs Args = Program.GetOpArgs<OP::ResourceConstantArgs>(item.At);
 
-			const FExtensionDataConstant::ELoadState LoadState = program.ConstantExtensionData[Args.value].LoadState;
+			const FExtensionDataConstant::ELoadState LoadState = Program.ConstantExtensionData[Args.value].LoadState;
 
 			check(LoadState != FExtensionDataConstant::ELoadState::Invalid);
 			
@@ -2672,19 +2668,19 @@ namespace mu
 			}
 			else
 			{
-				Issued = MakeShared<FLoadExtensionDataTask>(item, program.ConstantExtensionData[Args.value].Data, Args.value);
+				Issued = MakeShared<FLoadExtensionDataTask>(item, Program.ConstantExtensionData[Args.value].Data, Args.value);
 			}
 			break;
 		}
 
 		case EOpType::IM_CONSTANT:
 		{
-			OP::ResourceConstantArgs args = program.GetOpArgs<OP::ResourceConstantArgs>(item.At);
+			OP::ResourceConstantArgs args = Program.GetOpArgs<OP::ResourceConstantArgs>(item.At);
 			int32 MipsToSkip = item.ExecutionOptions;
 			int32 ImageIndex = args.value;
-			int32 ReallySkip = FMath::Min(MipsToSkip, program.ConstantImages[ImageIndex].LODCount - 1);
-			int32 LODIndexIndex = program.ConstantImages[ImageIndex].FirstIndex + ReallySkip;
-			int32 LODIndexCount = program.ConstantImages[ImageIndex].LODCount - ReallySkip;
+			int32 ReallySkip = FMath::Min(MipsToSkip, Program.ConstantImages[ImageIndex].LODCount - 1);
+			int32 LODIndexIndex = Program.ConstantImages[ImageIndex].FirstIndex + ReallySkip;
+			int32 LODIndexCount = Program.ConstantImages[ImageIndex].LODCount - ReallySkip;
 			check(LODIndexCount > 0);
 
 			// We always need to follow this path, or roms may not be protected for long enough and might be unloaded 
@@ -2693,8 +2689,8 @@ namespace mu
 			//bool bAnyMissing = false;
 			//for (int32 i=0; i<LODIndexCount; ++i)
 			//{
-			//	uint32 LODIndex = program.ConstantImageLODIndices[LODIndexIndex+i];
-			//	if ( !program.ConstantImageLODs[LODIndex].Value )
+			//	uint32 LODIndex = Program.ConstantImageLODIndices[LODIndexIndex+i];
+			//	if ( !Program.ConstantImageLODs[LODIndex].Value )
 			//	{
 			//		bAnyMissing = true;
 			//		break;
@@ -2719,7 +2715,7 @@ namespace mu
 
 		case EOpType::IM_PARAMETER:
 		{
-			OP::ParameterArgs args = program.GetOpArgs<OP::ParameterArgs>(item.At);
+			OP::ParameterArgs args = Program.GetOpArgs<OP::ParameterArgs>(item.At);
 			TSharedPtr<FRangeIndex> Index = BuildCurrentOpRangeIndex(item, Params, Model.Get(), args.variable);
 
 			const FName Id = Params->GetImageValue(args.variable, Index.Get());
@@ -2776,7 +2772,7 @@ namespace mu
 		{
 			if (item.Stage == 1)
 			{
-				OP::ImagePixelFormatArgs Args = program.GetOpArgs<OP::ImagePixelFormatArgs>(item.At);
+				OP::ImagePixelFormatArgs Args = Program.GetOpArgs<OP::ImagePixelFormatArgs>(item.At);
 				Issued = MakeShared<FImagePixelFormatTask>(item, Args);
 			}
 			break;
@@ -2786,7 +2782,7 @@ namespace mu
 		{
 			if (item.Stage == 1)
 			{
-				OP::ImageLayerColourArgs Args = program.GetOpArgs<OP::ImageLayerColourArgs>(item.At);
+				OP::ImageLayerColourArgs Args = Program.GetOpArgs<OP::ImageLayerColourArgs>(item.At);
 				Issued = MakeShared<FImageLayerColourTask>(item, Args);
 			}
 			break;
@@ -2799,7 +2795,7 @@ namespace mu
 				(ExecutionStrategy != EExecutionStrategy::MinimizeMemory && item.Stage == 1)
 				)
 			{
-				OP::ImageLayerArgs Args = program.GetOpArgs<OP::ImageLayerArgs>(item.At);
+				OP::ImageLayerArgs Args = Program.GetOpArgs<OP::ImageLayerArgs>(item.At);
 				Issued = MakeShared<FImageLayerTask>(item, Args);
 			}
 			break;
@@ -2809,7 +2805,7 @@ namespace mu
 		{
 			if (item.Stage == 1)
 			{
-				OP::ImageMipmapArgs Args = program.GetOpArgs<OP::ImageMipmapArgs>(item.At);
+				OP::ImageMipmapArgs Args = Program.GetOpArgs<OP::ImageMipmapArgs>(item.At);
 				Issued = MakeShared<FImageMipmapTask>(item, Args);
 			}
 			break;
@@ -2819,7 +2815,7 @@ namespace mu
 		{
 			if (item.Stage == 1)
 			{
-				OP::ImageSwizzleArgs Args = program.GetOpArgs<OP::ImageSwizzleArgs>(item.At);
+				OP::ImageSwizzleArgs Args = Program.GetOpArgs<OP::ImageSwizzleArgs>(item.At);
 				Issued = MakeShared<FImageSwizzleTask>(item, Args);
 			}
 			break;
@@ -2829,7 +2825,7 @@ namespace mu
 		{
 			if (item.Stage == 1)
 			{
-				OP::ImageSaturateArgs Args = program.GetOpArgs<OP::ImageSaturateArgs>(item.At);
+				OP::ImageSaturateArgs Args = Program.GetOpArgs<OP::ImageSaturateArgs>(item.At);
 				Issued = MakeShared<FImageSaturateTask>(item, Args);
 			}
 			break;
@@ -2839,7 +2835,7 @@ namespace mu
 		{
 			if (item.Stage == 1)
 			{
-				OP::ImageInvertArgs Args = program.GetOpArgs<OP::ImageInvertArgs>(item.At);
+				OP::ImageInvertArgs Args = Program.GetOpArgs<OP::ImageInvertArgs>(item.At);
 				Issued = MakeShared<FImageInvertTask>(item, Args);
 			}
 			break;
@@ -2849,7 +2845,7 @@ namespace mu
 		{
 			if (item.Stage == 1)
 			{
-				OP::ImageResizeArgs Args = program.GetOpArgs<OP::ImageResizeArgs>(item.At);
+				OP::ImageResizeArgs Args = Program.GetOpArgs<OP::ImageResizeArgs>(item.At);
 				Issued = MakeShared<FImageResizeTask>(item, Args);
 			}
 			break;
@@ -2859,7 +2855,7 @@ namespace mu
 		{
 			if (item.Stage == 1)
 			{
-				OP::ImageResizeRelArgs Args = program.GetOpArgs<OP::ImageResizeRelArgs>(item.At);
+				OP::ImageResizeRelArgs Args = Program.GetOpArgs<OP::ImageResizeRelArgs>(item.At);
 				Issued = MakeShared<FImageResizeRelTask>(item, Args);
 			}
 			break;
@@ -2870,7 +2866,7 @@ namespace mu
 			if ((ExecutionStrategy == EExecutionStrategy::MinimizeMemory && item.Stage == 3) ||
 				(ExecutionStrategy != EExecutionStrategy::MinimizeMemory && item.Stage == 2))
 			{
-				OP::ImageComposeArgs Args = program.GetOpArgs<OP::ImageComposeArgs>(item.At);
+				OP::ImageComposeArgs Args = Program.GetOpArgs<OP::ImageComposeArgs>(item.At);
 				TSharedPtr<const FLayout> ComposeLayout = StaticCastSharedPtr<const FLayout>( HeapData[item.CustomState].Resource);
 				Issued = MakeShared<FImageComposeTask>(item, Args, ComposeLayout);
 			}
diff --git a/Engine/Plugins/Mutable/Source/MutableRuntime/Private/MuR/Model.cpp b/Engine/Plugins/Mutable/Source/MutableRuntime/Private/MuR/Model.cpp
--- a/Engine/Plugins/Mutable/Source/MutableRuntime/Private/MuR/Model.cpp
+++ b/Engine/Plugins/Mutable/Source/MutableRuntime/Private/MuR/Model.cpp
@@ -24,6 +24,8 @@ namespace mu
 {
     MUTABLE_IMPLEMENT_POD_SERIALISABLE(FRomData);
 	MUTABLE_IMPLEMENT_POD_SERIALISABLE(FImageLODRange);
+	MUTABLE_IMPLEMENT_POD_SERIALISABLE(FConstantResourceIndex);
+	MUTABLE_IMPLEMENT_POD_VECTOR_SERIALISABLE(FConstantResourceIndex);
 
 	
     //---------------------------------------------------------------------------------------------
@@ -138,52 +140,50 @@ namespace mu
     {
 		LLM_SCOPE_BYNAME(TEXT("MutableRuntime"));
 
-		mu::FProgram& program = p->m_pD->Program;
+		mu::FProgram& Program = p->m_pD->Program;
 
-		TArray<TPair<int32, TSharedPtr<const FImage>>> InitialImages;
 		TArray<TPair<int32, TSharedPtr<const FMesh>>> InitialMeshes;
 		if (!bDropData)
 		{
-			InitialImages = program.ConstantImageLODs;
-			InitialMeshes = program.ConstantMeshes;
+			InitialMeshes = Program.ConstantMeshes;
 		}
 
 		FOutputMemoryStream MemStream(16 * 1024 * 1024);
 
 		// Save images and unload from memory
-		for (int32 ResourceIndex = 0; ResourceIndex < program.ConstantImageLODs.Num(); ++ResourceIndex)
+		for (TPair<uint32, TSharedPtr<const FImage>>& Entry: Program.ConstantImageLODsStreamed )
 		{
-			TPair<int32, TSharedPtr<const mu::FImage>>& ResData = program.ConstantImageLODs[ResourceIndex];
-
-			// This shouldn't have been serialised with rom support before.
-			if (ResData.Key < 0)
-			{
-				continue;
-			}
-
-			int32 RomIndex = ResData.Key;
-			const FRomData& RomData = program.Roms[RomIndex];
+			int32 RomIndex = Entry.Key;
+			const FRomData& RomData = Program.Roms[RomIndex];
 			check(RomData.ResourceType == DT_IMAGE);
-			check(RomData.ResourceIndex == ResourceIndex);
+			check(RomData.ResourceIndex.Index == RomIndex);
 
 			// Serialize to memory, to find out final size of this rom
 			MemStream.Reset();
 			FOutputArchive MemoryArch(&MemStream);
-			FImage::Serialise(ResData.Value.Get(), MemoryArch);
+			FImage::Serialise(Entry.Value.Get(), MemoryArch);
 			check(RomData.Size == MemStream.GetBufferSize());
 
 			streamer.OpenWriteFile(RomData.Id);
 			streamer.Write(MemStream.GetBuffer(), MemStream.GetBufferSize());
 			streamer.CloseWriteFile();
 
-			// We clear it to avoid including it also with the main model data
-			ResData.Value = nullptr;
+			// Do this progressively to avoid duplicating all data in memory.
+			if (bDropData)
+			{
+				Entry.Value.Reset();
+			}
+		}
+
+		if (bDropData)
+		{
+			Program.ConstantImageLODsStreamed.Empty(0);
 		}
 
 		// Save meshes and unload from memory
-		for (int32 ResourceIndex = 0; ResourceIndex < program.ConstantMeshes.Num(); ++ResourceIndex)
+		for (int32 ResourceIndex = 0; ResourceIndex < Program.ConstantMeshes.Num(); ++ResourceIndex)
 		{
-			TPair<int32, TSharedPtr<const FMesh>>& ResData = program.ConstantMeshes[ResourceIndex];
+			TPair<int32, TSharedPtr<const FMesh>>& ResData = Program.ConstantMeshes[ResourceIndex];
 
 			// This shouldn't have been serialised with rom support before.
 			if (ResData.Key < 0)
@@ -192,9 +192,9 @@ namespace mu
 			}
 
 			int32 RomIndex = ResData.Key;
-			const FRomData& RomData = program.Roms[RomIndex];
+			const FRomData& RomData = Program.Roms[RomIndex];
 			check(RomData.ResourceType == DT_MESH);
-			check(RomData.ResourceIndex == ResourceIndex);
+			check(RomData.ResourceIndex.Index == ResourceIndex);
 
 			// Serialize to memory, to find out final size of this rom
 			MemStream.Reset();
@@ -224,8 +224,7 @@ namespace mu
 		// Restore full data
 		if (!bDropData)
 		{
-			program.ConstantImageLODs = InitialImages;
-			program.ConstantMeshes = InitialMeshes;
+			Program.ConstantMeshes = InitialMeshes;
 		}
 	}
 
@@ -253,13 +252,7 @@ namespace mu
     {
 		LLM_SCOPE_BYNAME(TEXT("MutableRuntime"));
 
-		for (int32 ResIndex = 0; ResIndex < m_pD->Program.ConstantImageLODs.Num(); ++ResIndex)
-		{
-			if (m_pD->Program.ConstantImageLODs[ResIndex].Key >= 0)
-			{
-				m_pD->Program.ConstantImageLODs[ResIndex].Value = nullptr;
-			}
-		}
+		m_pD->Program.ConstantImageLODsStreamed.Empty(0);
 
 		for (int32 ResIndex = 0; ResIndex < m_pD->Program.ConstantMeshes.Num(); ++ResIndex)
 		{
diff --git a/Engine/Plugins/Mutable/Source/MutableTools/Private/MuT/ASTOpConstantResource.cpp b/Engine/Plugins/Mutable/Source/MutableTools/Private/MuT/ASTOpConstantResource.cpp
--- a/Engine/Plugins/Mutable/Source/MutableTools/Private/MuT/ASTOpConstantResource.cpp
+++ b/Engine/Plugins/Mutable/Source/MutableTools/Private/MuT/ASTOpConstantResource.cpp
@@ -348,11 +348,11 @@ namespace mu
 
 				if (MipIndex<0)
 				{
-					MipIndex = Program.ConstantImageLODs.Add(TPair<int32, TSharedPtr<const FImage>>(-1, pMip));
+					MipIndex = Program.ConstantImageLODsPermanent.Add(pMip);
 					Options.ImageConstantMipMap.Add(pMip, MipIndex);
 				}
 
-				Program.ConstantImageLODIndices.Add(uint32(MipIndex));
+				Program.ConstantImageLODIndices.Add({ uint32(MipIndex), 0 });
 
 				// Generate next mip if necessary
 				if (Mip + 1 < MipsToStore)
diff --git a/Engine/Plugins/Mutable/Source/MutableTools/Private/MuT/Compiler.cpp b/Engine/Plugins/Mutable/Source/MutableTools/Private/MuT/Compiler.cpp
--- a/Engine/Plugins/Mutable/Source/MutableTools/Private/MuT/Compiler.cpp
+++ b/Engine/Plugins/Mutable/Source/MutableTools/Private/MuT/Compiler.cpp
@@ -489,6 +489,17 @@ namespace mu
 
 		UE_LOG(LogMutableCore, Verbose, TEXT("(int) %s : %ld"), TEXT("Program size"), int64(Program.OpAddress.Num()));
 
+		// We are not touching the program anymore. Ensure we are not wasting memory.
+		Program.Roms.Shrink();
+		Program.ConstantImageLODsPermanent.Shrink();
+		Program.ConstantImageLODIndices.Shrink();
+		Program.ConstantImages.Shrink();
+		Program.ConstantMeshes.Shrink();
+		Program.ConstantStrings.Shrink();
+		Program.ConstantSkeletons.Shrink();
+		Program.ConstantPhysicsBodies.Shrink();
+		Program.Parameters.Shrink();
+
         return pResult;
     }
 
@@ -537,136 +548,156 @@ namespace mu
 		int32 NumRoms = 0;
 		uint64 NumRomsBytes = 0;
 		int32 NumEmbedded = 0;
-		uint64 NumEmbeddedBytes = 0;
+		int32 NumEmbeddedBytes = 0;
 		int32 NumHighRes = 0;
 		uint64 NumHighResBytes = 0;
 
 		// Maximum number of roms
-		int32 MaxRomCount = Program.ConstantImageLODs.Num() + Program.ConstantMeshes.Num();
+		int32 MaxRomCount = Program.ConstantImageLODsPermanent.Num() + Program.ConstantMeshes.Num();
 		Program.Roms.Reserve(MaxRomCount);
 
 		TSet<uint32> UsedIds;
 		UsedIds.Reserve(MaxRomCount);
 
 		TArray<FRomData> RomDatas;
-		RomDatas.SetNumZeroed(FMath::Max(Program.ConstantImageLODs.Num(), Program.ConstantMeshes.Num()));
+		RomDatas.SetNumZeroed(FMath::Max(Program.ConstantImageLODsPermanent.Num(), Program.ConstantMeshes.Num()));
 
 		// Images
 		{
 			MUTABLE_CPUPROFILER_SCOPE(GenerateRoms_ImageIds);
 
-			ParallelFor(Program.ConstantImageLODs.Num(),
-				[EmbeddedDataBytesLimit, &Program, &RomDatas](uint32 ResourceIndex)
-				{
-					TPair<int32, TSharedPtr<const FImage>>& ResData = Program.ConstantImageLODs[ResourceIndex];
+			TArray<TSharedPtr<const FImage>> AllMips = MoveTemp(Program.ConstantImageLODsPermanent);
+			Program.ConstantImageLODsPermanent = {};
+			Program.ConstantImageLODsPermanent.Reserve(AllMips.Num());
 
-					// This shouldn't have been serialised with rom support before.
-					check(ResData.Key < 0);
+
+			ParallelFor(AllMips.Num(),
+				[EmbeddedDataBytesLimit, &Program, &AllMips, &RomDatas](uint32 MipIndex)
+				{
+					TSharedPtr<const FImage>& ResData = AllMips[MipIndex];
 
 					// Serialize to find out final size of this rom
 					FOutputSizeStream SizeStream;
 					FOutputArchive MemoryArch(&SizeStream);
-					FImage::Serialise(ResData.Value.Get(), MemoryArch);
+					FImage::Serialise(ResData.Get(), MemoryArch);
+
+
+					FRomData& RomData = RomDatas[MipIndex];
+					RomData.ResourceType = DT_IMAGE;
+					RomData.Size = SizeStream.GetBufferSize();
 
 					// If the resource uses less memory than the threshold, don't save it in a separate rom.
-					if (SizeStream.GetBufferSize() <= EmbeddedDataBytesLimit)
+					if (RomData.Size <= EmbeddedDataBytesLimit)
 					{
 						return;
 					}
 
-					FRomData& RomData = RomDatas[ResourceIndex];
-					RomData.ResourceType = DT_IMAGE;
-					RomData.ResourceIndex = ResourceIndex;
-					RomData.Size = SizeStream.GetBufferSize();
-					RomData.Flags = ERomFlags::None;
+					// RomData.ResourceIndex will be filled below when we split in permanent and streamable
+					// RomData.Flags will be filled below
 
-					const FImage* Resource = ResData.Value.Get();
+					const FImage* Resource = ResData.Get();
 					uint64 FullHash = CityHash64(reinterpret_cast<const char*>(Resource->GetLODData(0)), Resource->GetLODDataSize(0));
 					RomData.Id = GetTypeHash(FullHash);
 				});
-		}
 
-		// Generate the high-res flags for images
-		{
-			// Initially all are high-res because if at least one reference to a mip is not we will set it to not-high-res
-			TArray<bool> IsLODHighRes;
-			IsLODHighRes.Init(true, Program.ConstantImageLODs.Num());
-
-			for (int32 ImageIndex = 0; ImageIndex < Program.ConstantImages.Num(); ++ImageIndex)
+			// Generate the high-res flags for images
 			{
-				const FImageLODRange& LODRange = Program.ConstantImages[ImageIndex];
+				// Initially all are high-res because if at least one reference to a mip is not we will set it to not-high-res
+				TArray<bool> IsLODHighRes;
+				IsLODHighRes.Init(true, AllMips.Num());
 
-				int32 NumHighResMips = FMath::Max(0,AdditionalData.SourceImagePerConstant[ImageIndex].SourceHighResMips);
-
-				for (int32 LODRangeIndex = NumHighResMips; LODRangeIndex < LODRange.LODCount; ++LODRangeIndex)
+				for (int32 ImageIndex = 0; ImageIndex < Program.ConstantImages.Num(); ++ImageIndex)
 				{
-					int32 LODIndex = Program.ConstantImageLODIndices[LODRange.FirstIndex+LODRangeIndex];
-					IsLODHighRes[LODIndex] = false;
+					const FImageLODRange& LODRange = Program.ConstantImages[ImageIndex];
+
+					int32 NumHighResMips = FMath::Max(0, AdditionalData.SourceImagePerConstant[ImageIndex].SourceHighResMips);
+
+					for (int32 LODRangeIndex = NumHighResMips; LODRangeIndex < LODRange.LODCount; ++LODRangeIndex)
+					{
+						FConstantResourceIndex LODIndex = Program.ConstantImageLODIndices[LODRange.FirstIndex + LODRangeIndex];
+						check(!LODIndex.Streamable); // Not classified yet
+						IsLODHighRes[LODIndex.Index] = false;
+					}
+
+					// Moreover, at least one mip of each image has to be non-highres
+					if (LODRange.LODCount > 0)
+					{
+						FConstantResourceIndex LastLODIndex = Program.ConstantImageLODIndices[LODRange.FirstIndex + LODRange.LODCount - 1];
+						check(!LastLODIndex.Streamable); // Not classified yet
+						IsLODHighRes[LastLODIndex.Index] = false;
+					}
 				}
 
-				// Moreover, at least one mip of each image has to be non-highres
-				if (LODRange.LODCount>0)
+				for (int32 MipIndex = 0; MipIndex < AllMips.Num(); ++MipIndex)
 				{
-					int32 LastLODIndex = Program.ConstantImageLODIndices[LODRange.FirstIndex + LODRange.LODCount - 1];
-					IsLODHighRes[LastLODIndex] = false;
+					// If this mip represents a high-quality mip, flag the rom as such
+					if (IsLODHighRes[MipIndex])
+					{
+						FRomData& RomData = RomDatas[MipIndex];
+						EnumAddFlags(RomData.Flags, ERomFlags::HighRes);
+
+						++NumHighRes;
+						NumHighResBytes += RomData.Size;
+					}
 				}
 			}
 
-			for (int32 ResourceIndex = 0; ResourceIndex < Program.ConstantImageLODs.Num(); ++ResourceIndex)
 			{
-				// If this mip represents a high-quality mip, flag the rom as such
-				if (IsLODHighRes[ResourceIndex])
+				MUTABLE_CPUPROFILER_SCOPE(GenerateRoms_ImageSourceIds);
+
+				for (int32 ImageIndex = 0; ImageIndex < Program.ConstantImages.Num(); ++ImageIndex)
 				{
-					FRomData& RomData = RomDatas[ResourceIndex];
-					EnumAddFlags(RomData.Flags, ERomFlags::HighRes);
-					
-					++NumHighRes;
-					NumHighResBytes += RomData.Size;
+					const FImageLODRange& LODRange = Program.ConstantImages[ImageIndex];
+
+					uint32 SourceId = AdditionalData.SourceImagePerConstant[ImageIndex].SourceId;
+
+					for (int32 LODRangeIndex = 0; LODRangeIndex < LODRange.LODCount; ++LODRangeIndex)
+					{
+						FConstantResourceIndex ResourceIndex = Program.ConstantImageLODIndices[LODRange.FirstIndex + LODRangeIndex];
+						check(!ResourceIndex.Streamable); // Not classified yet
+						RomDatas[ResourceIndex.Index].SourceId = SourceId;
+					}
 				}
 			}
-		}
-
-		{
-			MUTABLE_CPUPROFILER_SCOPE(GenerateRoms_ImageSourceIds);
 
-			for (int32 ImageIndex = 0; ImageIndex < Program.ConstantImages.Num(); ++ImageIndex)
+			// Split the data in permanent and streamable and assign final FConstantResourceIndex
 			{
-				const FImageLODRange& LODRange = Program.ConstantImages[ImageIndex];
-
-				uint32 SourceId = AdditionalData.SourceImagePerConstant[ImageIndex].SourceId;
-
-				for (int32 LODRangeIndex = 0; LODRangeIndex < LODRange.LODCount; ++LODRangeIndex)
+				for (int32 MipIndex = 0; MipIndex < AllMips.Num(); ++MipIndex)
 				{
-					int32 RomIndex = Program.ConstantImageLODIndices[LODRange.FirstIndex + LODRangeIndex];
-					RomDatas[RomIndex].SourceId = SourceId;
-				}
-			}
-		}
+					FRomData& RomData = RomDatas[MipIndex];
+					if (RomData.Size > EmbeddedDataBytesLimit)
+					{
+						NumRoms++;
+						NumRomsBytes += RomData.Size;
 
-		{
-			MUTABLE_CPUPROFILER_SCOPE(GenerateRoms_ImageIdsUnique);
+						// Ensure that the Id is unique
+						EnsureUniqueRomId(UsedIds, RomData.Id);
 
-			for (int32 ResourceIndex = 0; ResourceIndex < Program.ConstantImageLODs.Num(); ++ResourceIndex)
-			{
-				FRomData& RomData = RomDatas[ResourceIndex];
+						uint32 RomIndex = Program.Roms.Num();
+						check(RomIndex < 0x7fffffff);
 
-				// If the resource uses less memory than the threshold, don't save it in a separate rom.
-				if (int32(RomData.Size) <= EmbeddedDataBytesLimit)
-				{
-					NumEmbedded++;
-					NumEmbeddedBytes += RomData.Size;
-					continue;
-				}
-				NumRoms++;
-				NumRomsBytes += RomData.Size;
+						RomData.ResourceIndex = { RomIndex, 1 };
+						Program.ConstantImageLODsStreamed.Add(RomIndex, AllMips[MipIndex]);
 
-				// Ensure that the Id is unique
-				EnsureUniqueRomId(UsedIds, RomData.Id);
+						Program.Roms.Add(RomData);
+					}
+					else
+					{
+						uint32 Index = uint32(Program.ConstantImageLODsPermanent.Num());
+						check(Index < 0x7fffffff);
+						RomData.ResourceIndex = { Index, 0 };
+						Program.ConstantImageLODsPermanent.Add(AllMips[MipIndex]);
 
-				int32 RomIndex = Program.Roms.Add(RomData);
+						NumEmbedded++;
+						NumEmbeddedBytes += RomData.Size;
+					}
+				}
 
-				TPair<int32, TSharedPtr<const FImage>>& ResData = Program.ConstantImageLODs[ResourceIndex];
-				ResData.Key = RomIndex;
+				for (int32 Index = 0; Index < Program.ConstantImageLODIndices.Num(); ++Index)
+				{
+					int32 RomDataIndex = Program.ConstantImageLODIndices[Index].Index;
+					Program.ConstantImageLODIndices[Index] = RomDatas[RomDataIndex].ResourceIndex;
+				}
 			}
 		}
 
@@ -687,21 +718,21 @@ namespace mu
 					const FMesh* Resource = ResData.Value.Get();
 					int32 ApproximateSize = Resource->GetDataSize();
 
-					// If the resource uses less memory than the threshold, don't save it in a separate rom.
-					if (ApproximateSize <= EmbeddedDataBytesLimit)
-					{
-						return;
-					}
-
 					FOutputMemoryStream MemStream(ApproximateSize + 64 * 1024);
 					FOutputArchive MemoryArch(&MemStream);
 					FMesh::Serialise(ResData.Value.Get(), MemoryArch);
 
 					FRomData& RomData = RomDatas[ResourceIndex];
 					RomData.ResourceType = DT_MESH;
-					RomData.ResourceIndex = ResourceIndex;
 					RomData.Size = MemStream.GetBufferSize();
-					RomData.Flags = ERomFlags::None;
+
+					// If the resource uses less memory than the threshold, don't save it in a separate rom.
+					if (RomData.Size <= EmbeddedDataBytesLimit)
+					{
+						return;
+					}
+
+					RomData.ResourceIndex = { ResourceIndex, 1 };
 
 					// Ensure that the Id is unique
 					uint64 FullHash = CityHash64(reinterpret_cast<const char*>(MemStream.GetBuffer()), MemStream.GetBufferSize());
