[Mutable] Reduced the Update Skeletal Mesh Async log spam.

#rnx

diff --git a/Engine/Plugins/Mutable/Source/CustomizableObject/Internal/MuCO/CustomizableObjectSystemPrivate.h b/Engine/Plugins/Mutable/Source/CustomizableObject/Internal/MuCO/CustomizableObjectSystemPrivate.h
--- a/Engine/Plugins/Mutable/Source/CustomizableObject/Internal/MuCO/CustomizableObjectSystemPrivate.h
+++ b/Engine/Plugins/Mutable/Source/CustomizableObject/Internal/MuCO/CustomizableObjectSystemPrivate.h
@@ -901,6 +901,11 @@ public:
 
 	int32 NumLODUpdatesLastTick = 0;
 
+	/** Time when the "Started Update Skeletal Mesh Async" log will be unmuted (in seconds). */
+	float LogStartedUpdateUnmute = 0.0;
+
+	/** Time of the last "Started Update Skeletal Mesh Async" log (in seconds). */
+	float LogStartedUpdateLast = 0;
 
 public:
 	/**
diff --git a/Engine/Plugins/Mutable/Source/CustomizableObject/Private/MuCO/CustomizableObjectSystem.cpp b/Engine/Plugins/Mutable/Source/CustomizableObject/Private/MuCO/CustomizableObjectSystem.cpp
--- a/Engine/Plugins/Mutable/Source/CustomizableObject/Private/MuCO/CustomizableObjectSystem.cpp
+++ b/Engine/Plugins/Mutable/Source/CustomizableObject/Private/MuCO/CustomizableObjectSystem.cpp
@@ -804,6 +804,8 @@ void FinishUpdateGlobal(const TSharedRef<FUpdateContextPrivate>& Context)
 
 	UCustomizableObjectInstance* Instance = Context->Instance.Get();
 
+	UCustomizableObject* Object = Instance ? Instance->GetCustomizableObject() : nullptr;
+
 	UCustomizableObjectSystem* System = UCustomizableObjectSystem::GetInstance();
 	UCustomizableObjectSystemPrivate* SystemPrivate = System ? System->GetPrivate() : nullptr;
 
@@ -898,9 +900,10 @@ void FinishUpdateGlobal(const TSharedRef<FUpdateContextPrivate>& Context)
 		Context->UpdateTime = FPlatformTime::Seconds() - Context->StartUpdateTime;		
 	}
 	
-	const uint32 InstanceId = Instance ? Instance->GetUniqueID() : 0;
-	UE_LOG(LogMutable, Log, TEXT("Finished UpdateSkeletalMesh Async. Instance=%d, Frame=%d, QueueTime=%f, UpdateTime=%f"), InstanceId, GFrameNumber, Context->QueueTime, Context->UpdateTime);
-
+	const FName ObjectName = Object ? Object->GetFName() : NAME_None;
+	const FName InstanceName = Instance ? Instance->GetFName() : NAME_None;
+	UE_LOG(LogMutable, Verbose, TEXT("Finished Update Skeletal Mesh Async. CustomizableObject=%s Instance=%s, Frame=%d  QueueTime=%f, UpdateTime=%f"), *ObjectName.ToString(), *InstanceName.ToString(), GFrameNumber, Context->QueueTime, Context->UpdateTime);
+	
 	if (SystemPrivate && FLogBenchmarkUtil::IsBenchmarkingReportingEnabled())
 	{
 		FFunctionGraphTask::CreateAndDispatchWhenReady( // Calling Benchmark in a task so we make sure we exited all scopes.
@@ -1196,11 +1199,11 @@ void UCustomizableObjectSystemPrivate::EnqueueUpdateSkeletalMesh(const TSharedRe
 	
 	UCustomizableInstancePrivate* InstancePrivate = Instance->GetPrivate();
 
-	const EQueuePriorityType Priority = GetUpdatePriority(*Instance, Context->bForceHighPriority);
-	const uint32 InstanceId = Instance->GetUniqueID();
-	const float Distance = FMath::Sqrt(InstancePrivate->LastMinSquareDistFromComponentToPlayer);
-	const bool bIsPlayerOrNearIt = InstancePrivate->HasCOInstanceFlags(UsedByPlayerOrNearIt);
-	UE_LOG(LogMutable, Log, TEXT("Enqueue UpdateSkeletalMesh Async. Instance=%d, Frame=%d, Priority=%d, dist=%f, bIsPlayerOrNearIt=%d"), InstanceId, GFrameNumber, static_cast<int32>(Priority), Distance, bIsPlayerOrNearIt);				
+	UCustomizableObject* Object = Instance->GetCustomizableObject();
+	
+	const FName ObjectName = Object ? Object->GetFName() : NAME_None;
+	const FName InstanceName = Instance->GetFName();
+	UE_LOG(LogMutable, Verbose, TEXT("Enqueue Update Skeletal Mesh Async. CustomizableObject=%s Instance=%s, Frame=%d"), *ObjectName.ToString(), *InstanceName.ToString(), GFrameNumber);
 	
 	if (!bIsMutableEnabled)
 	{
@@ -3733,7 +3736,7 @@ int32 UCustomizableObjectSystem::TickInternal(const bool bBlocking)
 			{
 				FMutablePendingInstanceUpdate& PendingUpdate = *Iterator;
 
-				if (PendingUpdate.Context->Instance.IsValid())
+				if (PendingUpdate.Context->Instance.IsValid() && PendingUpdate.Context->Instance->GetCustomizableObject())
 				{
 					const EQueuePriorityType PriorityType = Private->GetUpdatePriority(*PendingUpdate.Context->Instance, false);
 					
@@ -3763,7 +3766,7 @@ int32 UCustomizableObjectSystem::TickInternal(const bool bBlocking)
 			{
 				const UCustomizableObjectInstance* Instance = LODUpdateTuple.Key;
 
-				if (Instance)
+				if (Instance && Instance->GetCustomizableObject())
 				{
 					FMutableUpdateCandidate& LODUpdateCandidate = LODUpdateTuple.Value;
 					ensure(LODUpdateCandidate.HasBeenIssued());
@@ -4361,11 +4364,31 @@ void UCustomizableObjectSystemPrivate::StartUpdateSkeletalMesh(const TSharedRef<
 	TRACE_BEGIN_REGION(UE_MUTABLE_UPDATE_REGION);
 
 	check(!CurrentMutableOperation); // Can not start an update if there is already another in progress
-	check(Context->Instance.IsValid()) // The instance has to be alive to start the update
+	check(Context->Instance.IsValid()) // The COI has to be alive to start the update
+
+	const UCustomizableObject* Object = Context->Instance->GetCustomizableObject();
+	check(Object) // The CO has to be alive to start the update
 
-	const uint32 InstanceId = Context->Instance->GetUniqueID();
-	UE_LOG(LogMutable, Log, TEXT("Started UpdateSkeletalMesh Async. Instance=%d, Frame=%d"), InstanceId, GFrameNumber);				
+	if (FPlatformTime::Seconds() > LogStartedUpdateUnmute)
+	{
+		const FName ObjectName = Object ? Object->GetFName() : NAME_None;
+		const FName InstanceName = Context->Instance->GetFName();
+		UE_LOG(LogMutable, Log, TEXT("Started Update Skeletal Mesh Async. CustomizableObject=%s Instance=%s, Frame=%d"), *ObjectName.ToString(), *InstanceName.ToString(), GFrameNumber);
+
+		const float CurrentTime = FPlatformTime::Seconds();
+
+		constexpr float LogInterval = 1.0 / 10.0; // Allow maximum 10 logs per second.
+		const bool bMute = CurrentTime - LogStartedUpdateLast < LogInterval;
+		LogStartedUpdateLast = CurrentTime;
 			
+		if (bMute)
+		{
+			constexpr float MuteTime = 5.0;
+			UE_LOG(LogMutable, Log, TEXT("Disabling \"Started Update Skeletal Mesh Async\" log during %f seconds due to spam"), MuteTime);
+			LogStartedUpdateUnmute = CurrentTime + MuteTime;
+		}
+	}
+	
 	CurrentMutableOperation = Context;
 }
 
