[Mutable] Replace check with a warning when there's no PlatformData in the CookAdditionalFiles.

#rnx

diff --git a/Engine/Plugins/Mutable/Source/CustomizableObject/Private/MuCO/CustomizableObject.cpp b/Engine/Plugins/Mutable/Source/CustomizableObject/Private/MuCO/CustomizableObject.cpp
--- a/Engine/Plugins/Mutable/Source/CustomizableObject/Private/MuCO/CustomizableObject.cpp
+++ b/Engine/Plugins/Mutable/Source/CustomizableObject/Private/MuCO/CustomizableObject.cpp
@@ -3586,7 +3586,14 @@ void UCustomizableObjectBulk::CookAdditionalFilesOverride(const TCHAR* PackageFi
 	UCustomizableObject* CustomizableObject = CastChecked<UCustomizableObject>(GetOutermostObject());
 
 	MutablePrivate::FMutableCachedPlatformData* PlatformData = CustomizableObject->GetPrivate()->CachedPlatformsData.Find(TargetPlatform->PlatformName());
-	check(PlatformData);
+	if (!PlatformData)
+	{
+		UE_LOG(LogMutable, Warning, TEXT("CookAdditionalFilesOverride: Customizable Object [%s] is missing [%s] platform data."),
+			*CustomizableObject->GetName(),
+			*TargetPlatform->PlatformName());
+
+		return;
+	}
 
 	const FString CookedBulkFileName = FString::Printf(TEXT("%s/%s"), *FPaths::GetPath(PackageFilename), *CustomizableObject->GetName());
 
