[mutable] Fix compilation with some clang editor targets.
#rnx

diff --git a/Engine/Plugins/Mutable/Source/MutableRuntime/Internal/MuR/SerialisationPrivate.h b/Engine/Plugins/Mutable/Source/MutableRuntime/Internal/MuR/SerialisationPrivate.h
--- a/Engine/Plugins/Mutable/Source/MutableRuntime/Internal/MuR/SerialisationPrivate.h
+++ b/Engine/Plugins/Mutable/Source/MutableRuntime/Internal/MuR/SerialisationPrivate.h
@@ -52,7 +52,7 @@ namespace mu
 
 
 #define MUTABLE_IMPLEMENT_POD_SERIALISABLE(Type)				     \
-    void DLLEXPORT operator<<(FOutputArchive& Arch, const Type& T)    \
+    void DLLEXPORT operator<<(FOutputArchive& Arch, const Type& T)   \
     {																 \
         Arch.Stream->Write(&T, sizeof(Type));						 \
     }																 \
@@ -65,18 +65,18 @@ namespace mu
 
 #define MUTABLE_IMPLEMENT_POD_VECTOR_SERIALISABLE(Type)                        \
     template<class Alloc>                                                      \
-	void operator<<(FOutputArchive& Arch, const TArray<Type, Alloc>& V)         \
+	void DLLEXPORT operator<<(FOutputArchive& Arch, const TArray<Type, Alloc>& V)         \
 	{                                                                          \
 		uint32 Num = uint32(V.Num());                                          \
 		Arch << Num;                                                           \
 		if (Num)                                                               \
 		{                                                                      \
-			Arch.Stream->Write(V.GetData(), Num * sizeof(Type));					   \
+			Arch.Stream->Write(V.GetData(), Num * sizeof(Type));			   \
 		}                                                                      \
 	}                                                                          \
                                                                                \
     template<class Alloc>                                                      \
-	void operator>>(FInputArchive& Arch, TArray<Type, Alloc>& V)                \
+	void DLLEXPORT operator>>(FInputArchive& Arch, TArray<Type, Alloc>& V)     \
 	{                                                                          \
 		uint32 Num;                                                            \
 		Arch >> Num;                                                           \
diff --git a/Engine/Plugins/Mutable/Source/MutableRuntime/Public/MuR/Serialisation.h b/Engine/Plugins/Mutable/Source/MutableRuntime/Public/MuR/Serialisation.h
--- a/Engine/Plugins/Mutable/Source/MutableRuntime/Public/MuR/Serialisation.h
+++ b/Engine/Plugins/Mutable/Source/MutableRuntime/Public/MuR/Serialisation.h
@@ -27,9 +27,9 @@ namespace mu
 
 #define MUTABLE_DEFINE_POD_VECTOR_SERIALISABLE(Type)							\
 	template<typename Alloc>													\
-	void operator<<(mu::FOutputArchive& Arch, const TArray<Type, Alloc>& V);	\
+	void DLLEXPORT operator<<(mu::FOutputArchive& Arch, const TArray<Type, Alloc>& V);	\
 	template<typename Alloc>													\
-	void operator>>(mu::FInputArchive& Arch, TArray<Type, Alloc>& V);
+	void DLLEXPORT operator>>(mu::FInputArchive& Arch, TArray<Type, Alloc>& V);
 
 #define MUTABLE_DEFINE_ENUM_SERIALISABLE(Type)									\
     void DLLEXPORT operator<<(mu::FOutputArchive& Arch, const Type& T);			\
