[mutable] CIS : Fixed case where the test commandlet waits indefinitely as no components got found after the instance update.
- Secondary submit performed to prevent robomerge issues as 33.20 and Main do have changes that would trigger a conflict

#rnx

diff --git a/Engine/Plugins/Mutable/Source/MutableValidation/Private/MuV/CustomizableObjectInstanceUpdateUtility.cpp b/Engine/Plugins/Mutable/Source/MutableValidation/Private/MuV/CustomizableObjectInstanceUpdateUtility.cpp
--- a/Engine/Plugins/Mutable/Source/MutableValidation/Private/MuV/CustomizableObjectInstanceUpdateUtility.cpp
+++ b/Engine/Plugins/Mutable/Source/MutableValidation/Private/MuV/CustomizableObjectInstanceUpdateUtility.cpp
@@ -95,22 +95,27 @@ void FCustomizableObjectInstanceUpdateUtility::OnInstanceUpdateResult(const FUpd
 {
 	LLM_SCOPE_BYNAME(TEXT("FCustomizableObjectInstanceUpdateUtility/OnInstanceUpdated"));
 	
+	check(Instance);
 	const FString InstanceName = Instance->GetName();
 	
 	if (UCustomizableObjectSystem::IsUpdateResultValid(Result.UpdateResult))
 	{
-		UE_LOG(LogMutable, Display, TEXT("Instance %s finished update succesfully."), *InstanceName);
+		UE_LOG(LogMutable, Display, TEXT("Instance %s finished update successfully."), *InstanceName);
 		bInstanceFailedUpdate = false;
 		
 		// Request load all MIPs
-		UE_LOG(LogMutable, Display, TEXT("Instance %s rquesting streaming all MIPs."), *InstanceName);
-
 		check(ComponentsBeingUpdated.IsEmpty());
 
-		UCustomizableObject* CO = Instance->GetCustomizableObject();
+		const TArray<FName> InstanceComponents = Instance->GetComponentNames();
+		
+		// Early out if no components are found in the instance and notify the Utility that the instance finished its update
+		if (InstanceComponents.IsEmpty())
+		{
+			bIsInstanceBeingUpdated = false;
+			return;
+		}
 
-		TArray<FName> InstanceComponents = Instance->GetComponentNames();
-		for (FName ComponentName : InstanceComponents)
+		for (const FName& ComponentName : InstanceComponents)
 		{
 			TStrongObjectPtr<USkeletalMeshComponent> SkeletalComponent = TStrongObjectPtr(NewObject<USkeletalMeshComponent>());
 			SkeletalComponent->SetSkeletalMesh(Instance->GetComponentMeshSkeletalMesh(ComponentName));
@@ -130,6 +135,8 @@ void FCustomizableObjectInstanceUpdateUtility::OnInstanceUpdateResult(const FUpd
 				Info.RenderAsset->StreamIn(MAX_int32, true);
 			}
 		}
+
+		UE_LOG(LogMutable, Display, TEXT("Instance %s requesting the streaming of all MIPs."), *InstanceName);
 	}
 	else
 	{
@@ -137,7 +144,7 @@ void FCustomizableObjectInstanceUpdateUtility::OnInstanceUpdateResult(const FUpd
 		UE_LOG(LogMutable, Error, TEXT("Instance %s finished update with anomalous state : %s."), *InstanceName, *OutputStatus);
 		bInstanceFailedUpdate = true;
 		
-		// Tell the system the instance finished it's update so we can continue the execution without waitting for the mips to stream in
+		// Tell the system the instance finished its update so we can continue the execution without waiting for the mips to stream in
 		bIsInstanceBeingUpdated = false;
 	}	
 }
