[mutable] Move mesh-closed test and warning to the core and runtime. Remove deprecated internal data structure.
#rnx

diff --git a/Engine/Plugins/Mutable/Source/CustomizableObjectEditor/Private/MuCOE/GenerateMutableSource/GenerateMutableSource.cpp b/Engine/Plugins/Mutable/Source/CustomizableObjectEditor/Private/MuCOE/GenerateMutableSource/GenerateMutableSource.cpp
--- a/Engine/Plugins/Mutable/Source/CustomizableObjectEditor/Private/MuCOE/GenerateMutableSource/GenerateMutableSource.cpp
+++ b/Engine/Plugins/Mutable/Source/CustomizableObjectEditor/Private/MuCOE/GenerateMutableSource/GenerateMutableSource.cpp
@@ -511,46 +511,6 @@ uint32 GetTypeHash(const FGeneratedKey& Key)
 }
 
 
-void FPinDataValue::Append(const FPinDataValue& From)
-{
-	MeshesData.Append(From.MeshesData);
-}
-
-
-FPinDataValue* FPinData::Find(const UEdGraphPin* Pin)
-{
-	return Data.Find(Pin);
-}
-
-
-FPinDataValue& FPinData::GetCurrent()
-{
-	return Data[PinStack.Last()];
-}
-
-
-void FPinData::Pop()
-{
-	const uint32 Num = PinStack.Num();
-
-	check(Num >= 1); // Pop called without a previous push
-
-	if (Num >= 2)
-	{
-		Data[PinStack[Num - 2]].Append(Data[PinStack[Num - 1]]);
-	}
-
-	PinStack.Pop();
-}
-
-
-void FPinData::Push(const UEdGraphPin* Pin)
-{
-	PinStack.Push(Pin);
-	Data.FindOrAdd(Pin);
-}
-
-
 FGraphCycleKey::FGraphCycleKey(const UEdGraphPin& Pin, const FString& Id) :
 	Pin(Pin),
 	Id(Id)
@@ -564,19 +524,6 @@ bool FGraphCycleKey::operator==(const FGraphCycleKey& Other) const
 }
 
 
-FScopedPinData::FScopedPinData(FMutableGraphGenerationContext& Context, const UEdGraphPin* Pin)
-	: Context(Context)
-{
-	Context.PinData.Push(Pin);
-}
-
-
-FScopedPinData::~FScopedPinData()
-{
-	Context.PinData.Pop();
-}
-
-
 UTexture2D* FindReferenceImage(const UEdGraphPin* Pin, FMutableGraphGenerationContext& GenerationContext)
 {
 	check(Pin)
diff --git a/Engine/Plugins/Mutable/Source/CustomizableObjectEditor/Private/MuCOE/GenerateMutableSource/GenerateMutableSource.h b/Engine/Plugins/Mutable/Source/CustomizableObjectEditor/Private/MuCOE/GenerateMutableSource/GenerateMutableSource.h
--- a/Engine/Plugins/Mutable/Source/CustomizableObjectEditor/Private/MuCOE/GenerateMutableSource/GenerateMutableSource.h
+++ b/Engine/Plugins/Mutable/Source/CustomizableObjectEditor/Private/MuCOE/GenerateMutableSource/GenerateMutableSource.h
@@ -385,39 +385,6 @@ struct FGroupProjectorImageInfo
 };
 
 
-/** Struct that defines a mesh (Mesh + LOD + MaterialIndex) */
-struct FMeshData
-{
-	/** Mesh which it may contain multiple LODs and Materials. */
-	const UObject* Mesh;
-
-	/** Specific LOD of the mesh. */
-	int LOD;
-
-	/** Specific index of the mesh material (MaterialIndex is an alias of Section). */
-	int MaterialIndex;
-
-	/** Node where the mesh is defined. Not a UCustomizableObjectNodeMesh due to Table nodes */
-	const UCustomizableObjectNode* Node;
-
-	bool operator==(const FMeshData& Other) const
-	{
-		return Mesh == Other.Mesh &&
-			LOD == Other.LOD &&
-			MaterialIndex == Other.MaterialIndex &&
-			Node == Other.Node;
-	}
-};
-
-inline uint32 GetTypeHash(const FMeshData& Key)
-{
-	uint32 MeshHash = GetTypeHash(Key.Mesh);
-	uint32 NodeHash = GetTypeHash(Key.Node->GetUniqueID());
-	
-	return HashCombine(HashCombine(HashCombine(MeshHash, Key.LOD), Key.MaterialIndex), NodeHash);
-}
-
-
 /** Struct used to store info specific to each component during compilation */
 struct FMutableComponentInfo
 {
@@ -447,64 +414,6 @@ struct FMutableComponentInfo
 };
 
 
-/** Struct of data which is behind a given output pin. 
- *
- * Each newly added field has should also be defined on the Append() function body.
- */
-struct FPinDataValue
-{
-	/** Set of all meshes behind a given output pin. */
-	TSet<FMeshData> MeshesData;
-
-	/** Add all data from other PinData. */
-	void Append(const FPinDataValue& From);
-};
-
-/** Eases the managment and access of data behind a given pin.
- * 
- * We understand as "behind a pin" as all pins which are connected to a given pin directly and indirectly. 
- * 
- * Given the the following pins:
- * C -> B -> A
- *      D ---^
- * 
- * FPinData will propagate the added data in the following way:
- * C (data added on C) -> B (data added on C, B) -> A (data added on C, B, A, D)
- *                           D (data added on D) ---^
- * 
- * FPinDataValue defines which data is going to be saved behind pins.
- *
- * When traversing the graph:
- * - Each time an output pin is explored recursively, call the Push() function.
- * - When exiting the recursion, call the Pop() function.
- * - Add data to be seen behind all the stacked pins using the GetCurrent() function.
- * - To avoid missing any returning branch, use the SCOPE_PIN_DATA wrapper.
- * 
- * The all data from a behind a pin can be queried using the Find() method.
- */
-class FPinData // DEPRECATED. DO NOT USE!
-{
-public:
-	/** Get the PinData from a given pin. */
-	FPinDataValue* Find(const UEdGraphPin* Pin);
-
-	/** Get the PinData of the current pin. The top pin in the stack. */
-	FPinDataValue& GetCurrent();
-
-	/** Pop the current pin from the stack and append all its PinData with to previous pin in the stack. */
-	void Pop();
-
-	/** Push a pin to the stack and create its PinData. */
-	void Push(const UEdGraphPin* Pin);
-
-private:
-	/** Data for each pin. */
-	TMap<const UEdGraphPin*, FPinDataValue> Data;
-
-	/** Stack of pins. Used to know which pin the data belongs to during the graph traversal recursion. */
-	TArray<const UEdGraphPin*> PinStack;
-};
-
 /** Graph cycle key.
  *
  * Pin is not enough since we can call multiple recursive functions with the same pin.
@@ -931,12 +840,6 @@ public:
 	TArray<const UEdGraphNode*> LimitedParameters;
 	int32 ParameterLimitationCount = 0;
 
-	/** Data which is behind a given output pin. 
-	 *
-	 * We only consider output pins. Do not use input pins as keys. See FPinData class definition. 
-	 */
-	FPinData PinData; // DEPRECATED. DO NOT USE!
-
 	// Stores all morphs to apply them directly to a skeletal mesh node
 	TArray<FMorphNodeData> MeshMorphStack;
 
@@ -981,20 +884,6 @@ public:
 	TObjectPtr<UObject> RootVersionBridge;
 };
 
-/** Pin Data scope wrapper. Pops the pin data on scope exit. */
-class FScopedPinData
-{
-public:
-	explicit FScopedPinData(FMutableGraphGenerationContext& Context, const UEdGraphPin* Pin);
-	~FScopedPinData();
-
-private:
-	FMutableGraphGenerationContext& Context;
-};
-
-#define SCOPED_PIN_DATA(Context, Pin) \
-	FScopedPinData ScopedPinData(Context, Pin);
-
 
 namespace Private
 {
diff --git a/Engine/Plugins/Mutable/Source/CustomizableObjectEditor/Private/MuCOE/GenerateMutableSource/GenerateMutableSourceMesh.cpp b/Engine/Plugins/Mutable/Source/CustomizableObjectEditor/Private/MuCOE/GenerateMutableSource/GenerateMutableSourceMesh.cpp
--- a/Engine/Plugins/Mutable/Source/CustomizableObjectEditor/Private/MuCOE/GenerateMutableSource/GenerateMutableSourceMesh.cpp
+++ b/Engine/Plugins/Mutable/Source/CustomizableObjectEditor/Private/MuCOE/GenerateMutableSource/GenerateMutableSourceMesh.cpp
@@ -2463,39 +2463,11 @@ TSharedPtr<mu::FMesh> GenerateMutableMesh(UObject * Mesh, const TSoftClassPtr<UA
 		else
 		{
 			MutableMesh = ConvertSkeletalMeshToMutable(SkeletalMesh, AnimInstance, LODIndexConnected, SectionIndexConnected, LODIndex, SectionIndex, SurfaceMetadataId, GenerationContext, CurrentNode, TableReferenceSkeletalMesh);
-
-			FSkeletalMeshModel* ImportedModel = SkeletalMesh->GetImportedModel();
-
-			if (MutableMesh &&
-				ImportedModel->LODModels.IsValidIndex(LODIndex) &&
-				ImportedModel->LODModels[LODIndex].Sections.IsValidIndex(SectionIndex))
-			{
-				FMeshData MeshData;
-				MeshData.Mesh = Mesh;
-				MeshData.LOD = LODIndex;
-				MeshData.MaterialIndex = SectionIndex;
-				MeshData.Node = CurrentNode;
-				GenerationContext.PinData.GetCurrent().MeshesData.Add(MeshData); // Set::Emplace only supports single element constructors
-			}
 		}
 	}
 	else if (const UStaticMesh* StaticMesh = Cast<UStaticMesh>(Mesh))
 	{
 		MutableMesh = ConvertStaticMeshToMutable(StaticMesh, LODIndex, SectionIndex, GenerationContext, CurrentNode);
-
-		const FStaticMeshRenderData* RenderData = StaticMesh->GetRenderData();
-
-		if (MutableMesh &&
-			RenderData->LODResources.IsValidIndex(LODIndex) &&
-			RenderData->LODResources[LODIndex].Sections.IsValidIndex(SectionIndex))
-		{
-			FMeshData MeshData;
-			MeshData.Mesh = Mesh;
-			MeshData.LOD = LODIndex;
-			MeshData.MaterialIndex = SectionIndex;
-			MeshData.Node = CurrentNode;
-			GenerationContext.PinData.GetCurrent().MeshesData.Add(MeshData); // Set::Emplace only supports single element constructors
-		}
 	}
 	else
 	{
@@ -2576,8 +2548,6 @@ TSharedPtr<mu::FMesh> BuildMorphedMutableMeshFromMesh(TSharedPtr<mu::FMesh> Base
 TSharedPtr<mu::FMesh> BuildMorphedMutableMesh(const UEdGraphPin* BaseSourcePin, const FString& MorphTargetName, FMutableGraphGenerationContext& GenerationContext, const bool bOnlyConnectedLOD, const FName& RowName)
 {
 	check(BaseSourcePin);
-	SCOPED_PIN_DATA(GenerationContext, BaseSourcePin)
-
 
 	if (!BaseSourcePin)
 	{
@@ -3128,8 +3098,6 @@ mu::NodeMeshPtr GenerateMorphMesh(const UEdGraphPin* Pin,
 	const FString& TableColumnName = "")
 {
 	MUTABLE_CPUPROFILER_SCOPE(GenerateMorphMesh);
-
-	SCOPED_PIN_DATA(GenerationContext, Pin)
 	
 	// SkeletalMesh node
 	const UEdGraphNode * MeshNode = Pin->GetOwningNode();
@@ -3308,7 +3276,6 @@ mu::Ptr<mu::NodeMesh> GenerateMutableSourceMesh(const UEdGraphPin* Pin,
 
 	check(Pin)
 	RETURN_ON_CYCLE(*Pin, GenerationContext)
-	SCOPED_PIN_DATA(GenerationContext, Pin)
 
 	CheckNumOutputs(*Pin, GenerationContext);
 
diff --git a/Engine/Plugins/Mutable/Source/CustomizableObjectEditor/Private/MuCOE/GenerateMutableSource/GenerateMutableSourceModifier.cpp b/Engine/Plugins/Mutable/Source/CustomizableObjectEditor/Private/MuCOE/GenerateMutableSource/GenerateMutableSourceModifier.cpp
--- a/Engine/Plugins/Mutable/Source/CustomizableObjectEditor/Private/MuCOE/GenerateMutableSource/GenerateMutableSourceModifier.cpp
+++ b/Engine/Plugins/Mutable/Source/CustomizableObjectEditor/Private/MuCOE/GenerateMutableSource/GenerateMutableSourceModifier.cpp
@@ -170,32 +170,6 @@ mu::Ptr<mu::NodeModifier> GenerateMutableSourceModifier(const UEdGraphPin * Pin,
 
 			mu::NodeMeshPtr ClipMesh = GenerateMutableSourceMesh(ConnectedPin, GenerationContext, DummyMeshData, 0, false, true);
 
-			FPinDataValue* PinData = GenerationContext.PinData.Find(ConnectedPin);
-			for (const FMeshData& MeshData : PinData->MeshesData)
-			{
-				bool bClosed = true;
-				if (const USkeletalMesh* SkeletalMesh = Cast<USkeletalMesh>(MeshData.Mesh))
-				{
-					bClosed = IsMeshClosed(SkeletalMesh, MeshData.LOD, MeshData.MaterialIndex);
-				}
-				else if (const UStaticMesh* StaticMesh = Cast<UStaticMesh>(MeshData.Mesh))
-				{
-					bClosed = IsMeshClosed(StaticMesh, MeshData.LOD, MeshData.MaterialIndex);
-				}
-				else
-				{
-					// TODO: We support the clip mesh not being constant. This message is not precise enough. It should say that it hasn't been 
-					// possible to check if the mesh is closed or not.
-					GenerationContext.Log(LOCTEXT("UnimplementedNode", "Node type not implemented yet."), MeshData.Node);
-				}
-
-				if (!bClosed)
-				{
-					FText ErrorMsg = FText::Format(LOCTEXT("Clipping mesh", "Clipping mesh [{0}] not closed (i.e., it does not enclose a volume)."), FText::FromName(MeshData.Mesh->GetFName()));
-					GenerationContext.Log(ErrorMsg, MeshData.Node, EMessageSeverity::Warning);
-				}
-			}
-
 			if (FMatrix Matrix = TypedNodeClipMesh->Transform.ToMatrixWithScale(); Matrix != FMatrix::Identity)
 			{
 				mu::NodeMeshTransformPtr TransformMesh = new mu::NodeMeshTransform();
@@ -634,33 +608,7 @@ mu::Ptr<mu::NodeModifier> GenerateMutableSourceModifier(const UEdGraphPin * Pin,
 		{
 			FMutableGraphMeshGenerationData DummyMeshData;
 
-			mu::NodeMeshPtr BoundingMesh = GenerateMutableSourceMesh(ConnectedPin, GenerationContext, DummyMeshData, 0, false, true);
-
-			const FPinDataValue* PinData = GenerationContext.PinData.Find(ConnectedPin);
-			for (const FMeshData& MeshData : PinData->MeshesData)
-			{
-				bool bClosed = true;
-				if (const USkeletalMesh* SkeletalMesh = Cast<USkeletalMesh>(MeshData.Mesh))
-				{
-					bClosed = IsMeshClosed(SkeletalMesh, MeshData.LOD, MeshData.MaterialIndex);
-				}
-				else if (const UStaticMesh* StaticMesh = Cast<UStaticMesh>(MeshData.Mesh))
-				{
-					bClosed = IsMeshClosed(StaticMesh, MeshData.LOD, MeshData.MaterialIndex);
-				}
-				else
-				{
-					// TODO: We support the bounding mesh not being constant. This message is not precise enough. It should say that it hasn't been 
-					// possible to check if the mesh is closed or not.
-					GenerationContext.Log(LOCTEXT("UnimplementedNode", "Node type not implemented yet."), MeshData.Node);
-				}
-
-				if (!bClosed)
-				{
-					FText ErrorMsg = FText::Format(LOCTEXT("Bounding mesh", "The bounding [{0}] not watertight (i.e. it does not fully enclose a volume)."), FText::FromName(MeshData.Mesh->GetFName()));
-					GenerationContext.Log(ErrorMsg, MeshData.Node, EMessageSeverity::Warning);
-				}
-			}
+			mu::Ptr<mu::NodeMesh> BoundingMesh = GenerateMutableSourceMesh(ConnectedPin, GenerationContext, DummyMeshData, 0, false, true);
 
 			if (FMatrix Matrix = TypedNodeTransformMesh->BoundingMeshTransform.ToMatrixWithScale(); Matrix != FMatrix::Identity)
 			{
diff --git a/Engine/Plugins/Mutable/Source/CustomizableObjectEditor/Private/MuCOE/GenerateMutableSource/GenerateMutableSourceTable.cpp b/Engine/Plugins/Mutable/Source/CustomizableObjectEditor/Private/MuCOE/GenerateMutableSource/GenerateMutableSourceTable.cpp
--- a/Engine/Plugins/Mutable/Source/CustomizableObjectEditor/Private/MuCOE/GenerateMutableSource/GenerateMutableSourceTable.cpp
+++ b/Engine/Plugins/Mutable/Source/CustomizableObjectEditor/Private/MuCOE/GenerateMutableSource/GenerateMutableSourceTable.cpp
@@ -922,8 +922,6 @@ bool GenerateTableColumn(const UCustomizableObjectNodeTable* TableNode, const UE
 {
 	MUTABLE_CPUPROFILER_SCOPE(GenerateTableColumn);
 
-	SCOPED_PIN_DATA(GenerationContext, Pin)
-
 	if (!TableNode)
 	{
 		return false;
diff --git a/Engine/Plugins/Mutable/Source/CustomizableObjectEditor/Private/MuCOE/MutableUtils.cpp b/Engine/Plugins/Mutable/Source/CustomizableObjectEditor/Private/MuCOE/MutableUtils.cpp
--- a/Engine/Plugins/Mutable/Source/CustomizableObjectEditor/Private/MuCOE/MutableUtils.cpp
+++ b/Engine/Plugins/Mutable/Source/CustomizableObjectEditor/Private/MuCOE/MutableUtils.cpp
@@ -73,183 +73,3 @@ bool HasNormalizedBounds(const FVector2f& Point)
 {
 	return Point.X >= 0.0f && Point.X <= 1.0 && Point.Y >= 0.0f && Point.Y <= 1.0;
 }
-
-
-bool IsMeshClosed(const USkeletalMesh* Mesh, int LOD, int MaterialIndex)
-{
-	check(Mesh->GetImportedModel() &&
-		Mesh->GetImportedModel()->LODModels.IsValidIndex(LOD) &&
-		Mesh->GetImportedModel()->LODModels[LOD].Sections.IsValidIndex(MaterialIndex));
-
-	const float VertexCollapseEpsilonSquared = 0.0001f * 0.0001f;
-
-	const FSkeletalMeshLODModel& LODResource = Mesh->GetImportedModel()->LODModels[LOD];
-	const FSkelMeshSection& Section = LODResource.Sections[MaterialIndex];
-
-	const TArray<FSoftSkinVertex>& VertexBuffer = Section.SoftVertices;
-	const uint32 VertexStart = Section.BaseVertexIndex;
-	const uint32 VertexCount = Section.NumVertices;
-
-	// Collapse vertices
-	TArray<uint32> CollapsedVertexMap;
-	CollapsedVertexMap.SetNum(VertexCount);
-
-	auto GetCollapsedVertex = [&CollapsedVertexMap, VertexStart](uint32 VertIndex) -> uint32&
-	{
-
-		check(CollapsedVertexMap.IsValidIndex(VertIndex - VertexStart));
-		return CollapsedVertexMap[VertIndex - VertexStart];
-	};
-
-	for (uint32 V = VertexStart; V < VertexCount; ++V)
-	{
-		GetCollapsedVertex(V) = V;
-
-		for (uint32 CandidateIndex = VertexStart; CandidateIndex < V; ++CandidateIndex)
-		{
-			uint32 Collapsed = GetCollapsedVertex(CandidateIndex);
-
-			FVector3f Vec = VertexBuffer[V].Position - VertexBuffer[Collapsed].Position;
-
-			if (FVector3f::DotProduct(Vec, Vec) <= VertexCollapseEpsilonSquared)
-			{
-				GetCollapsedVertex(V) = Collapsed;
-				break;
-			}
-		}
-	}
-
-	// Count faces per edge
-	TMap<TPair<uint32, uint32>, uint32> FaceCountPerEdge;
-
-	const TArray<uint32>& IndexBuffer = LODResource.IndexBuffer;//.GetData();
-	const uint32 IndexStart = Section.BaseIndex;
-	const uint32 IndexCount = Section.NumTriangles * 3;
-
-	for (uint32 I = IndexStart; I < IndexCount; I += 3)
-	{
-		uint32 Face[3];
-		uint32 Index0 = IndexBuffer[I + 0];
-		Face[0] = GetCollapsedVertex(Index0);
-		uint32 Index1 = IndexBuffer[I + 1];
-		Face[1] = GetCollapsedVertex(Index1);
-		uint32 Index2 = IndexBuffer[I + 2];
-		Face[2] = GetCollapsedVertex(Index2);
-
-		for (int E = 0; E < 3; ++E)
-		{
-			uint32 Vertex0 = Face[E];
-			uint32 Vertex1 = Face[(E + 1) % 3];
-
-			if (Vertex0 == Vertex1)
-			{
-				return false; // Degenerated mesh
-			}
-
-			TPair<uint32, uint32> Edge;
-			Edge.Key = FGenericPlatformMath::Min(Vertex0, Vertex1);
-			Edge.Value = FGenericPlatformMath::Max(Vertex0, Vertex1);
-
-			FaceCountPerEdge.FindOrAdd(Edge)++;
-		}
-	}
-
-	// See if every edge has 2 faces
-	for (const TPair<TPair<uint32, uint32>, uint32>& Entry : FaceCountPerEdge)
-	{
-		if (Entry.Value != 2)
-		{
-			return false;
-		}
-	}
-
-	return true;
-}
-
-
-bool IsMeshClosed(const UStaticMesh* Mesh, int LOD, int MaterialIndex)
-{
-	check(Mesh->GetRenderData() &&
-		Mesh->GetRenderData()->LODResources.IsValidIndex(LOD) &&
-		Mesh->GetRenderData()->LODResources[LOD].Sections.IsValidIndex(MaterialIndex));
-
-	const float VertexCollapseEpsilonSquared = 0.0001f * 0.0001f;
-
-	const FStaticMeshLODResources& LODResource = Mesh->GetRenderData()->LODResources[LOD];
-	const FStaticMeshSection& Section = LODResource.Sections[MaterialIndex];
-
-	const uint32 VertexStart = Section.MinVertexIndex;
-	const uint32 VertexCount = Section.MaxVertexIndex - VertexStart + 1;
-
-	// Collapse vertices
-	TArray<uint32> CollapsedVertexMap;
-	CollapsedVertexMap.SetNum(VertexCount);
-
-	auto GetCollapsedVertex = [&CollapsedVertexMap, VertexStart](uint32 VertIndex) -> uint32&
-	{
-
-		check(CollapsedVertexMap.IsValidIndex(VertIndex - VertexStart));
-		return CollapsedVertexMap[VertIndex - VertexStart];
-	};
-
-	for (uint32 V = VertexStart; V < VertexCount; ++V)
-	{
-		GetCollapsedVertex(V) = V;
-
-		for (uint32 CandidateIndex = VertexStart; CandidateIndex < V; ++CandidateIndex)
-		{
-			uint32 Collapsed = GetCollapsedVertex(CandidateIndex);
-
-			FVector3f Vec = LODResource.VertexBuffers.PositionVertexBuffer.VertexPosition(V) - LODResource.VertexBuffers.PositionVertexBuffer.VertexPosition(Collapsed);
-
-			if (FVector3f::DotProduct(Vec, Vec) <= VertexCollapseEpsilonSquared)
-			{
-				GetCollapsedVertex(V) = Collapsed;
-				break;
-			}
-		}
-	}
-
-	// Count faces per edge
-	TMap<TPair<uint32, uint32>, uint32> FaceCountPerEdge;
-
-	const FIndexArrayView& IndexBuffer = LODResource.IndexBuffer.GetArrayView();
-	const uint32 IndexStart = Section.FirstIndex;
-	const uint32 IndexCount = Section.NumTriangles * 3;
-
-	for (uint32 I = IndexStart; I < IndexCount; I += 3)
-	{
-		uint32 Face[3];
-		Face[0] = GetCollapsedVertex(IndexBuffer[I + 0]);
-		Face[1] = GetCollapsedVertex(IndexBuffer[I + 1]);
-		Face[2] = GetCollapsedVertex(IndexBuffer[I + 2]);
-
-		for (int E = 0; E < 3; ++E)
-		{
-			uint32 Vertex0 = Face[E];
-			uint32 Vertex1 = Face[(E + 1) % 3];
-
-			if (Vertex0 == Vertex1)
-			{
-				return false; // Degenerated mesh
-			}
-
-			TPair<uint32, uint32> Edge;
-			Edge.Key = FGenericPlatformMath::Min(Vertex0, Vertex1);
-			Edge.Value = FGenericPlatformMath::Max(Vertex0, Vertex1);
-
-			FaceCountPerEdge.FindOrAdd(Edge)++;
-		}
-	}
-
-	// See if every edge has 2 faces
-	for (const TPair<TPair<uint32, uint32>, uint32>& Entry : FaceCountPerEdge)
-	{
-		if (Entry.Value != 2)
-		{
-			return false;
-		}
-	}
-
-	return true;
-}
diff --git a/Engine/Plugins/Mutable/Source/CustomizableObjectEditor/Public/MuCOE/MutableUtils.h b/Engine/Plugins/Mutable/Source/CustomizableObjectEditor/Public/MuCOE/MutableUtils.h
--- a/Engine/Plugins/Mutable/Source/CustomizableObjectEditor/Public/MuCOE/MutableUtils.h
+++ b/Engine/Plugins/Mutable/Source/CustomizableObjectEditor/Public/MuCOE/MutableUtils.h
@@ -11,18 +11,8 @@ class UStaticMesh;
 /** Returns the mesh UV. */
 TArray<FVector2f> GetUV(const USkeletalMesh& SkeletalMesh, const int32 LODIndex, const int32 SectionIndex, const int32 UVIndex);
 
-
 /** Returns the mesh UV. */
 TArray<FVector2f> GetUV(const UStaticMesh& StaticMesh, const int32 LODIndex, const int32 SectionIndex, const int32 UVIndex);
 
-
 /** Returns if the given point is inside the 0-1 bounds. */
 bool HasNormalizedBounds(const FVector2f& Point);
-
-
-/** Check if a given LOD and material in a skeletal mesh is closed. */
-bool IsMeshClosed(const USkeletalMesh* Mesh, int LOD, int MaterialIndex);
-
-
-/** Check if a given LOD and material in a skeletal mesh is closed. */
-bool IsMeshClosed(const UStaticMesh* Mesh, int LOD, int MaterialIndex);
diff --git a/Engine/Plugins/Mutable/Source/MutableRuntime/Private/MuR/CodeRunner.cpp b/Engine/Plugins/Mutable/Source/MutableRuntime/Private/MuR/CodeRunner.cpp
--- a/Engine/Plugins/Mutable/Source/MutableRuntime/Private/MuR/CodeRunner.cpp
+++ b/Engine/Plugins/Mutable/Source/MutableRuntime/Private/MuR/CodeRunner.cpp
@@ -2172,17 +2172,17 @@ namespace mu
 				MUTABLE_CPUPROFILER_SCOPE(ME_CLIPWITHMESH_1)
 
                 TSharedPtr<const FMesh> Source = LoadMesh(FCacheAddress(args.source, item));
-                TSharedPtr<const FMesh> pClip = LoadMesh(FCacheAddress(args.clipMesh, item));
+                TSharedPtr<const FMesh> Clip = LoadMesh(FCacheAddress(args.clipMesh, item));
 
                 // Only if both are valid.
-                if (Source && pClip)
+                if (Source && Clip)
                 {
 					TSharedPtr<FMesh> Result = CreateMesh(Source->GetDataSize());
 
 					bool bOutSuccess = false;
-                    MeshClipWithMesh(Result.Get(), Source.Get(), pClip.Get(), bOutSuccess);
+                    MeshClipWithMesh(Result.Get(), Source.Get(), Clip.Get(), bOutSuccess);
 
-					Release(pClip);
+					Release(Clip);
 					if (!bOutSuccess)
 					{
 						Release(Result);
@@ -2196,7 +2196,7 @@ namespace mu
                 }
                 else
                 {
-					Release(pClip);
+					Release(Clip);
 					StoreMesh(item, Source);
                 }
 
diff --git a/Engine/Plugins/Mutable/Source/MutableRuntime/Private/MuR/Mesh.cpp b/Engine/Plugins/Mutable/Source/MutableRuntime/Private/MuR/Mesh.cpp
--- a/Engine/Plugins/Mutable/Source/MutableRuntime/Private/MuR/Mesh.cpp
+++ b/Engine/Plugins/Mutable/Source/MutableRuntime/Private/MuR/Mesh.cpp
@@ -9,6 +9,7 @@
 #include "HAL/UnrealMemory.h"
 #include "MuR/MeshPrivate.h"
 #include "MuR/MutableTrace.h"
+#include "MuR/OpMeshClipWithMesh.h"
 
 namespace mu
 {
@@ -1294,6 +1295,12 @@ void FMesh::CheckIntegrity() const
 }
 
 
+bool FMesh::IsClosed() const
+{
+	return IsMeshClosed(this);
+}
+
+
 static bool StaticMeshFormatIdentify_Project(const FMesh* InMesh)
 {
     // This format is used internally for the mesh project
diff --git a/Engine/Plugins/Mutable/Source/MutableRuntime/Private/MuR/OpMeshClipWithMesh.cpp b/Engine/Plugins/Mutable/Source/MutableRuntime/Private/MuR/OpMeshClipWithMesh.cpp
--- a/Engine/Plugins/Mutable/Source/MutableRuntime/Private/MuR/OpMeshClipWithMesh.cpp
+++ b/Engine/Plugins/Mutable/Source/MutableRuntime/Private/MuR/OpMeshClipWithMesh.cpp
@@ -8,17 +8,17 @@
 #include "MuR/ConvertData.h"
 #include "MuR/Platform.h"
 #include "MuR/MutableTrace.h"
+#include "MuR/MutableRuntimeModule.h"
 #include "Math/Ray.h"
 #include "TriangleTypes.h"
 #include "BoxTypes.h"
 #include "Intersection/IntrRay3Triangle3.h"
 #include "MathUtil.h"
 #include "IntVectorTypes.h"
-
 #include "Math/UnrealMathUtility.h"
-
 #include "Spatial/PointHashGrid3.h"
 
+
 namespace mu { namespace 
 {
 
@@ -77,72 +77,61 @@ namespace mu { namespace
 	}
 
 
-    //---------------------------------------------------------------------------------------------
-    //! Return true if a mesh is closed
-    //---------------------------------------------------------------------------------------------
-  //  bool MeshIsClosed( const FMesh* pMesh )
-  //  {
-  //      MUTABLE_CPUPROFILER_SCOPE(MeshIsClosed);
-
-  //      int vcount = pMesh->GetVertexCount();
-  //      int fcount = pMesh->GetFaceCount();
-
-  //      TArray<FVector3f> vertices;
-
-  //      // Map in ClipMesh from vertices to the one they are collapsed to because they are very similar, if they aren't collapsed then they are mapped to themselves
-		//TArray<int> collapsedVertexMap;
-
-  //      MeshCreateCollapsedVertexMap( pMesh, collapsedVertexMap, vertices );
-
-  //      // Acumulate edges
-  //      TMap< TPair<int,int>, int > faceCountPerEdge;
-
-  //      UntypedMeshBufferIteratorConst itClipMesh(pMesh->GetIndexBuffers(), EMeshBufferSemantic::VertexIndex);
-  //      for (int f = 0; f < fcount; ++f)
-  //      {
-  //          int face[3];
-  //          face[0] = collapsedVertexMap[itClipMesh.GetAsUINT32()]; ++itClipMesh;
-  //          face[1] = collapsedVertexMap[itClipMesh.GetAsUINT32()]; ++itClipMesh;
-  //          face[2] = collapsedVertexMap[itClipMesh.GetAsUINT32()]; ++itClipMesh;
-
-  //          for (int e=0; e<3; ++e)
-  //          {
-  //              int v0 = face[e];
-  //              int v1 = face[(e+1)%3];
-
-  //              if(v0==v1)
-  //              {
-  //                  // Degenerated mesh
-  //                  return false;
-  //              }
-
-  //              TPair<int,int> edge;
-  //              edge.Key = FMath::Min( v0, v1 );
-  //              edge.Value = FMath::Max( v0, v1 );
-
-		//		if (faceCountPerEdge.Contains(edge))
-		//		{
-		//			faceCountPerEdge[edge]++;
-		//		}
-		//		else
-		//		{
-		//			faceCountPerEdge.Add(edge,1);
-		//		}
-  //          }
-  //      }
-
-
-  //      // See if every edge has 2 faces
-  //      for( const auto& e: faceCountPerEdge)
-  //      {
-  //          if (e.Value!=2)
-  //          {
-  //              return false;
-  //          }
-  //      }
-
-  //      return true;
-  //  }
+	/** Return true if the mesh is closed. Usually used to validate clipping meshes. */
+	bool IsMeshClosed( const FMesh* Mesh, const TArray<int32>& CollapsedVertexMap)
+    {
+        MUTABLE_CPUPROFILER_SCOPE(IsMeshClosed);
+
+		if (!Mesh)
+		{
+			return true;
+		}
+
+        int32 FaceCount = Mesh->GetFaceCount();
+
+        // Acumulate edges
+		using FEdge = TPair<int32, int32>;
+        TMap< FEdge, int32 > FaceCountPerEdge;
+
+        UntypedMeshBufferIteratorConst ItClipMesh(Mesh->GetIndexBuffers(), EMeshBufferSemantic::VertexIndex);
+        for (int32 FaceIndex = 0; FaceIndex < FaceCount; ++FaceIndex)
+        {
+            int32 Face[3];
+			Face[0] = CollapsedVertexMap[ItClipMesh.GetAsUINT32()]; ++ItClipMesh;
+			Face[1] = CollapsedVertexMap[ItClipMesh.GetAsUINT32()]; ++ItClipMesh;
+			Face[2] = CollapsedVertexMap[ItClipMesh.GetAsUINT32()]; ++ItClipMesh;
+
+            for (int32 EdgeIndex=0; EdgeIndex <3; ++EdgeIndex)
+            {
+                int32 v0 = Face[EdgeIndex];
+                int32 v1 = Face[(EdgeIndex+1)%3];
+
+                if(v0==v1)
+                {
+                    // Degenerated mesh
+                    return false;
+                }
+
+				FEdge Edge;
+				Edge.Key = FMath::Min( v0, v1 );
+				Edge.Value = FMath::Max( v0, v1 );
+
+				int32& Count = FaceCountPerEdge.FindOrAdd(Edge);
+				++Count;
+            }
+        }
+
+        // See if every edge has 2 faces
+        for( const TPair<FEdge, int32>& Entry: FaceCountPerEdge)
+        {
+            if (Entry.Value!=2)
+            {
+                return false;
+            }
+        }
+
+        return true;
+    }
 
 
     //---------------------------------------------------------------------------------------------
@@ -499,21 +488,39 @@ namespace mu { namespace
 
 namespace mu
 {
+
+	bool IsMeshClosed(const FMesh* Mesh)
+	{
+		if (!Mesh)
+		{
+			return true;
+		}
+
+		const int32 VCount = Mesh->GetVertexBuffers().GetElementCount();
+
+		TArray<FVector3f> Vertices;
+		Vertices.SetNumUninitialized(VCount);
+
+		TArray<int32> CollapsedVertexMap;
+		CollapsedVertexMap.AddUninitialized(VCount);
+
+		MeshCreateCollapsedVertexMap(Mesh, CollapsedVertexMap, Vertices);
+
+		return IsMeshClosed(Mesh, CollapsedVertexMap);
+	}
+
+
     /** Core Geometry version */
-    void MeshClipMeshClassifyVertices(TBitArray<>& VertexInClipMesh, const FMesh* pBase, const FMesh* pClipMesh)
+    void MeshClipMeshClassifyVertices(TBitArray<>& VertexInClipMesh, const FMesh* Base, const FMesh* ClipMesh)
     {
 		using namespace UE::Geometry;
 
         MUTABLE_CPUPROFILER_SCOPE(MeshClipMeshClassifyVertices);
     	
-        // TODO: this fails in bro for some object. do it at graph construction time to report it 
-        // nicely.
-        //check( MeshIsClosed(pClipMesh) );
-
-        const int32 VCount = pClipMesh->GetVertexBuffers().GetElementCount();
-        const int32 FCount = pClipMesh->GetFaceCount();
+        const int32 VCount = ClipMesh->GetVertexBuffers().GetElementCount();
+        const int32 FCount = ClipMesh->GetFaceCount();
 
-        int32 OrigVertCount = pBase->GetVertexBuffers().GetElementCount();
+        int32 OrigVertCount = Base->GetVertexBuffers().GetElementCount();
 
         // Stores whether each vertex in the original mesh in the clip mesh volume
         VertexInClipMesh.SetNum(OrigVertCount, false);
@@ -534,10 +541,17 @@ namespace mu
         TArray<int32> CollapsedVertexMap;
 		CollapsedVertexMap.AddUninitialized(VCount);
 
-        MeshCreateCollapsedVertexMap( pClipMesh, CollapsedVertexMap, Vertices );
+        MeshCreateCollapsedVertexMap( ClipMesh, CollapsedVertexMap, Vertices );
+
+#if !UE_BUILD_SHIPPING
+		if (!IsMeshClosed(ClipMesh, CollapsedVertexMap))
+		{
+			UE_LOG(LogMutableCore, Warning, TEXT("Mesh operation with a mesh that is not closed as required."));
+		}
+#endif
 
         // Create cache of the faces
-        UntypedMeshBufferIteratorConst ItClipMesh(pClipMesh->GetIndexBuffers(), EMeshBufferSemantic::VertexIndex);
+        UntypedMeshBufferIteratorConst ItClipMesh(ClipMesh->GetIndexBuffers(), EMeshBufferSemantic::VertexIndex);
 
         for ( int32 F = 0; F < FCount; ++F )
         {
@@ -599,9 +613,9 @@ namespace mu
         }	
 
         // Now go through all vertices in the mesh and record whether they are inside or outside of the ClipMesh
-        uint32 DestVertexCount = pBase->GetVertexCount();
+        uint32 DestVertexCount = Base->GetVertexCount();
 
-        const FMeshBufferSet& MBSPriv2 = pBase->GetVertexBuffers();
+        const FMeshBufferSet& MBSPriv2 = Base->GetVertexBuffers();
         for (int32 b = 0; b < MBSPriv2.Buffers.Num(); ++b)
         {
             for (int32 c = 0; c < MBSPriv2.Buffers[b].Channels.Num(); ++c)
@@ -609,7 +623,7 @@ namespace mu
                 EMeshBufferSemantic Sem = MBSPriv2.Buffers[b].Channels[c].Semantic;
                 int32 SemIndex = MBSPriv2.Buffers[b].Channels[c].SemanticIndex;
 
-                UntypedMeshBufferIteratorConst It(pBase->GetVertexBuffers(), Sem, SemIndex);
+                UntypedMeshBufferIteratorConst It(Base->GetVertexBuffers(), Sem, SemIndex);
 
                 TArray<uint8> VertexAlreadyIntersected;
 				VertexAlreadyIntersected.AddZeroed(VCount);
@@ -685,12 +699,12 @@ namespace mu
         }
     }	
 
-    void MeshClipWithMesh(FMesh* Result, const FMesh* pBase, const FMesh* pClipMesh, bool& bOutSuccess)
+    void MeshClipWithMesh(FMesh* Result, const FMesh* pBase, const FMesh* ClipMesh, bool& bOutSuccess)
     {
         MUTABLE_CPUPROFILER_SCOPE(MeshClipWithMesh);
 		bOutSuccess = true;
 
-        uint32 VCount = pClipMesh->GetVertexBuffers().GetElementCount();
+        uint32 VCount = ClipMesh->GetVertexBuffers().GetElementCount();
         if (!VCount)
         {
 			bOutSuccess = false;
@@ -700,7 +714,7 @@ namespace mu
 		Result->CopyFrom(*pBase);
 
         TBitArray<> VertexInClipMesh;  // Stores whether each vertex in the original mesh is in the clip mesh volume
-        MeshClipMeshClassifyVertices(VertexInClipMesh, pBase, pClipMesh);
+        MeshClipMeshClassifyVertices(VertexInClipMesh, pBase, ClipMesh);
 
         // Now remove all the faces from the result mesh that have all the vertices outside the clip volume
         UntypedMeshBufferIteratorConst ItBase(Result->GetIndexBuffers(), EMeshBufferSemantic::VertexIndex);
diff --git a/Engine/Plugins/Mutable/Source/MutableRuntime/Private/MuR/OpMeshClipWithMesh.h b/Engine/Plugins/Mutable/Source/MutableRuntime/Private/MuR/OpMeshClipWithMesh.h
--- a/Engine/Plugins/Mutable/Source/MutableRuntime/Private/MuR/OpMeshClipWithMesh.h
+++ b/Engine/Plugins/Mutable/Source/MutableRuntime/Private/MuR/OpMeshClipWithMesh.h
@@ -28,4 +28,7 @@ namespace mu
 	/** Generate a mask mesh with the faces of the base mesh that have all 3 vertices inside any block of the given layout. */
 	extern void MakeMeshMaskFromLayout(FMesh* Result, const FMesh* Base, const FMesh* BaseForUVs, const FLayout* Mask, uint8 LayoutIndex, bool& bOutSuccess);
 
+	/** Return true if the mesh is closed. Usually used to validate clipping meshes. */
+	extern bool IsMeshClosed(const FMesh* Mesh);
+
 }
diff --git a/Engine/Plugins/Mutable/Source/MutableRuntime/Public/MuR/Mesh.h b/Engine/Plugins/Mutable/Source/MutableRuntime/Public/MuR/Mesh.h
--- a/Engine/Plugins/Mutable/Source/MutableRuntime/Public/MuR/Mesh.h
+++ b/Engine/Plugins/Mutable/Source/MutableRuntime/Public/MuR/Mesh.h
@@ -546,6 +546,9 @@ namespace mu
 		/** Check mesh buffer data for possible inconsistencies */
 		void CheckIntegrity() const;
 
+		/** Return true if this mesh is closed. It is still closed if it has dangling vertices. */
+		bool IsClosed() const;
+
 		/** 
 		 * Change the buffer descriptions so that all buffer indices start at 0 and are in the
 		 * same order than memory.
diff --git a/Engine/Plugins/Mutable/Source/MutableTools/Internal/MuT/AST.h b/Engine/Plugins/Mutable/Source/MutableTools/Internal/MuT/AST.h
--- a/Engine/Plugins/Mutable/Source/MutableTools/Internal/MuT/AST.h
+++ b/Engine/Plugins/Mutable/Source/MutableTools/Internal/MuT/AST.h
@@ -880,6 +880,16 @@ namespace mu
         //!
         virtual bool GetNonBlackRect( FImageRect& maskUsage ) const;
 
+		enum class EClosedMeshTest : uint8
+		{
+			No,
+			Yes,
+			Unknown
+		};
+
+		/** This can be overriden to help detect sugraph mesh properties.  */
+		virtual EClosedMeshTest IsClosedMesh(TMap<const ASTOp*, EClosedMeshTest>* Cache=nullptr ) const;
+
         // Logic expression evaluation
         //---------------------------------------------------------------------------------------------
         typedef enum
diff --git a/Engine/Plugins/Mutable/Source/MutableTools/Private/MuT/AST.cpp b/Engine/Plugins/Mutable/Source/MutableTools/Private/MuT/AST.cpp
--- a/Engine/Plugins/Mutable/Source/MutableTools/Private/MuT/AST.cpp
+++ b/Engine/Plugins/Mutable/Source/MutableTools/Private/MuT/AST.cpp
@@ -915,6 +915,13 @@ bool ASTOp::GetNonBlackRect( FImageRect& ) const
 }
 
 
+ASTOp::EClosedMeshTest ASTOp::IsClosedMesh(TMap<const ASTOp*, EClosedMeshTest>* Cache) const
+{ 
+	// If this is hit, consider implementing it for that subclass.
+	return EClosedMeshTest::Unknown; 
+}
+
+
 void ASTOp::LinkRange(FProgram& program,
 	const FRangeData& range,
 	OP::ADDRESS& rangeSize,
diff --git a/Engine/Plugins/Mutable/Source/MutableTools/Private/MuT/ASTOpConstantResource.cpp b/Engine/Plugins/Mutable/Source/MutableTools/Private/MuT/ASTOpConstantResource.cpp
--- a/Engine/Plugins/Mutable/Source/MutableTools/Private/MuT/ASTOpConstantResource.cpp
+++ b/Engine/Plugins/Mutable/Source/MutableTools/Private/MuT/ASTOpConstantResource.cpp
@@ -784,4 +784,41 @@ namespace mu
 		return SourceDataDescriptor;
 	}
 
+	
+	ASTOp::EClosedMeshTest ASTOpConstantResource::IsClosedMesh(TMap<const ASTOp*, EClosedMeshTest>* Cache) const
+	{
+		if (Cache)
+		{
+			const ASTOp::EClosedMeshTest* Cached = Cache->Find(this);
+			if (Cached)
+			{
+				return *Cached;
+			}
+		}
+
+		ASTOp::EClosedMeshTest Result = EClosedMeshTest::Unknown;
+		if (Type == EOpType::ME_CONSTANT)
+		{
+			const FMesh* Mesh = static_cast<const FMesh*>(GetValue().Get());
+			if (Mesh)
+			{
+				if (!Mesh->IsClosed())
+				{
+					Result = EClosedMeshTest::Yes;
+				}
+				else
+				{
+					Result = EClosedMeshTest::No;
+				}
+			}
+		}
+
+		if (Cache)
+		{
+			Cache->Add(this,Result);
+		}
+
+		return Result;
+	}
+
 }
diff --git a/Engine/Plugins/Mutable/Source/MutableTools/Private/MuT/ASTOpConstantResource.h b/Engine/Plugins/Mutable/Source/MutableTools/Private/MuT/ASTOpConstantResource.h
--- a/Engine/Plugins/Mutable/Source/MutableTools/Private/MuT/ASTOpConstantResource.h
+++ b/Engine/Plugins/Mutable/Source/MutableTools/Private/MuT/ASTOpConstantResource.h
@@ -70,6 +70,7 @@ namespace mu
 		virtual bool IsImagePlainConstant(FVector4f& colour) const override;
 		virtual Ptr<ImageSizeExpression> GetImageSizeExpression() const override;
 		virtual FSourceDataDescriptor GetSourceDataDescriptor(FGetSourceDataDescriptorContext*) const override;
+		virtual EClosedMeshTest IsClosedMesh(TMap<const ASTOp*, EClosedMeshTest>* Cache) const override;
 	};
 
 
diff --git a/Engine/Plugins/Mutable/Source/MutableTools/Private/MuT/CodeGenerator.cpp b/Engine/Plugins/Mutable/Source/MutableTools/Private/MuT/CodeGenerator.cpp
--- a/Engine/Plugins/Mutable/Source/MutableTools/Private/MuT/CodeGenerator.cpp
+++ b/Engine/Plugins/Mutable/Source/MutableTools/Private/MuT/CodeGenerator.cpp
@@ -2551,15 +2551,15 @@ namespace mu
 				// If a matrix node is not connected, the op won't do anything, so let's not create it at all.
 				if (TypedTransformNode->MatrixNode)
 				{
-					Ptr<ASTOpMeshTransformWithBoundingMesh> transformOp = new ASTOpMeshTransformWithBoundingMesh();
-					transformOp->source = LastMeshOp;
+					Ptr<ASTOpMeshTransformWithBoundingMesh> TransformOp = new ASTOpMeshTransformWithBoundingMesh();
+					TransformOp->source = LastMeshOp;
 
 					// Transform matrix.
 					if (TypedTransformNode->MatrixNode)
 					{
 						FMatrixGenerationResult ChildResult;
 						GenerateMatrix(ChildResult, Options, TypedTransformNode->MatrixNode);
-						transformOp->matrix = ChildResult.op;
+						TransformOp->matrix = ChildResult.op;
 					}
 					
 					if (TypedTransformNode->BoundingMesh)
@@ -2571,13 +2571,21 @@ namespace mu
 
 						FMeshGenerationResult BoundingMeshResult;
 						GenerateMesh(MeshOptions, BoundingMeshResult, TypedTransformNode->BoundingMesh);
-						transformOp->boundingMesh = BoundingMeshResult.MeshOp;
+						TransformOp->boundingMesh = BoundingMeshResult.MeshOp;
 
-						if (!transformOp->boundingMesh)
+						if (!TransformOp->boundingMesh)
 						{
 							ErrorLog->GetPrivate()->Add("Bounding mesh has not been generated", ELMT_ERROR, ErrorContext);
 							continue;
 						}
+						else
+						{
+							ASTOp::EClosedMeshTest IsClosed = TransformOp->boundingMesh->IsClosedMesh();
+							if (IsClosed==ASTOp::EClosedMeshTest::No)
+							{
+								ErrorLog->GetPrivate()->Add("Mesh used for clipping is not closed.", ELMT_WARNING, ErrorContext);
+							}
+						}
 					}
 
 					// Condition to apply the transform op
@@ -2586,13 +2594,13 @@ namespace mu
 						Ptr<ASTOpConditional> conditionalAd = new ASTOpConditional();
 						conditionalAd->type = EOpType::ME_CONDITIONAL;
 						conditionalAd->no = LastMeshOp;
-						conditionalAd->yes = transformOp;
+						conditionalAd->yes = TransformOp;
 						conditionalAd->condition = m.FinalCondition;
 						LastMeshOp = conditionalAd;
 					}
 					else
 					{
-						LastMeshOp = transformOp;
+						LastMeshOp = TransformOp;
 					}
 				}
 			}
