[Mutable] Load Runtime Referenced Textures async.

- Added CacheRuntimeTexturesEvent as a new precondition to Task_Mutable_GetMeshes.
- Created new FMutableStreamableManager to manage async loads.
- Moved existing async loads to use the new FMutableStreamableManager.
- Compiler.
- Load Additional Assets.
- Extension Data Streamer.

#rnx

diff --git a/Engine/Plugins/Mutable/Source/CustomizableObject/Internal/MuCO/CustomizableObjectInstancePrivate.h b/Engine/Plugins/Mutable/Source/CustomizableObject/Internal/MuCO/CustomizableObjectInstancePrivate.h
--- a/Engine/Plugins/Mutable/Source/CustomizableObject/Internal/MuCO/CustomizableObjectInstancePrivate.h
+++ b/Engine/Plugins/Mutable/Source/CustomizableObject/Internal/MuCO/CustomizableObjectInstancePrivate.h
@@ -227,11 +227,8 @@ public:
 	// \TODO: Create a key based on a struct instead of generating strings dynamically.
 	UPROPERTY(Transient)
 	TMap<FString, TWeakObjectPtr<UTexture2D>> TextureReuseCache;
-
-	// Handle used to store a streaming request operation if one is ongoing.
-	TSharedPtr<FStreamableHandle> StreamingHandle;
+	
 	UE::Tasks::FTaskEvent AssetAsyncLoadCompletionEvent = UE::Tasks::FTaskEvent(TEXT(""));
-	TArray<FSoftObjectPath> AssetsToStream;
 
 	// Only used in LiveUpdateMode to reuse core instances between updates and their temp data to speed up updates, but spend way more memory
 	mu::FInstance::FID LiveUpdateModeInstanceID = 0;
@@ -264,9 +261,9 @@ public:
 
 	/** Returns the task that will be called when all assets and data are loaded, may be already completed if no assets or data needs loading.
 	 * If no StreamableManager is provided, it will load assets synchronously. */
-	UE::Tasks::FTask LoadAdditionalAssetsAndData(const TSharedRef<FUpdateContextPrivate>& OperationData, FStreamableManager& StreamableManager);
+	UE::Tasks::FTask LoadAdditionalAssetsAndData(const TSharedRef<FUpdateContextPrivate>& OperationData);
 
-	void AdditionalAssetsAsyncLoaded(UCustomizableObjectInstance* Public);
+	void AdditionalAssetsAsyncLoaded();
 
 	void TickUpdateCloseCustomizableObjects(UCustomizableObjectInstance& Public, FMutableInstanceUpdateMap& InOutRequestedUpdates);
 	void UpdateInstanceIfNotGenerated(UCustomizableObjectInstance& Public, FMutableInstanceUpdateMap& InOutRequestedUpdates);
@@ -344,9 +341,7 @@ public:
 	int32 GetState() const;
 
 	void SetState(int32 InState);
-
-	void AdditionalAssetsAsyncLoaded();
-
+	
 	FCustomizableObjectInstanceDescriptor& GetDescriptor() const;
 
 	void SetDescriptor(const FCustomizableObjectInstanceDescriptor& InDescriptor);
diff --git a/Engine/Plugins/Mutable/Source/CustomizableObject/Internal/MuCO/CustomizableObjectSystemPrivate.h b/Engine/Plugins/Mutable/Source/CustomizableObject/Internal/MuCO/CustomizableObjectSystemPrivate.h
--- a/Engine/Plugins/Mutable/Source/CustomizableObject/Internal/MuCO/CustomizableObjectSystemPrivate.h
+++ b/Engine/Plugins/Mutable/Source/CustomizableObject/Internal/MuCO/CustomizableObjectSystemPrivate.h
@@ -26,6 +26,7 @@
 #include "AssetRegistry/AssetData.h"
 #include "ContentStreaming.h"
 #include "Animation/MorphTarget.h"
+#include "MuCO/MutableStreamableManager.h"
 #include "UObject/StrongObjectPtr.h"
 
 #include "Tasks/Task.h"
@@ -684,7 +685,7 @@ public:
 	TSet<TWeakObjectPtr<UCustomizableObjectInstanceUsage>> AttachedParentUpdated;
 	
 	/** Hard references to objects. Avoids GC to collect them. */
-	TArray<TStrongObjectPtr<const UObject>> Objects;
+	TArray<TStrongObjectPtr<const UObject>> Objects;	
 };
 
 
@@ -885,9 +886,8 @@ public:
 	UPROPERTY(Transient)
 	TArray<TObjectPtr<UTexture2D>> ProtectedCachedTextures;
 	
-	// For async material loading
-	FStreamableManager StreamableManager;
-	
+	TSharedRef<FMutableStreamableManager> StreamableManager = MakeShared<FMutableStreamableManager>();
+
 #if WITH_EDITORONLY_DATA
 	// Array to keep track of cached objects
 	TArray<FGuid> UncompiledCustomizableObjectIds;
diff --git a/Engine/Plugins/Mutable/Source/CustomizableObject/Internal/MuCO/FMutableTaskGraph.h b/Engine/Plugins/Mutable/Source/CustomizableObject/Internal/MuCO/FMutableTaskGraph.h
--- a/Engine/Plugins/Mutable/Source/CustomizableObject/Internal/MuCO/FMutableTaskGraph.h
+++ b/Engine/Plugins/Mutable/Source/CustomizableObject/Internal/MuCO/FMutableTaskGraph.h
@@ -42,7 +42,7 @@ class FMutableTaskGraph
 public:
 
 	/** Create and launch a task on the Mutable Thread with Normal priority. */
-	UE::Tasks::FTask AddMutableThreadTask(const TCHAR* DebugName, TUniqueFunction<void()>&& TaskBody);
+	UE::Tasks::FTask AddMutableThreadTask(const TCHAR* DebugName, TUniqueFunction<void()>&& TaskBody, const TArray<UE::Tasks::FTask>& Prerequisites = {});
 
 	/** Create and launch a task on the Mutable Thread with Low priority. */
 	uint32 AddMutableThreadTaskLowPriority(const TCHAR* DebugName, TFunction<void()>&& TaskBody);
diff --git a/Engine/Plugins/Mutable/Source/CustomizableObject/Internal/MuCO/MutableStreamableManager.h b/Engine/Plugins/Mutable/Source/CustomizableObject/Internal/MuCO/MutableStreamableManager.h
new file mode 100644
--- /dev/null
+++ b/Engine/Plugins/Mutable/Source/CustomizableObject/Internal/MuCO/MutableStreamableManager.h
@@ -0,0 +1,90 @@
+﻿// Copyright Epic Games, Inc. All Rights Reserved.
+
+#pragma once
+
+#include "Engine/StreamableManager.h"
+
+#include "Templates/SharedPointer.h"
+
+struct FStreamableHandle;
+
+
+/** Wrapper to the FStreamableManager that allows to kill in-flight async loads and force them sync. */
+class FMutableStreamableManager : public TSharedFromThis<FMutableStreamableManager>
+{
+public:
+	template<typename PathContainerType = TArray<FSoftObjectPath>, typename FuncType = FStreamableDelegate, typename DebugNameType = UE::FSourceLocation>
+		TSharedPtr<FStreamableHandle> RequestAsyncLoad(
+			PathContainerType&& TargetsToStream,
+			FuncType&& DelegateToCall = FStreamableDelegate(),
+			TAsyncLoadPriority Priority = FStreamableManager::DefaultAsyncLoadPriority,
+			bool bManageActiveHandle = false,
+			bool bStartStalled = false,
+			DebugNameType&& DebugNameOrLocation = UE::FSourceLocation::Current())
+	{
+		check(IsInGameThread());
+
+		TSharedRef<FMutableStreamableHandle> MutableStreamableHandle = MakeShared<FMutableStreamableHandle>();
+
+		if constexpr (std::is_constructible_v<FStreamableDelegate, FuncType>)
+		{
+			MutableStreamableHandle->Delegate = DelegateToCall;
+		}
+		else
+		{
+			MutableStreamableHandle->Delegate = FStreamableDelegate::CreateLambda(DelegateToCall);
+		}
+
+		FStreamableDelegate Delegate = FStreamableDelegate::CreateLambda([MutableStreamableHandle, WeakThis = SharedThis(this)->AsWeak()]
+		{
+			check(IsInGameThread());
+			
+			MutableStreamableHandle->Delegate.ExecuteIfBound();
+			MutableStreamableHandle->bCompleted = true;
+
+			if (TSharedPtr<FMutableStreamableManager> This = WeakThis.Pin())
+			{
+				This->StreamableHandles.RemoveSingle(MutableStreamableHandle);
+			}
+		});
+		
+		
+		TSharedPtr<FStreamableHandle> Handle = StreamableManager.RequestAsyncLoad(TargetsToStream, Delegate, Priority, bManageActiveHandle, bStartStalled, DebugNameOrLocation);
+		if (Handle)
+		{
+			MutableStreamableHandle->Handle = Handle;
+			StreamableHandles.Emplace(MutableStreamableHandle);
+		}
+		
+		return MutableStreamableHandle->Handle;
+	}
+	
+	template<typename PathContainerType = TArray<FSoftObjectPath>, typename DebugNameType = UE::FSourceLocation>
+	TSharedPtr<FStreamableHandle> RequestSyncLoad(
+		PathContainerType&& TargetsToStream,
+		bool bManageActiveHandle = false,
+		DebugNameType&& DebugNameOrLocation = UE::FSourceLocation::Current())
+	{
+		return StreamableManager.RequestSyncLoad(TargetsToStream, bManageActiveHandle, DebugNameOrLocation);
+	}
+
+	int32 Tick(bool bBlocking);
+	
+private:
+	struct FMutableStreamableHandle
+	{
+		/** Always valid. */
+		TSharedPtr<FStreamableHandle> Handle;
+
+		FStreamableDelegate Delegate;
+
+		/** Load completed and delegate called.
+		 * Used since FStreamableHandle does not have away of knowing when the delegate has been called. */
+		bool bCompleted = false;
+	};
+	
+	FStreamableManager StreamableManager;
+
+	TArray<TSharedRef<FMutableStreamableHandle>> StreamableHandles;
+};
+
diff --git a/Engine/Plugins/Mutable/Source/CustomizableObject/Private/MuCO/CustomizableObjectInstance.cpp b/Engine/Plugins/Mutable/Source/CustomizableObject/Private/MuCO/CustomizableObjectInstance.cpp
--- a/Engine/Plugins/Mutable/Source/CustomizableObject/Private/MuCO/CustomizableObjectInstance.cpp
+++ b/Engine/Plugins/Mutable/Source/CustomizableObject/Private/MuCO/CustomizableObjectInstance.cpp
@@ -475,13 +475,6 @@ void UCustomizableObjectInstance::BeginDestroy()
 		PrivateData->BindObjectDelegates(GetCustomizableObject(), nullptr);
 #endif
 
-		if (PrivateData->StreamingHandle.IsValid() && PrivateData->StreamingHandle->IsActive())
-		{
-			PrivateData->StreamingHandle->CancelHandle();
-		}
-
-		PrivateData->StreamingHandle = nullptr;
-
 		PrivateData.Get()->ReleaseMutableResources(true, *this);
 	}
 	
@@ -5310,8 +5303,7 @@ FAutoConsoleVariableRef CVarMutableHighPriorityLoading(
 	TEXT("If enabled, the request to load additional assets will have high priority."));
 
 
-UE::Tasks::FTask UCustomizableInstancePrivate::LoadAdditionalAssetsAndData(
-		const TSharedRef<FUpdateContextPrivate>& OperationData, FStreamableManager& StreamableManager)
+UE::Tasks::FTask UCustomizableInstancePrivate::LoadAdditionalAssetsAndData(const TSharedRef<FUpdateContextPrivate>& OperationData)
 {
 	MUTABLE_CPUPROFILER_SCOPE(UCustomizableInstancePrivate::LoadAdditionalAssetsAndDataAsync);
 
@@ -5320,7 +5312,7 @@ UE::Tasks::FTask UCustomizableInstancePrivate::LoadAdditionalAssetsAndData(
 	const FModelResources& ModelResources = CustomizableObject->GetPrivate()->GetModelResources();
 	const TSharedPtr<FModelStreamableBulkData>& ModelStreamableBulkData = CustomizableObject->GetPrivate()->GetModelStreamableBulkData();
 
-	AssetsToStream.Empty();
+	TArray<FSoftObjectPath> AssetsToStream;
 	TArray<uint32> RealTimeMorphStreamableBlocksToStream;
 	TArray<uint32> ClothingStreamableBlocksToStream;
 
@@ -5638,12 +5630,13 @@ UE::Tasks::FTask UCustomizableInstancePrivate::LoadAdditionalAssetsAndData(
 		FSlateThrottleManager::Get().DisableThrottle(true);
 #endif
 
-		StreamingHandle = StreamableManager.RequestAsyncLoad(
-				AssetsToStream, 
-				FStreamableDelegate::CreateUObject(this, &UCustomizableInstancePrivate::AdditionalAssetsAsyncLoaded),
-				bEnableHighPriorityLoading ? FStreamableManager::AsyncLoadHighPriority : FStreamableManager::DefaultAsyncLoadPriority);			
-	}
+		UCustomizableObjectSystemPrivate* PrivateSystem = UCustomizableObjectSystem::GetInstance()->GetPrivate();
 
+		PrivateSystem->StreamableManager->RequestAsyncLoad(
+			AssetsToStream,
+			FStreamableDelegate::CreateUObject(this, &UCustomizableInstancePrivate::AdditionalAssetsAsyncLoaded),
+			bEnableHighPriorityLoading ? FStreamableManager::AsyncLoadHighPriority : FStreamableManager::DefaultAsyncLoadPriority);
+	}
 	
 	// File handles will end up owned by the gather task.
 	TArray<TUniquePtr<IAsyncReadFileHandle>> OpenFileHandles;
@@ -5956,22 +5949,6 @@ UE::Tasks::FTask UCustomizableInstancePrivate::LoadAdditionalAssetsAndData(
 	}
 }
 
-void UCustomizableInstancePrivate::AdditionalAssetsAsyncLoaded()
-{
-	// TODO: Do we need this separated?
-	//check(IsInGameThread())
-	AdditionalAssetsAsyncLoaded(GetPublic());
-	AssetAsyncLoadCompletionEvent.Trigger(); // TODO: we know it is game thread?
-
-	StreamingHandle = nullptr;
-
-#if WITH_EDITOR
-	// TODO: Remove with UE-217665 when the underlying bug in the ColorPicker is solved
-	// Reenable the throttling which disabled when launching the Async Load
-	FSlateThrottleManager::Get().DisableThrottle(false);
-#endif
-}
-
 
 FCustomizableObjectInstanceDescriptor& UCustomizableInstancePrivate::GetDescriptor() const
 {
@@ -5993,11 +5970,15 @@ const TArray<TObjectPtr<UMaterialInterface>>* UCustomizableObjectInstance::GetOv
 }
 
 
-void UCustomizableInstancePrivate::AdditionalAssetsAsyncLoaded(UCustomizableObjectInstance* Public)
+void UCustomizableInstancePrivate::AdditionalAssetsAsyncLoaded()
 {
 	MUTABLE_CPUPROFILER_SCOPE(UCustomizableInstancePrivate::AdditionalAssetsAsyncLoaded);
 
-	UCustomizableObjectPrivate* CustomizableObjectPrivate = Public->GetCustomizableObject()->GetPrivate();
+	check(IsInGameThread())
+
+	AssetAsyncLoadCompletionEvent.Trigger();
+
+	UCustomizableObjectPrivate* CustomizableObjectPrivate = GetPublic()->GetCustomizableObject()->GetPrivate();
 
 	const FModelResources& ModelResources = CustomizableObjectPrivate->GetModelResources();
 
@@ -6161,6 +6142,12 @@ void UCustomizableInstancePrivate::AdditionalAssetsAsyncLoaded(UCustomizableObje
 	}
 
 	PassThroughMeshesToLoad.Empty();
+
+#if WITH_EDITOR
+	// TODO: Remove with UE-217665 when the underlying bug in the ColorPicker is solved
+	// Reenable the throttling which disabled when launching the Async Load
+	FSlateThrottleManager::Get().DisableThrottle(false);
+#endif
 }
 
 
diff --git a/Engine/Plugins/Mutable/Source/CustomizableObject/Private/MuCO/CustomizableObjectSystem.cpp b/Engine/Plugins/Mutable/Source/CustomizableObject/Private/MuCO/CustomizableObjectSystem.cpp
--- a/Engine/Plugins/Mutable/Source/CustomizableObject/Private/MuCO/CustomizableObjectSystem.cpp
+++ b/Engine/Plugins/Mutable/Source/CustomizableObject/Private/MuCO/CustomizableObjectSystem.cpp
@@ -3166,7 +3166,7 @@ namespace impl
 
 		// Next Task: Load Unreal Assets
 		//-------------------------------------------------------------
-		UE::Tasks::FTask Game_LoadUnrealAssets = ObjectInstancePrivateData->LoadAdditionalAssetsAndData(OperationData, System->GetPrivate()->StreamableManager);
+		UE::Tasks::FTask Game_LoadUnrealAssets = ObjectInstancePrivateData->LoadAdditionalAssetsAndData(OperationData);
 
 		// Next-next Task: Convert Resources
 		//-------------------------------------------------------------
@@ -3207,7 +3207,7 @@ namespace impl
 	}
 
 
-	void Task_Game_LockMeshCache(const TSharedRef<FUpdateContextPrivate>& Operation)
+	void Task_Game_LockMeshCache(const TSharedRef<FUpdateContextPrivate>& Operation, UE::Tasks::FTask Requisite)
 	{
 		MUTABLE_CPUPROFILER_SCOPE(Task_Game_LockMeshCache);
 
@@ -3230,7 +3230,7 @@ namespace impl
 			[Operation]()
 			{
 				Task_Mutable_GetMeshes(Operation);
-			});
+			}, { Requisite });
 
 		SystemPrivate->AddGameThreadTask(
 			FMutableTask
@@ -3423,10 +3423,46 @@ namespace impl
 			CandidateInstancePrivateData->LiveUpdateModeInstanceID = 0;
 		}
 
+
 		Operation->Model = CustomizableObject->GetPrivate()->GetModel().ToSharedRef();
 
+		UE::Tasks::FTask CacheRuntimeTexturesEvent = UE::Tasks::MakeCompletedTask<void>();
+
 #if WITH_EDITOR
-		SystemPrivateData->GetResourceProviderChecked()->CacheRuntimeReferencedImages(Operation->Model.ToSharedRef(), CustomizableObject->GetPrivate()->GetModelResources().RuntimeReferencedTextures);
+		// Async load all Runtime Referenced Textures.
+		const TArray<TSoftObjectPtr<const UTexture>>& RuntimeReferencedTextures = CustomizableObject->GetPrivate()->GetModelResources().RuntimeReferencedTextures;
+		if (!RuntimeReferencedTextures.IsEmpty())
+		{
+			UE::Tasks::FTaskEvent Event = UE::Tasks::FTaskEvent(TEXT("Texture"));
+			CacheRuntimeTexturesEvent = Event;
+
+			TArray<FSoftObjectPath> Textures;
+			Textures.Reserve(RuntimeReferencedTextures.Num());
+			for (const TSoftObjectPtr<const UTexture>& Texture : RuntimeReferencedTextures)
+			{
+				Textures.Add(Texture.ToSoftObjectPath());
+			}
+		
+			SystemPrivateData->StreamableManager->RequestAsyncLoad(Textures, FStreamableDelegate::CreateLambda([Operation, Event]() mutable
+			{
+				UCustomizableObjectSystem* System = UCustomizableObjectSystem::GetInstance();
+				if (!System)
+				{
+					Event.Trigger();
+					return;
+				}
+				
+				UCustomizableObject* CustomizableObject = Operation->GetCapturedDescriptor().GetCustomizableObject();
+				if (!CustomizableObject)
+				{
+					Event.Trigger();
+					return;
+				}
+			
+				System->GetPrivate()->GetResourceProviderChecked()->CacheRuntimeReferencedImages(Operation->Model.ToSharedRef(), CustomizableObject->GetPrivate()->GetModelResources().RuntimeReferencedTextures);
+				Event.Trigger();
+			}));
+		}
 #endif
 		
 		// Task: Mutable Update and GetMesh
@@ -3506,9 +3542,9 @@ namespace impl
 				FMutableTask 
 				{
 					FMutableTaskDelegate::CreateLambda(
-					[Operation]()
+					[Operation, CacheRuntimeTexturesEvent]()
 					{
-						impl::Task_Game_LockMeshCache(Operation);
+						impl::Task_Game_LockMeshCache(Operation, CacheRuntimeTexturesEvent);
 					}),
 					{ Mutable_GetMeshTask }
 				});
@@ -3520,7 +3556,7 @@ namespace impl
 				[Operation]()
 				{
 					Task_Mutable_GetMeshes(Operation);
-				});
+				}, { CacheRuntimeTexturesEvent });
 
 			// Task: Lock cache
 			//-------------------------------------------------------------
@@ -3892,20 +3928,10 @@ int32 UCustomizableObjectSystem::TickInternal(const bool bBlocking)
 #if WITH_EDITOR
 		RemainingWork += EditorModule ? EditorModule->Tick(true) : 0;
 #endif
-		
-		if (GetPrivate()->CurrentMutableOperation)
-		{
-			UCustomizableInstancePrivate* InstancePrivate = GetPrivate()->CurrentMutableOperation->Instance->GetPrivate();
-			
-			if (InstancePrivate->StreamingHandle)
-			{
-				InstancePrivate->StreamingHandle->CancelHandle();
-				GetPrivate()->StreamableManager.RequestSyncLoad(InstancePrivate->AssetsToStream);
-				InstancePrivate->AdditionalAssetsAsyncLoaded();
-			}
-		}
 	}
 	
+	RemainingWork += GetPrivate()->StreamableManager->Tick(bBlocking);
+	
 	return RemainingWork;
 }
 
diff --git a/Engine/Plugins/Mutable/Source/CustomizableObject/Private/MuCO/FMutableTaskGraph.cpp b/Engine/Plugins/Mutable/Source/CustomizableObject/Private/MuCO/FMutableTaskGraph.cpp
--- a/Engine/Plugins/Mutable/Source/CustomizableObject/Private/MuCO/FMutableTaskGraph.cpp
+++ b/Engine/Plugins/Mutable/Source/CustomizableObject/Private/MuCO/FMutableTaskGraph.cpp
@@ -21,14 +21,17 @@ static TAutoConsoleVariable<bool> CVarEnableMutableTaskLowPriority(
 	ECVF_Scalability);
 
 
-UE::Tasks::FTask FMutableTaskGraph::AddMutableThreadTask(const TCHAR* DebugName, TUniqueFunction<void()>&& TaskBody)
+
+UE::Tasks::FTask FMutableTaskGraph::AddMutableThreadTask(const TCHAR* DebugName, TUniqueFunction<void()>&& TaskBody, const TArray<UE::Tasks::FTask>& InPrerequisites)
 {
 	FScopeLock Lock(&MutableTaskLock);
 
-	TArray<UE::Tasks::FTask, TFixedAllocator<1>> Prerequisites = LastMutableTask.IsValid() 
-			? TArray<UE::Tasks::FTask, TFixedAllocator<1>>{LastMutableTask} 
-			: TArray<UE::Tasks::FTask, TFixedAllocator<1>>{};
+	TArray<UE::Tasks::FTask, TInlineAllocator<1>> Prerequisites = LastMutableTask.IsValid() 
+		? TArray<UE::Tasks::FTask, TInlineAllocator<1>>{LastMutableTask} 
+		: TArray<UE::Tasks::FTask, TInlineAllocator<1>>{};
 
+	Prerequisites.Append(InPrerequisites);
+	
 	LastMutableTask = UE::Tasks::Launch(DebugName, MoveTemp(TaskBody), Prerequisites, TASKGRAPH_PRIORITY);	
 	
 	return LastMutableTask;
diff --git a/Engine/Plugins/Mutable/Source/CustomizableObject/Private/MuCO/MutableStreamableManager.cpp b/Engine/Plugins/Mutable/Source/CustomizableObject/Private/MuCO/MutableStreamableManager.cpp
new file mode 100644
--- /dev/null
+++ b/Engine/Plugins/Mutable/Source/CustomizableObject/Private/MuCO/MutableStreamableManager.cpp
@@ -0,0 +1,28 @@
+﻿// Copyright Epic Games, Inc. All Rights Reserved.
+
+#include "MuCO/MutableStreamableManager.h"
+
+
+int32 FMutableStreamableManager::Tick(bool bBlocking)
+{
+	if (bBlocking)
+	{
+		for (const TSharedRef<FMutableStreamableHandle> MutableStreamableHandle : StreamableHandles)
+		{
+			if (!MutableStreamableHandle->bCompleted)
+			{
+				TArray<FSoftObjectPath> RequestedAssets;
+				MutableStreamableHandle->Handle->GetRequestedAssets(RequestedAssets);
+
+				StreamableManager.RequestSyncLoad(RequestedAssets);
+				MutableStreamableHandle->Delegate.ExecuteIfBound();
+			
+				MutableStreamableHandle->Handle->CancelHandle();
+			}			
+		}
+
+		StreamableHandles.Empty();
+	}
+
+	return 0;
+}
diff --git a/Engine/Plugins/Mutable/Source/CustomizableObject/Private/MuCO/UnrealExtensionDataStreamer.cpp b/Engine/Plugins/Mutable/Source/CustomizableObject/Private/MuCO/UnrealExtensionDataStreamer.cpp
--- a/Engine/Plugins/Mutable/Source/CustomizableObject/Private/MuCO/UnrealExtensionDataStreamer.cpp
+++ b/Engine/Plugins/Mutable/Source/CustomizableObject/Private/MuCO/UnrealExtensionDataStreamer.cpp
@@ -224,9 +224,9 @@ TSharedPtr<FStreamableHandle> FUnrealExtensionDataStreamer::StartLoadOnGameThrea
 	DebugString = FString::Printf(TEXT("UnrealExtensionDataStreamer for %s"), *Object->GetPathName());
 #endif
 
-	FStreamableManager& Manager = UCustomizableObjectSystem::GetInstance()->GetPrivate()->StreamableManager;
+	const TSharedRef<FMutableStreamableManager>& Manager = UCustomizableObjectSystem::GetInstance()->GetPrivate()->StreamableManager;
 
-	return Manager.RequestAsyncLoad(
+	return Manager->RequestAsyncLoad(
 		TargetsToStream,
 		MoveTemp(OnLoadComplete),
 		FStreamableManager::DefaultAsyncLoadPriority,
diff --git a/Engine/Plugins/Mutable/Source/CustomizableObject/Private/MuCO/UnrealMutableImageProvider.cpp b/Engine/Plugins/Mutable/Source/CustomizableObject/Private/MuCO/UnrealMutableImageProvider.cpp
--- a/Engine/Plugins/Mutable/Source/CustomizableObject/Private/MuCO/UnrealMutableImageProvider.cpp
+++ b/Engine/Plugins/Mutable/Source/CustomizableObject/Private/MuCO/UnrealMutableImageProvider.cpp
@@ -747,11 +747,17 @@ void FUnrealMutableResourceProvider::CacheRuntimeReferencedImages(const TSharedR
 	ModelImages.SourceTextures.Reset();
 	for (const TSoftObjectPtr<const UTexture>& RuntimeReferencedTexture : RuntimeReferencedTextures)
 	{
-		const UTexture* Texture = RuntimeReferencedTexture.LoadSynchronous();
+		const UTexture* Texture = RuntimeReferencedTexture.Get(); // Is already loaded.
 		if (!Texture)
 		{
-			UE_LOG(LogMutable, Warning, TEXT("Failed to load texture [%s]."), *RuntimeReferencedTexture->GetPathName());
-			continue;
+			UE_LOG(LogMutable, Warning, TEXT("Runtime Referenced Texture [%s] was not async loaded. Forcing load sync."), *RuntimeReferencedTexture->GetPathName());
+			
+			Texture = RuntimeReferencedTexture.LoadSynchronous();
+			if (!Texture)
+			{
+				UE_LOG(LogMutable, Warning, TEXT("Failed to force load sync [%s]."), *RuntimeReferencedTexture->GetPathName());
+				continue;
+			}
 		}
 
 		ModelImages.SourceTextures.Emplace(*Texture); // Perform a CopyTornOff. Once done, we no longer need the texture loaded.
diff --git a/Engine/Plugins/Mutable/Source/CustomizableObject/Private/MuCO/UnrealMutableImageProvider.h b/Engine/Plugins/Mutable/Source/CustomizableObject/Private/MuCO/UnrealMutableImageProvider.h
--- a/Engine/Plugins/Mutable/Source/CustomizableObject/Private/MuCO/UnrealMutableImageProvider.h
+++ b/Engine/Plugins/Mutable/Source/CustomizableObject/Private/MuCO/UnrealMutableImageProvider.h
@@ -18,7 +18,7 @@ class UTexture2D;
 
 
 /** Implementation of a mutable core provider for image parameters that are application-specific. */
-class FUnrealMutableResourceProvider : public mu::FExternalResourceProvider, public FGCObject
+class FUnrealMutableResourceProvider : public TSharedFromThis<FUnrealMutableResourceProvider>, public mu::FExternalResourceProvider, public FGCObject
 {
 
 public:
diff --git a/Engine/Plugins/Mutable/Source/CustomizableObjectEditor/Private/MuCOE/CustomizableObjectCompiler.cpp b/Engine/Plugins/Mutable/Source/CustomizableObjectEditor/Private/MuCOE/CustomizableObjectCompiler.cpp
--- a/Engine/Plugins/Mutable/Source/CustomizableObjectEditor/Private/MuCOE/CustomizableObjectCompiler.cpp
+++ b/Engine/Plugins/Mutable/Source/CustomizableObjectEditor/Private/MuCOE/CustomizableObjectCompiler.cpp
@@ -55,16 +55,9 @@ bool FCustomizableObjectCompiler::Tick(bool bBlocking)
 		bFinished = false;
 	}
 
-	if (AsynchronousStreamableHandlePtr)
+	if (AsynchronousStreamableHandlePtr && AsynchronousStreamableHandlePtr->IsActive())
 	{
 		bFinished = false;
-
-		if (bBlocking)
-		{
-			AsynchronousStreamableHandlePtr->CancelHandle();
-			UCustomizableObjectSystem::GetInstance()->GetPrivate()->StreamableManager.RequestSyncLoad(ArrayAssetToStream);
-			PreloadingReferencerAssetsCallback(false);
-		}
 	}
 
 	if (CompileTask)
@@ -121,14 +114,18 @@ void FCustomizableObjectCompiler::PreloadingReferencerAssetsCallback(bool bAsync
 	check(IsInGameThread());
 
 	check(ArrayGCProtect.IsEmpty());
-	for (FSoftObjectPath AssetToStream : ArrayAssetToStream)
-	{
-		ArrayGCProtect.Add(AssetToStream.TryLoad());
-	}
 	
 	if (AsynchronousStreamableHandlePtr)
 	{
-		AsynchronousStreamableHandlePtr.Reset();
+		TArray<FSoftObjectPath> AssetsToStream;
+		AsynchronousStreamableHandlePtr->GetRequestedAssets(AssetsToStream);
+
+		for (FSoftObjectPath AssetToStream : AssetsToStream)
+		{
+			ArrayGCProtect.Add(AssetToStream.TryLoad()); // Already loaded.
+		}
+
+		AsynchronousStreamableHandlePtr = nullptr;
 	}
 
 	UE_LOG(LogMutable, Verbose, TEXT("PROFILE: [ %16.8f ] Preload asynchronously assets end."), FPlatformTime::Seconds());
@@ -269,7 +266,7 @@ void FCustomizableObjectCompiler::Compile(const TSharedRef<FCompilationRequest>&
 	TArray<FAssetData> ReferencingAssets;
 	GetReferencingPackages(*CurrentObject, ReferencingAssets);
 
-	ArrayAssetToStream.Empty();
+	TArray<FSoftObjectPath> ArrayAssetToStream;
 	for (FAssetData& Element : ReferencingAssets)
 	{
 		ArrayAssetToStream.Add(Element.GetSoftObjectPath());
@@ -284,18 +281,18 @@ void FCustomizableObjectCompiler::Compile(const TSharedRef<FCompilationRequest>&
 		// The ones that need to be kept will be copied into SoftObjectPath on the object during save.
 		FCookLoadScope CookLoadScope(ECookLoadType::EditorOnly);
 
-		FStreamableManager& Streamable = System->GetPrivate()->StreamableManager;
+		const TSharedRef<FMutableStreamableManager>& Streamable = System->GetPrivate()->StreamableManager;
 
 		if (bAsync && !CurrentOptions.bIsCooking)
 		{
 			AddCompileNotification(LOCTEXT("LoadingReferencerAssets", "Loading assets"));
 
-			AsynchronousStreamableHandlePtr = Streamable.RequestAsyncLoad(ArrayAssetToStream, FStreamableDelegate::CreateRaw(this, &FCustomizableObjectCompiler::PreloadingReferencerAssetsCallback, bAsync));
+			AsynchronousStreamableHandlePtr = Streamable->RequestAsyncLoad(ArrayAssetToStream, FStreamableDelegate::CreateRaw(this, &FCustomizableObjectCompiler::PreloadingReferencerAssetsCallback, bAsync));
 			bAssetsLoaded = false;
 		}
 		else
 		{
-			Streamable.RequestSyncLoad(ArrayAssetToStream);
+			Streamable->RequestSyncLoad(ArrayAssetToStream);
 		}
 	}
 
@@ -1340,7 +1337,7 @@ void FCustomizableObjectCompiler::ForceFinishCompilation()
 	if (AsynchronousStreamableHandlePtr)
 	{
 		AsynchronousStreamableHandlePtr->CancelHandle();
-		AsynchronousStreamableHandlePtr.Reset();
+		AsynchronousStreamableHandlePtr = nullptr;
 	}
 
 	else if (CompileTask.IsValid())
diff --git a/Engine/Plugins/Mutable/Source/CustomizableObjectEditor/Private/MuCOE/CustomizableObjectCompiler.h b/Engine/Plugins/Mutable/Source/CustomizableObjectEditor/Private/MuCOE/CustomizableObjectCompiler.h
--- a/Engine/Plugins/Mutable/Source/CustomizableObjectEditor/Private/MuCOE/CustomizableObjectCompiler.h
+++ b/Engine/Plugins/Mutable/Source/CustomizableObjectEditor/Private/MuCOE/CustomizableObjectCompiler.h
@@ -118,7 +118,6 @@ private:
 
 	/** Pointer to the Asynchronous Preloading process call back */
 	TSharedPtr<FStreamableHandle> AsynchronousStreamableHandlePtr;
-	TArray<FSoftObjectPath> ArrayAssetToStream;
 
 	/** Compile task and thread. */
 	TSharedPtr<FCustomizableObjectCompileRunnable> CompileTask;
