[Mutable] Allow enqueuing compilations even if the CO is not loaded yet.

- Moved AddPendingLoad to private.

#rnx

diff --git a/Engine/Plugins/Mutable/Source/CustomizableObject/Internal/MuCO/CustomizableObjectSystemPrivate.h b/Engine/Plugins/Mutable/Source/CustomizableObject/Internal/MuCO/CustomizableObjectSystemPrivate.h
--- a/Engine/Plugins/Mutable/Source/CustomizableObject/Internal/MuCO/CustomizableObjectSystemPrivate.h
+++ b/Engine/Plugins/Mutable/Source/CustomizableObject/Internal/MuCO/CustomizableObjectSystemPrivate.h
@@ -851,6 +851,13 @@ public:
 	void CacheTextureParameters(const TArray<FCustomizableObjectTextureParameterValue>& TextureParameters) const;
 
 	void UnCacheTextureParameters(const TArray<FCustomizableObjectTextureParameterValue>& TextureParameters) const;
+
+#if WITH_EDITOR
+	/** Add a CO to the pending - load list.COs need to wait until all related objects are fully loaded before being able
+	 * to do things like check-if-up-to-date or compile.
+	 */
+	void AddPendingLoad(UCustomizableObject*);
+#endif
 	
 	/** Mutable TaskGraph system (Mutable Thread). */
 	FMutableTaskGraph MutableTaskGraph;
diff --git a/Engine/Plugins/Mutable/Source/CustomizableObject/Private/MuCO/CustomizableObject.cpp b/Engine/Plugins/Mutable/Source/CustomizableObject/Private/MuCO/CustomizableObject.cpp
--- a/Engine/Plugins/Mutable/Source/CustomizableObject/Private/MuCO/CustomizableObject.cpp
+++ b/Engine/Plugins/Mutable/Source/CustomizableObject/Private/MuCO/CustomizableObject.cpp
@@ -371,7 +371,7 @@ void UCustomizableObject::PostLoad()
 	if (!IsRunningCookCommandlet())
 	{
 		GetPrivate()->Status.NextState(FCustomizableObjectStatusTypes::EState::Loading);
-		UCustomizableObjectSystem::GetInstance()->AddPendingLoad(this);
+		UCustomizableObjectSystem::GetInstance()->GetPrivate()->AddPendingLoad(this);
 	}
 #endif
 }
@@ -1098,12 +1098,6 @@ bool UCustomizableObject::ConditionalAutoCompile()
 		return true;
 	}
 
-	// Model has not loaded yet
-	if (GetPrivate()->Status.Get() == FCustomizableObjectStatusTypes::EState::Loading)
-	{
-		return false;
-	}
-
 	UCustomizableObjectSystem* System = UCustomizableObjectSystem::GetInstance();
 	if (!System || !System->IsValidLowLevel() || System->HasAnyFlags(RF_BeginDestroyed))
 	{
@@ -1133,10 +1127,16 @@ bool UCustomizableObject::ConditionalAutoCompile()
 	ICustomizableObjectEditorModule* EditorModule = ICustomizableObjectEditorModule::Get();
 	if (ensure(EditorModule))
 	{
-		// Sync/Async compilation
-		TSharedRef<FCompilationRequest> CompileRequest = MakeShared<FCompilationRequest>(*this, !System->IsAutoCompilationSync());
-		CompileRequest->GetCompileOptions().bSilentCompilation = true;
-		EditorModule->CompileCustomizableObject(CompileRequest);
+		if (System->IsAutoCompilationSync() && GetPrivate()->Status.Get() == FCustomizableObjectStatusTypes::EState::Loading)
+		{
+			return false;
+		}
+		else
+		{
+			TSharedRef<FCompilationRequest> CompileRequest = MakeShared<FCompilationRequest>(*this, true);
+			CompileRequest->GetCompileOptions().bSilentCompilation = true;
+			EditorModule->CompileCustomizableObject(CompileRequest);
+		}
 	}
 
 	return IsCompiled();
diff --git a/Engine/Plugins/Mutable/Source/CustomizableObject/Private/MuCO/CustomizableObjectSystem.cpp b/Engine/Plugins/Mutable/Source/CustomizableObject/Private/MuCO/CustomizableObjectSystem.cpp
--- a/Engine/Plugins/Mutable/Source/CustomizableObject/Private/MuCO/CustomizableObjectSystem.cpp
+++ b/Engine/Plugins/Mutable/Source/CustomizableObject/Private/MuCO/CustomizableObjectSystem.cpp
@@ -1372,10 +1372,10 @@ void UCustomizableObjectSystem::ClearResourceCacheProtected()
 
 #if WITH_EDITOR
 
-void UCustomizableObjectSystem::AddPendingLoad(UCustomizableObject* CO)
+void UCustomizableObjectSystemPrivate::AddPendingLoad(UCustomizableObject* CO)
 {
 	check(IsInGameThread());
-	GetPrivate()->ObjectsPendingLoad.AddUnique(CO);
+	ObjectsPendingLoad.AddUnique(CO);
 }
 
 
diff --git a/Engine/Plugins/Mutable/Source/CustomizableObject/Public/MuCO/CustomizableObjectSystem.h b/Engine/Plugins/Mutable/Source/CustomizableObject/Public/MuCO/CustomizableObjectSystem.h
--- a/Engine/Plugins/Mutable/Source/CustomizableObject/Public/MuCO/CustomizableObjectSystem.h
+++ b/Engine/Plugins/Mutable/Source/CustomizableObject/Public/MuCO/CustomizableObjectSystem.h
@@ -209,12 +209,6 @@ public:
 	void ClearResourceCacheProtected();
 
 #if WITH_EDITOR
-
-	/** Add a CO to the pending - load list.COs need to wait until all related objects are fully loaded before being able
-	 * to do things like check-if-up-to-date or compile.
-	 */
-	void AddPendingLoad(UCustomizableObject*);
-
 	// Lock a CustomizableObjects, preventing the generation or update of any of its instances
 	// Will return true if successful, false if it fails to lock because an update is already underway
 	// This is usually only used in the editor
diff --git a/Engine/Plugins/Mutable/Source/CustomizableObjectEditor/Private/MuCOE/CustomizableObjectCompiler.cpp b/Engine/Plugins/Mutable/Source/CustomizableObjectEditor/Private/MuCOE/CustomizableObjectCompiler.cpp
--- a/Engine/Plugins/Mutable/Source/CustomizableObjectEditor/Private/MuCOE/CustomizableObjectCompiler.cpp
+++ b/Engine/Plugins/Mutable/Source/CustomizableObjectEditor/Private/MuCOE/CustomizableObjectCompiler.cpp
@@ -1172,8 +1172,32 @@ bool FCustomizableObjectCompiler::TryPopCompileRequest()
 		return false;
 	}
 
-	Compile(CompileRequests.Pop());
-	return true;
+	for (TArray<TSharedRef<FCompilationRequest>>::TIterator It = CompileRequests.CreateIterator(); It; ++It)
+	{
+		TSharedRef<FCompilationRequest>& CompileRequest = *It;
+		
+		switch (CompileRequest->GetCustomizableObject()->GetPrivate()->Status.Get())
+		{
+		case FCustomizableObjectStatusTypes::EState::Loading:
+			// Wait.
+			break;
+
+		case FCustomizableObjectStatusTypes::EState::NoModel:
+		case FCustomizableObjectStatusTypes::EState::ModelLoaded:
+			{
+				TSharedRef<FCompilationRequest> SelectedCompileRequest = CompileRequest;
+				It.RemoveCurrent(); // Preserve order.
+
+				Compile(SelectedCompileRequest);
+				return true;
+			}
+			
+		default:
+			unimplemented();
+		}
+	}
+
+	return false;
 }
 
 
