[Mutable] Fixed loop variable creating a copy warning.

#rnx

diff --git a/Engine/Plugins/Mutable/Source/CustomizableObject/Private/MuCO/MutableStreamableManager.cpp b/Engine/Plugins/Mutable/Source/CustomizableObject/Private/MuCO/MutableStreamableManager.cpp
--- a/Engine/Plugins/Mutable/Source/CustomizableObject/Private/MuCO/MutableStreamableManager.cpp
+++ b/Engine/Plugins/Mutable/Source/CustomizableObject/Private/MuCO/MutableStreamableManager.cpp
@@ -7,7 +7,7 @@ int32 FMutableStreamableManager::Tick(bool bBlocking)
 {
 	if (bBlocking)
 	{
-		for (const TSharedRef<FMutableStreamableHandle> MutableStreamableHandle : StreamableHandles)
+		for (const TSharedRef<FMutableStreamableHandle>& MutableStreamableHandle : StreamableHandles)
 		{
 			if (!MutableStreamableHandle->bCompleted)
 			{
