[Mutable] Fix MorphTarget and Clothing data file offsets when splitting resources in multiple files.

#rnx

diff --git a/Engine/Plugins/Mutable/Source/CustomizableObject/Internal/MuCO/CustomizableObjectPrivate.h b/Engine/Plugins/Mutable/Source/CustomizableObject/Internal/MuCO/CustomizableObjectPrivate.h
--- a/Engine/Plugins/Mutable/Source/CustomizableObject/Internal/MuCO/CustomizableObjectPrivate.h
+++ b/Engine/Plugins/Mutable/Source/CustomizableObject/Internal/MuCO/CustomizableObjectPrivate.h
@@ -916,10 +916,10 @@ namespace MutablePrivate
 		FModelStreamableData ModelStreamableData;
 
 		/** */
-		TArray64<uint8> MorphData;
+		FModelStreamableData MorphStreamableData;
 
 		/** */
-		TArray64<uint8> ClothingData;
+		FModelStreamableData ClothingStreamableData;
 
 		/** List of files to serialize. Each file has a list of binary blocks to be serialized. */
 		TArray<FFile> BulkDataFiles;
@@ -1634,6 +1634,6 @@ public:
 	
 	static constexpr int32 CurrentSupportedVersion = ECustomizableObjectVersions::LastCustomizableObjectVersion;
 	
-	static constexpr int32 DerivedDataVersion = 1;
+	static constexpr int32 DerivedDataVersion = 2;
 };
 
diff --git a/Engine/Plugins/Mutable/Source/CustomizableObject/Private/MuCO/CustomizableObject.cpp b/Engine/Plugins/Mutable/Source/CustomizableObject/Private/MuCO/CustomizableObject.cpp
--- a/Engine/Plugins/Mutable/Source/CustomizableObject/Private/MuCO/CustomizableObject.cpp
+++ b/Engine/Plugins/Mutable/Source/CustomizableObject/Private/MuCO/CustomizableObject.cpp
@@ -2534,23 +2534,27 @@ namespace MutablePrivate
 		}
 		else if (DataType == EDataType::RealTimeMorph)
 		{
-			SourceData = PlatformData->MorphData.GetData();
+			for (int32 BlockIndex = 0; BlockIndex < NumBlocks; ++BlockIndex)
+			{
+				const FBlock& Block = Blocks[BlockIndex];
+				check(Block.Offset + Block.Size <= DestSize);
+				PlatformData->MorphStreamableData.Get(Block.Id, TArrayView64<uint8>(DestData.GetData() + Block.Offset, Block.Size), bDropData);
+			}
 		}
 		else if (DataType == EDataType::Clothing)
 		{
-			SourceData = PlatformData->ClothingData.GetData();
+
+			for (int32 BlockIndex = 0; BlockIndex < NumBlocks; ++BlockIndex)
+			{
+				const FBlock& Block = Blocks[BlockIndex];
+				check(Block.Offset + Block.Size <= DestSize);
+				PlatformData->ClothingStreamableData.Get(Block.Id, TArrayView64<uint8>(DestData.GetData() + Block.Offset, Block.Size), bDropData);
+			}
 		}
 		else
 		{
 			checkf(false, TEXT("Unknown file DataType found."));
 		}
-
-		for (int32 BlockIndex = 0; BlockIndex < NumBlocks; ++BlockIndex)
-		{
-			const FBlock& Block = Blocks[BlockIndex];
-			check(Block.Offset + Block.Size <= DestSize);
-			FMemory::Memcpy(DestData.GetData() + Block.Offset, SourceData + Block.Offset, Block.Size);
-		}
 	}
 
 	FFileCategoryID::FFileCategoryID(EDataType InDataType, uint16 InResourceType, uint16 InFlags)
@@ -2875,14 +2879,19 @@ namespace MutablePrivate
 				const mu::ERomFlags BlockFlags = (mu::ERomFlags)(IgnoreMask & (uint16)Model->GetRomFlags(RomIndex));
 
 				FFileCategoryID CurrentCategory = { EDataType::Model, BlockResourceType, uint16(BlockFlags) };
-				FBlock CurrentBlock = { BlockId, SourceBlockId, BlockSize, 0 };
+				FBlock CurrentBlock = 
+				{ 
+					.Id = BlockId, 
+					.SourceId = SourceBlockId, 
+					.Size = BlockSize, 
+					.Offset = 0 
+				};
+
 				AddNode(RootNode, NumRoms, CurrentCategory, CurrentBlock);
 			}
 		}
 
 		{
-			uint64 SourceOffset = 0;
-
 			const TMap<uint32, FRealTimeMorphStreamable>& RealTimeMorphStreamables = ModelStreamableBulkData.RealTimeMorphStreamables;
 
 			const FFileCategoryID RealTimeMorphCategory = { EDataType::RealTimeMorph, (uint16)mu::DATATYPE::DT_NONE, (uint16)mu::ERomFlags::None };
@@ -2891,16 +2900,19 @@ namespace MutablePrivate
 			{
 				const uint32 BlockSize = MorphStreamable.Value.Size;
 
-				FBlock CurrentBlock = { MorphStreamable.Key, MorphStreamable.Value.SourceId, BlockSize, SourceOffset };
-				AddNode(RootNode, NumRoms, RealTimeMorphCategory, CurrentBlock);
+				FBlock CurrentBlock = 
+				{ 
+					.Id = MorphStreamable.Key, 
+					.SourceId = MorphStreamable.Value.SourceId, 
+					.Size = BlockSize,
+					.Offset = 0 
+				};
 
-				SourceOffset += BlockSize;
+				AddNode(RootNode, NumRoms, RealTimeMorphCategory, CurrentBlock);
 			}
 		}
 
 		{
-			uint64 SourceOffset = 0;
-
 			const TMap<uint32, FClothingStreamable>& ClothingStreamables = ModelStreamableBulkData.ClothingStreamables;
 
 			const FFileCategoryID ClothingCategory = { EDataType::Clothing, (uint16)mu::DATATYPE::DT_NONE, (uint16)mu::ERomFlags::None };
@@ -2909,10 +2921,15 @@ namespace MutablePrivate
 			{
 				const uint32 BlockSize = ClothStreamable.Value.Size;
 
-				FBlock CurrentBlock = { ClothStreamable.Key, ClothStreamable.Value.SourceId, BlockSize, SourceOffset };
-				AddNode(RootNode, NumRoms, ClothingCategory, CurrentBlock);
+				FBlock CurrentBlock = 
+				{ 
+					.Id = ClothStreamable.Key, 
+					.SourceId = ClothStreamable.Value.SourceId, 
+					.Size = BlockSize, 
+					.Offset = 0 
+				};
 
-				SourceOffset += BlockSize;
+				AddNode(RootNode, NumRoms, ClothingCategory, CurrentBlock);
 			}
 		}
 
@@ -3319,29 +3336,17 @@ void UCustomizableObjectPrivate::LoadCompiledDataFromDDC(FCompilationOptions Opt
 							}
 							case MutablePrivate::EDataType::RealTimeMorph:
 							{
-								// Store the MorphData in a single array.
-								uint64 Offset = CachedData.MorphData.Num();
-								CachedData.MorphData.Append(TempData.GetData(), RawSize);
-
-								// Fix offsets
 								for (const MutablePrivate::FBlock& Block : File->Blocks)
 								{
-									ModelStreamables->RealTimeMorphStreamables[Block.Id].Block.Offset = Offset;
-									Offset += Block.Size;
+									CachedData.MorphStreamableData.Set(Block.Id, TempData.GetData() + Block.Offset, Block.Size);
 								}
 								break;
 							}
 							case MutablePrivate::EDataType::Clothing:
 							{
-								// Store the ClothingData in a single array.
-								uint64 Offset = CachedData.ClothingData.Num();
-								CachedData.ClothingData.Append(TempData.GetData(), RawSize);
-
-								// Fix offsets
 								for (const MutablePrivate::FBlock& Block : File->Blocks)
 								{
-									ModelStreamables->ClothingStreamables[Block.Id].Block.Offset = Offset;
-									Offset += Block.Size;
+									CachedData.ClothingStreamableData.Set(Block.Id, TempData.GetData() + Block.Offset, Block.Size);
 								}
 								break;
 							}
diff --git a/Engine/Plugins/Mutable/Source/CustomizableObjectEditor/Private/MuCOE/CustomizableObjectCompileRunnable.cpp b/Engine/Plugins/Mutable/Source/CustomizableObjectEditor/Private/MuCOE/CustomizableObjectCompileRunnable.cpp
--- a/Engine/Plugins/Mutable/Source/CustomizableObjectEditor/Private/MuCOE/CustomizableObjectCompileRunnable.cpp
+++ b/Engine/Plugins/Mutable/Source/CustomizableObjectEditor/Private/MuCOE/CustomizableObjectCompileRunnable.cpp
@@ -320,15 +320,24 @@ FCustomizableObjectSaveDDRunnable::FCustomizableObjectSaveDDRunnable(const TShar
 		static_assert(TCanBulkSerialize<FMorphTargetVertexData>::Value);
 		TArray<FMorphTargetVertexData>& MorphVertexData = InModelResources.EditorOnlyMorphTargetReconstructionData;
 
-		PlatformData.MorphData.SetNum(MorphVertexData.Num() * sizeof(FMorphTargetVertexData));
-		FMemory::Memcpy(PlatformData.MorphData.GetData(), MorphVertexData.GetData(), PlatformData.MorphData.Num());
+		for (const TPair<int32, FRealTimeMorphStreamable>& MorphStreamable : InModelStreamables->RealTimeMorphStreamables)
+		{
+			PlatformData.MorphStreamableData.Set(
+					MorphStreamable.Key,
+					reinterpret_cast<uint8*>(MorphVertexData.GetData()) + MorphStreamable.Value.Block.Offset,
+					MorphStreamable.Value.Size);
+		}
 
-		// Cloth
 		static_assert(TCanBulkSerialize<FCustomizableObjectMeshToMeshVertData>::Value);
 		TArray<FCustomizableObjectMeshToMeshVertData>& ClothingVertexData = InModelResources.EditorOnlyClothingMeshToMeshVertData;
-
-		PlatformData.ClothingData.SetNum(ClothingVertexData.Num() * sizeof(FCustomizableObjectMeshToMeshVertData));
-		FMemory::Memcpy(PlatformData.ClothingData.GetData(), ClothingVertexData.GetData(), PlatformData.ClothingData.Num());
+		
+		for (const TPair<int32, FClothingStreamable>& ClothingStreamable : InModelStreamables->ClothingStreamables)
+		{
+			PlatformData.ClothingStreamableData.Set(
+					ClothingStreamable.Key,
+					reinterpret_cast<uint8*>(ClothingVertexData.GetData()) + ClothingStreamable.Value.Block.Offset,
+					ClothingStreamable.Value.Size);
+		}
 
 		if (Options.bIsCooking)
 		{
@@ -615,8 +624,8 @@ void FCustomizableObjectSaveDDRunnable::StoreCachedPlatformDataToDisk(bool& bSto
 		StreamableMemoryWriter->Flush();
 		StreamableMemoryWriter->Close();
 
-		PlatformData.MorphData.Empty();
-		PlatformData.ClothingData.Empty();
+		PlatformData.MorphStreamableData.Data.Empty();
+		PlatformData.ClothingStreamableData.Data.Empty();
 	}
 
 	// Serialize Model and ModelResources. Store after SerializeBulkDataFiles fixes the HashToStreamableFiles offsets.
