[Mutable] Fix Linux compilation error.

#rnx

diff --git a/Engine/Plugins/Mutable/Source/CustomizableObjectEditor/Private/MuCOE/CustomizableObjectCompileRunnable.cpp b/Engine/Plugins/Mutable/Source/CustomizableObjectEditor/Private/MuCOE/CustomizableObjectCompileRunnable.cpp
--- a/Engine/Plugins/Mutable/Source/CustomizableObjectEditor/Private/MuCOE/CustomizableObjectCompileRunnable.cpp
+++ b/Engine/Plugins/Mutable/Source/CustomizableObjectEditor/Private/MuCOE/CustomizableObjectCompileRunnable.cpp
@@ -320,7 +320,7 @@ FCustomizableObjectSaveDDRunnable::FCustomizableObjectSaveDDRunnable(const TShar
 		static_assert(TCanBulkSerialize<FMorphTargetVertexData>::Value);
 		TArray<FMorphTargetVertexData>& MorphVertexData = InModelResources.EditorOnlyMorphTargetReconstructionData;
 
-		for (const TPair<int32, FRealTimeMorphStreamable>& MorphStreamable : InModelStreamables->RealTimeMorphStreamables)
+		for (const TPair<uint32, FRealTimeMorphStreamable>& MorphStreamable : InModelStreamables->RealTimeMorphStreamables)
 		{
 			PlatformData.MorphStreamableData.Set(
 					MorphStreamable.Key,
@@ -331,7 +331,7 @@ FCustomizableObjectSaveDDRunnable::FCustomizableObjectSaveDDRunnable(const TShar
 		static_assert(TCanBulkSerialize<FCustomizableObjectMeshToMeshVertData>::Value);
 		TArray<FCustomizableObjectMeshToMeshVertData>& ClothingVertexData = InModelResources.EditorOnlyClothingMeshToMeshVertData;
 		
-		for (const TPair<int32, FClothingStreamable>& ClothingStreamable : InModelStreamables->ClothingStreamables)
+		for (const TPair<uint32, FClothingStreamable>& ClothingStreamable : InModelStreamables->ClothingStreamables)
 		{
 			PlatformData.ClothingStreamableData.Set(
 					ClothingStreamable.Key,
