[mutable] Restructure CustomizableObject constant mesh data to avoid one index mapping and save memory.
- The constant op in bytecode now contains directly the resource index, including the flag for streaming.
- Split streamed and permanent meshes in different arrays to prevent having a permanently allocated array of mostly null pointers.
#rnx

diff --git a/Engine/Plugins/Mutable/Source/CustomizableObject/Internal/MuCO/CustomizableObjectPrivate.h b/Engine/Plugins/Mutable/Source/CustomizableObject/Internal/MuCO/CustomizableObjectPrivate.h
--- a/Engine/Plugins/Mutable/Source/CustomizableObject/Internal/MuCO/CustomizableObjectPrivate.h
+++ b/Engine/Plugins/Mutable/Source/CustomizableObject/Internal/MuCO/CustomizableObjectPrivate.h
@@ -1628,6 +1628,8 @@ public:
 
 		RestructureConstantImageData,
 
+		RestructureConstantMeshData,
+
 		// -----<new versions can be added above this line>--------
 		LastCustomizableObjectVersion
 	};
diff --git a/Engine/Plugins/Mutable/Source/CustomizableObject/Private/MuCO/CustomizableObject.cpp b/Engine/Plugins/Mutable/Source/CustomizableObject/Private/MuCO/CustomizableObject.cpp
--- a/Engine/Plugins/Mutable/Source/CustomizableObject/Private/MuCO/CustomizableObject.cpp
+++ b/Engine/Plugins/Mutable/Source/CustomizableObject/Private/MuCO/CustomizableObject.cpp
@@ -3424,13 +3424,13 @@ void UCustomizableObjectPrivate::LogMemory()
 		}
 		int32 ImageTotal = ConstantImages + ConstantImageLODIndices + ConstantImageLODPermanent + PermanentImages + ConstantImageLODStreamed;
 
-		int32 ConstantMeshes = Program.ConstantMeshes.GetAllocatedSize();
+		int32 ConstantMeshes = Program.ConstantMeshesPermanent.GetAllocatedSize();
 		int32 PermanentMeshes = 0;
-		for (const TPair<int32, TSharedPtr<const mu::FMesh>>& Entry : Program.ConstantMeshes)
+		for (const TSharedPtr<const mu::FMesh>& Entry : Program.ConstantMeshesPermanent)
 		{
-			if (Entry.Key == -1)
+			if (Entry)
 			{
-				PermanentMeshes += Entry.Value->GetDataSize();
+				PermanentMeshes += Entry->GetDataSize();
 			}
 		}
 		int32 MeshesTotal = ConstantMeshes + PermanentMeshes;
diff --git a/Engine/Plugins/Mutable/Source/CustomizableObjectEditor/Private/MuCOE/SMutableConstantsWidget.cpp b/Engine/Plugins/Mutable/Source/CustomizableObjectEditor/Private/MuCOE/SMutableConstantsWidget.cpp
--- a/Engine/Plugins/Mutable/Source/CustomizableObjectEditor/Private/MuCOE/SMutableConstantsWidget.cpp
+++ b/Engine/Plugins/Mutable/Source/CustomizableObjectEditor/Private/MuCOE/SMutableConstantsWidget.cpp
@@ -1213,7 +1213,15 @@ void SMutableConstantsWidget::LoadConstantMeshes()
 {
 	check (MutableProgramPtr);
 	
-	const int32 ConstantsCount = MutableProgramPtr->ConstantMeshes.Num();
+	TArray<TSharedPtr<const mu::FMesh>> AllMeshes;
+	AllMeshes.Reserve(MutableProgramPtr->ConstantMeshesPermanent.Num() + MutableProgramPtr->ConstantMeshesStreamed.Num() );
+	AllMeshes.Append(MutableProgramPtr->ConstantMeshesPermanent);
+	for (const TPair<uint32, TSharedPtr<const mu::FMesh>>& Entry : MutableProgramPtr->ConstantMeshesStreamed)
+	{
+		AllMeshes.Add(Entry.Value);
+	}
+
+	const int32 ConstantsCount = AllMeshes.Num();
 	if (ConstantMeshElements)
 	{
 		ConstantMeshElements->Empty(ConstantsCount);
@@ -1229,7 +1237,7 @@ void SMutableConstantsWidget::LoadConstantMeshes()
 	for (int32 MeshIndex = 0; MeshIndex < ConstantsCount; MeshIndex++)
 	{
 		TSharedPtr< FMutableConstantMeshElement> ConstantMeshElement = MakeShared<FMutableConstantMeshElement>();
-		ConstantMeshElement->MeshPtr = MutableProgramPtr->ConstantMeshes[MeshIndex].Value;
+		ConstantMeshElement->MeshPtr = AllMeshes[MeshIndex];
 		ConstantMeshElement->IndexOnSourceVector = MeshIndex;
 
 		// Actual core disk size would be:
diff --git a/Engine/Plugins/Mutable/Source/MutableRuntime/Internal/MuR/ModelPrivate.h b/Engine/Plugins/Mutable/Source/MutableRuntime/Internal/MuR/ModelPrivate.h
--- a/Engine/Plugins/Mutable/Source/MutableRuntime/Internal/MuR/ModelPrivate.h
+++ b/Engine/Plugins/Mutable/Source/MutableRuntime/Internal/MuR/ModelPrivate.h
@@ -48,11 +48,6 @@ namespace mu
 		//! Size of the rom
 		uint32 Size=0;
 
-		/** Index of the resource in its type-specific arrays. 
-		* TODO: Remove when streamed meshes also use the rom index as resource index.
-		*/
-		FConstantResourceIndex ResourceIndex = {};
-
 		//! Index of the resource in its type-specific array (one of DATATYPE values)
 		uint16 ResourceType=0;
 
@@ -244,12 +239,12 @@ namespace mu
 		TSparseArray<uint8> LoadedMemTrackedRoms;
 
 		/** Constant image mip data is split in 2 sets: ConstantImageLODsPermanent constains data that is always loaded. 
-		* Index with FConstantImageLODIndex::Index, when Streamable is 0.
+		* Index with FConstantResourceIndex::Index, when Streamable is 0.
 		*/
 		TArray<TSharedPtr<const FImage>> ConstantImageLODsPermanent;
 
 		/** Constant image mip data is split in 2 sets: ConstantImageLODsStreamed constains data that is streamed in and out. 
-		* Index with FConstantImageLODIndex::Index, when Streamable is 1.
+		* Index with FConstantResourceIndex::Index, when Streamable is 1.
 		* This part is empty for an unused Model, and shouldn't be serialised.
 		*/
 		TMap<uint32, TSharedPtr<const FImage>> ConstantImageLODsStreamed;
@@ -260,8 +255,16 @@ namespace mu
 		//! Constant image data.
 		TArray<FImageLODRange> ConstantImages;
 
-        //! Constant mesh data: the first is the index in Roms for each mesh or -1 if it is always loaded.
-		TArray<TPair<int32, TSharedPtr<const FMesh>>> ConstantMeshes;
+		/** Constant mesh data is split in 2 sets: ConstantMeshesPermanent constains data that is always loaded.
+		* Index with FConstantResourceIndex::Index, when Streamable is 0.
+		*/
+		TArray<TSharedPtr<const FMesh>> ConstantMeshesPermanent;
+
+		/** Constant mesh data is split in 2 sets: ConstantMeshesStreamed constains data that is streamed in and out.
+		* Index with FConstantResourceIndex::Index, when Streamable is 1.
+		* This part is empty for an unused Model, and shouldn't be serialised.
+		*/
+		TMap<uint32, TSharedPtr<const FMesh>> ConstantMeshesStreamed;
 
 		//! Constant FExtensionData
 		TArray<FExtensionDataConstant> ConstantExtensionData;
@@ -317,7 +320,7 @@ namespace mu
 			arch << ConstantImageLODsPermanent;
 			arch << ConstantImageLODIndices;
 			arch << ConstantImages;
-			arch << ConstantMeshes;
+			arch << ConstantMeshesPermanent;
 			arch << ConstantExtensionData;
 			arch << ConstantStrings;
             arch << ConstantLayouts;
@@ -342,7 +345,7 @@ namespace mu
 			arch >> ConstantImageLODsPermanent;
 			arch >> ConstantImageLODIndices;
 			arch >> ConstantImages;
-			arch >> ConstantMeshes;
+			arch >> ConstantMeshesPermanent;
 			arch >> ConstantExtensionData;
 			arch >> ConstantStrings;
             arch >> ConstantLayouts;
@@ -366,18 +369,12 @@ namespace mu
 		//! Return true if the given ROM is loaded.
 		FORCEINLINE bool IsRomLoaded(int32 RomIndex) const
 		{
-			FConstantResourceIndex ResourceIndex = Roms[RomIndex].ResourceIndex;
-			if (!ResourceIndex.Streamable)
-			{
-				return true;
-			}
-
 			switch (Roms[RomIndex].ResourceType)
 			{
 				case DT_IMAGE: 
-					return ConstantImageLODsStreamed.Contains(ResourceIndex.Index);
+					return ConstantImageLODsStreamed.Contains(RomIndex);
 				case DT_MESH: 
-					return ConstantMeshes[ResourceIndex.Index].Value != nullptr;
+					return ConstantMeshesStreamed.Contains(RomIndex);
 				default:
 					check(false);
 					break;
@@ -386,11 +383,11 @@ namespace mu
 			return false;
 		}
 
-		/** Unload a rom resource. Return the size of the unloaded rom.*/
-		FORCEINLINE int32 UnloadRom(int32 RomIndex)
+		/** Unload a rom resource. 
+		* If a pointer to a size is passed in OutDataSize, the value will be set with the size of the unloaded rom.
+		*/
+		FORCEINLINE void UnloadRom(int32 RomIndex, int32* OutDataSize=nullptr)
 		{
-			int32 RomSize = 0;
-
 			if (DebugRom && (DebugRomAll||RomIndex == DebugRomIndex))
 				UE_LOG(LogMutableCore, Log, TEXT("Unloading rom %d."), RomIndex);
 
@@ -399,62 +396,50 @@ namespace mu
 				LoadedMemTrackedRoms.RemoveAt(RomIndex);
 			}
 
-			FConstantResourceIndex ResourceIndex = Roms[RomIndex].ResourceIndex;
-			check(ResourceIndex.Streamable);
-
 			switch (Roms[RomIndex].ResourceType)
 			{
 			case DT_IMAGE:
 			{
 				TSharedPtr<const FImage> Data;
-				ConstantImageLODsStreamed.RemoveAndCopyValue(ResourceIndex.Index, Data);
-				if (Data)
+				ConstantImageLODsStreamed.RemoveAndCopyValue(RomIndex, Data);
+				if (OutDataSize && Data)
 				{
-					RomSize = Data->GetDataSize();
+					*OutDataSize = Data->GetDataSize();
 				}
 				break;
 			}
 			case DT_MESH:
-				if (ConstantMeshes[ResourceIndex.Index].Value)
+			{
+				TSharedPtr<const FMesh> Data;
+				ConstantMeshesStreamed.RemoveAndCopyValue(RomIndex, Data);
+				if (OutDataSize && Data)
 				{
-					RomSize = ConstantMeshes[ResourceIndex.Index].Value->GetDataSize();
-					ConstantMeshes[ResourceIndex.Index].Value = nullptr;
+					*OutDataSize = Data->GetDataSize();
 				}
 				break;
+			}
 			default:
 				check(false);
 				break;
 			}
-
-			return RomSize;
 		}
 
-		FORCEINLINE void SetMeshRomValue(int32 RomIndex, const TSharedPtr<FMesh>& Value)
+		FORCEINLINE void SetMeshRomValue(uint32 RomIndex, const TSharedPtr<FMesh>& Value)
 		{
 			check(Roms[RomIndex].ResourceType == DT_MESH);
-			
-			FConstantResourceIndex ResourceIndex = Roms[RomIndex].ResourceIndex;
-			check(ResourceIndex.Streamable);
+			check(!ConstantMeshesStreamed.Contains(RomIndex));
 
 			LoadedMemTrackedRoms.EmplaceAt(RomIndex, (uint8)Roms[RomIndex].ResourceType);
-			ConstantMeshes[ResourceIndex.Index].Value = Value;
+			ConstantMeshesStreamed.Add(RomIndex, Value);
 		}
 
-		FORCEINLINE void SetImageRomValue(int32 RomIndex, const TSharedPtr<FImage>& Value)
+		FORCEINLINE void SetImageRomValue(uint32 RomIndex, const TSharedPtr<FImage>& Value)
 		{
 			check(Roms[RomIndex].ResourceType == DT_IMAGE);
-
 			check(!ConstantImageLODsStreamed.Contains(RomIndex));
 
 			LoadedMemTrackedRoms.EmplaceAt(RomIndex, (uint8)Roms[RomIndex].ResourceType);
-			ConstantImageLODsStreamed.Add(RomIndex, Value );
-		}
-
-
-		int32 AddConstant(TSharedPtr<const FMesh> pMesh)
-		{
-			// Uniques needs to be ensured outside
-			return ConstantMeshes.Add(TPair<int32, TSharedPtr<const FMesh>>( -1, pMesh ));
+			ConstantImageLODsStreamed.Add(RomIndex, Value);
 		}
 
 		OP::ADDRESS AddConstant(TSharedPtr<const FExtensionData> Data)
@@ -624,7 +609,7 @@ namespace mu
 
         //! Get a constant image, assuming at least some mips are loaded. The image constant will be composed with lodaded mips if necessary.
 		template <typename CreateImageFunc>
-        void GetConstant( int32 ConstantIndex, TSharedPtr<const FImage>& res, int32 MipsToSkip, const CreateImageFunc& CreateImage) const
+        void GetConstant( uint32 ConstantIndex, TSharedPtr<const FImage>& res, int32 MipsToSkip, const CreateImageFunc& CreateImage) const
         {
 			int32 ReallySkippedLODs = FMath::Min(ConstantImages[ConstantIndex].LODCount - 1, MipsToSkip);
 			int32 FirstLODIndexIndex = ConstantImages[ConstantIndex].FirstIndex;
@@ -717,9 +702,16 @@ namespace mu
 			}
 		}
 
-        void GetConstant(int32 ConstantIndex, TSharedPtr<const FMesh>& res) const
+        void GetConstant(FConstantResourceIndex Index, TSharedPtr<const FMesh>& OutMesh) const
         {
-			res = ConstantMeshes[ConstantIndex].Value;
+			if (!Index.Streamable)
+			{
+				OutMesh = ConstantMeshesPermanent[Index.Index];
+			}
+			else
+			{
+				OutMesh = ConstantMeshesStreamed[Index.Index];
+			}
 		}
 
 		void GetExtensionDataConstant(int32 ConstantIndex, TSharedPtr<const FExtensionData>& Result) const
@@ -743,15 +735,23 @@ namespace mu
             return result;
         }
 
-        template<typename ARGS>
-        inline const ARGS GetOpArgs( OP::ADDRESS at ) const
-        {
-            ARGS result;
-            uint64 byteCodeAddress = OpAddress[at];
-            byteCodeAddress += sizeof(EOpType);
-			FMemory::Memcpy( &result, &ByteCode[byteCodeAddress], sizeof(ARGS));
-            return result;
-        }
+		template<typename ARGS>
+		inline const ARGS GetOpArgs(OP::ADDRESS at) const
+		{
+			ARGS result;
+			uint64 byteCodeAddress = OpAddress[at];
+			byteCodeAddress += sizeof(EOpType);
+			FMemory::Memcpy(&result, &ByteCode[byteCodeAddress], sizeof(ARGS));
+			return result;
+		}
+
+		template<typename ARGS>
+		inline void SetOpArgs(OP::ADDRESS at, const ARGS& Args)
+		{
+			uint64 byteCodeAddress = OpAddress[at];
+			byteCodeAddress += sizeof(EOpType);
+			FMemory::Memcpy(&ByteCode[byteCodeAddress], &Args, sizeof(ARGS));
+		}
 
         inline const uint8* GetOpArgsPointer( OP::ADDRESS at ) const
         {
diff --git a/Engine/Plugins/Mutable/Source/MutableRuntime/Internal/MuR/SystemPrivate.h b/Engine/Plugins/Mutable/Source/MutableRuntime/Internal/MuR/SystemPrivate.h
--- a/Engine/Plugins/Mutable/Source/MutableRuntime/Internal/MuR/SystemPrivate.h
+++ b/Engine/Plugins/Mutable/Source/MutableRuntime/Internal/MuR/SystemPrivate.h
@@ -1628,9 +1628,9 @@ namespace mu
 						Result += ImageLOD.Value->GetDataSize();
 					}
 				}
-				for (const TPair<int32, TSharedPtr<const FMesh>>& Rom : Program.ConstantMeshes)
+				for (const TPair<uint32, TSharedPtr<const FMesh>>& Rom : Program.ConstantMeshesStreamed)
 				{
-					if (Rom.Value && Rom.Key >= 0)
+					if (Rom.Value)
 					{
 						Result += Rom.Value->GetDataSize();
 					}
diff --git a/Engine/Plugins/Mutable/Source/MutableRuntime/Private/MuR/CodeRunner.cpp b/Engine/Plugins/Mutable/Source/MutableRuntime/Private/MuR/CodeRunner.cpp
--- a/Engine/Plugins/Mutable/Source/MutableRuntime/Private/MuR/CodeRunner.cpp
+++ b/Engine/Plugins/Mutable/Source/MutableRuntime/Private/MuR/CodeRunner.cpp
@@ -952,13 +952,10 @@ namespace mu
         {
 			OP::MeshConstantArgs args = Program.GetOpArgs<OP::MeshConstantArgs>(item.At);
 
-            OP::ADDRESS cat = args.value;
-
-            // Assume the ROM has been loaded previously
-            check(Program.ConstantMeshes[cat].Value)
+            FConstantResourceIndex ResourceIndex = *reinterpret_cast<const FConstantResourceIndex*>(&args.value);
 
             TSharedPtr<const FMesh> SourceConst;
-			Program.GetConstant(cat, SourceConst);
+			Program.GetConstant(ResourceIndex, SourceConst);
 
 			check(SourceConst);
 			TSharedPtr<FMesh> Source = CreateMesh(SourceConst->GetDataSize());
diff --git a/Engine/Plugins/Mutable/Source/MutableRuntime/Private/MuR/CodeRunner.h b/Engine/Plugins/Mutable/Source/MutableRuntime/Private/MuR/CodeRunner.h
--- a/Engine/Plugins/Mutable/Source/MutableRuntime/Private/MuR/CodeRunner.h
+++ b/Engine/Plugins/Mutable/Source/MutableRuntime/Private/MuR/CodeRunner.h
@@ -261,7 +261,7 @@ namespace  mu
 		class FLoadMeshRomTask : public CodeRunner::FIssuedTask
 		{
 		public:
-			FLoadMeshRomTask( const FScheduledOp& InOp, int32 InRomIndex)
+			FLoadMeshRomTask( const FScheduledOp& InOp, uint32 InRomIndex)
 				: FIssuedTask(InOp)
 			{
 				RomIndex = InRomIndex;
@@ -272,7 +272,7 @@ namespace  mu
 			virtual bool Complete(CodeRunner*) override;
 
 		private:
-			int32 RomIndex = -1;
+			uint32 RomIndex = 0;
 		};
 
 		class FLoadExtensionDataTask : public CodeRunner::FIssuedTask
diff --git a/Engine/Plugins/Mutable/Source/MutableRuntime/Private/MuR/CodeRunnerTasked.cpp b/Engine/Plugins/Mutable/Source/MutableRuntime/Private/MuR/CodeRunnerTasked.cpp
--- a/Engine/Plugins/Mutable/Source/MutableRuntime/Private/MuR/CodeRunnerTasked.cpp
+++ b/Engine/Plugins/Mutable/Source/MutableRuntime/Private/MuR/CodeRunnerTasked.cpp
@@ -2109,7 +2109,7 @@ namespace mu
 
 		const FProgram& Program = Runner->Model->GetPrivate()->Program;
 
-		check(RomIndex < Program.Roms.Num());
+		check(int32(RomIndex) < Program.Roms.Num());
 		
 		FWorkingMemoryManager::FModelCacheEntry* ModelCache = Runner->System->WorkingMemoryManager.FindModelCache(Runner->Model.Get());
 		++ModelCache->PendingOpsPerRom[RomIndex];
@@ -2183,14 +2183,13 @@ namespace mu
 		{
 			bool bSuccess = Runner->System->StreamInterface->EndRead(RomLoadOp->m_streamID);
 
-			const FConstantResourceIndex ResIndex = Program.Roms[RomIndex].ResourceIndex;
-			check(!Program.ConstantMeshes[ResIndex.Index].Value)
+			check(!Program.ConstantMeshesStreamed.Contains(RomIndex));
+
 			MUTABLE_CPUPROFILER_SCOPE(Unserialise);
 
 			FInputMemoryStream Stream(RomLoadOp->m_streamBuffer.GetData(), RomLoadOp->m_streamBuffer.Num());
 			FInputArchive Arch(&Stream);
 
-			check(!Program.ConstantMeshes[ResIndex.Index].Value);
 			TSharedPtr<FMesh> Value;
 			if (bSuccess)
 			{
@@ -2204,7 +2203,7 @@ namespace mu
 
 			Program.SetMeshRomValue(RomIndex, Value);
 
-			check(Program.ConstantMeshes[ResIndex.Index].Value);
+			check(Program.ConstantMeshesStreamed.Contains(RomIndex));
 
 			Runner->RomLoadOps.Remove(*RomLoadOp);
 		}
@@ -2641,10 +2640,10 @@ namespace mu
 		case EOpType::ME_CONSTANT:
 		{
 			OP::MeshConstantArgs args = Program.GetOpArgs<OP::MeshConstantArgs>(item.At);
-			int32 RomIndex = Program.ConstantMeshes[args.value].Key;
-			if (RomIndex >= 0 && !Program.ConstantMeshes[args.value].Value)
+			FConstantResourceIndex ResourceIndex = *reinterpret_cast<const FConstantResourceIndex*>(&args.value);
+			if (ResourceIndex.Streamable && !Program.ConstantMeshesStreamed.Contains(ResourceIndex.Index))
 			{
-				Issued = MakeShared<FLoadMeshRomTask>(item, RomIndex);
+				Issued = MakeShared<FLoadMeshRomTask>(item, uint32(ResourceIndex.Index));
 			}
 			else
 			{
diff --git a/Engine/Plugins/Mutable/Source/MutableRuntime/Private/MuR/Model.cpp b/Engine/Plugins/Mutable/Source/MutableRuntime/Private/MuR/Model.cpp
--- a/Engine/Plugins/Mutable/Source/MutableRuntime/Private/MuR/Model.cpp
+++ b/Engine/Plugins/Mutable/Source/MutableRuntime/Private/MuR/Model.cpp
@@ -142,12 +142,6 @@ namespace mu
 
 		mu::FProgram& Program = p->m_pD->Program;
 
-		TArray<TPair<int32, TSharedPtr<const FMesh>>> InitialMeshes;
-		if (!bDropData)
-		{
-			InitialMeshes = Program.ConstantMeshes;
-		}
-
 		FOutputMemoryStream MemStream(16 * 1024 * 1024);
 
 		// Save images and unload from memory
@@ -156,7 +150,6 @@ namespace mu
 			int32 RomIndex = Entry.Key;
 			const FRomData& RomData = Program.Roms[RomIndex];
 			check(RomData.ResourceType == DT_IMAGE);
-			check(RomData.ResourceIndex.Index == RomIndex);
 
 			// Serialize to memory, to find out final size of this rom
 			MemStream.Reset();
@@ -181,20 +174,11 @@ namespace mu
 		}
 
 		// Save meshes and unload from memory
-		for (int32 ResourceIndex = 0; ResourceIndex < Program.ConstantMeshes.Num(); ++ResourceIndex)
+		for (TPair<uint32, TSharedPtr<const FMesh>>& ResData : Program.ConstantMeshesStreamed)
 		{
-			TPair<int32, TSharedPtr<const FMesh>>& ResData = Program.ConstantMeshes[ResourceIndex];
-
-			// This shouldn't have been serialised with rom support before.
-			if (ResData.Key < 0)
-			{
-				continue;
-			}
-
 			int32 RomIndex = ResData.Key;
 			const FRomData& RomData = Program.Roms[RomIndex];
 			check(RomData.ResourceType == DT_MESH);
-			check(RomData.ResourceIndex.Index == ResourceIndex);
 
 			// Serialize to memory, to find out final size of this rom
 			MemStream.Reset();
@@ -206,8 +190,16 @@ namespace mu
 			streamer.Write(MemStream.GetBuffer(), MemStream.GetBufferSize());
 			streamer.CloseWriteFile();
 
-			// We clear it to avoid including it also with the main model data
-			ResData.Value = nullptr;
+			// Do this progressively to avoid duplicating all data in memory.
+			if (bDropData)
+			{
+				ResData.Value.Reset();
+			}
+		}
+
+		if (bDropData)
+		{
+			Program.ConstantMeshesStreamed.Empty(0);
 		}
 
 		// Store the main data of the model
@@ -220,12 +212,6 @@ namespace mu
 
 			streamer.CloseWriteFile();
 		}
-
-		// Restore full data
-		if (!bDropData)
-		{
-			Program.ConstantMeshes = InitialMeshes;
-		}
 	}
 
 
@@ -253,14 +239,7 @@ namespace mu
 		LLM_SCOPE_BYNAME(TEXT("MutableRuntime"));
 
 		m_pD->Program.ConstantImageLODsStreamed.Empty(0);
-
-		for (int32 ResIndex = 0; ResIndex < m_pD->Program.ConstantMeshes.Num(); ++ResIndex)
-		{
-			if (m_pD->Program.ConstantMeshes[ResIndex].Key >= 0)
-			{
-				m_pD->Program.ConstantMeshes[ResIndex].Value = nullptr;
-			}
-		}
+		m_pD->Program.ConstantMeshesStreamed.Empty(0);
 	}
 
 
diff --git a/Engine/Plugins/Mutable/Source/MutableRuntime/Private/MuR/System.cpp b/Engine/Plugins/Mutable/Source/MutableRuntime/Private/MuR/System.cpp
--- a/Engine/Plugins/Mutable/Source/MutableRuntime/Private/MuR/System.cpp
+++ b/Engine/Plugins/Mutable/Source/MutableRuntime/Private/MuR/System.cpp
@@ -1358,7 +1358,8 @@ namespace mu
 				Candidates.HeapPop(Candidate, CompareCandidates, EAllowShrinking::No);
 
 				// UE_LOG(LogMutableCore,Log, "Unloading rom because of memory budget: %d.", lowestPriorityRom);
-				int32 UnloadedSize = Candidate.Value.Model->GetPrivate()->Program.UnloadRom(Candidate.Value.RomIndex);
+				int32 UnloadedSize = 0;
+				Candidate.Value.Model->GetPrivate()->Program.UnloadRom(Candidate.Value.RomIndex, &UnloadedSize);
 				TotalBytes -= UnloadedSize;
 				bFinished = TotalBytes <= BudgetBytes;
 			}
diff --git a/Engine/Plugins/Mutable/Source/MutableTools/Private/MuT/ASTOpConstantResource.cpp b/Engine/Plugins/Mutable/Source/MutableTools/Private/MuT/ASTOpConstantResource.cpp
--- a/Engine/Plugins/Mutable/Source/MutableTools/Private/MuT/ASTOpConstantResource.cpp
+++ b/Engine/Plugins/Mutable/Source/MutableTools/Private/MuT/ASTOpConstantResource.cpp
@@ -279,8 +279,8 @@ namespace mu
 
 	namespace
 	{
-		/** Adds a constant image data to a program and returns its constant index. */
-		int32 AddConstantImage(FProgram& Program, const TSharedPtr<const FImage>& pImage, FLinkerOptions& Options)
+		/** Adds a constant image data to a Program and returns its constant index. */
+		uint32 AddConstantImage(FProgram& Program, const TSharedPtr<const FImage>& pImage, FLinkerOptions& Options)
 		{
 			MUTABLE_CPUPROFILER_SCOPE(AddConstantImage);
 
@@ -395,13 +395,13 @@ namespace mu
 			LODRange.ImageFormat = pImage->GetFormat();
 			LODRange.ImageSizeX = pImage->GetSizeX();
 			LODRange.ImageSizeY = pImage->GetSizeY();
-			int32 ImageIndex = Program.ConstantImages.Add(LODRange);
+			uint32 ImageIndex = Program.ConstantImages.Add(LODRange);
 			return ImageIndex;
 		}
 	}
 
 	//-------------------------------------------------------------------------------------------------
-	void ASTOpConstantResource::Link(FProgram& program, FLinkerOptions* Options)
+	void ASTOpConstantResource::Link(FProgram& Program, FLinkerOptions* Options)
 	{
 		MUTABLE_CPUPROFILER_SCOPE(ASTOpConstantResource_Link);
 
@@ -418,14 +418,14 @@ namespace mu
 				args.skeleton = -1;
 				if (TSharedPtr<const FSkeleton> pSkeleton = MeshData->GetSkeleton())
 				{
-					args.skeleton = program.AddConstant(pSkeleton);
+					args.skeleton = Program.AddConstant(pSkeleton);
 					MeshData->SetSkeleton(nullptr);
 				}
 
 				args.physicsBody = -1;
 				if (TSharedPtr<const FPhysicsBody> pPhysicsBody = MeshData->GetPhysicsBody())
 				{
-					args.physicsBody = program.AddConstant(pPhysicsBody);
+					args.physicsBody = Program.AddConstant(pPhysicsBody);
 					MeshData->SetPhysicsBody(nullptr);
 				}
 
@@ -434,7 +434,7 @@ namespace mu
 				const int32* IndexPtr = Options->MeshConstantMap.Find(Key);
 				if (!IndexPtr)
 				{
-					args.value = program.AddConstant(MeshData);
+					args.value = Program.ConstantMeshesPermanent.Add(MeshData);
 
 					int32 DataDescIndex = Options->AdditionalData.SourceMeshPerConstant.Add(SourceDataDescriptor);
 					check(DataDescIndex == args.value);
@@ -446,10 +446,10 @@ namespace mu
 					args.value = *IndexPtr;
 				}
 
-				linkedAddress = (OP::ADDRESS)program.OpAddress.Num();
-				program.OpAddress.Add((uint32_t)program.ByteCode.Num());
-				AppendCode(program.ByteCode, Type);
-				AppendCode(program.ByteCode, args);
+				linkedAddress = (OP::ADDRESS)Program.OpAddress.Num();
+				Program.OpAddress.Add((uint32)Program.ByteCode.Num());
+				AppendCode(Program.ByteCode, Type);
+				AppendCode(Program.ByteCode, args);
 			}
 			else
 			{
@@ -472,7 +472,7 @@ namespace mu
 					}
 					else
 					{
-						args.value = AddConstantImage( program, pTyped, *Options);
+						args.value = AddConstantImage( Program, pTyped, *Options);
 
 						int32 DataDescIndex = Options->AdditionalData.SourceImagePerConstant.Add(SourceDataDescriptor);
 						check(DataDescIndex == args.value);
@@ -484,7 +484,7 @@ namespace mu
 				{
 					TSharedPtr<const FLayout> pTyped = StaticCastSharedPtr<const FLayout>(GetValue());
 					check(pTyped);
-					args.value = program.AddConstant(pTyped);
+					args.value = Program.AddConstant(pTyped);
 					break;
 				}
 				default:
@@ -493,10 +493,10 @@ namespace mu
 
 				if (bValidData)
 				{
-					linkedAddress = (OP::ADDRESS)program.OpAddress.Num();
-					program.OpAddress.Add((uint32)program.ByteCode.Num());
-					AppendCode(program.ByteCode, Type);
-					AppendCode(program.ByteCode, args);
+					linkedAddress = (OP::ADDRESS)Program.OpAddress.Num();
+					Program.OpAddress.Add((uint32)Program.ByteCode.Num());
+					AppendCode(Program.ByteCode, Type);
+					AppendCode(Program.ByteCode, args);
 				}
 				else
 				{
diff --git a/Engine/Plugins/Mutable/Source/MutableTools/Private/MuT/Compiler.cpp b/Engine/Plugins/Mutable/Source/MutableTools/Private/MuT/Compiler.cpp
--- a/Engine/Plugins/Mutable/Source/MutableTools/Private/MuT/Compiler.cpp
+++ b/Engine/Plugins/Mutable/Source/MutableTools/Private/MuT/Compiler.cpp
@@ -494,7 +494,7 @@ namespace mu
 		Program.ConstantImageLODsPermanent.Shrink();
 		Program.ConstantImageLODIndices.Shrink();
 		Program.ConstantImages.Shrink();
-		Program.ConstantMeshes.Shrink();
+		Program.ConstantMeshesPermanent.Shrink();
 		Program.ConstantStrings.Shrink();
 		Program.ConstantSkeletons.Shrink();
 		Program.ConstantPhysicsBodies.Shrink();
@@ -553,14 +553,14 @@ namespace mu
 		uint64 NumHighResBytes = 0;
 
 		// Maximum number of roms
-		int32 MaxRomCount = Program.ConstantImageLODsPermanent.Num() + Program.ConstantMeshes.Num();
+		int32 MaxRomCount = Program.ConstantImageLODsPermanent.Num() + Program.ConstantMeshesPermanent.Num();
 		Program.Roms.Reserve(MaxRomCount);
 
 		TSet<uint32> UsedIds;
 		UsedIds.Reserve(MaxRomCount);
 
 		TArray<FRomData> RomDatas;
-		RomDatas.SetNumZeroed(FMath::Max(Program.ConstantImageLODsPermanent.Num(), Program.ConstantMeshes.Num()));
+		RomDatas.SetNumZeroed(FMath::Max(Program.ConstantImageLODsPermanent.Num(), Program.ConstantMeshesPermanent.Num()));
 
 		// Images
 		{
@@ -574,12 +574,12 @@ namespace mu
 			ParallelFor(AllMips.Num(),
 				[EmbeddedDataBytesLimit, &Program, &AllMips, &RomDatas](uint32 MipIndex)
 				{
-					TSharedPtr<const FImage>& ResData = AllMips[MipIndex];
+					TSharedPtr<const FImage>& Resource = AllMips[MipIndex];
 
 					// Serialize to find out final size of this rom
 					FOutputSizeStream SizeStream;
 					FOutputArchive MemoryArch(&SizeStream);
-					FImage::Serialise(ResData.Get(), MemoryArch);
+					FImage::Serialise(Resource.Get(), MemoryArch);
 
 
 					FRomData& RomData = RomDatas[MipIndex];
@@ -595,7 +595,6 @@ namespace mu
 					// RomData.ResourceIndex will be filled below when we split in permanent and streamable
 					// RomData.Flags will be filled below
 
-					const FImage* Resource = ResData.Get();
 					uint64 FullHash = CityHash64(reinterpret_cast<const char*>(Resource->GetLODData(0)), Resource->GetLODDataSize(0));
 					RomData.Id = GetTypeHash(FullHash);
 				});
@@ -662,6 +661,11 @@ namespace mu
 
 			// Split the data in permanent and streamable and assign final FConstantResourceIndex
 			{
+				MUTABLE_CPUPROFILER_SCOPE(GenerateRoms_ImageSplit);
+
+				TArray<FConstantResourceIndex> IndexPerMip;
+				IndexPerMip.SetNumUninitialized(AllMips.Num());
+
 				for (int32 MipIndex = 0; MipIndex < AllMips.Num(); ++MipIndex)
 				{
 					FRomData& RomData = RomDatas[MipIndex];
@@ -676,7 +680,7 @@ namespace mu
 						uint32 RomIndex = Program.Roms.Num();
 						check(RomIndex < 0x7fffffff);
 
-						RomData.ResourceIndex = { RomIndex, 1 };
+						IndexPerMip[MipIndex] = {RomIndex, 1};
 						Program.ConstantImageLODsStreamed.Add(RomIndex, AllMips[MipIndex]);
 
 						Program.Roms.Add(RomData);
@@ -685,7 +689,7 @@ namespace mu
 					{
 						uint32 Index = uint32(Program.ConstantImageLODsPermanent.Num());
 						check(Index < 0x7fffffff);
-						RomData.ResourceIndex = { Index, 0 };
+						IndexPerMip[MipIndex] = { Index, 0 };
 						Program.ConstantImageLODsPermanent.Add(AllMips[MipIndex]);
 
 						NumEmbedded++;
@@ -696,31 +700,36 @@ namespace mu
 				for (int32 Index = 0; Index < Program.ConstantImageLODIndices.Num(); ++Index)
 				{
 					int32 RomDataIndex = Program.ConstantImageLODIndices[Index].Index;
-					Program.ConstantImageLODIndices[Index] = RomDatas[RomDataIndex].ResourceIndex;
+					Program.ConstantImageLODIndices[Index] = IndexPerMip[RomDataIndex];
 				}
 			}
 		}
 
 		// Meshes
+		TArray<TSharedPtr<const FMesh>> AllMeshes = MoveTemp(Program.ConstantMeshesPermanent);
+		Program.ConstantMeshesPermanent = {};
+		Program.ConstantMeshesPermanent.Reserve(AllMeshes.Num());
+
+		// Meshes indices will be remapped as they are classified between streamed and permanent
+		TArray<uint32> OldToNewMeshIndex;
+		OldToNewMeshIndex.SetNumUninitialized(AllMeshes.Num());
+
 		FMemory::Memzero( RomDatas.GetData(), RomDatas.GetAllocatedSize() );
 		{
 			MUTABLE_CPUPROFILER_SCOPE(GenerateRoms_MeshIds);
 
-			ParallelFor(Program.ConstantMeshes.Num(),
-				[EmbeddedDataBytesLimit, &Program, &RomDatas](uint32 ResourceIndex)
+			ParallelFor(AllMeshes.Num(),
+				[EmbeddedDataBytesLimit, &AllMeshes , &Program, &RomDatas](uint32 ResourceIndex)
 				{
-					TPair<int32, TSharedPtr<const FMesh>>& ResData = Program.ConstantMeshes[ResourceIndex];
-
-					// This shouldn't have been serialised with rom support before.
-					check(ResData.Key < 0);
+					TSharedPtr<const FMesh>& ResData = AllMeshes[ResourceIndex];
 
 					// Serialize to memory, to find out final size of this rom
-					const FMesh* Resource = ResData.Value.Get();
+					const FMesh* Resource = ResData.Get();
 					int32 ApproximateSize = Resource->GetDataSize();
 
 					FOutputMemoryStream MemStream(ApproximateSize + 64 * 1024);
 					FOutputArchive MemoryArch(&MemStream);
-					FMesh::Serialise(ResData.Value.Get(), MemoryArch);
+					FMesh::Serialise(ResData.Get(), MemoryArch);
 
 					FRomData& RomData = RomDatas[ResourceIndex];
 					RomData.ResourceType = DT_MESH;
@@ -732,7 +741,8 @@ namespace mu
 						return;
 					}
 
-					RomData.ResourceIndex = { ResourceIndex, 1 };
+					// RomData.ResourceIndex will be filled below when we split in permanent and streamable
+					// RomData.Flags will be filled below
 
 					// Ensure that the Id is unique
 					uint64 FullHash = CityHash64(reinterpret_cast<const char*>(MemStream.GetBuffer()), MemStream.GetBufferSize());
@@ -740,31 +750,57 @@ namespace mu
 				});
 		}
 
+		// Split the data in permanent and streamable and assign final FConstantResourceIndex
 		{
-			MUTABLE_CPUPROFILER_SCOPE(GenerateRoms_MeshIdsUnique);
+			MUTABLE_CPUPROFILER_SCOPE(GenerateRoms_MeshSplit);
 
-			for (int32 ResourceIndex = 0; ResourceIndex < Program.ConstantMeshes.Num(); ++ResourceIndex)
+			for (int32 MeshIndex = 0; MeshIndex < AllMeshes.Num(); ++MeshIndex)
 			{
-				FRomData& RomData = RomDatas[ResourceIndex];
-				RomData.SourceId = AdditionalData.SourceMeshPerConstant[ResourceIndex].SourceId;
+				FRomData& RomData = RomDatas[MeshIndex];
+				RomData.SourceId = AdditionalData.SourceMeshPerConstant[MeshIndex].SourceId;
+
+				FConstantResourceIndex ResourceIndex;
 
-				// If the resource uses less memory than the threshold, don't save it in a separate rom.
-				if (int32(RomData.Size) <= EmbeddedDataBytesLimit)
+				if (RomData.Size > EmbeddedDataBytesLimit)
 				{
+					NumRoms++;
+					NumRomsBytes += RomData.Size;
+
+					// Ensure that the Id is unique
+					EnsureUniqueRomId(UsedIds, RomData.Id);
+
+					uint32 RomIndex = Program.Roms.Num();
+					check(RomIndex < 0x7fffffff);
+
+					ResourceIndex = { RomIndex, 1 };
+					Program.ConstantMeshesStreamed.Add(RomIndex, AllMeshes[MeshIndex]);
+
+					Program.Roms.Add(RomData);
+				}
+				else
+				{
+					uint32 Index = uint32(Program.ConstantMeshesPermanent.Num());
+					check(Index < 0x7fffffff);
+					ResourceIndex = { Index, 0 };
+					Program.ConstantMeshesPermanent.Add(AllMeshes[MeshIndex]);
+
 					NumEmbedded++;
 					NumEmbeddedBytes += RomData.Size;
-					continue;
 				}
-				NumRoms++;
-				NumRomsBytes += RomData.Size;
 
-				// Ensure that the Id is unique
-				EnsureUniqueRomId(UsedIds, RomData.Id);
-
-				int32 RomIndex = Program.Roms.Add(RomData);
+				OldToNewMeshIndex[MeshIndex] = *reinterpret_cast<const uint32*>(&ResourceIndex);
+			}
+		}
 
-				TPair<int32, TSharedPtr<const FMesh>>& ResData = Program.ConstantMeshes[ResourceIndex];
-				ResData.Key = RomIndex;
+		// Remap mesh constant references in bytecode
+		for (int32 OpIndex=0; OpIndex<Program.OpAddress.Num(); ++OpIndex)
+		{
+			EOpType Type = Program.GetOpType(OpIndex);
+			if (Type == EOpType::ME_CONSTANT)
+			{
+				OP::ResourceConstantArgs Args = Program.GetOpArgs<OP::ResourceConstantArgs>(OpIndex);
+				Args.value = OldToNewMeshIndex[Args.value];
+				Program.SetOpArgs<OP::ResourceConstantArgs>(OpIndex, Args);
 			}
 		}
 
