[Mutable] Fix Mutable Low Priority Tasks overlapping Initial Generation Updates.

Mutable Low Priority Tasks should never be allowed to execute when other updates are being executed. This restriction was already in place but not all cases were contemplated. Here is a list of fixed edge cades:
- Task_Mutable_ReleaseInstance could be delayed so much that it would execute after the Task_Game_StartUpdate of the next Update.
- If a Low Priority Task was split, TryLaunchMutableTaskLowPriority was being executed after finishing the first split task.
- FinishUpdateGlobal was setting AllowLaunchingMutableTaskLowPriority to true regardless if it didn't start the Update.

Fix is CVar toggleable.

#rnx

diff --git a/Engine/Plugins/Mutable/Source/CustomizableObject/Internal/MuCO/CustomizableObjectSystemPrivate.h b/Engine/Plugins/Mutable/Source/CustomizableObject/Internal/MuCO/CustomizableObjectSystemPrivate.h
--- a/Engine/Plugins/Mutable/Source/CustomizableObject/Internal/MuCO/CustomizableObjectSystemPrivate.h
+++ b/Engine/Plugins/Mutable/Source/CustomizableObject/Internal/MuCO/CustomizableObjectSystemPrivate.h
@@ -585,6 +585,9 @@ public:
 	/** Whether the mesh to generate should support Mesh LOD streaming or not. */
 	bool bStreamMeshLODs = false;
 
+	/** true if the Update has blocked Low Priority Tasks from launching. */
+	bool bLowPriorityTasksBlocked = false;
+
 	/** This option comes from the operation request. It is used to reduce the number of mipmaps that mutable must generate for images.  */
 	int32 MipsToSkip = 0;
 
diff --git a/Engine/Plugins/Mutable/Source/CustomizableObject/Private/MuCO/CustomizableObjectSystem.cpp b/Engine/Plugins/Mutable/Source/CustomizableObject/Private/MuCO/CustomizableObjectSystem.cpp
--- a/Engine/Plugins/Mutable/Source/CustomizableObject/Private/MuCO/CustomizableObjectSystem.cpp
+++ b/Engine/Plugins/Mutable/Source/CustomizableObject/Private/MuCO/CustomizableObjectSystem.cpp
@@ -143,6 +143,11 @@ TAutoConsoleVariable<bool> CVarEnableReleaseMeshResources(
 	true,
 	TEXT("Allow releasing resources when discarding instances."));
 
+TAutoConsoleVariable<bool> CVarFixLowPriorityTasksOverlap(
+	TEXT("mutable.rollback.FixLowPriorityTasksOverlap"),
+	true,
+	TEXT("If true, use code that fixes the Low Priority Tasks overlap."));
+
 int32 UCustomizableObjectSystemPrivate::SkeletalMeshMinLodQualityLevel = -1;
 
 
@@ -889,10 +894,19 @@ void FinishUpdateGlobal(const TSharedRef<FUpdateContextPrivate>& Context)
 	Context->UpdateCallback.ExecuteIfBound(ContextPublic);
 	Context->UpdateNativeCallback.Broadcast(ContextPublic);
 
-
-	if (SystemPrivate)
+	if (CVarFixLowPriorityTasksOverlap.GetValueOnGameThread())
+	{
+		if (SystemPrivate && Context->bLowPriorityTasksBlocked)
+		{
+			SystemPrivate->MutableTaskGraph.AllowLaunchingMutableTaskLowPriority(true, false);
+		}
+	}
+	else
 	{
-		SystemPrivate->MutableTaskGraph.AllowLaunchingMutableTaskLowPriority(true, false);
+		if (SystemPrivate)
+		{
+			SystemPrivate->MutableTaskGraph.AllowLaunchingMutableTaskLowPriority(true, false);
+		}
 	}
 	
 	if (Context->StartUpdateTime != 0.0) // Update started.
@@ -2993,15 +3007,13 @@ namespace impl
 
 		// Next Task: Release Mutable. We need this regardless if we cancel or not
 		//-------------------------------------------------------------		
-		{
-			const TSharedPtr<mu::FSystem> MutableSystem = CustomizableObjectSystemPrivateData->MutableSystem;
-			CustomizableObjectSystemPrivateData->MutableTaskGraph.AddMutableThreadTask(
-				TEXT("Task_Mutable_ReleaseInstance"),
-				[InstanceID = OperationData->InstanceID, MutableSystem, bLiveUpdateMode = OperationData->bLiveUpdateMode]()
-				{
-					Task_Mutable_ReleaseInstance(InstanceID, MutableSystem, bLiveUpdateMode);
-				});
-		}
+		const  TSharedPtr<mu::FSystem> MutableSystem = CustomizableObjectSystemPrivateData->MutableSystem;
+		UE::Tasks::FTask Mutable_ReleaseInstanceTask = CustomizableObjectSystemPrivateData->MutableTaskGraph.AddMutableThreadTask(
+			TEXT("Task_Mutable_ReleaseInstance"),
+			[InstanceID = OperationData->InstanceID, MutableSystem, bLiveUpdateMode = OperationData->bLiveUpdateMode]()
+			{
+				Task_Mutable_ReleaseInstance(InstanceID, MutableSystem, bLiveUpdateMode);
+			});
 
 
 		// Next Task: Release Platform Data
@@ -3039,6 +3051,12 @@ namespace impl
 
 			// Next Task: Callbacks
 			//-------------------------------------------------------------
+			TArray<UE::Tasks::FTask, TFixedAllocator<2>> Dependencies;
+			if (CVarFixLowPriorityTasksOverlap.GetValueOnGameThread())
+			{
+				Dependencies.Add(Mutable_ReleaseInstanceTask);
+			}
+			
 			CustomizableObjectSystemPrivateData->AddGameThreadTask(
 				{
 				FMutableTaskDelegate::CreateLambda(
@@ -3046,7 +3064,7 @@ namespace impl
 					{
 						Task_Game_Callbacks(OperationData);
 					}),
-					{}
+					 Dependencies 
 				});
 		}
 		else
@@ -3293,6 +3311,7 @@ namespace impl
 
 		Operation->StartUpdateTime = FPlatformTime::Seconds();
 
+		Operation->bLowPriorityTasksBlocked = true;
 		UCustomizableObjectSystem::GetInstance()->GetPrivate()->MutableTaskGraph.AllowLaunchingMutableTaskLowPriority(false, false);
 		
 		UCustomizableObjectSystem* System = UCustomizableObjectSystem::GetInstance();
diff --git a/Engine/Plugins/Mutable/Source/CustomizableObject/Private/MuCO/FMutableTaskGraph.cpp b/Engine/Plugins/Mutable/Source/CustomizableObject/Private/MuCO/FMutableTaskGraph.cpp
--- a/Engine/Plugins/Mutable/Source/CustomizableObject/Private/MuCO/FMutableTaskGraph.cpp
+++ b/Engine/Plugins/Mutable/Source/CustomizableObject/Private/MuCO/FMutableTaskGraph.cpp
@@ -3,6 +3,7 @@
 #include "MuCO/FMutableTaskGraph.h"
 
 #include "MuCO/CustomizableObject.h"
+#include "MuCO/CustomizableObjectSystem.h"
 #include "MuR/MutableTrace.h"
 
 
@@ -159,20 +160,46 @@ void FMutableTaskGraph::TryLaunchMutableTaskLowPriority(bool bFromMutableTask)
 			return;
 		}
 
-		LastMutableTaskLowPriorityID = NextTask.Id;
-		LastMutableTaskLowPriority = AddMutableThreadTask(*NextTask.DebugName, [this, Task = MoveTemp(NextTask)]() // Moves the task, not the pointer.
+		if (CVarFixLowPriorityTasksOverlap.GetValueOnAnyThread())
 		{
-			MUTABLE_CPUPROFILER_SCOPE(LowPriorityTaskBody)
-			Task.Body();
+			LastMutableTaskLowPriorityID = NextTask.Id;
+			LastMutableTaskLowPriority = AddMutableThreadTask(*NextTask.DebugName, [this, Task = MoveTemp(NextTask)]() // Moves the task, not the pointer.
+			{
+				MUTABLE_CPUPROFILER_SCOPE(LowPriorityTaskBody)
+				Task.Body();
+			});
 
+			const FString TaskName = NextTask.DebugName + TEXT(" End");
+			AddMutableThreadTask(*TaskName, [this]()
 			{
+				MUTABLE_CPUPROFILER_SCOPE(LowPriorityTaskBodyEnd)
+
 				FScopeLock Lock(&MutableTaskLock);
 				LastMutableTaskLowPriorityID = INVALID_ID;
 				LastMutableTaskLowPriority = {};
-				
+
 				TryLaunchMutableTaskLowPriority(true);
-			}
-		});
+			},
+			{ LastMutableTaskLowPriority });
+
+		}
+		else
+		{
+			LastMutableTaskLowPriorityID = NextTask.Id;
+			LastMutableTaskLowPriority = AddMutableThreadTask(*NextTask.DebugName, [this, Task = MoveTemp(NextTask)]() // Moves the task, not the pointer.
+			{
+				MUTABLE_CPUPROFILER_SCOPE(LowPriorityTaskBody)
+				Task.Body();
+
+				{
+					FScopeLock Lock(&MutableTaskLock);
+					LastMutableTaskLowPriorityID = INVALID_ID;
+					LastMutableTaskLowPriority = {};
+				
+					TryLaunchMutableTaskLowPriority(true);
+				}
+			});
+		}
 
 		QueueMutableTasksLowPriority.RemoveAt(0);
 	}
diff --git a/Engine/Plugins/Mutable/Source/CustomizableObject/Public/MuCO/CustomizableObjectSystem.h b/Engine/Plugins/Mutable/Source/CustomizableObject/Public/MuCO/CustomizableObjectSystem.h
--- a/Engine/Plugins/Mutable/Source/CustomizableObject/Public/MuCO/CustomizableObjectSystem.h
+++ b/Engine/Plugins/Mutable/Source/CustomizableObject/Public/MuCO/CustomizableObjectSystem.h
@@ -43,6 +43,8 @@ extern TAutoConsoleVariable<bool> CVarRollbackFixModelDiskStreamerDataRace;
 
 extern TAutoConsoleVariable<bool> CVarEnableReleaseMeshResources;
 
+extern TAutoConsoleVariable<bool> CVarFixLowPriorityTasksOverlap;
+
 
 #if WITH_EDITOR
 
