[Mutable] Add FTextureLODGroup info to the DDC key.

#rnx

diff --git a/Engine/Plugins/Mutable/Source/CustomizableObject/Private/MuCO/CustomizableObject.cpp b/Engine/Plugins/Mutable/Source/CustomizableObject/Private/MuCO/CustomizableObject.cpp
--- a/Engine/Plugins/Mutable/Source/CustomizableObject/Private/MuCO/CustomizableObject.cpp
+++ b/Engine/Plugins/Mutable/Source/CustomizableObject/Private/MuCO/CustomizableObject.cpp
@@ -44,6 +44,7 @@
 
 #if WITH_EDITOR
 #include "Editor.h"
+#include "Engine/TextureLODSettings.h"
 
 #include "DerivedDataCache.h"
 #include "DerivedDataCacheInterface.h"
@@ -3071,6 +3072,34 @@ void SerializeCompilationOptionsForDDC(FArchive& Ar, FCompilationOptions& Option
 }
 
 
+void SerializeTextureGroupSettingsForDDC(FArchive& Ar, const ITargetPlatform& Platform)
+{
+	const UTextureLODSettings& LODSettings = Platform.GetTextureLODSettings();
+	for (FTextureLODGroup LODGroup : LODSettings.TextureLODGroups)
+	{
+		Ar << LODGroup.Filter;
+		Ar << LODGroup.Group;
+		Ar << LODGroup.LODBias;
+		Ar << LODGroup.LODBias_Smaller;
+		Ar << LODGroup.LODBias_Smallest;
+		Ar << LODGroup.LossyCompressionAmount;
+		Ar << LODGroup.MaxAniso;
+		Ar << LODGroup.MaxLODMipCount;
+		Ar << LODGroup.MaxLODSize;
+		Ar << LODGroup.MaxLODSize_Smaller;
+		Ar << LODGroup.MaxLODSize_Smallest;
+		Ar << LODGroup.MinLODSize;
+
+		FString MipFilter = LODGroup.MipFilter.ToString().ToLower();
+		Ar << MipFilter;
+		Ar << LODGroup.MipGenSettings;
+		Ar << LODGroup.MipLoadOptions;
+		Ar << LODGroup.NumStreamedMips;
+		Ar << LODGroup.OptionalLODBias;
+		Ar << LODGroup.OptionalMaxLODSize;
+	}
+}
+
 
 TArray<uint8> UCustomizableObjectPrivate::BuildDerivedDataKey(FCompilationOptions Options)
 {
@@ -3116,6 +3145,11 @@ TArray<uint8> UCustomizableObjectPrivate::BuildDerivedDataKey(FCompilationOption
 		SerializeCompilationOptionsForDDC(Ar, Options);
 	}
 
+	// Texture Settings
+	{
+		SerializeTextureGroupSettingsForDDC(Ar, *Options.TargetPlatform);
+	}
+
 	// Release Version
 	if (const ICustomizableObjectEditorModule* Module = ICustomizableObjectEditorModule::Get())
 	{
