[mutable] Fixed lod restriction settings for cases where no lods should be discarded
- We were discarding the two top lods when no lods should have been discarded.
- Fixed typo
- Removed unnecessary includes

#rnx

diff --git a/Engine/Plugins/Mutable/Source/CustomizableObject/Internal/MuCO/CustomizableObjectInstancePrivate.h b/Engine/Plugins/Mutable/Source/CustomizableObject/Internal/MuCO/CustomizableObjectInstancePrivate.h
--- a/Engine/Plugins/Mutable/Source/CustomizableObject/Internal/MuCO/CustomizableObjectInstancePrivate.h
+++ b/Engine/Plugins/Mutable/Source/CustomizableObject/Internal/MuCO/CustomizableObjectInstancePrivate.h
@@ -66,7 +66,7 @@ enum ECOInstanceFlags
 	UsedByComponent					= 1 << 5,	// If any components are using this instance, they will set flag every frame
 	UsedByComponentInPlay			= 1 << 6,	// If any components are using this instance in play, they will set flag every frame
 	UsedByPlayerOrNearIt			= 1 << 7,	// The instance is used by the player or is near the player, used to give more priority to its updates
-	DiscardedByNumInstancesLimit	= 1 << 8,	// The instance is descarded because we exceeded the limit of instances generated 
+	DiscardedByNumInstancesLimit	= 1 << 8,	// The instance is discarded because we exceeded the limit of instances generated 
 
 	// Types of updates
 	PendingLODsUpdate				= 1 << 9,	// Used to queue an update due to a change in LODs required by the instance
@@ -227,7 +227,7 @@ public:
 	// \TODO: Create a key based on a struct instead of generating strings dynamically.
 	UPROPERTY(Transient)
 	TMap<FString, TWeakObjectPtr<UTexture2D>> TextureReuseCache;
-	
+
 	UE::Tasks::FTaskEvent AssetAsyncLoadCompletionEvent = UE::Tasks::FTaskEvent(TEXT(""));
 
 	// Only used in LiveUpdateMode to reuse core instances between updates and their temp data to speed up updates, but spend way more memory
@@ -310,7 +310,7 @@ private:
 	// 
 	FCustomizableInstanceComponentData* GetComponentData(FCustomizableObjectComponentIndex ObjectComponentIndex);
 	const FCustomizableInstanceComponentData* GetComponentData(FCustomizableObjectComponentIndex ObjectComponentIndex) const;
-
+	
 	//
 	USkeleton* MergeSkeletons(UCustomizableObject& CustomizableObject, const FMutableRefSkeletalMeshData& RefSkeletalMeshData, FCustomizableObjectComponentIndex ObjectComponentIndex, bool& bOutCreatedNewSkeleton);
 
@@ -341,7 +341,7 @@ public:
 	int32 GetState() const;
 
 	void SetState(int32 InState);
-	
+
 	FCustomizableObjectInstanceDescriptor& GetDescriptor() const;
 
 	void SetDescriptor(const FCustomizableObjectInstanceDescriptor& InDescriptor);
diff --git a/Engine/Plugins/Mutable/Source/CustomizableObject/Private/MuCO/CustomizableInstanceLODManagement.cpp b/Engine/Plugins/Mutable/Source/CustomizableObject/Private/MuCO/CustomizableInstanceLODManagement.cpp
--- a/Engine/Plugins/Mutable/Source/CustomizableObject/Private/MuCO/CustomizableInstanceLODManagement.cpp
+++ b/Engine/Plugins/Mutable/Source/CustomizableObject/Private/MuCO/CustomizableInstanceLODManagement.cpp
@@ -23,9 +23,6 @@
 
 #if WITH_EDITOR
 #include "LevelEditorViewport.h"
-#else
-#include "Camera/CameraComponent.h"
-#include "PrimitiveDrawingUtils.h"
 #endif
 
 static TAutoConsoleVariable<int32> CVarNumGeneratedInstancesLimit(
@@ -253,16 +250,16 @@ void UCustomizableInstanceLODManagement::UpdateInstanceDistsAndLODs(FMutableInst
 			for (int32 i = 0; i < NumGeneratedInstancesLimitLODs && i < SortedInstances.Num(); ++i)
 			{
 				const UCustomizableObject* CO = SortedInstances[i]->GetCustomizableObject();
-				int32 NumObjectComponents = CO ? CO->GetComponentCount() : 0;
-				RequestedLODs.Init(MAX_uint16, NumObjectComponents);
+				const int32 NumObjectComponents = CO ? CO->GetComponentCount() : 0;
+				RequestedLODs.Init(0, NumObjectComponents);
 
 				if (!bAlreadyReachedFixedLOD && SortedInstances[i]->GetMinSquareDistToPlayer() < DistanceForFixedLODSquared)
 				{
-					SortedInstances[i]->SetRequestedLODs(0, 0, RequestedLODs, InOutRequestedUpdates);
+					SortedInstances[i]->SetRequestedLODs(0, INDEX_NONE, RequestedLODs, InOutRequestedUpdates);
 				}
 				else
 				{
-					SortedInstances[i]->SetRequestedLODs(2, 0, RequestedLODs, InOutRequestedUpdates);
+					SortedInstances[i]->SetRequestedLODs(MAX_int32, INDEX_NONE, RequestedLODs, InOutRequestedUpdates);
 					bAlreadyReachedFixedLOD = true;
 				}
 			}
@@ -270,16 +267,16 @@ void UCustomizableInstanceLODManagement::UpdateInstanceDistsAndLODs(FMutableInst
 			for (int32 i = NumGeneratedInstancesLimitLODs; i < NumGeneratedInstancesLimitLODs + NumGeneratedInstancesLimitLOD1 && i < SortedInstances.Num(); ++i)
 			{
 				const UCustomizableObject* CO = SortedInstances[i]->GetCustomizableObject();
-				int32 NumObjectComponents = CO ? CO->GetComponentCount() : 0;
-				RequestedLODs.Init(MAX_uint16, NumObjectComponents);
+				const int32 NumObjectComponents = CO ? CO->GetComponentCount() : 0;
+				RequestedLODs.Init(0, NumObjectComponents);
 
 				if (!bAlreadyReachedFixedLOD && SortedInstances[i]->GetMinSquareDistToPlayer() < DistanceForFixedLODSquared)
 				{
-					SortedInstances[i]->SetRequestedLODs(1, 0, RequestedLODs, InOutRequestedUpdates);
+					SortedInstances[i]->SetRequestedLODs(1, INDEX_NONE, RequestedLODs, InOutRequestedUpdates);
 				}
 				else
 				{
-					SortedInstances[i]->SetRequestedLODs(2, 0, RequestedLODs, InOutRequestedUpdates);
+					SortedInstances[i]->SetRequestedLODs(MAX_int32, INDEX_NONE, RequestedLODs, InOutRequestedUpdates);
 					bAlreadyReachedFixedLOD = true;
 				}
 			}
@@ -287,9 +284,9 @@ void UCustomizableInstanceLODManagement::UpdateInstanceDistsAndLODs(FMutableInst
 			for (int32 i = NumGeneratedInstancesLimitLODs + NumGeneratedInstancesLimitLOD1; i < NumGeneratedInstancesLimitLODs + NumGeneratedInstancesLimitLOD1 + NumGeneratedInstancesLimitLOD2 && i < SortedInstances.Num(); ++i)
 			{
 				const UCustomizableObject* CO = SortedInstances[i]->GetCustomizableObject();
-				int32 NumObjectComponents = CO ? CO->GetComponentCount() : 0;
-				RequestedLODs.Init(MAX_uint16, NumObjectComponents);
-				SortedInstances[i]->SetRequestedLODs(2, 0, RequestedLODs, InOutRequestedUpdates);
+				const int32 NumObjectComponents = CO ? CO->GetComponentCount() : 0;
+				RequestedLODs.Init(0, NumObjectComponents);
+				SortedInstances[i]->SetRequestedLODs(2, INDEX_NONE, RequestedLODs, InOutRequestedUpdates);
 			}
 
 			for (int32 i = NumGeneratedInstancesLimitLODs + NumGeneratedInstancesLimitLOD1 + NumGeneratedInstancesLimitLOD2; i < SortedInstances.Num(); ++i)
@@ -303,9 +300,9 @@ void UCustomizableInstanceLODManagement::UpdateInstanceDistsAndLODs(FMutableInst
 			for (int32 i = 0; i < SortedInstances.Num(); ++i)
 			{
 				const UCustomizableObject* CO = SortedInstances[i]->GetCustomizableObject();
-				int32 NumObjectComponents = CO ? CO->GetComponentCount() : 0;
-				RequestedLODs.Init(MAX_uint16, NumObjectComponents);
-				SortedInstances[i]->SetRequestedLODs(2, 0, RequestedLODs, InOutRequestedUpdates);
+				const int32 NumObjectComponents = CO ? CO->GetComponentCount() : 0;
+				RequestedLODs.Init(0, NumObjectComponents);
+				SortedInstances[i]->SetRequestedLODs(0, INDEX_NONE, RequestedLODs, InOutRequestedUpdates);
 			}
 		}
 	}
