[mutable] More optimization in fixed-memory cost for COs.
- Separate compile-time data.
- Always use rom index as resource id, and remove resource id.
#rnx

diff --git a/Engine/Plugins/Mutable/Source/CustomizableObject/Internal/MuCO/CustomizableObjectPrivate.h b/Engine/Plugins/Mutable/Source/CustomizableObject/Internal/MuCO/CustomizableObjectPrivate.h
--- a/Engine/Plugins/Mutable/Source/CustomizableObject/Internal/MuCO/CustomizableObjectPrivate.h
+++ b/Engine/Plugins/Mutable/Source/CustomizableObject/Internal/MuCO/CustomizableObjectPrivate.h
@@ -507,6 +507,7 @@ struct CUSTOMIZABLEOBJECT_API FMutableSkinWeightProfileInfo
 };
 
 
+// TODO: Optimize this struct
 USTRUCT()
 struct CUSTOMIZABLEOBJECT_API FMutableStreamableBlock
 {
@@ -753,6 +754,14 @@ struct CUSTOMIZABLEOBJECT_API FMutableStateData
 };
 
 
+/** This is encoded in exact bits so if extended, review its uses everywhere. */
+enum class EMutableFileFlags : uint8
+{
+	None	= 0,
+	HighRes = 1 << 0
+};
+
+
 #if WITH_EDITOR
 namespace MutablePrivate
 {
@@ -799,7 +808,7 @@ namespace MutablePrivate
 		/** Rom ResourceType. */
 		uint16 ResourceType = 0;
 
-		/** Common flags of the data stored in this file. See mu::ERomFlags. */
+		/** Common flags of the data stored in this file. See EMutableFileFlags.*/
 		uint16 Flags = 0;
 
 		/** Id generated from a hash of the file content + offset to avoid collisions. */
@@ -1630,6 +1639,12 @@ public:
 
 		RestructureConstantMeshData,
 
+		RestructureRomData,
+
+		RestructureRomDataRemovingRomHash,
+
+		ModifiedRomCompiledDataSerialization,
+
 		// -----<new versions can be added above this line>--------
 		LastCustomizableObjectVersion
 	};
diff --git a/Engine/Plugins/Mutable/Source/CustomizableObject/Internal/MuCO/UnrealMutableModelDiskStreamer.h b/Engine/Plugins/Mutable/Source/CustomizableObject/Internal/MuCO/UnrealMutableModelDiskStreamer.h
--- a/Engine/Plugins/Mutable/Source/CustomizableObject/Internal/MuCO/UnrealMutableModelDiskStreamer.h
+++ b/Engine/Plugins/Mutable/Source/CustomizableObject/Internal/MuCO/UnrealMutableModelDiskStreamer.h
@@ -129,7 +129,7 @@ public:
 	FUnrealMutableModelBulkWriterEditor(FArchive* InMainDataArchive = nullptr, FArchive* InStreamedDataArchive = nullptr);
 
 	// mu::FModelWriter interface
-	void OpenWriteFile(uint32 BlockKey) override;
+	void OpenWriteFile(uint32 BlockKey, bool bIsStreamable) override;
 	void Write(const void* pBuffer, uint64 size) override;
 	void CloseWriteFile() override;
 
@@ -154,7 +154,7 @@ public:
 	FUnrealMutableModelBulkWriterCook(FArchive* InMainDataArchive = nullptr, MutablePrivate::FModelStreamableData* InStreamedData = nullptr);
 
 	// mu::FModelWriter interface
-	void OpenWriteFile(uint32 BlockKey) override;
+	void OpenWriteFile(uint32 BlockKey, bool bIsStreamable) override;
 	void Write(const void* pBuffer, uint64 size) override;
 	void CloseWriteFile() override;
 
@@ -167,7 +167,7 @@ protected:
 	MutablePrivate::FModelStreamableData* StreamedData = nullptr;
 
 	uint32 CurrentKey = 0;
-
+	bool CurrentIsStreamable = false;
 };
 
 
diff --git a/Engine/Plugins/Mutable/Source/CustomizableObject/Private/MuCO/CustomizableObject.cpp b/Engine/Plugins/Mutable/Source/CustomizableObject/Private/MuCO/CustomizableObject.cpp
--- a/Engine/Plugins/Mutable/Source/CustomizableObject/Private/MuCO/CustomizableObject.cpp
+++ b/Engine/Plugins/Mutable/Source/CustomizableObject/Private/MuCO/CustomizableObject.cpp
@@ -225,7 +225,7 @@ void UCustomizableObject::PreSave(FObjectPreSaveContext ObjectSaveContext)
 						ByteBulkData.Unlock();
 
 						uint32 BulkDataFlags = BULKDATA_PayloadInSeperateFile | BULKDATA_Force_NOT_InlinePayload;
-						if (File.Flags == uint16(mu::ERomFlags::HighRes))
+						if (File.Flags == uint16(EMutableFileFlags::HighRes))
 						{
 							BulkDataFlags |= BULKDATA_OptionalPayload;
 						}
@@ -1213,12 +1213,13 @@ void UCustomizableObjectPrivate::SaveEmbeddedData(FArchive& Ar)
 
 	if (Model)
 	{	
-		// Serialise the entire model, but unload the streamable data first.
+		if (Ar.IsCooking())
 		{
-			UnrealMutableOutputStream Stream(Ar);
-			mu::FOutputArchive Arch(&Stream);
-			mu::FModel::Serialise(Model.Get(), Arch);
+			Model->GetPrivate()->Program.RomsCompileData.Empty();
 		}
+		UnrealMutableOutputStream Stream(Ar);
+		mu::FOutputArchive Arch(&Stream);
+		mu::FModel::Serialise(Model.Get(), Arch);
 
 		UE_LOG(LogMutable, Verbose, TEXT("Saved embedded data for Customizable Object [%s] now at position %d."), *GetName(), int(Ar.Tell()));
 	}
@@ -2646,15 +2647,14 @@ namespace MutablePrivate
 			DefaultBucket.Categories.Add({ FFileCategoryID(EDataType::Model, mu::DATATYPE::DT_IMAGE, 0), 0, 0, 0 });
 
 			// Add High-res textures to the Optional bucket
-			OptionalBucket.Categories.Add({ FFileCategoryID(EDataType::Model, mu::DATATYPE::DT_IMAGE, (uint16)mu::ERomFlags::HighRes), 0, 0, 0});
+			OptionalBucket.Categories.Add({ FFileCategoryID(EDataType::Model, mu::DATATYPE::DT_IMAGE, uint16(EMutableFileFlags::HighRes)), 0, 0, 0});
 
 			const int32 NumRoms = Model->GetRomCount();
 			for (int32 RomIndex = 0; RomIndex < NumRoms; ++RomIndex)
 			{
-				uint32 BlockId = Model->GetRomId(RomIndex);
 				const uint32 BlockSize = Model->GetRomSize(RomIndex);
-				const uint16 BlockResourceType = Model->GetRomType(RomIndex);
-				const mu::ERomFlags BlockFlags = Model->GetRomFlags(RomIndex);
+				const uint16 BlockResourceType = Model->IsMeshData(RomIndex) ? mu::DATATYPE::DT_MESH : mu::DATATYPE::DT_IMAGE;
+				const EMutableFileFlags BlockFlags = Model->IsRomHighRes(RomIndex) ? EMutableFileFlags::HighRes : EMutableFileFlags::None;
 
 				FFileCategoryID CategoryID = { EDataType::Model, BlockResourceType, uint16(BlockFlags) };
 				TPair<FFileBucket&, FFileCategory&> It = FindOrAddCategory(FileBuckets, DefaultBucket, CategoryID); // Add block to an existing or new category
@@ -2862,27 +2862,30 @@ namespace MutablePrivate
 			return;
 		}
 
-		// TODO: Temp. Remove after unifying generated output files code between editor an package. UE-222777 
-		const uint16 IgnoreMask = bUseRomTypeAndFlagsToFilter ? MAX_uint16 : 0;
-
 		// Root nodes by flags.
-		const int32 NumRoms = Model->GetRomCount();
+		const uint32 NumRoms = Model->GetRomCount();
 		TMap<FFileCategoryID, FClassifyNode> RootNode;
 
 		// Create blocks data.
 		{
-			for (int32 RomIndex = 0; RomIndex < NumRoms; ++RomIndex)
+			for (uint32 RomIndex = 0; RomIndex < NumRoms; ++RomIndex)
 			{
-				uint32 BlockId = Model->GetRomId(RomIndex);
 				uint32 SourceBlockId = Model->GetRomSourceId(RomIndex);
 				const uint32 BlockSize = Model->GetRomSize(RomIndex);
-				const uint16 BlockResourceType = IgnoreMask & Model->GetRomType(RomIndex);
-				const mu::ERomFlags BlockFlags = (mu::ERomFlags)(IgnoreMask & (uint16)Model->GetRomFlags(RomIndex));
+				
+				uint16 BlockResourceType = 0;
+				EMutableFileFlags BlockFlags = EMutableFileFlags::None;
+				// TODO: Temp. Remove after unifying generated output files code between editor an package. UE-222777 
+				if (bUseRomTypeAndFlagsToFilter)
+				{
+					BlockResourceType = Model->IsMeshData(RomIndex) ? mu::DATATYPE::DT_MESH : mu::DATATYPE::DT_IMAGE;
+					BlockFlags = Model->IsRomHighRes(RomIndex) ? EMutableFileFlags::HighRes : EMutableFileFlags::None;
+				}
 
 				FFileCategoryID CurrentCategory = { EDataType::Model, BlockResourceType, uint16(BlockFlags) };
 				FBlock CurrentBlock = 
 				{ 
-					.Id = BlockId, 
+					.Id = RomIndex, 
 					.SourceId = SourceBlockId, 
 					.Size = BlockSize, 
 					.Offset = 0 
@@ -2895,7 +2898,7 @@ namespace MutablePrivate
 		{
 			const TMap<uint32, FRealTimeMorphStreamable>& RealTimeMorphStreamables = ModelStreamableBulkData.RealTimeMorphStreamables;
 
-			const FFileCategoryID RealTimeMorphCategory = { EDataType::RealTimeMorph, (uint16)mu::DATATYPE::DT_NONE, (uint16)mu::ERomFlags::None };
+			const FFileCategoryID RealTimeMorphCategory = { EDataType::RealTimeMorph, (uint16)mu::DATATYPE::DT_NONE, (uint16)EMutableFileFlags::None };
 
 			for (const TPair<uint32, FRealTimeMorphStreamable>& MorphStreamable : RealTimeMorphStreamables)
 			{
@@ -2916,7 +2919,7 @@ namespace MutablePrivate
 		{
 			const TMap<uint32, FClothingStreamable>& ClothingStreamables = ModelStreamableBulkData.ClothingStreamables;
 
-			const FFileCategoryID ClothingCategory = { EDataType::Clothing, (uint16)mu::DATATYPE::DT_NONE, (uint16)mu::ERomFlags::None };
+			const FFileCategoryID ClothingCategory = { EDataType::Clothing, (uint16)mu::DATATYPE::DT_NONE, (uint16)EMutableFileFlags::None };
 
 			for (const TPair<uint32, FClothingStreamable>& ClothStreamable : ClothingStreamables)
 			{
@@ -3481,7 +3484,7 @@ void UCustomizableObjectPrivate::LogMemory()
 			SkeletonsTotal += sizeof(mu::FSkeleton) + Entry->DebugBoneNames.GetAllocatedSize() + Entry->BoneIds.GetAllocatedSize() + Entry->BoneParents.GetAllocatedSize();
 		}
 
-		int32 PhysicsTotal = Program.ConstantSkeletons.GetAllocatedSize();
+		int32 PhysicsTotal = Program.ConstantPhysicsBodies.GetAllocatedSize();
 		for (const TSharedPtr<const mu::FPhysicsBody>& Entry : Program.ConstantPhysicsBodies)
 		{
 			PhysicsTotal += sizeof(mu::FPhysicsBody) + Entry->BoneIds.GetAllocatedSize() + Entry->Bodies.GetAllocatedSize() + Entry->BodiesCustomIds.GetAllocatedSize();
@@ -3597,7 +3600,7 @@ TUniquePtr<IAsyncReadFileHandle> UCustomizableObjectBulk::OpenFileAsyncRead(uint
 	check(IsInGameThread());
 
 	FString FilePath = FString::Printf(TEXT("%s-%08x.mut"), *BulkFilePrefix, FileId);
-	if (Flags == uint32(mu::ERomFlags::HighRes))
+	if (Flags == uint16(EMutableFileFlags::HighRes))
 	{
 		FilePath += TEXT(".high");
 	}
@@ -3640,7 +3643,7 @@ void UCustomizableObjectBulk::CookAdditionalFilesOverride(const TCHAR* PackageFi
 		{
 			FString FileName = CookedBulkFileName + FString::Printf(TEXT("-%08x.mut"), File.Id);
 
-			if (File.Flags == uint32(mu::ERomFlags::HighRes))
+			if (File.Flags == uint16(EMutableFileFlags::HighRes))
 			{
 				// We can do something different here for high-res data.
 				// For example: change the file name. We also need to detect it when generating the file name for loading.
diff --git a/Engine/Plugins/Mutable/Source/CustomizableObject/Private/MuCO/UnrealMutableModelDiskStreamer.cpp b/Engine/Plugins/Mutable/Source/CustomizableObject/Private/MuCO/UnrealMutableModelDiskStreamer.cpp
--- a/Engine/Plugins/Mutable/Source/CustomizableObject/Private/MuCO/UnrealMutableModelDiskStreamer.cpp
+++ b/Engine/Plugins/Mutable/Source/CustomizableObject/Private/MuCO/UnrealMutableModelDiskStreamer.cpp
@@ -363,7 +363,7 @@ mu::FModelReader::FOperationID FUnrealMutableModelBulkReader::BeginReadBlock(con
 					FString FilePath = ObjectData->BulkFilePrefix;
 #else
 					FString FilePath = FString::Printf(TEXT("%s-%08x.mut"), *ObjectData->BulkFilePrefix, Block->FileId);
-					if (Block->Flags == uint32(mu::ERomFlags::HighRes))
+					if (Block->Flags == uint16(EMutableFileFlags::HighRes))
 					{
 						FilePath += TEXT(".high");
 					}
@@ -499,9 +499,9 @@ FUnrealMutableModelBulkWriterEditor::FUnrealMutableModelBulkWriterEditor(FArchiv
 }
 
 
-void FUnrealMutableModelBulkWriterEditor::OpenWriteFile(uint32 BlockKey)
+void FUnrealMutableModelBulkWriterEditor::OpenWriteFile(uint32 BlockKey, bool bIsStreamable)
 {
-	if (BlockKey == 0) // Model
+	if (!bIsStreamable) // Model
 	{
 		check(MainDataArchive);
 		CurrentWriteFile = MainDataArchive;
@@ -534,18 +534,20 @@ FUnrealMutableModelBulkWriterCook::FUnrealMutableModelBulkWriterCook(FArchive* I
 	MainDataArchive = InMainDataArchive;
 	StreamedData = InStreamedData;
 	CurrentKey = 0;
+	CurrentIsStreamable = false;
 }
 
 
-void FUnrealMutableModelBulkWriterCook::OpenWriteFile(uint32 BlockKey)
+void FUnrealMutableModelBulkWriterCook::OpenWriteFile(uint32 BlockKey, bool bIsStreamable)
 {
 	CurrentKey = BlockKey;
+	CurrentIsStreamable = bIsStreamable;
 }
 
 
 void FUnrealMutableModelBulkWriterCook::Write(const void* pBuffer, uint64 size)
 {
-	if (CurrentKey == 0)
+	if (!CurrentIsStreamable)
 	{
 		MainDataArchive->Serialize(const_cast<void*>(pBuffer), size);
 	}
diff --git a/Engine/Plugins/Mutable/Source/CustomizableObjectEditor/Private/MuCOE/CustomizableObjectCompileRunnable.cpp b/Engine/Plugins/Mutable/Source/CustomizableObjectEditor/Private/MuCOE/CustomizableObjectCompileRunnable.cpp
--- a/Engine/Plugins/Mutable/Source/CustomizableObjectEditor/Private/MuCOE/CustomizableObjectCompileRunnable.cpp
+++ b/Engine/Plugins/Mutable/Source/CustomizableObjectEditor/Private/MuCOE/CustomizableObjectCompileRunnable.cpp
@@ -361,7 +361,7 @@ uint32 FCustomizableObjectSaveDDRunnable::Run()
 
 	if (Model)
 	{
-		CachePlatfromData();
+		CachePlatformData();
 
 		bool bStoredSuccessfully = false;
 
@@ -395,9 +395,9 @@ const ITargetPlatform* FCustomizableObjectSaveDDRunnable::GetTargetPlatform() co
 }
 
 
-void FCustomizableObjectSaveDDRunnable::CachePlatfromData()
+void FCustomizableObjectSaveDDRunnable::CachePlatformData()
 {
-	MUTABLE_CPUPROFILER_SCOPE(CachePlatfromData);
+	MUTABLE_CPUPROFILER_SCOPE(CachePlatformData);
 
 	if (!Model || !ModelStreamables)
 	{
@@ -405,7 +405,7 @@ void FCustomizableObjectSaveDDRunnable::CachePlatfromData()
 		return;
 	}
 
-	// Cache ModelStreamalbes
+	// Cache ModelStreamables
 	{
 		PlatformData.ModelStreamables = ModelStreamables;
 
diff --git a/Engine/Plugins/Mutable/Source/CustomizableObjectEditor/Private/MuCOE/CustomizableObjectCompileRunnable.h b/Engine/Plugins/Mutable/Source/CustomizableObjectEditor/Private/MuCOE/CustomizableObjectCompileRunnable.h
--- a/Engine/Plugins/Mutable/Source/CustomizableObjectEditor/Private/MuCOE/CustomizableObjectCompileRunnable.h
+++ b/Engine/Plugins/Mutable/Source/CustomizableObjectEditor/Private/MuCOE/CustomizableObjectCompileRunnable.h
@@ -105,7 +105,7 @@ public:
 
 private:
 
-	void CachePlatfromData();
+	void CachePlatformData();
 
 	void StoreCachedPlatformDataInDDC(bool& bStoredSuccessfully);
 
diff --git a/Engine/Plugins/Mutable/Source/CustomizableObjectEditor/Private/MuCOE/CustomizableObjectCompiler.cpp b/Engine/Plugins/Mutable/Source/CustomizableObjectEditor/Private/MuCOE/CustomizableObjectCompiler.cpp
--- a/Engine/Plugins/Mutable/Source/CustomizableObjectEditor/Private/MuCOE/CustomizableObjectCompiler.cpp
+++ b/Engine/Plugins/Mutable/Source/CustomizableObjectEditor/Private/MuCOE/CustomizableObjectCompiler.cpp
@@ -865,7 +865,7 @@ void FCustomizableObjectCompiler::CompileInternal(bool bAsync)
 
 			ResourceMeshData.NameResolutionMap = MeshData.Value.NameResolutionMap;
 			ResourceMeshData.Size = DataSizeInBytes;
-			mu::ERomFlags Flags = mu::ERomFlags::None;
+			EMutableFileFlags Flags = EMutableFileFlags::None;
 			ResourceMeshData.Block = FMutableStreamableBlock { uint32(0), uint32(Flags), RealTimeMorphDataOffsetInBytes };
 			ResourceMeshData.SourceId = MeshData.Value.SourceId;
 
@@ -897,7 +897,7 @@ void FCustomizableObjectCompiler::CompileInternal(bool bAsync)
 			ResourceMeshData.ClothingAssetLOD = MeshData.Value.ClothingAssetLOD;
 			ResourceMeshData.PhysicsAssetIndex = MeshData.Value.PhysicsAssetIndex;
 			ResourceMeshData.Size = DataSizeInBytes;
-			mu::ERomFlags Flags = mu::ERomFlags::None;
+			EMutableFileFlags Flags = EMutableFileFlags::None;
 			ResourceMeshData.Block = FMutableStreamableBlock{ uint32(0), uint32(Flags), ClothingDataOffsetInBytes };
 			ResourceMeshData.SourceId = MeshData.Value.SourceId;
 
@@ -1303,10 +1303,13 @@ void FCustomizableObjectCompiler::FinishCompilationTask()
 		uint64 Offset = 0;
 		for (int32 FileIndex = 0; FileIndex < NumStreamingFiles; ++FileIndex)
 		{
-			const uint32 ResourceId = Model->GetRomId(FileIndex);
 			const uint32 ResourceSize = Model->GetRomSize(FileIndex);
-			mu::ERomFlags Flags = bRequiresCookedData ? Model->GetRomFlags(FileIndex) : mu::ERomFlags::None;
-			ModelStreamables.Add(ResourceId, FMutableStreamableBlock{ 0, uint32(Flags), Offset });
+			EMutableFileFlags Flags = EMutableFileFlags::None;
+			if (bRequiresCookedData)
+			{
+				Flags = Model->IsRomHighRes(FileIndex) ? EMutableFileFlags::HighRes : EMutableFileFlags::None;
+			}
+			ModelStreamables.Add(FileIndex, FMutableStreamableBlock{ 0, uint32(Flags), Offset });
 			Offset += ResourceSize;
 		}
 
diff --git a/Engine/Plugins/Mutable/Source/MutableRuntime/Internal/MuR/ModelPrivate.h b/Engine/Plugins/Mutable/Source/MutableRuntime/Internal/MuR/ModelPrivate.h
--- a/Engine/Plugins/Mutable/Source/MutableRuntime/Internal/MuR/ModelPrivate.h
+++ b/Engine/Plugins/Mutable/Source/MutableRuntime/Internal/MuR/ModelPrivate.h
@@ -36,30 +36,42 @@ namespace mu
 	MUTABLE_DEFINE_POD_SERIALISABLE(FConstantResourceIndex);
 	MUTABLE_DEFINE_POD_VECTOR_SERIALISABLE(FConstantResourceIndex);
 
-    //! Data stored for a rom even if it is not loaded
-    struct FRomData
-    {
-		// TODO: Remove. Use Rom index instead to find the file
-        //! This is used to identify a ROM file. It is usually a hash from its data.
-        uint32 Id=0; 
-
-		uint32 SourceId=0;
+	/** This is encoded with minimal bits. Make sure to review all uses if extended. */
+	enum class ERomDataType : uint32
+	{
+		Image = 0,
+		Mesh  = 1
+	};
 
+    /** Data stored for a rom even if it is not loaded.
+	* This struct is size-sensitive since there may be many roms and it is always loaded in memory when a CO is.
+	*/
+    struct FRomDataRuntime
+    {
 		//! Size of the rom
-		uint32 Size=0;
+		uint32 Size : 30;
 
-		//! Index of the resource in its type-specific array (one of DATATYPE values)
-		uint16 ResourceType=0;
+		/** Index of the resource in its type-specific array. See ERomDataType. */
+		uint32 ResourceType : 1;
 
 		/** Properties of the rom data. */
-		ERomFlags Flags = ERomFlags::None;
+		uint32 IsHighRes : 1;
 
-		// TODO: Store the offset here and delete the FModelStreamableBlock map
+		// TODO: Store the offset here and delete the FModelStreamableBlock map?
 		//! Offset in file
-		// uint64 Offset = 0;
+		// uint32 Offset = 0;
     };
 
-    MUTABLE_DEFINE_POD_SERIALISABLE(FRomData);
+	/** Not critical to keep this size, but it is memory-usage sensitive. */
+	static_assert(sizeof(FRomDataRuntime) == 4);
+	MUTABLE_DEFINE_POD_SERIALISABLE(FRomDataRuntime);
+
+	struct FRomDataCompile
+	{
+		/** ID used to identify the origin of this data and used for grouping. */
+		uint32 SourceId;
+	};
+	MUTABLE_DEFINE_POD_SERIALISABLE(FRomDataCompile);
 
     //!
     template<typename DATA>
@@ -232,8 +244,11 @@ namespace mu
         //!
 		TArray<FState> States;
 
-		//! Data for every rom.
-		TArray<FRomData> Roms;
+		/** Data for every rom required in-game. */
+		TArray<FRomDataRuntime> Roms;
+
+		/** Data for every rom required at compile-time. It is empty in cooked data. */
+		TArray<FRomDataCompile> RomsCompileData;
 	
 		/** Loaded roms worth tracking (only images and meshes for now). Stores the rom's data type.*/
 		TSparseArray<uint8> LoadedMemTrackedRoms;
@@ -317,6 +332,7 @@ namespace mu
             arch << ByteCode;
             arch << States;
 			arch << Roms;
+			arch << RomsCompileData;
 			arch << ConstantImageLODsPermanent;
 			arch << ConstantImageLODIndices;
 			arch << ConstantImages;
@@ -342,6 +358,7 @@ namespace mu
             arch >> ByteCode;
             arch >> States;
 			arch >> Roms;
+			arch >> RomsCompileData;
 			arch >> ConstantImageLODsPermanent;
 			arch >> ConstantImageLODIndices;
 			arch >> ConstantImages;
@@ -371,9 +388,9 @@ namespace mu
 		{
 			switch (Roms[RomIndex].ResourceType)
 			{
-				case DT_IMAGE: 
+				case uint32(ERomDataType::Image):
 					return ConstantImageLODsStreamed.Contains(RomIndex);
-				case DT_MESH: 
+				case uint32(ERomDataType::Mesh):
 					return ConstantMeshesStreamed.Contains(RomIndex);
 				default:
 					check(false);
@@ -398,7 +415,7 @@ namespace mu
 
 			switch (Roms[RomIndex].ResourceType)
 			{
-			case DT_IMAGE:
+			case uint32(ERomDataType::Image):
 			{
 				TSharedPtr<const FImage> Data;
 				ConstantImageLODsStreamed.RemoveAndCopyValue(RomIndex, Data);
@@ -408,7 +425,7 @@ namespace mu
 				}
 				break;
 			}
-			case DT_MESH:
+			case uint32(ERomDataType::Mesh):
 			{
 				TSharedPtr<const FMesh> Data;
 				ConstantMeshesStreamed.RemoveAndCopyValue(RomIndex, Data);
@@ -426,7 +443,7 @@ namespace mu
 
 		FORCEINLINE void SetMeshRomValue(uint32 RomIndex, const TSharedPtr<FMesh>& Value)
 		{
-			check(Roms[RomIndex].ResourceType == DT_MESH);
+			check(Roms[RomIndex].ResourceType == uint32(ERomDataType::Mesh));
 			check(!ConstantMeshesStreamed.Contains(RomIndex));
 
 			LoadedMemTrackedRoms.EmplaceAt(RomIndex, (uint8)Roms[RomIndex].ResourceType);
@@ -435,7 +452,7 @@ namespace mu
 
 		FORCEINLINE void SetImageRomValue(uint32 RomIndex, const TSharedPtr<FImage>& Value)
 		{
-			check(Roms[RomIndex].ResourceType == DT_IMAGE);
+			check(Roms[RomIndex].ResourceType == uint32(ERomDataType::Image));
 			check(!ConstantImageLODsStreamed.Contains(RomIndex));
 
 			LoadedMemTrackedRoms.EmplaceAt(RomIndex, (uint8)Roms[RomIndex].ResourceType);
diff --git a/Engine/Plugins/Mutable/Source/MutableRuntime/Private/MuR/CodeRunnerTasked.cpp b/Engine/Plugins/Mutable/Source/MutableRuntime/Private/MuR/CodeRunnerTasked.cpp
--- a/Engine/Plugins/Mutable/Source/MutableRuntime/Private/MuR/CodeRunnerTasked.cpp
+++ b/Engine/Plugins/Mutable/Source/MutableRuntime/Private/MuR/CodeRunnerTasked.cpp
@@ -2152,7 +2152,7 @@ namespace mu
 			ReadCompletionEvent.Trigger();
 		};
 		
-		const uint32 RomId = Program.Roms[RomIndex].Id;
+		const uint32 RomId = RomIndex;
 		RomLoadOp.m_streamID = Runner->System->StreamInterface->BeginReadBlock(Runner->Model.Get(), RomId, RomLoadOp.m_streamBuffer.GetData(), RomSize, &Callback);
 		if (RomLoadOp.m_streamID < 0)
 		{
@@ -2388,7 +2388,7 @@ namespace mu
 				ReadCompletionEvent.Trigger();
 			};
 			
-			const uint32 RomId = Program.Roms[RomIndex].Id;
+			const uint32 RomId = RomIndex;
 			RomLoadOp.m_streamID = Runner->System->StreamInterface->BeginReadBlock(Runner->Model.Get(), RomId, RomLoadOp.m_streamBuffer.GetData(), RomSize, &Callback);
 			if (RomLoadOp.m_streamID < 0)
 			{
diff --git a/Engine/Plugins/Mutable/Source/MutableRuntime/Private/MuR/Model.cpp b/Engine/Plugins/Mutable/Source/MutableRuntime/Private/MuR/Model.cpp
--- a/Engine/Plugins/Mutable/Source/MutableRuntime/Private/MuR/Model.cpp
+++ b/Engine/Plugins/Mutable/Source/MutableRuntime/Private/MuR/Model.cpp
@@ -22,7 +22,8 @@
 
 namespace mu
 {
-    MUTABLE_IMPLEMENT_POD_SERIALISABLE(FRomData);
+	MUTABLE_IMPLEMENT_POD_SERIALISABLE(FRomDataRuntime);
+	MUTABLE_IMPLEMENT_POD_SERIALISABLE(FRomDataCompile);
 	MUTABLE_IMPLEMENT_POD_SERIALISABLE(FImageLODRange);
 	MUTABLE_IMPLEMENT_POD_SERIALISABLE(FConstantResourceIndex);
 	MUTABLE_IMPLEMENT_POD_VECTOR_SERIALISABLE(FConstantResourceIndex);
@@ -148,8 +149,8 @@ namespace mu
 		for (TPair<uint32, TSharedPtr<const FImage>>& Entry: Program.ConstantImageLODsStreamed )
 		{
 			int32 RomIndex = Entry.Key;
-			const FRomData& RomData = Program.Roms[RomIndex];
-			check(RomData.ResourceType == DT_IMAGE);
+			const FRomDataRuntime& RomData = Program.Roms[RomIndex];
+			check(RomData.ResourceType == uint32(ERomDataType::Image));
 
 			// Serialize to memory, to find out final size of this rom
 			MemStream.Reset();
@@ -157,7 +158,7 @@ namespace mu
 			FImage::Serialise(Entry.Value.Get(), MemoryArch);
 			check(RomData.Size == MemStream.GetBufferSize());
 
-			streamer.OpenWriteFile(RomData.Id);
+			streamer.OpenWriteFile(RomIndex, true);
 			streamer.Write(MemStream.GetBuffer(), MemStream.GetBufferSize());
 			streamer.CloseWriteFile();
 
@@ -177,8 +178,8 @@ namespace mu
 		for (TPair<uint32, TSharedPtr<const FMesh>>& ResData : Program.ConstantMeshesStreamed)
 		{
 			int32 RomIndex = ResData.Key;
-			const FRomData& RomData = Program.Roms[RomIndex];
-			check(RomData.ResourceType == DT_MESH);
+			const FRomDataRuntime& RomData = Program.Roms[RomIndex];
+			check(RomData.ResourceType == uint32(ERomDataType::Mesh));
 
 			// Serialize to memory, to find out final size of this rom
 			MemStream.Reset();
@@ -186,7 +187,7 @@ namespace mu
 			FMesh::Serialise(ResData.Value.Get(), MemoryArch);
 			check(RomData.Size == MemStream.GetBufferSize());
 
-			streamer.OpenWriteFile(RomData.Id);
+			streamer.OpenWriteFile(RomIndex, true);
 			streamer.Write(MemStream.GetBuffer(), MemStream.GetBufferSize());
 			streamer.CloseWriteFile();
 
@@ -204,7 +205,7 @@ namespace mu
 
 		// Store the main data of the model
 		{
-			streamer.OpenWriteFile(0);
+			streamer.OpenWriteFile(0, false);
 			FOutputModelStream stream(&streamer);
 			FOutputArchive arch(&stream);
 
@@ -385,31 +386,26 @@ namespace mu
     	return m_pD->Program.Roms.Num();
     }
 
-
-    uint32 FModel::GetRomId(int32 Index) const
-    {
-    	return m_pD->Program.Roms[Index].Id;
-    }
-
+#if WITH_EDITOR
 	uint32 FModel::GetRomSourceId(int32 Index) const
 	{
-		return m_pD->Program.Roms[Index].SourceId;
+		return m_pD->Program.RomsCompileData[Index].SourceId;
 	}
-
+#endif
 
     uint32 FModel::GetRomSize(int32 Index) const
     {
     	return m_pD->Program.Roms[Index].Size;
     }
 
-	uint16 FModel::GetRomType(int32 Index) const
+	bool FModel::IsMeshData(int32 Index) const
 	{
-		return m_pD->Program.Roms[Index].ResourceType;
+		return m_pD->Program.Roms[Index].ResourceType==uint32(ERomDataType::Mesh);
 	}
 
-	ERomFlags FModel::GetRomFlags(int32 Index) const
+	bool FModel::IsRomHighRes(int32 Index) const
 	{
-		return m_pD->Program.Roms[Index].Flags;
+		return m_pD->Program.Roms[Index].IsHighRes==1;
 	}
 
 
diff --git a/Engine/Plugins/Mutable/Source/MutableRuntime/Private/MuR/System.cpp b/Engine/Plugins/Mutable/Source/MutableRuntime/Private/MuR/System.cpp
--- a/Engine/Plugins/Mutable/Source/MutableRuntime/Private/MuR/System.cpp
+++ b/Engine/Plugins/Mutable/Source/MutableRuntime/Private/MuR/System.cpp
@@ -1326,7 +1326,7 @@ namespace mu
 					{
 						const int32 RomIndex = Iter.GetIndex();
 
-						const FRomData& Rom = Program.Roms[RomIndex];
+						const FRomDataRuntime& Rom = Program.Roms[RomIndex];
 						check(Program.IsRomLoaded(RomIndex));
 						check(Rom.ResourceType == (uint32)*Iter);
 
diff --git a/Engine/Plugins/Mutable/Source/MutableRuntime/Public/MuR/Model.h b/Engine/Plugins/Mutable/Source/MutableRuntime/Public/MuR/Model.h
--- a/Engine/Plugins/Mutable/Source/MutableRuntime/Public/MuR/Model.h
+++ b/Engine/Plugins/Mutable/Source/MutableRuntime/Public/MuR/Model.h
@@ -17,16 +17,6 @@ namespace mu
 	class FSystem;
     class FModelParametersGenerator;
 
-	/** Flags for the model data chunks. */
-	enum class ERomFlags : uint16
-	{
-		/** Standard data. */
-		None = 0,
-
-		/** Bigger mips and mesh lods that are optional in some devices of a platform. */
-		HighRes = 1 << 0,
-	};
-
 
     /** A Model represents a customisable object with any number of parameters.
     * When values are given to the parameters, specific Instances can be built, which hold the
@@ -146,15 +136,15 @@ namespace mu
 
     	int32 GetRomCount() const;
 
-		uint32 GetRomId(int32 Index) const;
-
+#if WITH_EDITOR
 		uint32 GetRomSourceId(int32 Index) const;
+#endif
 
 		uint32 GetRomSize(int32 Index) const;
 
-		uint16 GetRomType(int32 Index) const;
+		bool IsMeshData(int32 Index) const;
 
-		ERomFlags GetRomFlags(int32 Index) const;
+		bool IsRomHighRes(int32 Index) const;
     	
 		//-----------------------------------------------------------------------------------------
 		// Interface pattern
diff --git a/Engine/Plugins/Mutable/Source/MutableRuntime/Public/MuR/Serialisation.h b/Engine/Plugins/Mutable/Source/MutableRuntime/Public/MuR/Serialisation.h
--- a/Engine/Plugins/Mutable/Source/MutableRuntime/Public/MuR/Serialisation.h
+++ b/Engine/Plugins/Mutable/Source/MutableRuntime/Public/MuR/Serialisation.h
@@ -136,23 +136,10 @@ namespace mu
 		// Writing interface
 		//-----------------------------------------------------------------------------------------
 
-		//! \brief Open a file for writing data.
-		//!
-		//! Only one file can be open a time for writing. This method is only required when
-		//! writing model data from tools.
-		//! \param strModelName Name of the model where the file to open belongs to. This string's
-		//!         meaning depends on the implementation of the ModelStreamer subclasses. It could
-		//!         be a file path, or a resource identifier, etc. It will be the same for all
-		//!         files in a model.
-		//! \param key key identifying the model data fragment that is requested.
-		//!         This key interpretation depends on the implementation of the ModelStreamer,
-		virtual void OpenWriteFile(uint32 BlockKey) = 0;
-
-		//! \brief Write a piece of data to the currently open file.
-		//!
-		//! There must be a file open with OpenWriteFile before calling this.
-		//! \param pBuffer pointer to the data that will be written to the file.
-		//! \param size size of the data to write on the file, in bytes.
+		/** */
+		virtual void OpenWriteFile(uint32 BlockKey, bool bIsStreamable) = 0;
+
+		/** */
 		virtual void Write(const void* Buffer, uint64 Size) = 0;
 
 		//! \brief Close the file open for writing in a previous call to OpenWriteFile in this
diff --git a/Engine/Plugins/Mutable/Source/MutableTools/Private/MuT/Compiler.cpp b/Engine/Plugins/Mutable/Source/MutableTools/Private/MuT/Compiler.cpp
--- a/Engine/Plugins/Mutable/Source/MutableTools/Private/MuT/Compiler.cpp
+++ b/Engine/Plugins/Mutable/Source/MutableTools/Private/MuT/Compiler.cpp
@@ -511,29 +511,6 @@ namespace mu
     }
 
 
-	//---------------------------------------------------------------------------------------------
-	namespace
-	{
-		inline void EnsureUniqueRomId(TSet<uint32>& UsedIds, uint32& RomId)
-		{
-			while (true)
-			{
-				bool bAlreadyUsed = false;
-				UsedIds.FindOrAdd( RomId, &bAlreadyUsed);
-
-				if (!bAlreadyUsed)
-				{
-					break;
-				}
-
-				// Generate a new Id. It is not going to be stable accross build, which means it may hurt content patching
-				// a little bit, but it shouldn't happen often.
-				RomId++;
-			}
-		}
-	}
-
-
 	//---------------------------------------------------------------------------------------------
 	void Compiler::Private::GenerateRoms(FModel* Model, const FLinkerOptions::FAdditionalData& AdditionalData )
 	{
@@ -559,8 +536,11 @@ namespace mu
 		TSet<uint32> UsedIds;
 		UsedIds.Reserve(MaxRomCount);
 
-		TArray<FRomData> RomDatas;
-		RomDatas.SetNumZeroed(FMath::Max(Program.ConstantImageLODsPermanent.Num(), Program.ConstantMeshesPermanent.Num()));
+		int32 MaxResources = FMath::Max(Program.ConstantImageLODsPermanent.Num(), Program.ConstantMeshesPermanent.Num());
+		TArray<FRomDataRuntime> RomDatas;
+		RomDatas.SetNumZeroed(MaxResources);
+		TArray<FRomDataCompile> RomDatasCompile;
+		RomDatasCompile.SetNumZeroed(MaxResources);
 
 		// Images
 		{
@@ -582,21 +562,11 @@ namespace mu
 					FImage::Serialise(Resource.Get(), MemoryArch);
 
 
-					FRomData& RomData = RomDatas[MipIndex];
-					RomData.ResourceType = DT_IMAGE;
-					RomData.Size = SizeStream.GetBufferSize();
-
-					// If the resource uses less memory than the threshold, don't save it in a separate rom.
-					if (RomData.Size <= EmbeddedDataBytesLimit)
-					{
-						return;
-					}
-
-					// RomData.ResourceIndex will be filled below when we split in permanent and streamable
-					// RomData.Flags will be filled below
-
-					uint64 FullHash = CityHash64(reinterpret_cast<const char*>(Resource->GetLODData(0)), Resource->GetLODDataSize(0));
-					RomData.Id = GetTypeHash(FullHash);
+					FRomDataRuntime& RomData = RomDatas[MipIndex];
+					RomData.ResourceType = uint32(ERomDataType::Image);
+					uint32 Size = SizeStream.GetBufferSize();
+					check(Size < uint32(1 << 30));
+					RomData.Size = Size;
 				});
 
 			// Generate the high-res flags for images
@@ -632,8 +602,8 @@ namespace mu
 					// If this mip represents a high-quality mip, flag the rom as such
 					if (IsLODHighRes[MipIndex])
 					{
-						FRomData& RomData = RomDatas[MipIndex];
-						EnumAddFlags(RomData.Flags, ERomFlags::HighRes);
+						FRomDataRuntime& RomData = RomDatas[MipIndex];
+						RomData.IsHighRes = true;
 
 						++NumHighRes;
 						NumHighResBytes += RomData.Size;
@@ -654,7 +624,7 @@ namespace mu
 					{
 						FConstantResourceIndex ResourceIndex = Program.ConstantImageLODIndices[LODRange.FirstIndex + LODRangeIndex];
 						check(!ResourceIndex.Streamable); // Not classified yet
-						RomDatas[ResourceIndex.Index].SourceId = SourceId;
+						RomDatasCompile[ResourceIndex.Index].SourceId = SourceId;
 					}
 				}
 			}
@@ -668,15 +638,12 @@ namespace mu
 
 				for (int32 MipIndex = 0; MipIndex < AllMips.Num(); ++MipIndex)
 				{
-					FRomData& RomData = RomDatas[MipIndex];
+					FRomDataRuntime& RomData = RomDatas[MipIndex];
 					if (RomData.Size > EmbeddedDataBytesLimit)
 					{
 						NumRoms++;
 						NumRomsBytes += RomData.Size;
 
-						// Ensure that the Id is unique
-						EnsureUniqueRomId(UsedIds, RomData.Id);
-
 						uint32 RomIndex = Program.Roms.Num();
 						check(RomIndex < 0x7fffffff);
 
@@ -684,6 +651,9 @@ namespace mu
 						Program.ConstantImageLODsStreamed.Add(RomIndex, AllMips[MipIndex]);
 
 						Program.Roms.Add(RomData);
+
+						FRomDataCompile& RomDataCompile = RomDatasCompile[MipIndex];
+						Program.RomsCompileData.Add(RomDataCompile);
 					}
 					else
 					{
@@ -725,28 +695,16 @@ namespace mu
 
 					// Serialize to memory, to find out final size of this rom
 					const FMesh* Resource = ResData.Get();
-					int32 ApproximateSize = Resource->GetDataSize();
 
-					FOutputMemoryStream MemStream(ApproximateSize + 64 * 1024);
+					FOutputSizeStream MemStream;
 					FOutputArchive MemoryArch(&MemStream);
 					FMesh::Serialise(ResData.Get(), MemoryArch);
 
-					FRomData& RomData = RomDatas[ResourceIndex];
-					RomData.ResourceType = DT_MESH;
-					RomData.Size = MemStream.GetBufferSize();
-
-					// If the resource uses less memory than the threshold, don't save it in a separate rom.
-					if (RomData.Size <= EmbeddedDataBytesLimit)
-					{
-						return;
-					}
-
-					// RomData.ResourceIndex will be filled below when we split in permanent and streamable
-					// RomData.Flags will be filled below
-
-					// Ensure that the Id is unique
-					uint64 FullHash = CityHash64(reinterpret_cast<const char*>(MemStream.GetBuffer()), MemStream.GetBufferSize());
-					RomData.Id = GetTypeHash(FullHash);
+					FRomDataRuntime& RomData = RomDatas[ResourceIndex];
+					RomData.ResourceType = uint32(ERomDataType::Mesh);
+					uint32 Size = MemStream.GetBufferSize();
+					check(Size < uint32(1 << 30));
+					RomData.Size = Size;
 				});
 		}
 
@@ -756,8 +714,10 @@ namespace mu
 
 			for (int32 MeshIndex = 0; MeshIndex < AllMeshes.Num(); ++MeshIndex)
 			{
-				FRomData& RomData = RomDatas[MeshIndex];
-				RomData.SourceId = AdditionalData.SourceMeshPerConstant[MeshIndex].SourceId;
+				FRomDataRuntime& RomData = RomDatas[MeshIndex];
+
+				FRomDataCompile& RomDataCompile = RomDatasCompile[MeshIndex];
+				RomDataCompile.SourceId = AdditionalData.SourceMeshPerConstant[MeshIndex].SourceId;
 
 				FConstantResourceIndex ResourceIndex;
 
@@ -766,9 +726,6 @@ namespace mu
 					NumRoms++;
 					NumRomsBytes += RomData.Size;
 
-					// Ensure that the Id is unique
-					EnsureUniqueRomId(UsedIds, RomData.Id);
-
 					uint32 RomIndex = Program.Roms.Num();
 					check(RomIndex < 0x7fffffff);
 
@@ -776,6 +733,7 @@ namespace mu
 					Program.ConstantMeshesStreamed.Add(RomIndex, AllMeshes[MeshIndex]);
 
 					Program.Roms.Add(RomData);
+					Program.RomsCompileData.Add(RomDataCompile);
 				}
 				else
 				{
