[Mutable] Added missing Expose API UCustomizableObjectEditorFunctionLibrary.

#rnx

diff --git a/Engine/Plugins/Mutable/Source/CustomizableObjectEditor/Public/MuCOE/CustomizableObjectEditorFunctionLibrary.h b/Engine/Plugins/Mutable/Source/CustomizableObjectEditor/Public/MuCOE/CustomizableObjectEditorFunctionLibrary.h
--- a/Engine/Plugins/Mutable/Source/CustomizableObjectEditor/Public/MuCOE/CustomizableObjectEditorFunctionLibrary.h
+++ b/Engine/Plugins/Mutable/Source/CustomizableObjectEditor/Public/MuCOE/CustomizableObjectEditorFunctionLibrary.h
@@ -54,7 +54,7 @@ struct FNewCustomizableObjectParameters
  * be exposed to cook as well.
  */
 UCLASS()
-class UCustomizableObjectEditorFunctionLibrary : public UBlueprintFunctionLibrary
+class CUSTOMIZABLEOBJECTEDITOR_API UCustomizableObjectEditorFunctionLibrary : public UBlueprintFunctionLibrary
 {
 	GENERATED_BODY()
 
