[Mutable] Serialize Extension Data with the compiled data.

Created a new UExtensionData wrapper to allow serializing optional data (extension plugins may not be loaded). It uses Unreal's serialization system.

- Moved the following data from UCustomizableObject to FModelResources:
- AlwaysLoadedExtensionData.
- StreamedExtensionData.
- Removed unnecessary getters.

#rnx

diff --git a/Engine/Plugins/Mutable/Source/CustomizableObject/Internal/MuCO/CustomizableObjectPrivate.h b/Engine/Plugins/Mutable/Source/CustomizableObject/Internal/MuCO/CustomizableObjectPrivate.h
--- a/Engine/Plugins/Mutable/Source/CustomizableObject/Internal/MuCO/CustomizableObjectPrivate.h
+++ b/Engine/Plugins/Mutable/Source/CustomizableObject/Internal/MuCO/CustomizableObjectPrivate.h
@@ -1029,6 +1029,26 @@ struct FMutableParamNameSet
 };
 
 
+/** Wrapper to Mutable Extensions data.
+  * Only used when serializing extensions data. Required to use Unreal's serialization system, which allows serializing optional data. */
+UCLASS()
+class UExtensionData : public UObject
+{
+	GENERATED_BODY()
+
+public:
+	void CopyFrom(const FModelResources& ModelResources);
+	
+	void MoveTo(FModelResources& ModelResources);
+	
+	UPROPERTY()
+	TArray<FCustomizableObjectResourceData> AlwaysLoadedExtensionData;
+
+	UPROPERTY()
+	TArray<FCustomizableObjectStreamedResourceData> StreamedExtensionData;
+};
+
+
 /** Struct containing all UE resources derived from a CO compilation. These resources will be embedded in the CO at cook time but not in the editor.
   * Editor compilations will serialize this struct to disk using the Serialize methods. Ensure new fields are serialized, too.
   * Variables and settings that should not change until the CO is re-compiled should be stored here. */
@@ -1168,6 +1188,14 @@ struct FModelResources
 	UPROPERTY()
 	TArray<FCustomizableObjectStreamedResourceData> StreamedResourceData;
 
+	// mu::FExtensionData::Index is an index into this array when mu::ExtensionData::Origin is ConstantAlwaysLoaded
+	UPROPERTY()
+	TArray<FCustomizableObjectResourceData> AlwaysLoadedExtensionData;
+
+	// mu::FExtensionData::Index is an index into this array when mu::ExtensionData::Origin is ConstantStreamed
+	UPROPERTY()
+	TArray<FCustomizableObjectStreamedResourceData> StreamedExtensionData;
+	
 	/** Max number of LODs in the compiled Model. */
 	UPROPERTY()
 	uint8 NumLODs = 0;
@@ -1327,6 +1355,7 @@ public:
 	const TSharedPtr<mu::FModel, ESPMode::ThreadSafe>& GetModel();
 	const TSharedPtr<const mu::FModel, ESPMode::ThreadSafe> GetModel() const;
 
+	FModelResources& GetModelResources();
 	const FModelResources& GetModelResources() const;
 
 #if WITH_EDITORONLY_DATA
@@ -1442,13 +1471,7 @@ public:
 	TArray<FProfileParameterDat>& GetInstancePropertiesProfiles();
 #endif
 	
-	TArray<FCustomizableObjectResourceData>& GetAlwaysLoadedExtensionData();
-	const TArray<FCustomizableObjectResourceData>& GetAlwaysLoadedExtensionData() const;
-
-	TArray<FCustomizableObjectStreamedResourceData>& GetStreamedExtensionData();
-	const TArray<FCustomizableObjectStreamedResourceData>& GetStreamedExtensionData() const;
-
-	const FCustomizableObjectResourceData* LoadStreamedResource(int32 ResourceIndex);
+	const FCustomizableObjectResourceData* GetLoadedStreamedResource(int32 ResourceIndex);
 	void UnloadStreamedResource(int32 ResourceIndex);
 
 #if WITH_EDITORONLY_DATA
@@ -1645,6 +1668,8 @@ public:
 
 		ModifiedRomCompiledDataSerialization,
 
+		ModelResourcesExtensionData,
+
 		// -----<new versions can be added above this line>--------
 		LastCustomizableObjectVersion
 	};
diff --git a/Engine/Plugins/Mutable/Source/CustomizableObject/Private/MuCO/CustomizableObject.cpp b/Engine/Plugins/Mutable/Source/CustomizableObject/Private/MuCO/CustomizableObject.cpp
--- a/Engine/Plugins/Mutable/Source/CustomizableObject/Private/MuCO/CustomizableObject.cpp
+++ b/Engine/Plugins/Mutable/Source/CustomizableObject/Private/MuCO/CustomizableObject.cpp
@@ -685,6 +685,10 @@ void FModelResources::Serialize(FObjectAndNameAsStringProxyArchive& MemoryWriter
 
 	SerializeStreamedResources(MemoryWriter, StreamedResourceData);
 
+	UExtensionData* ExtensionData = NewObject<UExtensionData>();
+	ExtensionData->CopyFrom(*this);
+	ExtensionData->Serialize(MemoryWriter);
+	
 	int32 NumReferencedMaterials = Materials.Num();
 	MemoryWriter << NumReferencedMaterials;
 
@@ -815,6 +819,10 @@ bool FModelResources::Unserialize(FObjectAndNameAsStringProxyArchive& MemoryRead
 
 	UnserializeStreamedResources(MemoryReader, &Outer, StreamedResourceData, bIsCooking);
 
+	UExtensionData* ExtensionData = NewObject<UExtensionData>();
+	ExtensionData->Serialize(MemoryReader);
+	ExtensionData->MoveTo(*this);
+
 	// Initialize resources. 
 	for (FMutableRefSkeletalMeshData& ReferenceSkeletalMeshData : ReferenceSkeletalMeshesData)
 	{
@@ -2251,6 +2259,20 @@ void UModelStreamableData::Serialize(FArchive& Ar)
 }
 
 
+void UExtensionData::CopyFrom(const FModelResources& ModelResources)
+{
+	AlwaysLoadedExtensionData = ModelResources.AlwaysLoadedExtensionData;
+	StreamedExtensionData = ModelResources.StreamedExtensionData;
+}
+
+
+void UExtensionData::MoveTo(FModelResources& ModelResources)
+{
+	ModelResources.AlwaysLoadedExtensionData = MoveTemp(AlwaysLoadedExtensionData);
+	ModelResources.StreamedExtensionData = MoveTemp(StreamedExtensionData);
+}
+
+
 void UCustomizableObjectPrivate::SetModel(const TSharedPtr<mu::FModel, ESPMode::ThreadSafe>& Model, const FGuid Id)
 {
 	if (MutableModel == Model
@@ -2326,6 +2348,13 @@ TSharedPtr<FModelStreamableBulkData> UCustomizableObjectPrivate::GetModelStreama
 }
 
 
+FModelResources& UCustomizableObjectPrivate::GetModelResources()
+{
+	const UCustomizableObjectPrivate* ConstThis = this;
+	return *const_cast<FModelResources*>(&ConstThis->GetModelResources());
+}
+
+
 const FModelResources& UCustomizableObjectPrivate::GetModelResources() const
 {
 #if WITH_EDITORONLY_DATA
@@ -2394,31 +2423,7 @@ TArray<FProfileParameterDat>& UCustomizableObjectPrivate::GetInstancePropertiesP
 }
 #endif
 
-
-TArray<FCustomizableObjectResourceData>& UCustomizableObjectPrivate::GetAlwaysLoadedExtensionData()
-{
-	return GetPublic()->AlwaysLoadedExtensionData;
-}
-
-const TArray<FCustomizableObjectResourceData>& UCustomizableObjectPrivate::GetAlwaysLoadedExtensionData() const
-{
-	return GetPublic()->AlwaysLoadedExtensionData;
-}
-
-
-TArray<FCustomizableObjectStreamedResourceData>& UCustomizableObjectPrivate::GetStreamedExtensionData()
-{
-	return GetPublic()->StreamedExtensionData;
-}
-
-
-const TArray<FCustomizableObjectStreamedResourceData>& UCustomizableObjectPrivate::GetStreamedExtensionData() const
-{
-	return GetPublic()->StreamedExtensionData;
-}
-
-
-const FCustomizableObjectResourceData* UCustomizableObjectPrivate::LoadStreamedResource(int32 ResourceIndex)
+const FCustomizableObjectResourceData* UCustomizableObjectPrivate::GetLoadedStreamedResource(int32 ResourceIndex)
 {
 #if WITH_EDITORONLY_DATA
 	FModelResources& LocalModelResources = ModelResourcesEditor;
diff --git a/Engine/Plugins/Mutable/Source/CustomizableObject/Private/MuCO/CustomizableObjectInstance.cpp b/Engine/Plugins/Mutable/Source/CustomizableObject/Private/MuCO/CustomizableObjectInstance.cpp
--- a/Engine/Plugins/Mutable/Source/CustomizableObject/Private/MuCO/CustomizableObjectInstance.cpp
+++ b/Engine/Plugins/Mutable/Source/CustomizableObject/Private/MuCO/CustomizableObjectInstance.cpp
@@ -2174,18 +2174,16 @@ bool UCustomizableInstancePrivate::UpdateSkeletalMesh_PostBeginUpdate0(UCustomiz
 			{
 				case mu::FExtensionData::EOrigin::ConstantAlwaysLoaded:
 				{
-					check(CustomizableObject->GetPrivate()->GetAlwaysLoadedExtensionData().IsValidIndex(ExtensionOutput.Data->Index));
-					ReferencedExtensionData = &CustomizableObject->GetPrivate()->GetAlwaysLoadedExtensionData()[ExtensionOutput.Data->Index];
+					check(ModelResources.AlwaysLoadedExtensionData.IsValidIndex(ExtensionOutput.Data->Index));
+					ReferencedExtensionData = &ModelResources.AlwaysLoadedExtensionData[ExtensionOutput.Data->Index];
 				}
 				break;
 
 				case mu::FExtensionData::EOrigin::ConstantStreamed:
 				{
-					check(CustomizableObject->GetPrivate()->GetStreamedExtensionData().IsValidIndex(ExtensionOutput.Data->Index));
-
-					const FCustomizableObjectStreamedResourceData& StreamedData =
-						CustomizableObject->GetPrivate()->GetStreamedExtensionData()[ExtensionOutput.Data->Index];
+					check(ModelResources.StreamedExtensionData.IsValidIndex(ExtensionOutput.Data->Index));
 
+					const FCustomizableObjectStreamedResourceData& StreamedData = ModelResources.StreamedExtensionData[ExtensionOutput.Data->Index];
 					if (!StreamedData.IsLoaded())
 					{
 						// The data should have been loaded as part of executing the CO program.
@@ -6013,7 +6011,7 @@ void UCustomizableInstancePrivate::AdditionalAssetsAsyncLoaded()
 	{
 		for (int32 ResourceIndex : ComponentData.StreamedResourceIndex)
 		{
-			const FCustomizableObjectResourceData* ResourceData = CustomizableObjectPrivate->LoadStreamedResource(ResourceIndex);
+			const FCustomizableObjectResourceData* ResourceData = CustomizableObjectPrivate->GetLoadedStreamedResource(ResourceIndex);
 			if (!ResourceData)
 			{
 				check(false); // Invalid resource index?
diff --git a/Engine/Plugins/Mutable/Source/CustomizableObject/Private/MuCO/UnrealExtensionDataStreamer.cpp b/Engine/Plugins/Mutable/Source/CustomizableObject/Private/MuCO/UnrealExtensionDataStreamer.cpp
--- a/Engine/Plugins/Mutable/Source/CustomizableObject/Private/MuCO/UnrealExtensionDataStreamer.cpp
+++ b/Engine/Plugins/Mutable/Source/CustomizableObject/Private/MuCO/UnrealExtensionDataStreamer.cpp
@@ -174,7 +174,9 @@ TSharedPtr<FStreamableHandle> FUnrealExtensionDataStreamer::StartLoadOnGameThrea
 		return nullptr;
 	}
 
-	if (!Object->GetPrivate()->GetStreamedExtensionData().IsValidIndex(LoadHandle->Data->Index))
+	FModelResources& ModelResources = Object->GetPrivate()->GetModelResources();
+	
+	if (!ModelResources.StreamedExtensionData.IsValidIndex(LoadHandle->Data->Index))
 	{
 		// The compiled data appears to be out of sync with the CO's properties
 
@@ -185,7 +187,7 @@ TSharedPtr<FStreamableHandle> FUnrealExtensionDataStreamer::StartLoadOnGameThrea
 		return nullptr;
 	}
 
-	FCustomizableObjectStreamedResourceData& StreamedData = Object->GetPrivate()->GetStreamedExtensionData()[LoadHandle->Data->Index];
+	FCustomizableObjectStreamedResourceData& StreamedData = ModelResources.StreamedExtensionData[LoadHandle->Data->Index];
 	if (StreamedData.IsLoaded())
 	{
 		// Already loaded
@@ -242,7 +244,9 @@ void FUnrealExtensionDataStreamer::NotifyLoadCompleted(
 	check(IsInGameThread());
 	check(Object);
 
-	if (!Object->GetPrivate()->GetStreamedExtensionData().IsValidIndex(LoadHandle->Data->Index))
+	FModelResources& ModelResources = Object->GetPrivate()->GetModelResources();
+
+	if (!ModelResources.StreamedExtensionData.IsValidIndex(LoadHandle->Data->Index))
 	{
 		// The compiled data appears to be out of sync with the CO's properties
 
@@ -253,7 +257,7 @@ void FUnrealExtensionDataStreamer::NotifyLoadCompleted(
 		return;
 	}
 
-	FCustomizableObjectStreamedResourceData& StreamedData = Object->GetPrivate()->GetStreamedExtensionData()[LoadHandle->Data->Index];
+	FCustomizableObjectStreamedResourceData& StreamedData = ModelResources.StreamedExtensionData[LoadHandle->Data->Index];
 
 	// The object could have been loaded by another request, in which case we can skip updating
 	// the StreamedData.
diff --git a/Engine/Plugins/Mutable/Source/CustomizableObject/Public/MuCO/CustomizableObject.h b/Engine/Plugins/Mutable/Source/CustomizableObject/Public/MuCO/CustomizableObject.h
--- a/Engine/Plugins/Mutable/Source/CustomizableObject/Public/MuCO/CustomizableObject.h
+++ b/Engine/Plugins/Mutable/Source/CustomizableObject/Public/MuCO/CustomizableObject.h
@@ -279,16 +279,6 @@ public:
 	UPROPERTY(EditAnywhere, Category = CustomizableObject, meta = (DisplayName = "LOD Settings"))
 	FMutableLODSettings LODSettings;
 	
-private:
-	// mu::FExtensionData::Index is an index into this array when mu::ExtensionData::Origin is ConstantAlwaysLoaded
-	UPROPERTY()
-	TArray<FCustomizableObjectResourceData> AlwaysLoadedExtensionData;
-
-	// mu::FExtensionData::Index is an index into this array when mu::ExtensionData::Origin is ConstantStreamed
-	UPROPERTY()
-	TArray<FCustomizableObjectStreamedResourceData> StreamedExtensionData;
-
-public:
 	/** Use the SkeletalMesh of reference as a placeholder until the custom mesh is ready to use.
 	  * Note: If disabled, a null mesh will be used to replace the discarded mesh due to 'ReplaceDiscardedWithReferenceMesh' being enabled. */
 	UPROPERTY(EditAnywhere, Category = CustomizableObject)
diff --git a/Engine/Plugins/Mutable/Source/CustomizableObjectEditor/Private/MuCOE/CustomizableObjectCompiler.cpp b/Engine/Plugins/Mutable/Source/CustomizableObjectEditor/Private/MuCOE/CustomizableObjectCompiler.cpp
--- a/Engine/Plugins/Mutable/Source/CustomizableObjectEditor/Private/MuCOE/CustomizableObjectCompiler.cpp
+++ b/Engine/Plugins/Mutable/Source/CustomizableObjectEditor/Private/MuCOE/CustomizableObjectCompiler.cpp
@@ -958,9 +958,9 @@ void FCustomizableObjectCompiler::CompileInternal(bool bAsync)
 
 		ModelResources.bCompiledWithHDTextureCompression = GenerationContext.Options.TextureCompression == ECustomizableObjectTextureCompression::HighQuality;
 
-		CurrentObject->GetPrivate()->GetAlwaysLoadedExtensionData() = MoveTemp(GenerationContext.AlwaysLoadedExtensionData);
+		ModelResources.AlwaysLoadedExtensionData = MoveTemp(GenerationContext.AlwaysLoadedExtensionData);
 
-		CurrentObject->GetPrivate()->GetStreamedExtensionData().Empty(GenerationContext.StreamedExtensionData.Num());
+		ModelResources.StreamedExtensionData.Empty(GenerationContext.StreamedExtensionData.Num());
 		for (const TPair<FName, UCustomizableObjectResourceDataContainer*>& Pair : GenerationContext.StreamedExtensionData)
 		{
 			const FName& ContainerName = Pair.Key;
@@ -977,7 +977,7 @@ void FCustomizableObjectCompiler::CompileInternal(bool bAsync)
 			
 			NewContainer->Data = Container->Data;
 			
-			CurrentObject->GetPrivate()->GetStreamedExtensionData().Emplace(NewContainer);
+			ModelResources.StreamedExtensionData.Emplace(NewContainer);
 		}
 
 #if WITH_EDITORONLY_DATA
diff --git a/Engine/Plugins/Mutable/Source/CustomizableObjectEditor/Private/MuCOE/CustomizableObjectCookPackageSplitter.cpp b/Engine/Plugins/Mutable/Source/CustomizableObjectEditor/Private/MuCOE/CustomizableObjectCookPackageSplitter.cpp
--- a/Engine/Plugins/Mutable/Source/CustomizableObjectEditor/Private/MuCOE/CustomizableObjectCookPackageSplitter.cpp
+++ b/Engine/Plugins/Mutable/Source/CustomizableObjectEditor/Private/MuCOE/CustomizableObjectCookPackageSplitter.cpp
@@ -154,7 +154,7 @@ bool FCustomizableObjectCookPackageSplitter::ShouldSplit(UObject* SplitData)
 	UCustomizableObject* Object = CastChecked<UCustomizableObject>(SplitData);
 	const FModelResources& ModelResources = Object->GetPrivate()->GetModelResources(true);
 
-	return !Object->IsChildObject() && (ModelResources.StreamedResourceData.Num() > 0 || Object->GetPrivate()->GetStreamedExtensionData().Num() > 0);
+	return !Object->IsChildObject() && (ModelResources.StreamedResourceData.Num() > 0 || ModelResources.StreamedExtensionData.Num() > 0);
 }
 
 TArray<ICookPackageSplitter::FGeneratedPackage> FCustomizableObjectCookPackageSplitter::GetGenerateList(
@@ -176,7 +176,7 @@ TArray<ICookPackageSplitter::FGeneratedPackage> FCustomizableObjectCookPackageSp
 	}
 
 	// Generate a new package for each streamed Extension Data
-	for (const FCustomizableObjectStreamedResourceData& StreamedData : Object->GetPrivate()->GetStreamedExtensionData())
+	for (const FCustomizableObjectStreamedResourceData& StreamedData : ModelResources.StreamedExtensionData)
 	{
 		GenerateNewPackage(StreamedData, OwnerPackage, OwnerObject, Result);
 	}
@@ -233,7 +233,7 @@ bool FCustomizableObjectCookPackageSplitter::PreSaveGeneratorPackage(
 
 	// There should be one generated package per streamed Resource Data
 	const int32 NumStreamedData = ModelResources.StreamedResourceData.Num();
-	const int32 NumStreamedExtensionData = Object->GetPrivate()->GetStreamedExtensionData().Num();
+	const int32 NumStreamedExtensionData = ModelResources.StreamedExtensionData.Num();
 	
 	check(NumStreamedData + NumStreamedExtensionData == PlaceholderPackages.Num());
 
@@ -258,7 +258,7 @@ bool FCustomizableObjectCookPackageSplitter::PreSaveGeneratorPackage(
 	for (int32 Index = 0; Index < NumStreamedExtensionData; ++Index)
 	{
 		const ICookPackageSplitter::FGeneratedPackageForPreSave& GeneratedPackage = PlaceholderPackages[NumStreamedData + Index];
-		if (!PreSavePackage(GeneratedPackage, Object->GetPrivate()->GetStreamedExtensionData()))
+		if (!PreSavePackage(GeneratedPackage, ModelResources.StreamedExtensionData))
 		{
 			return false;
 		}
@@ -304,11 +304,11 @@ void FCustomizableObjectCookPackageSplitter::PostSaveGeneratorPackage(UPackage*
 
 	for (const FString& ContainerName : SavedExtensionContainerNames)
 	{
-		FCustomizableObjectStreamedResourceData* ResourceData = FindStreamedResourceData(Object->GetPrivate()->GetStreamedExtensionData(), ContainerName);
+		FCustomizableObjectStreamedResourceData* ResourceData = FindStreamedResourceData(ModelResources.StreamedExtensionData, ContainerName);
 		if (!ResourceData)
 		{
 			UE_LOG(LogMutable, Error, TEXT("Couldn't find streamed Extension Data container with name %s in array of %d entries"),
-				*ContainerName, Object->GetPrivate()->GetStreamedExtensionData().Num());
+				*ContainerName, ModelResources.StreamedExtensionData.Num());
 
 			continue;
 		}
@@ -320,7 +320,7 @@ void FCustomizableObjectCookPackageSplitter::PostSaveGeneratorPackage(UPackage*
 		NewArray.Emplace(Container);
 	}
 
-	Object->GetPrivate()->GetStreamedExtensionData() = NewArray;
+	ModelResources.StreamedExtensionData = NewArray;
 }
 
 bool FCustomizableObjectCookPackageSplitter::PopulateGeneratedPackage(
@@ -338,13 +338,13 @@ bool FCustomizableObjectCookPackageSplitter::PopulateGeneratedPackage(
 	FCustomizableObjectStreamedResourceData* ResourceData = FindStreamedResourceData(ModelResources.StreamedResourceData, GeneratedPackage.RelativePath);
 	if (!ResourceData)
 	{
-		ResourceData = FindStreamedResourceData(Object->GetPrivate()->GetStreamedExtensionData(), GeneratedPackage.RelativePath);
+		ResourceData = FindStreamedResourceData(ModelResources.StreamedExtensionData, GeneratedPackage.RelativePath);
 	}
 
 	if (!ResourceData)
 	{
 		UE_LOG(LogMutable, Error, TEXT("Couldn't find streamed resource Data container with name %s in arrays of %d and %d entries"),
-			*GeneratedPackage.RelativePath, ModelResources.StreamedResourceData.Num(), Object->GetPrivate()->GetStreamedExtensionData().Num());
+			*GeneratedPackage.RelativePath, ModelResources.StreamedResourceData.Num(), ModelResources.StreamedExtensionData.Num());
 
 		return false;
 	}
@@ -386,13 +386,13 @@ void FCustomizableObjectCookPackageSplitter::PostSaveGeneratedPackage(
 	FCustomizableObjectStreamedResourceData* ResourceData = FindStreamedResourceData(ModelResources.StreamedResourceData, GeneratedPackage.RelativePath);
 	if (!ResourceData)
 	{
-		ResourceData = FindStreamedResourceData(Object->GetPrivate()->GetStreamedExtensionData(), GeneratedPackage.RelativePath);
+		ResourceData = FindStreamedResourceData(ModelResources.StreamedExtensionData, GeneratedPackage.RelativePath);
 	}
 
 	if (!ResourceData)
 	{
 		UE_LOG(LogMutable, Error, TEXT("Couldn't find streamed resource Data container with name %s in arrays of %d and %d entries"),
-			*GeneratedPackage.RelativePath, ModelResources.StreamedResourceData.Num(), Object->GetPrivate()->GetStreamedExtensionData().Num());
+			*GeneratedPackage.RelativePath, ModelResources.StreamedResourceData.Num(), ModelResources.StreamedExtensionData.Num());
 
 		return;
 	}
