[Mutable] Fix possible crash when scalability settings change between instance updates.

- UCustomizablInstancerPrivate::DoComponentNeedsUpdate() now checks if the SkeletalMesh sections are populated correctly for the current update FirstLODAvailable.

#rnx

diff --git a/Engine/Plugins/Mutable/Source/CustomizableObject/Private/MuCO/CustomizableObjectInstance.cpp b/Engine/Plugins/Mutable/Source/CustomizableObject/Private/MuCO/CustomizableObjectInstance.cpp
--- a/Engine/Plugins/Mutable/Source/CustomizableObject/Private/MuCO/CustomizableObjectInstance.cpp
+++ b/Engine/Plugins/Mutable/Source/CustomizableObject/Private/MuCO/CustomizableObjectInstance.cpp
@@ -2034,6 +2034,24 @@ bool UCustomizableInstancePrivate::DoComponentsNeedUpdate(UCustomizableObjectIns
 		{
 			OperationData->MeshChangedPerInstanceComponent[InstanceComponentIndex] = MeshIDs[ComponentOffset + MeshIndex] != ComponentData->LastMeshIdPerLOD[MeshIndex];
 		}
+
+		// If the component SkeletalMesh was build with scalability settings different than the current, we need to update.
+		// The mismatch can happen when the scalability setting is changed and an update is requested with a MinLOD lower 
+		// than the new FirstLODAvailable. This is an artifact of our LOD management system setting MinLOD.
+		if (!OperationData->MeshChangedPerInstanceComponent[InstanceComponentIndex] && bHadSkeletalMesh) 
+		{
+			int32 MinLOD = OperationData->GetMinLOD();
+			for (int32 LODIndex = OperationData->FirstLODAvailable; LODIndex < MinLOD; ++LODIndex)
+			{
+				TIndirectArray<FSkeletalMeshLODRenderData>& LODRenderData = (*SkeletalMeshPtr)->GetResourceForRendering()->LODRenderData;
+				if (LODRenderData[LODIndex].RenderSections.Num() != LODRenderData[MinLOD].RenderSections.Num())
+				{
+					OperationData->MeshChangedPerInstanceComponent[InstanceComponentIndex] = true;
+					break;
+				}
+			}
+		}
+
 	}
 
 	bool bChanged =
