[Mutable] Add support to stream Mesh LOD with AltSkinWeightProfiles enabled.

#rnx

diff --git a/Engine/Plugins/Mutable/Source/CustomizableObject/Internal/MuCO/UnrealConversionUtils.h b/Engine/Plugins/Mutable/Source/CustomizableObject/Internal/MuCO/UnrealConversionUtils.h
--- a/Engine/Plugins/Mutable/Source/CustomizableObject/Internal/MuCO/UnrealConversionUtils.h
+++ b/Engine/Plugins/Mutable/Source/CustomizableObject/Internal/MuCO/UnrealConversionUtils.h
@@ -61,7 +61,7 @@ namespace UnrealConversionUtils
 
 	/** Performs a copy of the data found on the vertex buffers on the mutable mesh to the buffers of the skeletal mesh
 	 * @param LODResource - LODRenderData to update
-	 * @param InMutableMesh - Mutable mesh to be used as reference for the section data update on the Skeletal Mesh
+	 * @param InMutableMesh - Mutable mesh to be used as reference for the vertex data update on the Skeletal Mesh
 	 * @param bAllowCPUAccess - Keeps this LODs data on the CPU so it can be used for things such as sampling in FX.
 	 */
 	CUSTOMIZABLEOBJECT_API void CopyMutableVertexBuffers(
@@ -90,43 +90,36 @@ namespace UnrealConversionUtils
 		const mu::FMesh* InMutableMesh);
 	
 
-	/**
-	 *Initializes the LODResource's SkinWeightProfilesData buffer.
-	 * @param LODResource - LODRenderData to be updated.
-	 * @param bIsStreamableMesh - If StreamMeshLODs is enabled for the outer SkeletalMesh.
-	 */
-	void InitSkinWeightProfilesData(
-		FSkeletalMeshLODRenderData& LODResource,
-		bool bIsStreamableMesh
-	);
-
-
 	/**
 	 *Performs a copy of the data found on the index buffers on the mutable mesh to the buffers of the skeletal mesh
 	 * @param LODResource - LODRenderData to be updated.
-	 * @param InProfileName - Name of the profile to generate.
-	 * @param InMutableMeshVertexBuffers - The mutable buffers to be reading data from
-	 * @param InBoneIndexBuffer - The buffer containing the indices for the bones.
-	 * @return True if the operation could be performed successfully,	 false if not.
+	 * @param Owner - Outer mesh of the LODResource. 
+	 * @param LODIndex - LOD Index of the LODResource.
+	 * @param InMutableMesh - Mutable mesh to be used to initialize the SkinWeightProfilesBuffers.
+	 * @param InActiveProfiles - Profiles to initialize. Values are ProfileID and ProfileName.
 	 */
 	CUSTOMIZABLEOBJECT_API void CopyMutableSkinWeightProfilesBuffers(
 		FSkeletalMeshLODRenderData& LODResource,
-		const FName InProfileName,
-		const mu::FMeshBufferSet& InMutableMeshVertexBuffers,
-		const int32 InBoneIndexBuffer);
+		USkeletalMesh& Owner,
+		int32 LODIndex,
+		const mu::FMesh* InMutableMesh,
+		const TArray<TPair<uint32, FName>>& InActiveProfiles
+		);
 
 
 	/**
 	 *Performs a copy of the render data of a specific Skeletal Mesh LOD to another Skeletal Mesh
 	 * @param LODResource - LODRenderData to copy to.
 	 * @param SourceLODResource - LODRenderData to copy from.
-	 * @param SkeletalMesh - Owner of both LODResource and SourceLODResource. 
+	 * @param Owner - Outer mesh of the LODResource and SourceLODResource. 
+	 * @param LODIndex - LOD Index to copy to.
 	 * @param bAllowCPUAccess - Keeps this LODs data on the CPU so it can be used for things such as sampling in FX.
 	 */
 	CUSTOMIZABLEOBJECT_API void CopySkeletalMeshLODRenderData(
 		FSkeletalMeshLODRenderData& LODResource,
 		FSkeletalMeshLODRenderData& SourceLODResource,
-		const USkeletalMesh& SkeletalMesh,
+		USkeletalMesh& Owner,
+		int32 LODIndex,
 		const bool bAllowCPUAccess
 	);
 
diff --git a/Engine/Plugins/Mutable/Source/CustomizableObject/Private/MuCO/CustomizableObjectInstance.cpp b/Engine/Plugins/Mutable/Source/CustomizableObject/Private/MuCO/CustomizableObjectInstance.cpp
--- a/Engine/Plugins/Mutable/Source/CustomizableObject/Private/MuCO/CustomizableObjectInstance.cpp
+++ b/Engine/Plugins/Mutable/Source/CustomizableObject/Private/MuCO/CustomizableObjectInstance.cpp
@@ -2349,7 +2349,7 @@ bool UCustomizableInstancePrivate::UpdateSkeletalMesh_PostBeginUpdate0(UCustomiz
 
 		if (OperationData->bStreamMeshLODs)
 		{
-			SkeletalMesh = UCustomizableObjectSkeletalMesh::CreateSkeletalMesh(OperationData, *Public, *CustomizableObject, Component->FirstLOD, Component->LODCount);
+			SkeletalMesh = UCustomizableObjectSkeletalMesh::CreateSkeletalMesh(*Public);
 		}
 		else
 		{
@@ -2472,6 +2472,11 @@ bool UCustomizableInstancePrivate::UpdateSkeletalMesh_PostBeginUpdate0(UCustomiz
 				CustomizableObject->GetPrivate()->MeshCache.Add(*MeshId, SkeletalMesh);
 			}
 		}
+
+		if (UCustomizableObjectSkeletalMesh* StreamableMesh = Cast<UCustomizableObjectSkeletalMesh>(SkeletalMesh))
+		{
+			StreamableMesh->InitMutableStreamingData(OperationData, *CustomizableObject, Component->FirstLOD, Component->LODCount);
+		}
 	}
 
 	return true;
@@ -5185,7 +5190,45 @@ bool UCustomizableInstancePrivate::BuildOrCopyRenderData(const TSharedRef<FUpdat
 			LOD.FirstBoneMap,
 			MeshSurfacesMetadata);
 
-		UnrealConversionUtils::InitSkinWeightProfilesData(LODResource, OperationData->bStreamMeshLODs);
+		// Set SkinWeightProfiles
+		LODResource.SkinWeightProfilesData.Init(&LODResource.SkinWeightVertexBuffer);
+		
+		// Active SkinWeightProfiles ID and Name
+		TArray<TPair<uint32, FName>> ActiveSkinWeightProfiles;
+		
+		const mu::FMeshBufferSet& MutableMeshVertexBuffers = LOD.Mesh->GetVertexBuffers();
+		const int32 NumBuffers = MutableMeshVertexBuffers.GetBufferCount();
+
+		for (int32 BufferIndex = 0; BufferIndex < NumBuffers; ++BufferIndex)
+		{
+			mu::EMeshBufferSemantic Semantic;
+			int32 SemanticIndex;
+			MutableMeshVertexBuffers.GetChannel(BufferIndex, 0, &Semantic, &SemanticIndex, nullptr, nullptr, nullptr);
+
+			if (Semantic != mu::EMeshBufferSemantic::AltSkinWeight)
+			{
+				continue;
+			}
+
+			const FMutableSkinWeightProfileInfo* ProfileInfo = ModelResources.SkinWeightProfilesInfo.FindByPredicate(
+				[&SemanticIndex](const FMutableSkinWeightProfileInfo& P) { return P.NameId == SemanticIndex; });
+
+			if (ensure(ProfileInfo))
+			{
+				const FSkinWeightProfileInfo* ExistingProfile = SkeletalMesh->GetSkinWeightProfiles().FindByPredicate(
+					[&ProfileInfo](const FSkinWeightProfileInfo& P) { return P.Name == ProfileInfo->Name; });
+
+				if (!ExistingProfile)
+				{
+					SkeletalMesh->AddSkinWeightProfile({ ProfileInfo->Name, ProfileInfo->DefaultProfile, ProfileInfo->DefaultProfileFromLODIndex });
+				}
+
+				ActiveSkinWeightProfiles.Add({ ProfileInfo->NameId, ProfileInfo->Name });
+
+				LODResource.SkinWeightProfilesData.AddOverrideData(ProfileInfo->Name);
+
+			}
+		}
 
 		if (LODResource.bStreamedDataInlined) // Non-streamable LOD
 		{
@@ -5203,57 +5246,12 @@ bool UCustomizableInstancePrivate::BuildOrCopyRenderData(const TSharedRef<FUpdat
 			}
 
 			// Copy SkinWeightProfiles
-			if (!ModelResources.SkinWeightProfilesInfo.IsEmpty())
-			{
-				const mu::FMeshBufferSet& MutableMeshVertexBuffers = LOD.Mesh->GetVertexBuffers();
-
-				bool bHasSkinWeightProfiles = false;
-			
-				const int32 NumBuffers = MutableMeshVertexBuffers.GetBufferCount();
-				for (int32 BufferIndex = 0; BufferIndex < NumBuffers; ++BufferIndex)
-				{
-					// Skin Weight Profiles have 3 channels
-					if (MutableMeshVertexBuffers.GetBufferChannelCount(BufferIndex) != 3)
-					{
-						continue;
-					}
-
-					mu::EMeshBufferSemantic Semantic;
-					int32 SemanticIndex;
-					MutableMeshVertexBuffers.GetChannel(BufferIndex, 0, &Semantic, &SemanticIndex, nullptr, nullptr, nullptr);
-
-					if (Semantic != mu::EMeshBufferSemantic::AltSkinWeight)
-					{
-						continue;
-					}
-
-					const FMutableSkinWeightProfileInfo* ProfileInfo = ModelResources.SkinWeightProfilesInfo.FindByPredicate(
-							[&SemanticIndex](const FMutableSkinWeightProfileInfo& P) { return P.NameId == SemanticIndex; });
-
-					if (ensure(ProfileInfo))
-					{
-						if (!bHasSkinWeightProfiles)
-						{
-							LODResource.SkinWeightProfilesData.Init(&LODResource.SkinWeightVertexBuffer);
-							bHasSkinWeightProfiles = true;
-						}
-						
-						const FSkinWeightProfileInfo* ExistingProfile = SkeletalMesh->GetSkinWeightProfiles().FindByPredicate(
-							[&ProfileInfo](const FSkinWeightProfileInfo& P) { return P.Name == ProfileInfo->Name; });
-
-						if (!ExistingProfile)
-						{
-							SkeletalMesh->AddSkinWeightProfile({ ProfileInfo->Name, ProfileInfo->DefaultProfile, ProfileInfo->DefaultProfileFromLODIndex });
-						}
-
-						UnrealConversionUtils::CopyMutableSkinWeightProfilesBuffers(
-							LODResource,
-							ProfileInfo->Name,
-							MutableMeshVertexBuffers,
-							BufferIndex);
-					}
-				}
-			}
+			UnrealConversionUtils::CopyMutableSkinWeightProfilesBuffers(
+				LODResource,
+				*SkeletalMesh,
+				LODIndex,
+				LOD.Mesh.Get(),
+				ActiveSkinWeightProfiles);
 		}
 		else // Streamable LOD. 
 		{
@@ -5265,8 +5263,6 @@ bool UCustomizableInstancePrivate::BuildOrCopyRenderData(const TSharedRef<FUpdat
 
 			// Init IndexBuffers for streaming
 			UnrealConversionUtils::InitIndexBuffersWithDummyData(LODResource, LOD.Mesh.Get());
-
-			// SkinWeightProfilesInfo Not supported yet
 		}
 
 		if (LODResource.StaticVertexBuffers.ColorVertexBuffer.GetNumVertices())
@@ -5305,6 +5301,7 @@ bool UCustomizableInstancePrivate::BuildOrCopyRenderData(const TSharedRef<FUpdat
 			LODResource,
 			SourceLODResource,
 			*SkeletalMesh,
+			LODIndex,
 			SkeletalMesh->GetLODInfo(LODIndex)->bAllowCPUAccess
 		);
 	}
diff --git a/Engine/Plugins/Mutable/Source/CustomizableObject/Private/MuCO/CustomizableObjectMeshUpdate.cpp b/Engine/Plugins/Mutable/Source/CustomizableObject/Private/MuCO/CustomizableObjectMeshUpdate.cpp
--- a/Engine/Plugins/Mutable/Source/CustomizableObject/Private/MuCO/CustomizableObjectMeshUpdate.cpp
+++ b/Engine/Plugins/Mutable/Source/CustomizableObject/Private/MuCO/CustomizableObjectMeshUpdate.cpp
@@ -255,7 +255,7 @@ void FCustomizableObjectMeshStreamIn::ConvertMesh(const FContext& Context)
 
 	check(!TaskSynchronization.GetValue());
 
-	const USkeletalMesh* Mesh = Context.Mesh;
+	const UCustomizableObjectSkeletalMesh* Mesh = Cast<UCustomizableObjectSkeletalMesh>(Context.Mesh);
 	FSkeletalMeshRenderData* RenderData = Context.RenderData;
 	if (IsCancelled() || !Mesh || !RenderData)
 	{
@@ -280,15 +280,16 @@ void FCustomizableObjectMeshStreamIn::ConvertMesh(const FContext& Context)
 			return;
 		}
 
+		const bool bNeedsCPUAccess = Mesh->GetResourceForRendering()->RequiresCPUSkinning(GMaxRHIFeatureLevel) || Mesh->NeedCPUData(LODIndex);
+
 		FSkeletalMeshLODRenderData& LODResource = *Context.LODResourcesView[LODIndex];
-		UnrealConversionUtils::CopyMutableVertexBuffers(LODResource, MutableMesh.Get(), true);
+		UnrealConversionUtils::CopyMutableVertexBuffers(LODResource, MutableMesh.Get(), bNeedsCPUAccess);
 		UnrealConversionUtils::CopyMutableIndexBuffers(LODResource, MutableMesh.Get());
+		UnrealConversionUtils::CopyMutableSkinWeightProfilesBuffers(LODResource, *const_cast<UCustomizableObjectSkeletalMesh*>(Mesh), LODIndex, MutableMesh.Get(), Mesh->SkinWeightProfileIDs);
 
 		// TODO PERE: Implement missing buffers
-		// SkinWeightProfilesData
 		// MorphTargets (MorphTargetVertexInfoBuffers)
 		// Clothing (ClothVertexBuffer)
-		UnrealConversionUtils::InitSkinWeightProfilesData(LODResource, true /* bIsStreamableMesh */);
 
 		UnrealConversionUtils::UpdateSkeletalMeshLODRenderDataBuffersSize(LODResource);
 	}
diff --git a/Engine/Plugins/Mutable/Source/CustomizableObject/Private/MuCO/CustomizableObjectSkeletalMesh.cpp b/Engine/Plugins/Mutable/Source/CustomizableObject/Private/MuCO/CustomizableObjectSkeletalMesh.cpp
--- a/Engine/Plugins/Mutable/Source/CustomizableObject/Private/MuCO/CustomizableObjectSkeletalMesh.cpp
+++ b/Engine/Plugins/Mutable/Source/CustomizableObject/Private/MuCO/CustomizableObjectSkeletalMesh.cpp
@@ -10,6 +10,8 @@
 #include "MuCO/CustomizableObjectInstancePrivate.h"
 #include "MuCO/CustomizableObjectMeshUpdate.h"
 
+#include "Animation/SkinWeightProfile.h"
+
 FName GenerateUniqueNameFromCOInstance(const UCustomizableObjectInstance& Instance)
 {
 	FString BaseName = FString("MutableSkMesh-") + Instance.GetCustomizableObject()->GetName();
@@ -20,31 +22,46 @@ FName GenerateUniqueNameFromCOInstance(const UCustomizableObjectInstance& Instan
 }
 
 
-UCustomizableObjectSkeletalMesh* UCustomizableObjectSkeletalMesh::CreateSkeletalMesh(const TSharedRef<FUpdateContextPrivate>& OperationData, 
-	const UCustomizableObjectInstance& Instance, UCustomizableObject& CustomizableObject, const int32 FirstLOD, const int32 LODCount)
+UCustomizableObjectSkeletalMesh* UCustomizableObjectSkeletalMesh::CreateSkeletalMesh(const UCustomizableObjectInstance& Instance)
 {
 	FName SkeletalMeshName = GenerateUniqueNameFromCOInstance(Instance);
 	UCustomizableObjectSkeletalMesh* OutSkeletalMesh = NewObject<UCustomizableObjectSkeletalMesh>(GetTransientPackage(), SkeletalMeshName, RF_Transient);
 	
+	return OutSkeletalMesh;
+}
+
+
+void UCustomizableObjectSkeletalMesh::InitMutableStreamingData(const TSharedRef<FUpdateContextPrivate>& InOperationData, UCustomizableObject& CustomizableObject, const int32 FirstLOD, const int32 LODCount)
+{
 	// Debug info
-	OutSkeletalMesh->CustomizableObjectPathName = GetNameSafe(&CustomizableObject);
-	OutSkeletalMesh->InstancePathName = Instance.GetName();
+	CustomizableObjectPathName = GetNameSafe(&CustomizableObject);
 
 	// Init properties
-	OutSkeletalMesh->Model = CustomizableObject.GetPrivate()->GetModel();
+	Model = CustomizableObject.GetPrivate()->GetModel();
 
-	OutSkeletalMesh->Parameters = OperationData->Parameters;
-	OutSkeletalMesh->State = OperationData->GetCapturedDescriptor().GetState();
-	
-	OutSkeletalMesh->MeshIDs.Init(MAX_uint64, MAX_MESH_LOD_COUNT);
-	
-	for (int32 LODIndex = OperationData->FirstLODAvailable; LODIndex < LODCount; ++LODIndex)
+	Parameters = InOperationData->Parameters;
+	State = InOperationData->GetCapturedDescriptor().GetState();
+
+	MeshIDs.Init(MAX_uint64, MAX_MESH_LOD_COUNT);
+
+	for (int32 LODIndex = InOperationData->FirstLODAvailable; LODIndex < LODCount; ++LODIndex)
 	{
-		const FInstanceUpdateData::FLOD& LOD = OperationData->InstanceUpdateData.LODs[FirstLOD + LODIndex];
-		OutSkeletalMesh->MeshIDs[LODIndex] = LOD.MeshID;
+		const FInstanceUpdateData::FLOD& LOD = InOperationData->InstanceUpdateData.LODs[FirstLOD + LODIndex];
+		MeshIDs[LODIndex] = LOD.MeshID;
 	}
 
-	return OutSkeletalMesh;
+	const FModelResources& ModelResources = CustomizableObject.GetPrivate()->GetModelResources();
+
+	for (const FSkinWeightProfileInfo& SkinWeightProfile : GetSkinWeightProfiles())
+	{
+		const FMutableSkinWeightProfileInfo* ProfileInfo = ModelResources.SkinWeightProfilesInfo.FindByPredicate(
+			[&SkinWeightProfile](const FMutableSkinWeightProfileInfo& P) { return P.Name == SkinWeightProfile.Name; });
+
+		if (ensure(ProfileInfo))
+		{
+			SkinWeightProfileIDs.Add({ ProfileInfo->NameId, ProfileInfo->Name });
+		}
+	}
 }
 
 
diff --git a/Engine/Plugins/Mutable/Source/CustomizableObject/Private/MuCO/UnrealConversionUtils.cpp b/Engine/Plugins/Mutable/Source/CustomizableObject/Private/MuCO/UnrealConversionUtils.cpp
--- a/Engine/Plugins/Mutable/Source/CustomizableObject/Private/MuCO/UnrealConversionUtils.cpp
+++ b/Engine/Plugins/Mutable/Source/CustomizableObject/Private/MuCO/UnrealConversionUtils.cpp
@@ -15,13 +15,6 @@
 
 class USkeleton;
 
-TAutoConsoleVariable<int32> CVarMutableInitSkinWightProfileData(
-	TEXT("mutable.InitSkinWeightProfileData"),
-	2,
-	TEXT("0 - Do not initialize SkinWeightProfilesData buffer by default")
-	TEXT("1 - Initialize on streamable SkeletalMeshes")
-	TEXT("2 - Initialize on all generated SkeletalMeshes"));
-
 namespace UnrealConversionUtils
 {
 	// Hidden functions only used internally to aid other functions
@@ -473,18 +466,6 @@ namespace UnrealConversionUtils
 	}
 
 
-	void InitSkinWeightProfilesData(FSkeletalMeshLODRenderData& LODResource, bool bIsStreamableMesh)
-	{
-		const int32 InitProfileDataMode = CVarMutableInitSkinWightProfileData.GetValueOnAnyThread();
-		const bool bInitSkinWeightProfilesData = (InitProfileDataMode == 1 && bIsStreamableMesh) || InitProfileDataMode == 2;
-
-		if (bInitSkinWeightProfilesData)
-		{
-			LODResource.SkinWeightProfilesData.Init(&LODResource.SkinWeightVertexBuffer);
-		}
-	}
-
-
 	struct FAuxMutableSkinWeightVertexKey
 	{
 		FAuxMutableSkinWeightVertexKey(const uint8* InKey, uint8 InKeySize, uint32 InHash)
@@ -510,89 +491,115 @@ namespace UnrealConversionUtils
 	
 	void CopyMutableSkinWeightProfilesBuffers(
 		FSkeletalMeshLODRenderData& LODResource,
-		const FName ProfileName,
-		const mu::FMeshBufferSet& MutableMeshVertexBuffers,
-		const int32 BoneIndexBuffer)
+		USkeletalMesh& SkeletalMesh,
+		int32 LODIndex,
+		const mu::FMesh* InMutableMesh,
+		const TArray<TPair<uint32, FName>>& ActiveSkinWeightProfiles)
 	{
 		MUTABLE_CPUPROFILER_SCOPE(CopyMutableSkinWeightProfilesBuffers);
 
-		// Basic Buffer override settings
-		FRuntimeSkinWeightProfileData& Override = LODResource.SkinWeightProfilesData.AddOverrideData(ProfileName);
-
 		const uint8 NumInfluences = LODResource.GetVertexBufferMaxBoneInfluences();
-		Override.NumWeightsPerVertex = NumInfluences;
-
 		const bool b16BitBoneIndices = LODResource.DoesVertexBufferUse16BitBoneIndex();
-		Override.b16BitBoneIndices = b16BitBoneIndices;
-		const int32 BoneIndexSize = b16BitBoneIndices ? 2 : 1;
-
-		// BoneWeights channel info
-		mu::EMeshBufferFormat Format;
-		int32 MutableNumInfluences;
-		int32 Offset;
-		MutableMeshVertexBuffers.GetChannel(BoneIndexBuffer, 2, nullptr, nullptr, &Format, &MutableNumInfluences, &Offset);
 
-		const uint8 MutBoneIndexSize = Format == mu::EMeshBufferFormat::UInt16 ? 2 : 1;
-		const uint8 MutBoneIndicesSize = Offset - sizeof(int32);
-		const uint8 MutBoneWeightsSize = MutableNumInfluences;
-		const uint8 MutBoneWeightVertexSize = MutBoneIndicesSize + MutBoneWeightsSize;
+		const int32 BoneIndexTypeByteSize = LODResource.SkinWeightVertexBuffer.GetBoneIndexByteSize();
+		const uint8 BoneIndicesDataSize = NumInfluences * BoneIndexTypeByteSize;
 
-		check(BoneIndexSize == MutBoneIndexSize);
-
-		const int32 ElementCount = MutableMeshVertexBuffers.GetElementCount();
-		Override.BoneIDs.Reserve(ElementCount * MutBoneIndicesSize);
-		Override.BoneWeights.Reserve(ElementCount * NumInfluences);
-		Override.VertexIndexToInfluenceOffset.Reserve(ElementCount);
+		const int32 BoneWeightTypeByteSize = LODResource.SkinWeightVertexBuffer.GetBoneWeightByteSize();
+		const uint8 BoneWeightsDataSize = NumInfluences * BoneWeightTypeByteSize;
 
 		TMap<FAuxMutableSkinWeightVertexKey, int32> HashToUniqueWeightIndexMap;
-		HashToUniqueWeightIndexMap.Reserve(ElementCount);
-		int32 UniqueWeightsCount = 0;
 
-		const uint8* SkinWeightsBuffer = reinterpret_cast<const uint8*>(MutableMeshVertexBuffers.GetBufferData(BoneIndexBuffer));
-		for (int32 ElementIndex = 0; ElementIndex < ElementCount; ++ElementIndex)
+		const mu::FMeshBufferSet& MutableMeshVertexBuffers = InMutableMesh->GetVertexBuffers();
+		for (const TPair<uint32, FName>& Profile : ActiveSkinWeightProfiles)
 		{
-			uint32 ElementHash;
-			FMemory::Memcpy(&ElementHash, SkinWeightsBuffer, sizeof(int32));
-			SkinWeightsBuffer += sizeof(int32);
+			int32 BufferIndex, ChannelIndex;
+			mu::EMeshBufferFormat Format;
+			int32 MutNumInfluences;
 
-			if (ElementHash > 0)
+			MutableMeshVertexBuffers.FindChannel(mu::EMeshBufferSemantic::AltSkinWeight, Profile.Key, &BufferIndex, &ChannelIndex);
+			
+			if (BufferIndex == INDEX_NONE)
 			{
-				if (int32* OverrideIndex = HashToUniqueWeightIndexMap.Find({ SkinWeightsBuffer, MutBoneWeightVertexSize, ElementHash }))
+				continue;
+			}
+
+			// Basic Buffer override settings
+			FRuntimeSkinWeightProfileData& Override = LODResource.SkinWeightProfilesData.AddOverrideData(Profile.Value);
+			Override.NumWeightsPerVertex = NumInfluences;
+			Override.b16BitBoneIndices = b16BitBoneIndices;
+
+			// BoneIndices channel info
+			MutableMeshVertexBuffers.GetChannel(BufferIndex, 1, nullptr, nullptr, &Format, &MutNumInfluences, nullptr);
+			const uint8 MutBoneIndexTypeByteSize = Format == mu::EMeshBufferFormat::UInt16 ? 2 : 1;
+			const uint8 MutBoneIndicesDataSize = MutBoneIndexTypeByteSize * MutNumInfluences;
+			check(BoneIndexTypeByteSize == MutBoneIndexTypeByteSize);
+
+			// BoneWeights channel info
+			MutableMeshVertexBuffers.GetChannel(BufferIndex, 2, nullptr, nullptr, &Format, &MutNumInfluences, nullptr);
+			const uint8 MutBoneWeightTypeByteSize = Format == mu::EMeshBufferFormat::NUInt16 ? 2 : 1;
+			const uint8 MutBoneWeightsDataSize = MutBoneWeightTypeByteSize * MutNumInfluences;
+			check(BoneWeightTypeByteSize == MutBoneWeightTypeByteSize);
+
+			const uint8 MutSkinWeightVertexSize = MutBoneIndicesDataSize + MutBoneWeightsDataSize;
+
+			const int32 ElementCount = MutableMeshVertexBuffers.GetElementCount();
+			Override.BoneIDs.Reserve(ElementCount * BoneIndicesDataSize);
+			Override.BoneWeights.Reserve(ElementCount * BoneWeightsDataSize);
+			Override.VertexIndexToInfluenceOffset.Reserve(ElementCount);
+
+			HashToUniqueWeightIndexMap.Empty(ElementCount);
+			int32 UniqueWeightsCount = 0;
+
+			const uint8* SkinWeightsBuffer = reinterpret_cast<const uint8*>(MutableMeshVertexBuffers.GetBufferData(BufferIndex));
+			for (int32 ElementIndex = 0; ElementIndex < ElementCount; ++ElementIndex)
+			{
+				uint32 ElementHash;
+				FMemory::Memcpy(&ElementHash, SkinWeightsBuffer, sizeof(int32));
+				SkinWeightsBuffer += sizeof(int32);
+
+				if (ElementHash > 0)
 				{
-					Override.VertexIndexToInfluenceOffset.Add(ElementIndex, *OverrideIndex);
-					SkinWeightsBuffer += MutBoneWeightVertexSize;
+					if (int32* OverrideIndex = HashToUniqueWeightIndexMap.Find({ SkinWeightsBuffer, MutSkinWeightVertexSize, ElementHash }))
+					{
+						Override.VertexIndexToInfluenceOffset.Add(ElementIndex, *OverrideIndex);
+						SkinWeightsBuffer += MutSkinWeightVertexSize;
+					}
+					else
+					{
+						Override.VertexIndexToInfluenceOffset.Add(ElementIndex, UniqueWeightsCount);
+						HashToUniqueWeightIndexMap.Add({ SkinWeightsBuffer, MutSkinWeightVertexSize, ElementHash }, UniqueWeightsCount);
+
+						Override.BoneIDs.SetNumUninitialized((UniqueWeightsCount + 1) * BoneIndicesDataSize, EAllowShrinking::No);
+						FMemory::Memcpy(&Override.BoneIDs[UniqueWeightsCount * BoneIndicesDataSize], SkinWeightsBuffer, BoneIndicesDataSize);
+						SkinWeightsBuffer += MutBoneIndicesDataSize;
+
+						Override.BoneWeights.SetNumUninitialized((UniqueWeightsCount + 1) * BoneWeightsDataSize, EAllowShrinking::No);
+						FMemory::Memcpy(&Override.BoneWeights[UniqueWeightsCount * BoneWeightsDataSize], SkinWeightsBuffer, BoneWeightsDataSize);
+						SkinWeightsBuffer += MutBoneWeightsDataSize;
+						++UniqueWeightsCount;
+					}
 				}
 				else
 				{
-					Override.VertexIndexToInfluenceOffset.Add(ElementIndex, UniqueWeightsCount);
-					HashToUniqueWeightIndexMap.Add({ SkinWeightsBuffer, MutBoneWeightVertexSize, ElementHash }, UniqueWeightsCount);
-
-					Override.BoneIDs.SetNumUninitialized((UniqueWeightsCount + 1) * MutBoneIndicesSize, EAllowShrinking::No);
-					FMemory::Memcpy(&Override.BoneIDs[UniqueWeightsCount * MutBoneIndicesSize], SkinWeightsBuffer, MutBoneIndicesSize);
-					SkinWeightsBuffer += MutBoneIndicesSize;
-
-					Override.BoneWeights.SetNumUninitialized((UniqueWeightsCount + 1) * MutBoneWeightsSize, EAllowShrinking::No);
-					FMemory::Memcpy(&Override.BoneWeights[UniqueWeightsCount * MutBoneWeightsSize], SkinWeightsBuffer, MutBoneWeightsSize);
-					SkinWeightsBuffer += MutBoneWeightsSize;
-					++UniqueWeightsCount;
+					SkinWeightsBuffer += MutSkinWeightVertexSize;
 				}
 			}
-			else
-			{
-				SkinWeightsBuffer += MutBoneWeightVertexSize;
-			}
+
+			Override.BoneIDs.Shrink();
+			Override.BoneWeights.Shrink();
+			Override.VertexIndexToInfluenceOffset.Shrink();
 		}
 
-		Override.BoneIDs.Shrink();
-		Override.BoneWeights.Shrink();
-		Override.VertexIndexToInfluenceOffset.Shrink();
+		LODResource.SkinWeightProfilesData.Init(&LODResource.SkinWeightVertexBuffer);
+		SkeletalMesh.SetSkinWeightProfilesData(LODIndex, LODResource.SkinWeightProfilesData);
 	}
 
 	
 	 void CopySkeletalMeshLODRenderData(
 		 FSkeletalMeshLODRenderData& LODResource,
 		 FSkeletalMeshLODRenderData& SourceLODResource,
-		 const USkeletalMesh& SkeletalMesh,
+		 USkeletalMesh& SkeletalMesh,
+		 int32 LODIndex,
 		 const bool bAllowCPUAccess
 	 )
 	 {
@@ -687,6 +694,9 @@ namespace UnrealConversionUtils
 				 
 				 DestProfile = *SourceProfile;
 			 }
+
+			 LODResource.SkinWeightProfilesData.Init(&LODResource.SkinWeightVertexBuffer);
+			 SkeletalMesh.SetSkinWeightProfilesData(LODIndex, LODResource.SkinWeightProfilesData);
 		 }
 
 		 // Copying Indices
diff --git a/Engine/Plugins/Mutable/Source/CustomizableObject/Public/MuCO/CustomizableObject.h b/Engine/Plugins/Mutable/Source/CustomizableObject/Public/MuCO/CustomizableObject.h
--- a/Engine/Plugins/Mutable/Source/CustomizableObject/Public/MuCO/CustomizableObject.h
+++ b/Engine/Plugins/Mutable/Source/CustomizableObject/Public/MuCO/CustomizableObject.h
@@ -294,7 +294,7 @@ public:
 	bool bEnableMeshCache = false;
 
 	/** Experimental - If true, Mesh LODs will be streamed on demand. It requires streaming of SkeletalMeshes and Mutable.StreamMeshLODsEnabled to be enabled.
-	 *  Does not support Clothing, Morphs, and alternative SkinWeightProfiles yet. */
+	 *  Does not support Clothing and Morphs, yet. */
 	UPROPERTY(EditAnywhere, Category = CustomizableObject)
 	bool bEnableMeshStreaming = false;
 	
diff --git a/Engine/Plugins/Mutable/Source/CustomizableObject/Public/MuCO/CustomizableObjectSkeletalMesh.h b/Engine/Plugins/Mutable/Source/CustomizableObject/Public/MuCO/CustomizableObjectSkeletalMesh.h
--- a/Engine/Plugins/Mutable/Source/CustomizableObject/Public/MuCO/CustomizableObjectSkeletalMesh.h
+++ b/Engine/Plugins/Mutable/Source/CustomizableObject/Public/MuCO/CustomizableObjectSkeletalMesh.h
@@ -17,6 +17,7 @@ namespace mu
 
 class UCustomizableObjectInstance;
 class UCustomizableObject;
+class UCustomizableInstancePrivate;
 class FUpdateContextPrivate;
 
 class FCustomizableObjectMeshStreamIn;
@@ -33,25 +34,20 @@ class CUSTOMIZABLEOBJECT_API UCustomizableObjectSkeletalMesh : public USkeletalM
 	GENERATED_BODY()
 
 	friend FCustomizableObjectMeshStreamIn;
+	friend UCustomizableInstancePrivate;
 
-public: // Public Methods
+private:
 
-	static UCustomizableObjectSkeletalMesh* CreateSkeletalMesh(const TSharedRef<FUpdateContextPrivate>& InOperationData, 
-		const UCustomizableObjectInstance& InInstance, UCustomizableObject& CustomizableObject, 
-		const int32 FirstLOD, const int32 LODCount);
+	static UCustomizableObjectSkeletalMesh* CreateSkeletalMesh(const UCustomizableObjectInstance& InInstance);
 
-private: // Private Methods
+	void InitMutableStreamingData(const TSharedRef<FUpdateContextPrivate>& InOperationData, UCustomizableObject& CustomizableObject,
+		const int32 FirstLOD, const int32 LODCount);
 
 	//~ Begin UStreamableRenderAsset Interface.
 	virtual bool StreamIn(int32 NewMipCount, bool bHighPrio) override;
 	//~ End UStreamableRenderAsset Interface.
 
-
-public: // Public Variables
-
-private: // Private Variables
 	FString CustomizableObjectPathName;
-	FString InstancePathName;
 
 	TSharedPtr<mu::FModel, ESPMode::ThreadSafe> Model;
 	
@@ -61,5 +57,6 @@ private: // Private Variables
 
 	/** Mesh IDs per LOD */
 	TArray<mu::FResourceID> MeshIDs;
+	TArray<TPair<uint32, FName>> SkinWeightProfileIDs;
 
 };
