[mutable] Check for GCd objects pending to load.
#rnx

diff --git a/Engine/Plugins/Mutable/Source/CustomizableObject/Private/MuCO/CustomizableObjectSystem.cpp b/Engine/Plugins/Mutable/Source/CustomizableObject/Private/MuCO/CustomizableObjectSystem.cpp
--- a/Engine/Plugins/Mutable/Source/CustomizableObject/Private/MuCO/CustomizableObjectSystem.cpp
+++ b/Engine/Plugins/Mutable/Source/CustomizableObject/Private/MuCO/CustomizableObjectSystem.cpp
@@ -3705,6 +3705,11 @@ int32 UCustomizableObjectSystem::TickInternal(const bool bBlocking)
 	for ( TArray<TObjectPtr<UCustomizableObject>>::TIterator It = GetPrivate()->ObjectsPendingLoad.CreateIterator(); It; ++It )
 	{
 		UCustomizableObject* CO = (*It);
+		if (!CO)
+		{
+			It.RemoveCurrent();
+			continue;
+		}
 
 		bool bReady = true;
 
