[mutable] Prevent crash when requesting image lod that is not there.

diff --git a/Engine/Plugins/Mutable/Source/MutableRuntime/Internal/MuR/ModelPrivate.h b/Engine/Plugins/Mutable/Source/MutableRuntime/Internal/MuR/ModelPrivate.h
--- a/Engine/Plugins/Mutable/Source/MutableRuntime/Internal/MuR/ModelPrivate.h
+++ b/Engine/Plugins/Mutable/Source/MutableRuntime/Internal/MuR/ModelPrivate.h
@@ -613,11 +613,18 @@ namespace mu
 			TSharedPtr<const FImage> Mip;
 			if (!Index.Streamable)
 			{
-				Mip = ConstantImageLODsPermanent[Index.Index];
+				if (ConstantImageLODsPermanent.IsValidIndex(Index.Index))
+				{
+					Mip = ConstantImageLODsPermanent[Index.Index];
+				}
 			}
 			else
 			{
-				Mip = ConstantImageLODsStreamed[Index.Index];
+				const TSharedPtr<const FImage>* Found = ConstantImageLODsStreamed.Find(Index.Index);
+				if (Found)
+				{
+					Mip = *Found;
+				}
 			}
 
 			return Mip;
