[Mutable] Re-enabled clampig of UV vertices automatically for shared layouts.

#rnx

diff --git a/Engine/Plugins/Mutable/Source/MutableTools/Private/MuT/CodeGenerator.cpp b/Engine/Plugins/Mutable/Source/MutableTools/Private/MuT/CodeGenerator.cpp
--- a/Engine/Plugins/Mutable/Source/MutableTools/Private/MuT/CodeGenerator.cpp
+++ b/Engine/Plugins/Mutable/Source/MutableTools/Private/MuT/CodeGenerator.cpp
@@ -946,7 +946,7 @@ namespace mu
 			// Ensure UV islands remain within their main layout block on lower LODs to avoid unexpected reordering 
 			// of the layout blocks when reusing a surface between LODs. Used to fix small displacements on vertices
 			// that may cause them to fall on a different block.
-			MeshOptions.bClampUVIslands = bShareSurface && bNormalizeUVs;
+			MeshOptions.bClampUVIslands = bShareSurface;
 
             GenerateMesh(MeshOptions, MeshResults, SurfaceNode->Mesh);
 
