[mutable] Add protection to access to mesh resources that are not streamed yet.
#rnx

diff --git a/Engine/Plugins/Mutable/Source/MutableRuntime/Internal/MuR/ModelPrivate.h b/Engine/Plugins/Mutable/Source/MutableRuntime/Internal/MuR/ModelPrivate.h
--- a/Engine/Plugins/Mutable/Source/MutableRuntime/Internal/MuR/ModelPrivate.h
+++ b/Engine/Plugins/Mutable/Source/MutableRuntime/Internal/MuR/ModelPrivate.h
@@ -730,11 +730,18 @@ namespace mu
         {
 			if (!Index.Streamable)
 			{
-				OutMesh = ConstantMeshesPermanent[Index.Index];
+				if (ConstantMeshesPermanent.IsValidIndex(Index.Index))
+				{
+					OutMesh = ConstantMeshesPermanent[Index.Index];
+				}
 			}
 			else
 			{
-				OutMesh = ConstantMeshesStreamed[Index.Index];
+				const TSharedPtr<const FMesh>* Found = ConstantMeshesStreamed.Find(Index.Index);
+				if (Found)
+				{
+					OutMesh = *Found;
+				}
 			}
 		}
 
