[mutable] Fix mesh-closed test log message and add support for the mesh transform op.
#rnx

diff --git a/Engine/Plugins/Mutable/Source/MutableTools/Private/MuT/ASTOpConstantResource.cpp b/Engine/Plugins/Mutable/Source/MutableTools/Private/MuT/ASTOpConstantResource.cpp
--- a/Engine/Plugins/Mutable/Source/MutableTools/Private/MuT/ASTOpConstantResource.cpp
+++ b/Engine/Plugins/Mutable/Source/MutableTools/Private/MuT/ASTOpConstantResource.cpp
@@ -802,7 +802,7 @@ namespace mu
 			const FMesh* Mesh = static_cast<const FMesh*>(GetValue().Get());
 			if (Mesh)
 			{
-				if (!Mesh->IsClosed())
+				if (Mesh->IsClosed())
 				{
 					Result = EClosedMeshTest::Yes;
 				}
diff --git a/Engine/Plugins/Mutable/Source/MutableTools/Private/MuT/ASTOpMeshTransform.cpp b/Engine/Plugins/Mutable/Source/MutableTools/Private/MuT/ASTOpMeshTransform.cpp
--- a/Engine/Plugins/Mutable/Source/MutableTools/Private/MuT/ASTOpMeshTransform.cpp
+++ b/Engine/Plugins/Mutable/Source/MutableTools/Private/MuT/ASTOpMeshTransform.cpp
@@ -91,4 +91,14 @@ namespace mu
 		return {};
 	}
 
+
+	ASTOp::EClosedMeshTest ASTOpMeshTransform::IsClosedMesh(TMap<const ASTOp*, EClosedMeshTest>* Cache) const
+	{
+		if (source)
+		{
+			return source->IsClosedMesh(Cache);
+		}
+		return EClosedMeshTest::Unknown;
+	}
+
 }
diff --git a/Engine/Plugins/Mutable/Source/MutableTools/Private/MuT/ASTOpMeshTransform.h b/Engine/Plugins/Mutable/Source/MutableTools/Private/MuT/ASTOpMeshTransform.h
--- a/Engine/Plugins/Mutable/Source/MutableTools/Private/MuT/ASTOpMeshTransform.h
+++ b/Engine/Plugins/Mutable/Source/MutableTools/Private/MuT/ASTOpMeshTransform.h
@@ -39,6 +39,7 @@ struct FProgram;
 		Ptr<ASTOp> Clone(MapChildFuncRef mapChild) const override;
 		void Link(FProgram& program, FLinkerOptions* Options) override;
 		virtual FSourceDataDescriptor GetSourceDataDescriptor(FGetSourceDataDescriptorContext*) const override;
+		virtual EClosedMeshTest IsClosedMesh(TMap<const ASTOp*, EClosedMeshTest>* Cache) const override;
 
 	};
 
