[Mutable] Add CVar to force run Coder Runner tasks inline.

#rnx

diff --git a/Engine/Plugins/Mutable/Source/MutableRuntime/Private/MuR/CodeRunnerTasked.cpp b/Engine/Plugins/Mutable/Source/MutableRuntime/Private/MuR/CodeRunnerTasked.cpp
--- a/Engine/Plugins/Mutable/Source/MutableRuntime/Private/MuR/CodeRunnerTasked.cpp
+++ b/Engine/Plugins/Mutable/Source/MutableRuntime/Private/MuR/CodeRunnerTasked.cpp
@@ -44,6 +44,13 @@
 #include "Async/Fundamental/Scheduler.h"
 
 
+static TAutoConsoleVariable<bool> CVarCodeRunnerForceInline(
+	TEXT("mutable.CodeRunnerForceInline"),
+	true,
+	TEXT("If true, force all Code Runners to be inline (do not split them into tasks)."),
+	ECVF_Default);
+
+
 DECLARE_CYCLE_STAT(TEXT("MutableCoreTask"), STAT_MutableCoreTask, STATGROUP_Game);
 
 namespace mu
@@ -269,6 +276,11 @@ namespace mu
 	{
 		check(RunnerCompletionEvent.IsCompleted());
 
+		if (CVarCodeRunnerForceInline.GetValueOnAnyThread())
+		{
+			bForceInlineExecution = true;
+		}
+		
 		bUnrecoverableError = false;
 
 		HeapData.SetNum(0, EAllowShrinking::No);
