[Mutable] Added rename method to some nodes.
+ Fixed small bug related to node titles.

#rnx

diff --git a/Engine/Plugins/Mutable/Source/CustomizableObjectEditor/Private/MuCOE/Nodes/CustomizableObjectNodeColorParameter.cpp b/Engine/Plugins/Mutable/Source/CustomizableObjectEditor/Private/MuCOE/Nodes/CustomizableObjectNodeColorParameter.cpp
--- a/Engine/Plugins/Mutable/Source/CustomizableObjectEditor/Private/MuCOE/Nodes/CustomizableObjectNodeColorParameter.cpp
+++ b/Engine/Plugins/Mutable/Source/CustomizableObjectEditor/Private/MuCOE/Nodes/CustomizableObjectNodeColorParameter.cpp
@@ -21,7 +21,7 @@ void UCustomizableObjectNodeColorParameter::AllocateDefaultPins(UCustomizableObj
 
 FText UCustomizableObjectNodeColorParameter::GetNodeTitle(ENodeTitleType::Type TitleType) const
 {
-	if (TitleType == ENodeTitleType::ListView || ParameterName.IsEmpty())
+	if (TitleType == ENodeTitleType::ListView)
 	{
 		return LOCTEXT("Color_Parameter", "Color Parameter");
 	}
diff --git a/Engine/Plugins/Mutable/Source/CustomizableObjectEditor/Private/MuCOE/Nodes/CustomizableObjectNodeComponentMesh.cpp b/Engine/Plugins/Mutable/Source/CustomizableObjectEditor/Private/MuCOE/Nodes/CustomizableObjectNodeComponentMesh.cpp
--- a/Engine/Plugins/Mutable/Source/CustomizableObjectEditor/Private/MuCOE/Nodes/CustomizableObjectNodeComponentMesh.cpp
+++ b/Engine/Plugins/Mutable/Source/CustomizableObjectEditor/Private/MuCOE/Nodes/CustomizableObjectNodeComponentMesh.cpp
@@ -50,16 +50,32 @@ void UCustomizableObjectNodeComponentMesh::BackwardsCompatibleFixup(int32 Custom
 
 FText UCustomizableObjectNodeComponentMesh::GetNodeTitle(ENodeTitleType::Type TitleType) const
 {
+	FText NodeTitle;
+	FFormatNamedArguments Args;
+	Args.Add(TEXT("ComponentName"), FText::FromName(ComponentName));
+
 	if (TitleType == ENodeTitleType::ListView)
 	{
-		return LOCTEXT("Component_Mesh", "Mesh Component");
+		NodeTitle = LOCTEXT("Component_Mesh", "Mesh Component");
+	}
+	else if (TitleType == ENodeTitleType::EditableTitle)
+	{
+		NodeTitle = LOCTEXT("ComponentMesh_Edit", "{ComponentName}");
 	}
 	else
 	{
-		FFormatNamedArguments Args;
-		Args.Add(TEXT("ComponentName"), FText::FromName(ComponentName));
+		NodeTitle = LOCTEXT("ComponentMesh_Title", "{ComponentName}\nMesh Component");
+	}
 
-		return FText::Format(LOCTEXT("ComponentMesh_Title", "{ComponentName}\nMesh Component"), Args);
+	return FText::Format(NodeTitle, Args);
+}
+
+
+void UCustomizableObjectNodeComponentMesh::OnRenameNode(const FString& NewName)
+{
+	if (!NewName.IsEmpty())
+	{
+		ComponentName = FName(*NewName);
 	}
 }
 
diff --git a/Engine/Plugins/Mutable/Source/CustomizableObjectEditor/Private/MuCOE/Nodes/CustomizableObjectNodeComponentMeshAddTo.cpp b/Engine/Plugins/Mutable/Source/CustomizableObjectEditor/Private/MuCOE/Nodes/CustomizableObjectNodeComponentMeshAddTo.cpp
--- a/Engine/Plugins/Mutable/Source/CustomizableObjectEditor/Private/MuCOE/Nodes/CustomizableObjectNodeComponentMeshAddTo.cpp
+++ b/Engine/Plugins/Mutable/Source/CustomizableObjectEditor/Private/MuCOE/Nodes/CustomizableObjectNodeComponentMeshAddTo.cpp
@@ -1,4 +1,4 @@
-﻿// Copyright Epic Games, Inc. All Rights Reserved.
+// Copyright Epic Games, Inc. All Rights Reserved.
 
 #include "MuCOE/Nodes/CustomizableObjectNodeComponentMeshAddTo.h"
 
@@ -28,16 +28,32 @@ void UCustomizableObjectNodeComponentMeshAddTo::PostEditChangeProperty(FProperty
 
 FText UCustomizableObjectNodeComponentMeshAddTo::GetNodeTitle(ENodeTitleType::Type TitleType) const
 {
+	FText NodeTitle;
+	FFormatNamedArguments Args;
+	Args.Add(TEXT("ComponentName"), FText::FromName(ParentComponentName));
+
 	if (TitleType == ENodeTitleType::ListView)
 	{
-		return LOCTEXT("ComponentMeshAdd", "Add To Mesh Component");
+		NodeTitle = LOCTEXT("ComponentMeshAdd", "Add To Mesh Component");
+	}
+	else if (TitleType == ENodeTitleType::EditableTitle)
+	{
+		NodeTitle = LOCTEXT("ComponentMeshAdd_Edit", "{ComponentName}");
 	}
 	else
 	{
-		FFormatNamedArguments Args;
-		Args.Add(TEXT("ComponentName"), FText::FromName(ParentComponentName));
+		NodeTitle = LOCTEXT("ComponentMeshAdd_Title", "{ComponentName}\nAdd To Mesh Component");
+	}
 
-		return FText::Format(LOCTEXT("ComponentMeshAdd_Title", "{ComponentName}\nAdd To Mesh Component"), Args);
+	return FText::Format(NodeTitle, Args);
+}
+
+
+void UCustomizableObjectNodeComponentMeshAddTo::OnRenameNode(const FString& NewName)
+{
+	if (!NewName.IsEmpty())
+	{
+		ParentComponentName = FName(*NewName);
 	}
 }
 
diff --git a/Engine/Plugins/Mutable/Source/CustomizableObjectEditor/Private/MuCOE/Nodes/CustomizableObjectNodeComponentPassthroughMesh.cpp b/Engine/Plugins/Mutable/Source/CustomizableObjectEditor/Private/MuCOE/Nodes/CustomizableObjectNodeComponentPassthroughMesh.cpp
--- a/Engine/Plugins/Mutable/Source/CustomizableObjectEditor/Private/MuCOE/Nodes/CustomizableObjectNodeComponentPassthroughMesh.cpp
+++ b/Engine/Plugins/Mutable/Source/CustomizableObjectEditor/Private/MuCOE/Nodes/CustomizableObjectNodeComponentPassthroughMesh.cpp
@@ -1,4 +1,4 @@
-﻿// Copyright Epic Games, Inc. All Rights Reserved.
+// Copyright Epic Games, Inc. All Rights Reserved.
 
 #include "MuCOE/Nodes/CustomizableObjectNodeComponentPassthroughMesh.h"
 
@@ -139,17 +139,24 @@ bool UCustomizableObjectNodeComponentPassthroughMesh::IsExperimental() const
 
 FText UCustomizableObjectNodeComponentPassthroughMesh::GetNodeTitle(ENodeTitleType::Type TitleType) const
 {
+	FText NodeTitle;
+	FFormatNamedArguments Args;
+	Args.Add(TEXT("ComponentName"), FText::FromName(ComponentName));
+
 	if (TitleType == ENodeTitleType::ListView)
 	{
-		return LOCTEXT("ComponentPassthroughMesh", "Passthrough Mesh Component");
+		NodeTitle = LOCTEXT("ComponentPassthroughMesh", "Passthrough Mesh Component");
+	}
+	else if (TitleType == ENodeTitleType::EditableTitle)
+	{
+		NodeTitle = LOCTEXT("ComponentPassthrough_Edit", "{ComponentName}");
 	}
 	else
 	{
-		FFormatNamedArguments Args;
-		Args.Add(TEXT("ComponentName"), FText::FromName(ComponentName));
-
-		return FText::Format(LOCTEXT("ComponentPassthroughMesh_Title", "{ComponentName}\n Passthrough Mesh Component"), Args);
+		NodeTitle = LOCTEXT("ComponentPassthroughMesh_Title", "{ComponentName}\n Passthrough Mesh Component");
 	}
+
+	return FText::Format(NodeTitle, Args);
 }
 
 
@@ -205,6 +212,15 @@ FText UCustomizableObjectNodeComponentPassthroughMesh::GetTooltipText() const
 }
 
 
+void UCustomizableObjectNodeComponentPassthroughMesh::OnRenameNode(const FString& NewName)
+{
+	if (!NewName.IsEmpty())
+	{
+		ComponentName = FName(*NewName);
+	}
+}
+
+
 UMaterialInterface* UCustomizableObjectNodeComponentPassthroughMesh::GetMaterialInterfaceFor(const int32 LODIndex, const int32 SectionIndex, const FSkeletalMeshModel* ImportedModel) const
 {
 	if (FSkeletalMaterial* SkeletalMaterial = GetSkeletalMaterialFor(LODIndex, SectionIndex, ImportedModel))
diff --git a/Engine/Plugins/Mutable/Source/CustomizableObjectEditor/Private/MuCOE/Nodes/CustomizableObjectNodeEnumParameter.cpp b/Engine/Plugins/Mutable/Source/CustomizableObjectEditor/Private/MuCOE/Nodes/CustomizableObjectNodeEnumParameter.cpp
--- a/Engine/Plugins/Mutable/Source/CustomizableObjectEditor/Private/MuCOE/Nodes/CustomizableObjectNodeEnumParameter.cpp
+++ b/Engine/Plugins/Mutable/Source/CustomizableObjectEditor/Private/MuCOE/Nodes/CustomizableObjectNodeEnumParameter.cpp
@@ -33,7 +33,7 @@ void UCustomizableObjectNodeEnumParameter::AllocateDefaultPins(UCustomizableObje
 
 FText UCustomizableObjectNodeEnumParameter::GetNodeTitle(ENodeTitleType::Type TitleType) const
 {
-	if (TitleType == ENodeTitleType::ListView || ParameterName.IsEmpty())
+	if (TitleType == ENodeTitleType::ListView)
 	{
 		return LOCTEXT("Enum_Parameter", "Enum Parameter");
 	}
diff --git a/Engine/Plugins/Mutable/Source/CustomizableObjectEditor/Private/MuCOE/Nodes/CustomizableObjectNodeExposePin.cpp b/Engine/Plugins/Mutable/Source/CustomizableObjectEditor/Private/MuCOE/Nodes/CustomizableObjectNodeExposePin.cpp
--- a/Engine/Plugins/Mutable/Source/CustomizableObjectEditor/Private/MuCOE/Nodes/CustomizableObjectNodeExposePin.cpp
+++ b/Engine/Plugins/Mutable/Source/CustomizableObjectEditor/Private/MuCOE/Nodes/CustomizableObjectNodeExposePin.cpp
@@ -29,16 +29,25 @@ void UCustomizableObjectNodeExposePin::AllocateDefaultPins(UCustomizableObjectNo
 
 FText UCustomizableObjectNodeExposePin::GetNodeTitle(ENodeTitleType::Type TitleType) const
 {
-	const FText PinCategoryName = UEdGraphSchema_CustomizableObject::GetPinCategoryFriendlyName(PinType);
+	FText NodeTitle;
+	FFormatNamedArguments Args;
+	Args.Add(TEXT("NodeName"), FText::FromString(Name));
+	Args.Add(TEXT("PinType"), UEdGraphSchema_CustomizableObject::GetPinCategoryFriendlyName(PinType));
 
 	if (TitleType == ENodeTitleType::ListView)
 	{
-		return FText::Format(LOCTEXT("Expose_Pin_Title_ListView", "Export {0} Pin"), PinCategoryName);
+		NodeTitle = LOCTEXT("Expose_Pin_Title_ListView", "Export {PinType} Pin");
+	}
+	else if (TitleType == ENodeTitleType::EditableTitle)
+	{
+		NodeTitle = LOCTEXT("Expose_Pin_Title_Edit", "{NodeName}");
 	}
 	else
 	{
-		return FText::Format(LOCTEXT("Expose_Pin_Title", "{0}\nExport {1} Pin"), FText::FromString(Name), PinCategoryName);
+		NodeTitle = LOCTEXT("Expose_Pin_Title", "{NodeName}\nExport {PinType} Pin");
 	}
+
+	return FText::Format(NodeTitle, Args);
 }
 
 
@@ -54,6 +63,16 @@ FText UCustomizableObjectNodeExposePin::GetTooltipText() const
 	return LOCTEXT("Expose_Pin_Tooltip", "Exposes a value to the rest of its Customizable Object hierarchy.");
 }
 
+
+void UCustomizableObjectNodeExposePin::OnRenameNode(const FString& NewName)
+{
+	if (!NewName.IsEmpty())
+	{
+		Name = NewName;
+	}
+}
+
+
 bool UCustomizableObjectNodeExposePin::CanConnect(const UEdGraphPin* InOwnedInputPin, const UEdGraphPin* InOutputPin, bool& bOutIsOtherNodeBlocklisted, bool& bOutArePinsCompatible) const
 {
 	// Chech the pin types do match
diff --git a/Engine/Plugins/Mutable/Source/CustomizableObjectEditor/Private/MuCOE/Nodes/CustomizableObjectNodeFloatParameter.cpp b/Engine/Plugins/Mutable/Source/CustomizableObjectEditor/Private/MuCOE/Nodes/CustomizableObjectNodeFloatParameter.cpp
--- a/Engine/Plugins/Mutable/Source/CustomizableObjectEditor/Private/MuCOE/Nodes/CustomizableObjectNodeFloatParameter.cpp
+++ b/Engine/Plugins/Mutable/Source/CustomizableObjectEditor/Private/MuCOE/Nodes/CustomizableObjectNodeFloatParameter.cpp
@@ -40,18 +40,22 @@ bool UCustomizableObjectNodeFloatParameter::IsAffectedByLOD() const
 
 FText UCustomizableObjectNodeFloatParameter::GetNodeTitle(ENodeTitleType::Type TitleType) const
 {
-	if (TitleType == ENodeTitleType::ListView || ParameterName.IsEmpty())
+	FText NodeTitle;
+
+	if (TitleType == ENodeTitleType::ListView)
 	{
-		return LOCTEXT("Float_Parameter", "Float Parameter");
+		NodeTitle = LOCTEXT("Float_Parameter", "Float Parameter");
 	}
 	else if (TitleType == ENodeTitleType::EditableTitle)
 	{
-		return FText::Format(LOCTEXT("Float_Parameter_EditableTitle", "{0}"), FText::FromString(ParameterName));
+		NodeTitle = FText::Format(LOCTEXT("Float_Parameter_EditableTitle", "{0}"), FText::FromString(ParameterName));
 	}
 	else
 	{
-		return FText::Format(LOCTEXT("Float_Parameter_Title", "{0}\nFloat Parameter"), FText::FromString(ParameterName));
+		NodeTitle = FText::Format(LOCTEXT("Float_Parameter_Title", "{0}\nFloat Parameter"), FText::FromString(ParameterName));
 	}
+
+	return NodeTitle;
 }
 
 
diff --git a/Engine/Plugins/Mutable/Source/CustomizableObjectEditor/Private/MuCOE/Nodes/CustomizableObjectNodeGroupProjectorParameter.cpp b/Engine/Plugins/Mutable/Source/CustomizableObjectEditor/Private/MuCOE/Nodes/CustomizableObjectNodeGroupProjectorParameter.cpp
--- a/Engine/Plugins/Mutable/Source/CustomizableObjectEditor/Private/MuCOE/Nodes/CustomizableObjectNodeGroupProjectorParameter.cpp
+++ b/Engine/Plugins/Mutable/Source/CustomizableObjectEditor/Private/MuCOE/Nodes/CustomizableObjectNodeGroupProjectorParameter.cpp
@@ -86,7 +86,7 @@ TArray<FGroupProjectorParameterImage> UCustomizableObjectNodeGroupProjectorParam
 
 FText UCustomizableObjectNodeGroupProjectorParameter::GetNodeTitle(ENodeTitleType::Type TitleType) const
 {
-	if (TitleType == ENodeTitleType::ListView || ParameterName.IsEmpty())
+	if (TitleType == ENodeTitleType::ListView)
 	{
 		return LOCTEXT("Group_Projector_Parameter", "Group Projector Parameter");
 	}
diff --git a/Engine/Plugins/Mutable/Source/CustomizableObjectEditor/Private/MuCOE/Nodes/CustomizableObjectNodeObject.cpp b/Engine/Plugins/Mutable/Source/CustomizableObjectEditor/Private/MuCOE/Nodes/CustomizableObjectNodeObject.cpp
--- a/Engine/Plugins/Mutable/Source/CustomizableObjectEditor/Private/MuCOE/Nodes/CustomizableObjectNodeObject.cpp
+++ b/Engine/Plugins/Mutable/Source/CustomizableObjectEditor/Private/MuCOE/Nodes/CustomizableObjectNodeObject.cpp
@@ -601,8 +601,10 @@ void UCustomizableObjectNodeObject::AllocateDefaultPins(UCustomizableObjectNodeR
 
 FText UCustomizableObjectNodeObject::GetNodeTitle(ENodeTitleType::Type TitleType) const
 {
-	if (TitleType == ENodeTitleType::ListView ||
-		ObjectName.IsEmpty())
+	FFormatNamedArguments Args;
+	Args.Add(TEXT("ObjectName"), FText::FromString(ObjectName));
+
+	if (TitleType == ENodeTitleType::ListView)
 	{
 		if (bIsBase)
 		{
@@ -613,11 +615,12 @@ FText UCustomizableObjectNodeObject::GetNodeTitle(ENodeTitleType::Type TitleType
 			return LOCTEXT("Base_Object_Deprecated", "Base Object (Deprecated)");
 		}
 	}
+	else if (TitleType == ENodeTitleType::EditableTitle)
+	{
+		return FText::Format(LOCTEXT("Base_Object_Edit", "{ObjectName}"), Args);
+	}
 	else
 	{
-		FFormatNamedArguments Args;	
-		Args.Add(TEXT("ObjectName"), FText::FromString(ObjectName) );
-
 		if (bIsBase)
 		{
 			return FText::Format(LOCTEXT("Base_Object_Title", "{ObjectName}\nBase Object"), Args);
@@ -716,5 +719,13 @@ bool UCustomizableObjectNodeObject::IsSingleOutputNode() const
 	return true;
 }
 
+void UCustomizableObjectNodeObject::OnRenameNode(const FString& NewName)
+{
+	if (!NewName.IsEmpty())
+	{
+		ObjectName = NewName;
+	}
+}
+
 
 #undef LOCTEXT_NAMESPACE
diff --git a/Engine/Plugins/Mutable/Source/CustomizableObjectEditor/Private/MuCOE/Nodes/CustomizableObjectNodeObjectChild.cpp b/Engine/Plugins/Mutable/Source/CustomizableObjectEditor/Private/MuCOE/Nodes/CustomizableObjectNodeObjectChild.cpp
--- a/Engine/Plugins/Mutable/Source/CustomizableObjectEditor/Private/MuCOE/Nodes/CustomizableObjectNodeObjectChild.cpp
+++ b/Engine/Plugins/Mutable/Source/CustomizableObjectEditor/Private/MuCOE/Nodes/CustomizableObjectNodeObjectChild.cpp
@@ -22,6 +22,10 @@ FText UCustomizableObjectNodeObjectChild::GetNodeTitle(ENodeTitleType::Type Titl
 	{
 		return FText::Format(LOCTEXT("Child_Object_Title_List", "{ObjectName} - Child Object"), Args);
 	}
+	else if (TitleType == ENodeTitleType::EditableTitle)
+	{
+		return FText::Format(LOCTEXT("Child_Object_Edit", "{ObjectName}"), Args);
+	}
 	else
 	{
 		return FText::Format(LOCTEXT("Child_Object_Title", "{ObjectName}\nChild Object"), Args);
diff --git a/Engine/Plugins/Mutable/Source/CustomizableObjectEditor/Private/MuCOE/Nodes/CustomizableObjectNodeObjectGroup.cpp b/Engine/Plugins/Mutable/Source/CustomizableObjectEditor/Private/MuCOE/Nodes/CustomizableObjectNodeObjectGroup.cpp
--- a/Engine/Plugins/Mutable/Source/CustomizableObjectEditor/Private/MuCOE/Nodes/CustomizableObjectNodeObjectGroup.cpp
+++ b/Engine/Plugins/Mutable/Source/CustomizableObjectEditor/Private/MuCOE/Nodes/CustomizableObjectNodeObjectGroup.cpp
@@ -86,18 +86,24 @@ void UCustomizableObjectNodeObjectGroup::PostEditChangeProperty(FPropertyChanged
 
 FText UCustomizableObjectNodeObjectGroup::GetNodeTitle(ENodeTitleType::Type TitleType) const
 {
-	if (TitleType == ENodeTitleType::ListView ||
-		GroupName.IsEmpty())
+	FText NodeTitle;
+	FFormatNamedArguments Args;
+	Args.Add(TEXT("GroupName"), FText::FromString(GroupName));
+
+	if (TitleType == ENodeTitleType::ListView)
+	{
+		NodeTitle = LOCTEXT("Object_Group", "Object Group");
+	}
+	else if(TitleType == ENodeTitleType::EditableTitle)
 	{
-		return LOCTEXT("Object_Group", "Object Group");
+		NodeTitle = LOCTEXT("ObjectGroup_Edit", "{GroupName}");
 	}
 	else
 	{
-		FFormatNamedArguments Args;
-		Args.Add(TEXT("GroupName"), FText::FromString(GroupName));
-
-		return FText::Format(LOCTEXT("Group_Object_Title", "{GroupName}\nObject Group"), Args);
+		NodeTitle = LOCTEXT("Group_Object_Title", "{GroupName}\nObject Group");
 	}
+
+	return FText::Format(NodeTitle, Args);
 }
 
 
@@ -115,6 +121,15 @@ FText UCustomizableObjectNodeObjectGroup::GetTooltipText() const
 }
 
 
+void UCustomizableObjectNodeObjectGroup::OnRenameNode(const FString& NewName)
+{
+	if (!NewName.IsEmpty())
+	{
+		GroupName = NewName;
+	}
+}
+
+
 void UCustomizableObjectNodeObjectGroup::AllocateDefaultPins(UCustomizableObjectNodeRemapPins* RemapPins)
 {
 	const UEdGraphSchema_CustomizableObject* Schema = GetDefault<UEdGraphSchema_CustomizableObject>();
diff --git a/Engine/Plugins/Mutable/Source/CustomizableObjectEditor/Private/MuCOE/Nodes/CustomizableObjectNodeProjectorParameter.cpp b/Engine/Plugins/Mutable/Source/CustomizableObjectEditor/Private/MuCOE/Nodes/CustomizableObjectNodeProjectorParameter.cpp
--- a/Engine/Plugins/Mutable/Source/CustomizableObjectEditor/Private/MuCOE/Nodes/CustomizableObjectNodeProjectorParameter.cpp
+++ b/Engine/Plugins/Mutable/Source/CustomizableObjectEditor/Private/MuCOE/Nodes/CustomizableObjectNodeProjectorParameter.cpp
@@ -124,7 +124,7 @@ void UCustomizableObjectNodeProjectorParameter::SetProjectorDefaultAngle(float A
 
 FText UCustomizableObjectNodeProjectorParameter::GetNodeTitle(ENodeTitleType::Type TitleType) const
 {
-	if (TitleType == ENodeTitleType::ListView || ParameterName.IsEmpty())
+	if (TitleType == ENodeTitleType::ListView)
 	{
 		return LOCTEXT("Projector_Parameter", "Projector Parameter");
 	}
diff --git a/Engine/Plugins/Mutable/Source/CustomizableObjectEditor/Private/MuCOE/Nodes/CustomizableObjectNodeTextureParameter.cpp b/Engine/Plugins/Mutable/Source/CustomizableObjectEditor/Private/MuCOE/Nodes/CustomizableObjectNodeTextureParameter.cpp
--- a/Engine/Plugins/Mutable/Source/CustomizableObjectEditor/Private/MuCOE/Nodes/CustomizableObjectNodeTextureParameter.cpp
+++ b/Engine/Plugins/Mutable/Source/CustomizableObjectEditor/Private/MuCOE/Nodes/CustomizableObjectNodeTextureParameter.cpp
@@ -28,7 +28,7 @@ bool UCustomizableObjectNodeTextureParameter::IsExperimental() const
 
 FText UCustomizableObjectNodeTextureParameter::GetNodeTitle(ENodeTitleType::Type TitleType) const
 {
-	if (TitleType == ENodeTitleType::ListView || ParameterName.IsEmpty())
+	if (TitleType == ENodeTitleType::ListView)
 	{
 		return LOCTEXT("Texture_Parameter", "Texture Parameter");
 	}
diff --git a/Engine/Plugins/Mutable/Source/CustomizableObjectEditor/Private/MuCOE/Nodes/CustomizableObjectNodeTransformParameter.cpp b/Engine/Plugins/Mutable/Source/CustomizableObjectEditor/Private/MuCOE/Nodes/CustomizableObjectNodeTransformParameter.cpp
--- a/Engine/Plugins/Mutable/Source/CustomizableObjectEditor/Private/MuCOE/Nodes/CustomizableObjectNodeTransformParameter.cpp
+++ b/Engine/Plugins/Mutable/Source/CustomizableObjectEditor/Private/MuCOE/Nodes/CustomizableObjectNodeTransformParameter.cpp
@@ -10,7 +10,7 @@
 
 FText UCustomizableObjectNodeTransformParameter::GetNodeTitle(ENodeTitleType::Type TitleType) const
 {
-	if (TitleType == ENodeTitleType::ListView || ParameterName.IsEmpty())
+	if (TitleType == ENodeTitleType::ListView)
 	{
 		return LOCTEXT("Transform_Parameter", "Transform Parameter");
 	}
diff --git a/Engine/Plugins/Mutable/Source/CustomizableObjectEditor/Public/MuCOE/Nodes/CustomizableObjectNodeComponentMesh.h b/Engine/Plugins/Mutable/Source/CustomizableObjectEditor/Public/MuCOE/Nodes/CustomizableObjectNodeComponentMesh.h
--- a/Engine/Plugins/Mutable/Source/CustomizableObjectEditor/Public/MuCOE/Nodes/CustomizableObjectNodeComponentMesh.h
+++ b/Engine/Plugins/Mutable/Source/CustomizableObjectEditor/Public/MuCOE/Nodes/CustomizableObjectNodeComponentMesh.h
@@ -56,7 +56,8 @@ public:
 
 	// UEdGraphNode interface
 	virtual FText GetNodeTitle(ENodeTitleType::Type TitleType) const override;
-	
+	virtual void OnRenameNode(const FString& NewName) override;
+
 	UPROPERTY(EditAnywhere, Category = ComponentMesh)
 	FName ComponentName;
 	
diff --git a/Engine/Plugins/Mutable/Source/CustomizableObjectEditor/Public/MuCOE/Nodes/CustomizableObjectNodeComponentMeshAddTo.h b/Engine/Plugins/Mutable/Source/CustomizableObjectEditor/Public/MuCOE/Nodes/CustomizableObjectNodeComponentMeshAddTo.h
--- a/Engine/Plugins/Mutable/Source/CustomizableObjectEditor/Public/MuCOE/Nodes/CustomizableObjectNodeComponentMeshAddTo.h
+++ b/Engine/Plugins/Mutable/Source/CustomizableObjectEditor/Public/MuCOE/Nodes/CustomizableObjectNodeComponentMeshAddTo.h
@@ -1,4 +1,4 @@
-﻿// Copyright Epic Games, Inc. All Rights Reserved.
+// Copyright Epic Games, Inc. All Rights Reserved.
 
 #pragma once
 
@@ -17,6 +17,7 @@ public:
 
 	// UEdGraphNode interface
 	virtual FText GetNodeTitle(ENodeTitleType::Type TitleType) const override;
+	virtual void OnRenameNode(const FString& NewName) override;
 	
 	UPROPERTY(EditAnywhere, Category = ComponentMesh)
 	FName ParentComponentName;
diff --git a/Engine/Plugins/Mutable/Source/CustomizableObjectEditor/Public/MuCOE/Nodes/CustomizableObjectNodeComponentMeshBase.h b/Engine/Plugins/Mutable/Source/CustomizableObjectEditor/Public/MuCOE/Nodes/CustomizableObjectNodeComponentMeshBase.h
--- a/Engine/Plugins/Mutable/Source/CustomizableObjectEditor/Public/MuCOE/Nodes/CustomizableObjectNodeComponentMeshBase.h
+++ b/Engine/Plugins/Mutable/Source/CustomizableObjectEditor/Public/MuCOE/Nodes/CustomizableObjectNodeComponentMeshBase.h
@@ -1,4 +1,4 @@
-﻿// Copyright Epic Games, Inc. All Rights Reserved.
+// Copyright Epic Games, Inc. All Rights Reserved.
 
 #pragma once
 
@@ -28,6 +28,7 @@ public:
 	// UEdGraphNode interface
 	virtual void AllocateDefaultPins(UCustomizableObjectNodeRemapPins* RemapPins) override;
 	virtual FLinearColor GetNodeTitleColor() const override;
+	virtual bool GetCanRenameNode() const override { return true; }
 
 	// UCustomizableObjectNode interface
 	virtual bool IsSingleOutputNode() const override;
diff --git a/Engine/Plugins/Mutable/Source/CustomizableObjectEditor/Public/MuCOE/Nodes/CustomizableObjectNodeComponentPassthroughMesh.h b/Engine/Plugins/Mutable/Source/CustomizableObjectEditor/Public/MuCOE/Nodes/CustomizableObjectNodeComponentPassthroughMesh.h
--- a/Engine/Plugins/Mutable/Source/CustomizableObjectEditor/Public/MuCOE/Nodes/CustomizableObjectNodeComponentPassthroughMesh.h
+++ b/Engine/Plugins/Mutable/Source/CustomizableObjectEditor/Public/MuCOE/Nodes/CustomizableObjectNodeComponentPassthroughMesh.h
@@ -80,6 +80,8 @@ public:
 	virtual FText GetNodeTitle(ENodeTitleType::Type TitleType) const override;
 	virtual FLinearColor GetNodeTitleColor() const override;
 	virtual FText GetTooltipText() const override;
+	virtual bool GetCanRenameNode() const override { return true; }
+	virtual void OnRenameNode(const FString& NewName) override;
 
 	// UCustomizableObjectNode interface
 	virtual void AllocateDefaultPins(UCustomizableObjectNodeRemapPins* RemapPins) override;
diff --git a/Engine/Plugins/Mutable/Source/CustomizableObjectEditor/Public/MuCOE/Nodes/CustomizableObjectNodeExposePin.h b/Engine/Plugins/Mutable/Source/CustomizableObjectEditor/Public/MuCOE/Nodes/CustomizableObjectNodeExposePin.h
--- a/Engine/Plugins/Mutable/Source/CustomizableObjectEditor/Public/MuCOE/Nodes/CustomizableObjectNodeExposePin.h
+++ b/Engine/Plugins/Mutable/Source/CustomizableObjectEditor/Public/MuCOE/Nodes/CustomizableObjectNodeExposePin.h
@@ -31,10 +31,11 @@ public:
 	FText GetNodeTitle(ENodeTitleType::Type TitleType) const override;
 	FLinearColor GetNodeTitleColor() const override;
 	FText GetTooltipText() const override;
-
-	virtual bool CanConnect(const UEdGraphPin* InOwnedInputPin, const UEdGraphPin* InOutputPin, bool& bOutIsOtherNodeBlocklisted, bool& bOutArePinsCompatible) const override;
+	virtual bool GetCanRenameNode() const override { return true; }
+	virtual void OnRenameNode(const FString& NewName) override;
 
 	// UCustomizableObjectNode interface
+	virtual bool CanConnect(const UEdGraphPin* InOwnedInputPin, const UEdGraphPin* InOutputPin, bool& bOutIsOtherNodeBlocklisted, bool& bOutArePinsCompatible) const override;
 	void AllocateDefaultPins(UCustomizableObjectNodeRemapPins* RemapPins) override;
 
 	/** Return the Expose Pin Node expose pin name. */
diff --git a/Engine/Plugins/Mutable/Source/CustomizableObjectEditor/Public/MuCOE/Nodes/CustomizableObjectNodeObject.h b/Engine/Plugins/Mutable/Source/CustomizableObjectEditor/Public/MuCOE/Nodes/CustomizableObjectNodeObject.h
--- a/Engine/Plugins/Mutable/Source/CustomizableObjectEditor/Public/MuCOE/Nodes/CustomizableObjectNodeObject.h
+++ b/Engine/Plugins/Mutable/Source/CustomizableObjectEditor/Public/MuCOE/Nodes/CustomizableObjectNodeObject.h
@@ -158,6 +158,8 @@ public:
 	void PrepareForCopying() override;
 	void PostPasteNode() override;
 	void PostDuplicate(bool bDuplicateForPIE) override;
+	virtual bool GetCanRenameNode() const override { return true; }
+	virtual void OnRenameNode(const FString& NewName) override;
 
 	// UCustomizableObjectNode interface
 	virtual void BackwardsCompatibleFixup(int32 CustomizableObjectCustomVersion) override;
diff --git a/Engine/Plugins/Mutable/Source/CustomizableObjectEditor/Public/MuCOE/Nodes/CustomizableObjectNodeObjectGroup.h b/Engine/Plugins/Mutable/Source/CustomizableObjectEditor/Public/MuCOE/Nodes/CustomizableObjectNodeObjectGroup.h
--- a/Engine/Plugins/Mutable/Source/CustomizableObjectEditor/Public/MuCOE/Nodes/CustomizableObjectNodeObjectGroup.h
+++ b/Engine/Plugins/Mutable/Source/CustomizableObjectEditor/Public/MuCOE/Nodes/CustomizableObjectNodeObjectGroup.h
@@ -49,6 +49,8 @@ public:
 	FText GetNodeTitle(ENodeTitleType::Type TitleType) const override;
 	FLinearColor GetNodeTitleColor() const override;
 	FText GetTooltipText() const override;
+	virtual bool GetCanRenameNode() const override { return true; }
+	virtual void OnRenameNode(const FString& NewName) override;
 
 	// UCustomizableObjectNode interface
 	void AllocateDefaultPins(UCustomizableObjectNodeRemapPins* RemapPins) override;
