[Mutable] Fix Components PImpl pattern + Actor improvements.

- Created ACustomizableSkeletalMeshActorPrivate.
- Renamed CustomizableObjectInstanceUsage due to old assets having the Usage serialized (must be transient).
- Prepared Mutable Actor to use dynamic components.
- Deprecated Actor debug material functions.

diff --git a/Engine/Plugins/Mutable/Source/CustomizableObject/Internal/MuCO/CustomizableObjectInstanceUsagePrivate.h b/Engine/Plugins/Mutable/Source/CustomizableObject/Internal/MuCO/CustomizableObjectInstanceUsagePrivate.h
--- a/Engine/Plugins/Mutable/Source/CustomizableObject/Internal/MuCO/CustomizableObjectInstanceUsagePrivate.h
+++ b/Engine/Plugins/Mutable/Source/CustomizableObject/Internal/MuCO/CustomizableObjectInstanceUsagePrivate.h
@@ -65,4 +65,7 @@ public:
 
 	// Returns true if the NetMode of the associated UCustomizableSkeletalComponent (or the associated SkeletalMeshComponent if the former does not exist) is equal to InNetMode
 	bool IsNetMode(ENetMode InNetMode) const;
+
+	// Used to replace the SkeletalMesh of the parent component by the ReferenceSkeletalMesh or the generated SkeletalMesh 
+	bool bUsedPendingSetSkeletalMesh = false;
 };
diff --git a/Engine/Plugins/Mutable/Source/CustomizableObject/Internal/MuCO/CustomizableSkeletalComponentPrivate.h b/Engine/Plugins/Mutable/Source/CustomizableObject/Internal/MuCO/CustomizableSkeletalComponentPrivate.h
--- a/Engine/Plugins/Mutable/Source/CustomizableObject/Internal/MuCO/CustomizableSkeletalComponentPrivate.h
+++ b/Engine/Plugins/Mutable/Source/CustomizableObject/Internal/MuCO/CustomizableSkeletalComponentPrivate.h
@@ -18,8 +18,6 @@ class CUSTOMIZABLEOBJECT_API UCustomizableSkeletalComponentPrivate : public UObj
 public:
 	UCustomizableSkeletalComponentPrivate();
 	
-	void CreateCustomizableObjectInstanceUsage();
-
 	/** Common end point of all updates. Even those which failed. */
 	void Callbacks() const;
 	
@@ -42,5 +40,8 @@ public:
 	UCustomizableSkeletalComponent* GetPublic();
 	
 	const UCustomizableSkeletalComponent* GetPublic() const;
+
+	UPROPERTY()
+	TObjectPtr<UCustomizableObjectInstanceUsage> InstanceUsage;
 };
 
diff --git a/Engine/Plugins/Mutable/Source/CustomizableObject/Internal/MuCO/CustomizableSkeletalMeshActorPrivate.h b/Engine/Plugins/Mutable/Source/CustomizableObject/Internal/MuCO/CustomizableSkeletalMeshActorPrivate.h
new file mode 100644
--- /dev/null
+++ b/Engine/Plugins/Mutable/Source/CustomizableObject/Internal/MuCO/CustomizableSkeletalMeshActorPrivate.h
@@ -0,0 +1,28 @@
+﻿// Copyright Epic Games, Inc. All Rights Reserved.
+
+#pragma once
+
+#include "UObject/Object.h"
+
+#include "UObject/ObjectPtr.h"
+#include "Containers/Array.h"
+
+#include "CustomizableSkeletalMeshActorPrivate.generated.h"
+
+class UCustomizableObjectInstance;
+class UCustomizableSkeletalComponent;
+class ACustomizableSkeletalMeshActor;
+
+
+UCLASS()
+class CUSTOMIZABLEOBJECT_API UCustomizableSkeletalMeshActorPrivate : public UObject
+{
+	GENERATED_BODY()
+
+public:
+	ACustomizableSkeletalMeshActor* GetPublic();
+
+	void Init(UCustomizableObjectInstance* Instance);
+	
+	TArray<TObjectPtr<UCustomizableSkeletalComponent>>& GetComponents();
+};
diff --git a/Engine/Plugins/Mutable/Source/CustomizableObject/Private/MuCO/CustomizableObjectInstanceUsage.cpp b/Engine/Plugins/Mutable/Source/CustomizableObject/Private/MuCO/CustomizableObjectInstanceUsage.cpp
--- a/Engine/Plugins/Mutable/Source/CustomizableObject/Private/MuCO/CustomizableObjectInstanceUsage.cpp
+++ b/Engine/Plugins/Mutable/Source/CustomizableObject/Private/MuCO/CustomizableObjectInstanceUsage.cpp
@@ -152,7 +152,7 @@ void UCustomizableObjectInstanceUsagePrivate::SetPendingSetSkeletalMesh(bool bIs
 	}
 	else
 	{
-		GetPublic()->bUsedPendingSetSkeletalMesh = bIsActive;
+		bUsedPendingSetSkeletalMesh = bIsActive;
 	}
 }
 
@@ -165,7 +165,7 @@ bool UCustomizableObjectInstanceUsagePrivate::GetPendingSetSkeletalMesh() const
 	}
 	else
 	{
-		return GetPublic()->bUsedPendingSetSkeletalMesh;
+		return bUsedPendingSetSkeletalMesh;
 	}
 }
 
diff --git a/Engine/Plugins/Mutable/Source/CustomizableObject/Private/MuCO/CustomizableSkeletalComponent.cpp b/Engine/Plugins/Mutable/Source/CustomizableObject/Private/MuCO/CustomizableSkeletalComponent.cpp
--- a/Engine/Plugins/Mutable/Source/CustomizableObject/Private/MuCO/CustomizableSkeletalComponent.cpp
+++ b/Engine/Plugins/Mutable/Source/CustomizableObject/Private/MuCO/CustomizableSkeletalComponent.cpp
@@ -11,60 +11,11 @@
 #include "MuCO/CustomizableObject.h"
 
 
-static TAutoConsoleVariable<bool> CVarCustomizableComponentPrivateNullptr(
-	TEXT("mutable.CustomizableComponentPrivateNullptr"),
-	true,
-	TEXT("Rollback. If true, force create the private part after the serialization."),
-	ECVF_Scalability);
-
-
-UCustomizableSkeletalComponentPrivate::UCustomizableSkeletalComponentPrivate()
-{
-	// This object may get instantiated into a level, as a part of UCustomizableSkeletalComponent, so needs to be public to ensure
-	// it can be serialized out.
-	if (!HasAllFlags(RF_ClassDefaultObject))
-	{
-		SetFlags(RF_Public);
-	}
-}
-
-void UCustomizableSkeletalComponentPrivate::CreateCustomizableObjectInstanceUsage()
-{
-	if (GetPublic()->CustomizableObjectInstanceUsage)
-	{
-		// CustomizableObjectInstanceUsage may already exist if duplicated from an existing Customizable Skeletal Component
-		if (GetPublic()->CustomizableObjectInstanceUsage->GetPrivate()->GetCustomizableSkeletalComponent() != GetPublic())
-		{
-			GetPublic()->CustomizableObjectInstanceUsage = nullptr;
-		}
-	}
-
-	AActor* RootActor = GetPublic()->GetAttachmentRootActor();
-	bool bIsDefaultActor = RootActor ?
-		RootActor->HasAnyFlags(RF_ClassDefaultObject) :
-		false;
-
-	if (!GetPublic()->CustomizableObjectInstanceUsage && !HasAnyFlags(RF_ClassDefaultObject) && !bIsDefaultActor)
-	{
-		GetPublic()->CustomizableObjectInstanceUsage = NewObject<UCustomizableObjectInstanceUsage>(this, TEXT("InstanceUsage"), RF_Transient);
-		GetPublic()->CustomizableObjectInstanceUsage->GetPrivate()->SetCustomizableSkeletalComponent(GetPublic());
-	}
-}
-
-
 UCustomizableSkeletalComponent::UCustomizableSkeletalComponent()
 {
-}
-
-void UCustomizableSkeletalComponent::Serialize(FArchive& Ar)
-{
-	Super::Serialize(Ar);
+	Private = CreateDefaultSubobject<UCustomizableSkeletalComponentPrivate>(TEXT("Private"));
 	
-	// Somehow we got some assets serialized with Private = nullptr. Fix them.
-	if (Ar.IsLoading() && !Private && CVarCustomizableComponentPrivateNullptr.GetValueOnAnyThread())
-	{
-		Private = NewObject<UCustomizableSkeletalComponentPrivate>(this, FName("Private"), RF_Public);
-	}
+	GetPrivate()->InstanceUsage->GetPrivate()->SetCustomizableSkeletalComponent(this);
 }
 
 
@@ -88,20 +39,26 @@ void UCustomizableSkeletalComponent::PostEditChangeProperty(FPropertyChangedEven
 #endif
 
 
+UCustomizableSkeletalComponentPrivate::UCustomizableSkeletalComponentPrivate()
+{
+	InstanceUsage = CreateDefaultSubobject<UCustomizableObjectInstanceUsage>(TEXT("InstanceUsage"));
+}
+
+
 void UCustomizableSkeletalComponentPrivate::Callbacks() const
 {
-	if (GetPublic()->CustomizableObjectInstanceUsage)
+	if (InstanceUsage)
 	{
-		GetPublic()->CustomizableObjectInstanceUsage->GetPrivate()->Callbacks();
+		InstanceUsage->GetPrivate()->Callbacks();
 	}
 }
 
 
 USkeletalMesh* UCustomizableSkeletalComponentPrivate::GetSkeletalMesh() const
 {
-	if (GetPublic()->CustomizableObjectInstanceUsage)
+	if (InstanceUsage)
 	{
-		return GetPublic()->CustomizableObjectInstanceUsage->GetPrivate()->GetSkeletalMesh();
+		return InstanceUsage->GetPrivate()->GetSkeletalMesh();
 	}
 
 	return nullptr;
@@ -110,36 +67,36 @@ USkeletalMesh* UCustomizableSkeletalComponentPrivate::GetSkeletalMesh() const
 
 void UCustomizableSkeletalComponentPrivate::SetSkeletalMesh(USkeletalMesh* SkeletalMesh)
 {
-	if (GetPublic()->CustomizableObjectInstanceUsage)
+	if (InstanceUsage)
 	{
-		GetPublic()->CustomizableObjectInstanceUsage->GetPrivate()->SetSkeletalMesh(SkeletalMesh);
+		InstanceUsage->GetPrivate()->SetSkeletalMesh(SkeletalMesh);
 	}
 }
 
 
 void UCustomizableSkeletalComponentPrivate::SetPhysicsAsset(UPhysicsAsset* PhysicsAsset)
 {
-	if (GetPublic()->CustomizableObjectInstanceUsage)
+	if (InstanceUsage)
 	{
-		GetPublic()->CustomizableObjectInstanceUsage->GetPrivate()->SetPhysicsAsset(PhysicsAsset);
+		InstanceUsage->GetPrivate()->SetPhysicsAsset(PhysicsAsset);
 	}
 }
 
 
 void UCustomizableSkeletalComponentPrivate::SetPendingSetSkeletalMesh(bool bIsActive)
 {
-	if (GetPublic()->CustomizableObjectInstanceUsage)
+	if (InstanceUsage)
 	{
-		GetPublic()->CustomizableObjectInstanceUsage->GetPrivate()->SetPendingSetSkeletalMesh(bIsActive);
+		InstanceUsage->GetPrivate()->SetPendingSetSkeletalMesh(bIsActive);
 	}
 }
 
 
 USkeletalMesh* UCustomizableSkeletalComponentPrivate::GetAttachedSkeletalMesh() const
 {
-	if (GetPublic()->CustomizableObjectInstanceUsage)
+	if (InstanceUsage)
 	{
-		return GetPublic()->CustomizableObjectInstanceUsage->GetPrivate()->GetAttachedSkeletalMesh();
+		return InstanceUsage->GetPrivate()->GetAttachedSkeletalMesh();
 	}
 
 	return nullptr;
@@ -204,18 +161,18 @@ bool UCustomizableSkeletalComponent::GetSkipSetSkeletalMeshOnAttach() const
 
 void UCustomizableSkeletalComponent::UpdateSkeletalMeshAsync(bool bNeverSkipUpdate)
 {
-	if (CustomizableObjectInstanceUsage)
+	if (GetPrivate()->InstanceUsage)
 	{
-		CustomizableObjectInstanceUsage->UpdateSkeletalMeshAsync(bNeverSkipUpdate);
+		GetPrivate()->InstanceUsage->UpdateSkeletalMeshAsync(bNeverSkipUpdate);
 	}
 }
 
 
 void UCustomizableSkeletalComponent::UpdateSkeletalMeshAsyncResult(FInstanceUpdateDelegate Callback, bool bIgnoreCloseDist, bool bForceHighPriority)
 {
-	if (CustomizableObjectInstanceUsage)
+	if (GetPrivate()->InstanceUsage)
 	{
-		CustomizableObjectInstanceUsage->UpdateSkeletalMeshAsyncResult(Callback, bIgnoreCloseDist, bForceHighPriority);
+		GetPrivate()->InstanceUsage->UpdateSkeletalMeshAsyncResult(Callback, bIgnoreCloseDist, bForceHighPriority);
 	}
 }
 
@@ -237,9 +194,9 @@ const UCustomizableSkeletalComponentPrivate* UCustomizableSkeletalComponent::Get
 #if WITH_EDITOR
 void UCustomizableSkeletalComponentPrivate::EditorUpdateComponent()
 {
-	if (GetPublic()->CustomizableObjectInstanceUsage)
+	if (InstanceUsage)
 	{
-		GetPublic()->CustomizableObjectInstanceUsage->GetPrivate()->EditorUpdateComponent();
+		InstanceUsage->GetPrivate()->EditorUpdateComponent();
 	}
 }
 #endif
@@ -283,22 +240,3 @@ void UCustomizableSkeletalComponent::OnAttachmentChanged()
 	}
 }
 
-
-void UCustomizableSkeletalComponent::PostInitProperties()
-{
-	Super::PostInitProperties();
-
-	if (!HasAllFlags(RF_ClassDefaultObject))
-	{
-		if (!Private)
-		{
-			Private = NewObject<UCustomizableSkeletalComponentPrivate>(this, FName("Private"), RF_Public);
-		}
-		else if (Private->GetOuter() != this)
-		{
-			Private = Cast<UCustomizableSkeletalComponentPrivate>(StaticDuplicateObject(Private, this, FName("Private")));
-		}
-	
-		GetPrivate()->CreateCustomizableObjectInstanceUsage();
-	}
-}
diff --git a/Engine/Plugins/Mutable/Source/CustomizableObject/Private/MuCO/CustomizableSkeletalMeshActor.cpp b/Engine/Plugins/Mutable/Source/CustomizableObject/Private/MuCO/CustomizableSkeletalMeshActor.cpp
--- a/Engine/Plugins/Mutable/Source/CustomizableObject/Private/MuCO/CustomizableSkeletalMeshActor.cpp
+++ b/Engine/Plugins/Mutable/Source/CustomizableObject/Private/MuCO/CustomizableSkeletalMeshActor.cpp
@@ -1,75 +1,112 @@
 // Copyright Epic Games, Inc. All Rights Reserved.
 
 #include "MuCO/CustomizableSkeletalMeshActor.h"
+#include "MuCO/CustomizableSkeletalMeshActorPrivate.h"
 
 #include "Components/SkeletalMeshComponent.h"
 #include "MuCO/CustomizableSkeletalComponent.h"
 #include "MuCO/CustomizableObject.h"
 #include "MuCO/CustomizableObjectPrivate.h"
 #include "MuCO/CustomizableObjectSystem.h"
+#include "MuCO/CustomizableSkeletalComponentPrivate.h"
+#include "Engine/SkeletalMesh.h"
+#include "MuCO/CustomizableObjectCustomVersion.h"
+#include "MuCO/CustomizableObjectPrivate.h"
+#include "MuCO/UnrealPortabilityHelpers.h"
+
 
 #include UE_INLINE_GENERATED_CPP_BY_NAME(CustomizableSkeletalMeshActor)
 
 #define LOCTEXT_NAMESPACE "CustomizableObject"
 
 
-ACustomizableSkeletalMeshActor::ACustomizableSkeletalMeshActor(const FObjectInitializer& ObjectInitializer)
-	: Super(ObjectInitializer)
+UCustomizableObjectInstance* ACustomizableSkeletalMeshActor::GetInstance()
 {
-	CustomizableSkeletalComponent = CreateDefaultSubobject<UCustomizableSkeletalComponent>(TEXT("CustomizableSkeletalComponent0"));
-	CustomizableSkeletalComponents.Add(CustomizableSkeletalComponent);
-	if (USkeletalMeshComponent* SkeletalMeshComp = GetSkeletalMeshComponent()) 
+	for (const UCustomizableSkeletalComponent* Component : CustomizableSkeletalComponents)
 	{
-		SkeletalMeshComponents.Add(SkeletalMeshComp);
-		bool Success = CustomizableSkeletalComponents[0]->AttachToComponent(SkeletalMeshComp, FAttachmentTransformRules::KeepRelativeTransform);
+		if (UCustomizableObjectInstance* COInstance = Component->CustomizableObjectInstance)
+		{
+			return COInstance;
+		}
 	}
+
+	return nullptr;
 }
 
 
-void ACustomizableSkeletalMeshActor::AttachNewComponent()
+ACustomizableSkeletalMeshActor::ACustomizableSkeletalMeshActor(const FObjectInitializer& Initializer)
 {
-	int32 CurrentIndex = CustomizableSkeletalComponents.Num();
-	FString CustomizableComponentName = FString::Printf(TEXT("CustomizableSkeletalComponent%d"), CurrentIndex);
-	FString SkeletalMeshComponentName = FString::Printf(TEXT("SkeletalMeshComponent%d"), CurrentIndex);
+	Private = CreateDefaultSubobject<UCustomizableSkeletalMeshActorPrivate>(TEXT("Private"));
 
-	//USkeletalMeshComponent* SkeletalMesh = CreateDefaultSubobject<USkeletalMeshComponent>(FName(*SkeletalMeshComponentName));
-	USkeletalMeshComponent* SkeletalMesh = NewObject<USkeletalMeshComponent>(this, USkeletalMeshComponent::StaticClass(), FName(*SkeletalMeshComponentName));
-	USkeletalMeshComponent* RootSkeletalMesh = GetSkeletalMeshComponent();
-	
-	if (SkeletalMesh && RootSkeletalMesh)
-	{		
-		bool Success = SkeletalMesh->AttachToComponent(RootSkeletalMesh, FAttachmentTransformRules::KeepRelativeTransform);
-		if (Success)
+	// Old assets used to create the first UCustomizableSkeletalComponent as a Subobject. To be able to deserialize them we need to create it.
+	// Creating instead an Object instead of a Suobject will not work. Only Component UCustomizableSkeletalComponent 0 is a Subobject.
+	{
+		UCustomizableSkeletalComponent* CustomizableSkeletalComponent = CreateDefaultSubobject<UCustomizableSkeletalComponent>(TEXT("CustomizableSkeletalComponent0"));
+		CustomizableSkeletalComponents.Add(CustomizableSkeletalComponent);
+		if (USkeletalMeshComponent* SkeletalMeshComp = GetSkeletalMeshComponent()) 
 		{
-			//UCustomizableSkeletalComponent* CustomizableSkeletalComponent = CreateDefaultSubobject<UCustomizableSkeletalComponent>(FName(*CustomizableComponentName));
-			UCustomizableSkeletalComponent* NewCustomizableSkeletalComponent = NewObject<UCustomizableSkeletalComponent>(this, UCustomizableSkeletalComponent::StaticClass(), FName(*CustomizableComponentName));
-
-			if (NewCustomizableSkeletalComponent)
-			{
-				Success = NewCustomizableSkeletalComponent->AttachToComponent(SkeletalMesh, FAttachmentTransformRules::KeepRelativeTransform);
-
-				if (Success)
-				{
-					SkeletalMeshComponents.Add(SkeletalMesh);
-					CustomizableSkeletalComponents.Add(NewCustomizableSkeletalComponent);
-				}
-			}
+			SkeletalMeshComponents.Add(SkeletalMeshComp);
+			CustomizableSkeletalComponents[0]->AttachToComponent(SkeletalMeshComp, FAttachmentTransformRules::KeepRelativeTransform);
 		}
 	}
 }
 
 
-UCustomizableObjectInstance* ACustomizableSkeletalMeshActor::GetComponentsCommonInstance()
+void UCustomizableSkeletalMeshActorPrivate::Init(UCustomizableObjectInstance* Instance)
 {
-	for (const UCustomizableSkeletalComponent* Component : CustomizableSkeletalComponents)
+	UCustomizableObject* Object = Instance->GetCustomizableObject();
+	if (!Object)
 	{
-		if (UCustomizableObjectInstance* COInstance = Component->CustomizableObjectInstance)
+		return;
+	}
+	
+	for (int32 ObjectComponentIndex = 0; ObjectComponentIndex < Object->GetComponentCount(); ++ObjectComponentIndex)
+	{
+		const FName& ComponentName = Object->GetComponentName(ObjectComponentIndex);
+		const FString ComponentNameString = ComponentName.ToString();
+
+		USkeletalMeshComponent* Component;
+		if (ObjectComponentIndex == 0)
 		{
-			return COInstance;
+			Component = GetPublic()->SkeletalMeshComponents[0];
+			// Renaming the component will make it disappear. Keep it as it is.
+		}
+		else
+		{
+			const FString SkeletalMeshComponentName = FString::Printf(TEXT("SkeletalMeshComponent %s"), *ComponentNameString);
+			Component = NewObject<USkeletalMeshComponent>(GetPublic(), USkeletalMeshComponent::StaticClass(), FName(*SkeletalMeshComponentName));
+			Component->CreationMethod = EComponentCreationMethod::Native;
+			Component->AttachToComponent(GetPublic()->GetRootComponent(), FAttachmentTransformRules::KeepRelativeTransform);
+			Component->RegisterComponent();
+		}
+					
+		if (!Component)
+		{
+			continue;
 		}
+
+		UCustomizableSkeletalComponent* CustomizableComponent;
+		if (ObjectComponentIndex == 0)
+		{
+			CustomizableComponent = GetPublic()->CustomizableSkeletalComponents[0];
+			// Renaming the component will make it disappear. Keep it as it is.
+		}
+		else
+		{
+			const FString CustomizableComponentName = FString::Printf(TEXT("CustomizableSkeletalComponent %s"), *ComponentNameString);
+			CustomizableComponent = NewObject<UCustomizableSkeletalComponent>(GetPublic(), UCustomizableSkeletalComponent::StaticClass(), FName(CustomizableComponentName));
+		}
+
+		CustomizableComponent->AttachToComponent(Component, FAttachmentTransformRules::KeepRelativeTransform);
+		CustomizableComponent->CustomizableObjectInstance = Instance;
+		CustomizableComponent->SetComponentName(ComponentName);
+		CustomizableComponent->RegisterComponent();
+		
+		GetPublic()->SkeletalMeshComponents.Add(Component);
+		GetPublic()->CustomizableSkeletalComponents.Add(CustomizableComponent);
 	}
 
-	return nullptr;
+	Instance->UpdateSkeletalMeshAsync();
 }
 
 
@@ -89,7 +126,7 @@ void ACustomizableSkeletalMeshActor::EnableDebugMaterial(bool bEnableDebugMateri
 	bRemoveDebugMaterial = bDebugMaterialEnabled && !bEnableDebugMaterial;
 	bDebugMaterialEnabled = bEnableDebugMaterial;
 
-	if (UCustomizableObjectInstance* COInstance = GetComponentsCommonInstance())
+	if (UCustomizableObjectInstance* COInstance = GetInstance())
 	{
 		//Bind Instance Update delegate to Actor
 		COInstance->UpdatedDelegate.AddUniqueDynamic(this, &ACustomizableSkeletalMeshActor::SwitchComponentsMaterials);
@@ -98,6 +135,20 @@ void ACustomizableSkeletalMeshActor::EnableDebugMaterial(bool bEnableDebugMateri
 }
 
 
+UCustomizableSkeletalMeshActorPrivate* ACustomizableSkeletalMeshActor::GetPrivate()
+{
+	check(Private);
+	return Private;
+}
+
+
+const UCustomizableSkeletalMeshActorPrivate* ACustomizableSkeletalMeshActor::GetPrivate() const
+{
+	check(Private);
+	return Private;
+}
+
+
 void ACustomizableSkeletalMeshActor::SwitchComponentsMaterials(UCustomizableObjectInstance* Instance)
 {
 	if (!DebugMaterial)
@@ -107,7 +158,7 @@ void ACustomizableSkeletalMeshActor::SwitchComponentsMaterials(UCustomizableObje
 
 	if (bDebugMaterialEnabled || bRemoveDebugMaterial)
 	{
-		UCustomizableObjectInstance* COInstance = GetComponentsCommonInstance();
+		UCustomizableObjectInstance* COInstance = GetInstance();
 
 		if (!COInstance)
 		{
@@ -149,4 +200,16 @@ void ACustomizableSkeletalMeshActor::SwitchComponentsMaterials(UCustomizableObje
 	}
 }
 
+ACustomizableSkeletalMeshActor* UCustomizableSkeletalMeshActorPrivate::GetPublic()
+{
+	return CastChecked<ACustomizableSkeletalMeshActor>(GetOuter());
+}
+
+
+TArray<TObjectPtr<UCustomizableSkeletalComponent>>& UCustomizableSkeletalMeshActorPrivate::GetComponents()
+{
+	return GetPublic()->CustomizableSkeletalComponents;
+}
+
+
 #undef LOCTEXT_NAMESPACE
diff --git a/Engine/Plugins/Mutable/Source/CustomizableObject/Public/MuCO/CustomizableObjectInstanceUsage.h b/Engine/Plugins/Mutable/Source/CustomizableObject/Public/MuCO/CustomizableObjectInstanceUsage.h
--- a/Engine/Plugins/Mutable/Source/CustomizableObject/Public/MuCO/CustomizableObjectInstanceUsage.h
+++ b/Engine/Plugins/Mutable/Source/CustomizableObject/Public/MuCO/CustomizableObjectInstanceUsage.h
@@ -111,9 +111,6 @@ private:
 	/** Only used if the UsedComponentIndex is INDEX_NONE. */
 	UPROPERTY()
 	FName UsedComponentName;
-	
-	// Used to replace the SkeletalMesh of the parent component by the ReferenceSkeletalMesh or the generated SkeletalMesh 
-	bool bUsedPendingSetSkeletalMesh = false;
 
 	// Used to avoid replacing the SkeletalMesh of the parent component by the ReferenceSkeletalMesh if bPendingSetSkeletalMesh is true
 	UPROPERTY()
@@ -122,7 +119,7 @@ private:
 	UPROPERTY()
 	bool bUsedSkipSetSkeletalMeshOnAttach = false;
 	
-	UPROPERTY()
+	UPROPERTY( VisibleAnywhere, Category=CO)
 	TObjectPtr<UCustomizableObjectInstanceUsagePrivate> Private;
 	
 	friend UCustomizableObjectInstanceUsagePrivate;
diff --git a/Engine/Plugins/Mutable/Source/CustomizableObject/Public/MuCO/CustomizableSkeletalComponent.h b/Engine/Plugins/Mutable/Source/CustomizableObject/Public/MuCO/CustomizableSkeletalComponent.h
--- a/Engine/Plugins/Mutable/Source/CustomizableObject/Public/MuCO/CustomizableSkeletalComponent.h
+++ b/Engine/Plugins/Mutable/Source/CustomizableObject/Public/MuCO/CustomizableSkeletalComponent.h
@@ -49,9 +49,6 @@ private:
 
 	UPROPERTY()
 	bool bSkipSkipSetSkeletalMeshOnAttach = false;
-
-	UPROPERTY(Transient)
-	TObjectPtr<UCustomizableObjectInstanceUsage> CustomizableObjectInstanceUsage;
 	
 	UPROPERTY()
 	TObjectPtr<UCustomizableSkeletalComponentPrivate> Private;
@@ -64,10 +61,6 @@ protected:
 #if WITH_EDITOR
 	virtual void PostEditChangeProperty(FPropertyChangedEvent& PropertyChangedEvent) override;
 #endif
-	virtual void Serialize(FArchive& Ar) override;
-
-	// UObject interface
-	virtual void PostInitProperties() override;
 
 	// USceneComponent interface
 	virtual void OnAttachmentChanged() override;
diff --git a/Engine/Plugins/Mutable/Source/CustomizableObject/Public/MuCO/CustomizableSkeletalMeshActor.h b/Engine/Plugins/Mutable/Source/CustomizableObject/Public/MuCO/CustomizableSkeletalMeshActor.h
--- a/Engine/Plugins/Mutable/Source/CustomizableObject/Public/MuCO/CustomizableSkeletalMeshActor.h
+++ b/Engine/Plugins/Mutable/Source/CustomizableObject/Public/MuCO/CustomizableSkeletalMeshActor.h
@@ -8,60 +8,59 @@
 
 class UCustomizableObjectInstance;
 class UCustomizableSkeletalComponent;
+class UCustomizableSkeletalMeshActorPrivate;
 
 class UObject;
 
 UCLASS(ClassGroup = CustomizableObject, Blueprintable, ComponentWrapperClass, ConversionRoot, meta = (ChildCanTick))
 class CUSTOMIZABLEOBJECT_API ACustomizableSkeletalMeshActor : public ASkeletalMeshActor
 {
-	GENERATED_UCLASS_BODY()
-public:
-	
-	class UCustomizableSkeletalComponent* GetCustomizableSkeletalComponent(int32 ObjectComponentIndex = 0) const { return CustomizableSkeletalComponents[ObjectComponentIndex]; }
-	
-	class USkeletalMeshComponent* GetSkeletalMeshComponentAt(int32 ObjectComponentIndex = 0) const { return SkeletalMeshComponents[ObjectComponentIndex]; }
-
-	void AttachNewComponent();
+	GENERATED_BODY()
 
-	int32 GetNumComponents() { return CustomizableSkeletalComponents.Num(); }
+	friend UCustomizableSkeletalMeshActorPrivate;
 
-	/** Sets which material will be used for debugging purposes */
-	UFUNCTION(BlueprintCallable, Category = CustomizableSkeletalMeshActor)
+public:
+	ACustomizableSkeletalMeshActor(const FObjectInitializer& Initializer);
+	
+	UE_DEPRECATED(5.6, "Plase replace with custom code.")
+	UFUNCTION(BlueprintCallable, Category = CustomizableSkeletalMeshActor, meta=(DeprecatedFunction))
 	void SetDebugMaterial(UMaterialInterface* InDebugMaterial);
 
-	/** Enables the debug material to all the componets of the actor */
-	UFUNCTION(BlueprintCallable, Category = CustomizableSkeletalMeshActor)
+	UE_DEPRECATED(5.6, "Plase replace with custom code.")
+	UFUNCTION(BlueprintCallable, Category = CustomizableSkeletalMeshActor, meta=(DeprecatedFunction))
 	void EnableDebugMaterial(bool bEnableDebugMaterial);
 
-private:
+	UCustomizableSkeletalMeshActorPrivate* GetPrivate();
+	const UCustomizableSkeletalMeshActorPrivate* GetPrivate() const;
 
-	/** Changes materials to all actor componets */
+private:
+	void Init(UCustomizableObjectInstance* Instance);
+	
+	// DEPRECATED
 	UFUNCTION()
 	void SwitchComponentsMaterials(UCustomizableObjectInstance* Instance);
 
-	/** return: the instance represented in all the actor components */
-	UCustomizableObjectInstance* GetComponentsCommonInstance();
+	/** @return the instance represented in all the actor components. */
+	UCustomizableObjectInstance* GetInstance();
 
-private:
-	
 	/** TODO: There are as many as components in the CO (not in the COI). */
-	UPROPERTY(Category = CustomizableSkeletalMesh, VisibleAnywhere, BlueprintReadOnly, meta = (AllowPrivateAccess = "true"))
-	TArray< TObjectPtr<class UCustomizableSkeletalComponent> > CustomizableSkeletalComponents;
+	UPROPERTY(VisibleAnywhere, Category = CustomizableSkeletalMeshActor)
+	TArray<TObjectPtr<UCustomizableSkeletalComponent>> CustomizableSkeletalComponents;
 
 	/** TODO: There are as many as components in the CO (not in the COI). */
-	UPROPERTY(Category = CustomizableSkeletalMesh, VisibleAnywhere, BlueprintReadOnly, meta = (AllowPrivateAccess = "true"))
-	TArray< TObjectPtr<class USkeletalMeshComponent>> SkeletalMeshComponents;
-
-	// TODO: This is a temporary fix to not break the demos, we should update the demos to support the arrays of components
-	UPROPERTY(Category = CustomizableSkeletalMesh, VisibleAnywhere, BlueprintReadOnly, meta = (AllowPrivateAccess = "true"))
-	TObjectPtr<class UCustomizableSkeletalComponent> CustomizableSkeletalComponent;
+	UPROPERTY(VisibleAnywhere, Category = CustomizableSkeletalMeshActor)
+	TArray<TObjectPtr<USkeletalMeshComponent>> SkeletalMeshComponents;
 
-	/** Material used for debugging */
+	// DEPRECATED
+	UPROPERTY()
 	TObjectPtr<UMaterialInterface> DebugMaterial;
 
-	/** bool that determines if the debug material has been enabled */
+	// DEPRECATED
 	bool bDebugMaterialEnabled = false;
 
-	/** bool used to remove the debug material whet it has been disabled */
+	// DEPRECATED
 	bool bRemoveDebugMaterial = false;
+
+	UPROPERTY()
+	TObjectPtr<UCustomizableSkeletalMeshActorPrivate> Private = nullptr;
 };
diff --git a/Engine/Plugins/Mutable/Source/CustomizableObjectEditor/Private/MuCOE/CustomizableObjectInstanceFactory.cpp b/Engine/Plugins/Mutable/Source/CustomizableObjectEditor/Private/MuCOE/CustomizableObjectInstanceFactory.cpp
--- a/Engine/Plugins/Mutable/Source/CustomizableObjectEditor/Private/MuCOE/CustomizableObjectInstanceFactory.cpp
+++ b/Engine/Plugins/Mutable/Source/CustomizableObjectEditor/Private/MuCOE/CustomizableObjectInstanceFactory.cpp
@@ -12,6 +12,7 @@
 #include "MuCO/CustomizableObjectPrivate.h"
 #include "MuCO/CustomizableSkeletalComponent.h"
 #include "MuCO/CustomizableSkeletalMeshActor.h"
+#include "MuCO/CustomizableSkeletalMeshActorPrivate.h"
 #include "MuCOE/Nodes/CustomizableObjectNodeObject.h"
 #include "MuCO/UnrealPortabilityHelpers.h"
 
@@ -48,63 +49,26 @@ void UCustomizableObjectInstanceFactory::PostSpawnActor(UObject* Asset, AActor*
 	{
 		return;
 	}
-	
-
-	for (int32 ObjectComponentIndex = 0; ObjectComponentIndex < Object->GetComponentCount(); ++ObjectComponentIndex)
-	{
-		const FName& ComponentName = Object->GetPrivate()->GetComponentName(FCustomizableObjectComponentIndex(ObjectComponentIndex));
-		
-		USkeletalMesh* SkeletalMesh = Instance->GetComponentMeshSkeletalMesh(ComponentName);
-		if (!SkeletalMesh)
-		{
-			SkeletalMesh = Object->GetComponentMeshReferenceSkeletalMesh(ComponentName);
-		}
-		
-		if (ObjectComponentIndex > 0 && ObjectComponentIndex > NewCSMActor->GetNumComponents() - 1)
-		{
-			NewCSMActor->AttachNewComponent();
-		}
-
-		if (USkeletalMeshComponent* SkeletalMeshComp = NewCSMActor->GetSkeletalMeshComponentAt(ObjectComponentIndex))
-		{
-			SkeletalMeshComp->UnregisterComponent();
-			UE_MUTABLE_SETSKINNEDASSET(SkeletalMeshComp, SkeletalMesh);
 
-			if (ObjectComponentIndex == 0 && NewCSMActor->GetWorld()->IsGameWorld())
-			{
-				NewCSMActor->ReplicatedMesh = SkeletalMesh;
-			}
-
-			if (UCustomizableSkeletalComponent* CustomSkeletalComp = NewCSMActor->GetCustomizableSkeletalComponent(ObjectComponentIndex))
-			{
-				CustomSkeletalComp->UnregisterComponent();
-				CustomSkeletalComp->CustomizableObjectInstance = Instance;
-				CustomSkeletalComp->SetComponentName(ComponentName);
-				CustomSkeletalComp->GetPrivate()->SetSkeletalMesh(SkeletalMesh);
-				CustomSkeletalComp->UpdateSkeletalMeshAsync();
-				CustomSkeletalComp->RegisterComponent();
-			}
-
-			SkeletalMeshComp->RegisterComponent();
-		}
-	}
+	NewCSMActor->GetPrivate()->Init(Instance);
 }
 
+
 UObject* UCustomizableObjectInstanceFactory::GetAssetFromActorInstance(AActor* ActorInstance)
 {
 	if (ACustomizableSkeletalMeshActor* CSMActor = CastChecked<ACustomizableSkeletalMeshActor>(ActorInstance))
 	{
-		if (CSMActor->GetNumComponents() > 0)
+		const TArray<TObjectPtr<UCustomizableSkeletalComponent>>& Components = CSMActor->GetPrivate()->GetComponents();
+		if (Components.IsValidIndex(0))
 		{
-			if (UCustomizableSkeletalComponent* CustomSkeletalComp = CSMActor->GetCustomizableSkeletalComponent(0))
-			{
-				return CustomSkeletalComp->CustomizableObjectInstance;
-			}
+			return Components[0];
 		}
 	}
+	
 	return nullptr;
 }
 
+
 bool UCustomizableObjectInstanceFactory::CanCreateActorFrom(const FAssetData& AssetData, FText& OutErrorMsg)
 {
     if (!AssetData.IsValid() ||
diff --git a/Engine/Plugins/Mutable/Source/CustomizableObjectEditor/Public/MuCOE/CustomizableObjectInstanceFactory.h b/Engine/Plugins/Mutable/Source/CustomizableObjectEditor/Public/MuCOE/CustomizableObjectInstanceFactory.h
--- a/Engine/Plugins/Mutable/Source/CustomizableObjectEditor/Public/MuCOE/CustomizableObjectInstanceFactory.h
+++ b/Engine/Plugins/Mutable/Source/CustomizableObjectEditor/Public/MuCOE/CustomizableObjectInstanceFactory.h
@@ -17,11 +17,9 @@ class UCustomizableObjectInstanceFactory : public UActorFactory
 	GENERATED_UCLASS_BODY()
 
 protected:
-	//~ Begin UActorFactory Interface
-	void PostSpawnActor(UObject* Asset, AActor* NewActor) override;
-	UObject* GetAssetFromActorInstance(AActor* ActorInstance) override;
-	bool CanCreateActorFrom(const FAssetData& AssetData, FText& OutErrorMsg) override;
-	//~ End UActorFactory Interface
-
-	FQuat AlignObjectToSurfaceNormal(const FVector& InSurfaceNormal, const FQuat& ActorRotation) const;
+	// UActorFactory interface
+	virtual void PostSpawnActor(UObject* Asset, AActor* NewActor) override;
+	virtual UObject* GetAssetFromActorInstance(AActor* ActorInstance) override;
+	virtual bool CanCreateActorFrom(const FAssetData& AssetData, FText& OutErrorMsg) override;
+	virtual FQuat AlignObjectToSurfaceNormal(const FVector& InSurfaceNormal, const FQuat& ActorRotation) const override;
 };
