[mutable] Fixed crash when editing a transform parameter though a slider widget in the CO Editor.
#rnx
#tests Editor

diff --git a/Engine/Plugins/Mutable/Source/CustomizableObjectEditor/Private/MuCOE/CustomizableInstanceDetails.cpp b/Engine/Plugins/Mutable/Source/CustomizableObjectEditor/Private/MuCOE/CustomizableInstanceDetails.cpp
--- a/Engine/Plugins/Mutable/Source/CustomizableObjectEditor/Private/MuCOE/CustomizableInstanceDetails.cpp
+++ b/Engine/Plugins/Mutable/Source/CustomizableObjectEditor/Private/MuCOE/CustomizableInstanceDetails.cpp
@@ -1503,6 +1503,21 @@ TSharedRef<SWidget> FCustomizableInstanceDetails::GenerateTransformWidget(const
 		}
 	};
 	
+	auto BeginSliderMovement = [WeakThis = SharedThis(this).ToWeakPtr()]()
+	{
+		if (TSharedPtr<FCustomizableInstanceDetails> SharedThis = WeakThis.Pin()) 
+		{ 
+			SharedThis->bUpdatingSlider = true; 
+		}
+	};
+
+	auto EndSliderMovement = [WeakThis = SharedThis(this).ToWeakPtr()](float NewValue)
+	{
+		if (TSharedPtr<FCustomizableInstanceDetails> SharedThis = WeakThis.Pin())
+		{
+			SharedThis->bUpdatingSlider = false;
+		}
+	};
 	
 	return SNew(SGridPanel)
 		.FillColumn(1, 1)
@@ -1531,6 +1546,13 @@ TSharedRef<SWidget> FCustomizableInstanceDetails::GenerateTransformWidget(const
 			.OnYCommitted_Lambda([OnLocationCommitted](float Value, ETextCommit::Type Type) { OnLocationCommitted(Value, Type, EAxis::Y); })
 			.OnZCommitted_Lambda([OnLocationCommitted](float Value, ETextCommit::Type Type) { OnLocationCommitted(Value, Type, EAxis::Z); })
 			.Font(IDetailLayoutBuilder::GetDetailFont())
+
+			.OnXBeginSliderMovement_Lambda([BeginSliderMovement]() { BeginSliderMovement(); })
+			.OnXEndSliderMovement_Lambda([EndSliderMovement](float NewValue) { EndSliderMovement(NewValue); })
+			.OnYBeginSliderMovement_Lambda([BeginSliderMovement]() { BeginSliderMovement(); })
+			.OnYEndSliderMovement_Lambda([EndSliderMovement](float NewValue) { EndSliderMovement(NewValue); })
+			.OnZBeginSliderMovement_Lambda([BeginSliderMovement]() { BeginSliderMovement(); })
+			.OnZEndSliderMovement_Lambda([EndSliderMovement](float NewValue) { EndSliderMovement(NewValue); })
 		]
 		
 		+ SGridPanel::Slot(0, 1)
@@ -1558,6 +1580,13 @@ TSharedRef<SWidget> FCustomizableInstanceDetails::GenerateTransformWidget(const
 			.OnPitchCommitted_Lambda([OnRotationCommitted](float Value, ETextCommit::Type Type) { OnRotationCommitted(Value, Type, EAxis::Y); })
 			.OnYawCommitted_Lambda([OnRotationCommitted](float Value, ETextCommit::Type Type) { OnRotationCommitted(Value, Type, EAxis::Z); })
 			.Font(IDetailLayoutBuilder::GetDetailFont())
+
+			.OnRollBeginSliderMovement_Lambda([BeginSliderMovement]() { BeginSliderMovement(); })
+			.OnRollEndSliderMovement_Lambda([EndSliderMovement](float NewValue) { EndSliderMovement(NewValue); })
+			.OnPitchBeginSliderMovement_Lambda([BeginSliderMovement]() { BeginSliderMovement(); })
+			.OnPitchEndSliderMovement_Lambda([EndSliderMovement](float NewValue) { EndSliderMovement(NewValue); })
+			.OnYawBeginSliderMovement_Lambda([BeginSliderMovement]() { BeginSliderMovement(); })
+			.OnYawEndSliderMovement_Lambda([EndSliderMovement](float NewValue) { EndSliderMovement(NewValue); })
 		]
 
 		+ SGridPanel::Slot(0, 2)
@@ -1586,6 +1615,13 @@ TSharedRef<SWidget> FCustomizableInstanceDetails::GenerateTransformWidget(const
 			.OnYCommitted_Lambda([OnScaleCommitted](float Value, ETextCommit::Type Type) { OnScaleCommitted(Value, Type, EAxis::Y); })
 			.OnZCommitted_Lambda([OnScaleCommitted](float Value, ETextCommit::Type Type) { OnScaleCommitted(Value, Type, EAxis::Z); })
 			.Font(IDetailLayoutBuilder::GetDetailFont())
+
+			.OnXBeginSliderMovement_Lambda([BeginSliderMovement]() { BeginSliderMovement(); })
+			.OnXEndSliderMovement_Lambda([EndSliderMovement](float NewValue) { EndSliderMovement(NewValue); })
+			.OnYBeginSliderMovement_Lambda([BeginSliderMovement]() { BeginSliderMovement(); })
+			.OnYEndSliderMovement_Lambda([EndSliderMovement](float NewValue) { EndSliderMovement(NewValue); })
+			.OnZBeginSliderMovement_Lambda([BeginSliderMovement]() { BeginSliderMovement(); })
+			.OnZEndSliderMovement_Lambda([EndSliderMovement](float NewValue) { EndSliderMovement(NewValue); })
 		];
 }
 
