[Backout] - CL38658761
Original CL Desc
-----------------------------------------------------------------
[Mutable] Fix Components PImpl pattern + Actor improvements.

- Created ACustomizableSkeletalMeshActorPrivate.
- Renamed CustomizableObjectInstanceUsage due to old assets having the Usage serialized (must be transient).
- Prepared Mutable Actor to use dynamic components.
- Deprecated Actor debug material functions.

diff --git a/Engine/Plugins/Mutable/Source/CustomizableObject/Internal/MuCO/CustomizableObjectInstanceUsagePrivate.h b/Engine/Plugins/Mutable/Source/CustomizableObject/Internal/MuCO/CustomizableObjectInstanceUsagePrivate.h
--- a/Engine/Plugins/Mutable/Source/CustomizableObject/Internal/MuCO/CustomizableObjectInstanceUsagePrivate.h
+++ b/Engine/Plugins/Mutable/Source/CustomizableObject/Internal/MuCO/CustomizableObjectInstanceUsagePrivate.h
@@ -65,7 +65,4 @@ public:
 
 	// Returns true if the NetMode of the associated UCustomizableSkeletalComponent (or the associated SkeletalMeshComponent if the former does not exist) is equal to InNetMode
 	bool IsNetMode(ENetMode InNetMode) const;
-
-	// Used to replace the SkeletalMesh of the parent component by the ReferenceSkeletalMesh or the generated SkeletalMesh 
-	bool bUsedPendingSetSkeletalMesh = false;
 };
diff --git a/Engine/Plugins/Mutable/Source/CustomizableObject/Internal/MuCO/CustomizableSkeletalComponentPrivate.h b/Engine/Plugins/Mutable/Source/CustomizableObject/Internal/MuCO/CustomizableSkeletalComponentPrivate.h
--- a/Engine/Plugins/Mutable/Source/CustomizableObject/Internal/MuCO/CustomizableSkeletalComponentPrivate.h
+++ b/Engine/Plugins/Mutable/Source/CustomizableObject/Internal/MuCO/CustomizableSkeletalComponentPrivate.h
@@ -18,6 +18,8 @@ class CUSTOMIZABLEOBJECT_API UCustomizableSkeletalComponentPrivate : public UObj
 public:
 	UCustomizableSkeletalComponentPrivate();
 	
+	void CreateCustomizableObjectInstanceUsage();
+
 	/** Common end point of all updates. Even those which failed. */
 	void Callbacks() const;
 	
@@ -40,8 +42,5 @@ public:
 	UCustomizableSkeletalComponent* GetPublic();
 	
 	const UCustomizableSkeletalComponent* GetPublic() const;
-
-	UPROPERTY()
-	TObjectPtr<UCustomizableObjectInstanceUsage> InstanceUsage;
 };
 
diff --git a/Engine/Plugins/Mutable/Source/CustomizableObject/Internal/MuCO/CustomizableSkeletalMeshActorPrivate.h b/Engine/Plugins/Mutable/Source/CustomizableObject/Internal/MuCO/CustomizableSkeletalMeshActorPrivate.h
deleted file mode 100644
diff --git a/Engine/Plugins/Mutable/Source/CustomizableObject/Private/MuCO/CustomizableObjectInstanceUsage.cpp b/Engine/Plugins/Mutable/Source/CustomizableObject/Private/MuCO/CustomizableObjectInstanceUsage.cpp
--- a/Engine/Plugins/Mutable/Source/CustomizableObject/Private/MuCO/CustomizableObjectInstanceUsage.cpp
+++ b/Engine/Plugins/Mutable/Source/CustomizableObject/Private/MuCO/CustomizableObjectInstanceUsage.cpp
@@ -152,7 +152,7 @@ void UCustomizableObjectInstanceUsagePrivate::SetPendingSetSkeletalMesh(bool bIs
 	}
 	else
 	{
-		bUsedPendingSetSkeletalMesh = bIsActive;
+		GetPublic()->bUsedPendingSetSkeletalMesh = bIsActive;
 	}
 }
 
@@ -165,7 +165,7 @@ bool UCustomizableObjectInstanceUsagePrivate::GetPendingSetSkeletalMesh() const
 	}
 	else
 	{
-		return bUsedPendingSetSkeletalMesh;
+		return GetPublic()->bUsedPendingSetSkeletalMesh;
 	}
 }
 
diff --git a/Engine/Plugins/Mutable/Source/CustomizableObject/Private/MuCO/CustomizableSkeletalComponent.cpp b/Engine/Plugins/Mutable/Source/CustomizableObject/Private/MuCO/CustomizableSkeletalComponent.cpp
--- a/Engine/Plugins/Mutable/Source/CustomizableObject/Private/MuCO/CustomizableSkeletalComponent.cpp
+++ b/Engine/Plugins/Mutable/Source/CustomizableObject/Private/MuCO/CustomizableSkeletalComponent.cpp
@@ -11,11 +11,60 @@
 #include "MuCO/CustomizableObject.h"
 
 
+static TAutoConsoleVariable<bool> CVarCustomizableComponentPrivateNullptr(
+	TEXT("mutable.CustomizableComponentPrivateNullptr"),
+	true,
+	TEXT("Rollback. If true, force create the private part after the serialization."),
+	ECVF_Scalability);
+
+
+UCustomizableSkeletalComponentPrivate::UCustomizableSkeletalComponentPrivate()
+{
+	// This object may get instantiated into a level, as a part of UCustomizableSkeletalComponent, so needs to be public to ensure
+	// it can be serialized out.
+	if (!HasAllFlags(RF_ClassDefaultObject))
+	{
+		SetFlags(RF_Public);
+	}
+}
+
+void UCustomizableSkeletalComponentPrivate::CreateCustomizableObjectInstanceUsage()
+{
+	if (GetPublic()->CustomizableObjectInstanceUsage)
+	{
+		// CustomizableObjectInstanceUsage may already exist if duplicated from an existing Customizable Skeletal Component
+		if (GetPublic()->CustomizableObjectInstanceUsage->GetPrivate()->GetCustomizableSkeletalComponent() != GetPublic())
+		{
+			GetPublic()->CustomizableObjectInstanceUsage = nullptr;
+		}
+	}
+
+	AActor* RootActor = GetPublic()->GetAttachmentRootActor();
+	bool bIsDefaultActor = RootActor ?
+		RootActor->HasAnyFlags(RF_ClassDefaultObject) :
+		false;
+
+	if (!GetPublic()->CustomizableObjectInstanceUsage && !HasAnyFlags(RF_ClassDefaultObject) && !bIsDefaultActor)
+	{
+		GetPublic()->CustomizableObjectInstanceUsage = NewObject<UCustomizableObjectInstanceUsage>(this, TEXT("InstanceUsage"), RF_Transient);
+		GetPublic()->CustomizableObjectInstanceUsage->GetPrivate()->SetCustomizableSkeletalComponent(GetPublic());
+	}
+}
+
+
 UCustomizableSkeletalComponent::UCustomizableSkeletalComponent()
 {
-	Private = CreateDefaultSubobject<UCustomizableSkeletalComponentPrivate>(TEXT("Private"));
+}
+
+void UCustomizableSkeletalComponent::Serialize(FArchive& Ar)
+{
+	Super::Serialize(Ar);
 	
-	GetPrivate()->InstanceUsage->GetPrivate()->SetCustomizableSkeletalComponent(this);
+	// Somehow we got some assets serialized with Private = nullptr. Fix them.
+	if (Ar.IsLoading() && !Private && CVarCustomizableComponentPrivateNullptr.GetValueOnAnyThread())
+	{
+		Private = NewObject<UCustomizableSkeletalComponentPrivate>(this, FName("Private"), RF_Public);
+	}
 }
 
 
@@ -39,26 +88,20 @@ void UCustomizableSkeletalComponent::PostEditChangeProperty(FPropertyChangedEven
 #endif
 
 
-UCustomizableSkeletalComponentPrivate::UCustomizableSkeletalComponentPrivate()
-{
-	InstanceUsage = CreateDefaultSubobject<UCustomizableObjectInstanceUsage>(TEXT("InstanceUsage"));
-}
-
-
 void UCustomizableSkeletalComponentPrivate::Callbacks() const
 {
-	if (InstanceUsage)
+	if (GetPublic()->CustomizableObjectInstanceUsage)
 	{
-		InstanceUsage->GetPrivate()->Callbacks();
+		GetPublic()->CustomizableObjectInstanceUsage->GetPrivate()->Callbacks();
 	}
 }
 
 
 USkeletalMesh* UCustomizableSkeletalComponentPrivate::GetSkeletalMesh() const
 {
-	if (InstanceUsage)
+	if (GetPublic()->CustomizableObjectInstanceUsage)
 	{
-		return InstanceUsage->GetPrivate()->GetSkeletalMesh();
+		return GetPublic()->CustomizableObjectInstanceUsage->GetPrivate()->GetSkeletalMesh();
 	}
 
 	return nullptr;
@@ -67,36 +110,36 @@ USkeletalMesh* UCustomizableSkeletalComponentPrivate::GetSkeletalMesh() const
 
 void UCustomizableSkeletalComponentPrivate::SetSkeletalMesh(USkeletalMesh* SkeletalMesh)
 {
-	if (InstanceUsage)
+	if (GetPublic()->CustomizableObjectInstanceUsage)
 	{
-		InstanceUsage->GetPrivate()->SetSkeletalMesh(SkeletalMesh);
+		GetPublic()->CustomizableObjectInstanceUsage->GetPrivate()->SetSkeletalMesh(SkeletalMesh);
 	}
 }
 
 
 void UCustomizableSkeletalComponentPrivate::SetPhysicsAsset(UPhysicsAsset* PhysicsAsset)
 {
-	if (InstanceUsage)
+	if (GetPublic()->CustomizableObjectInstanceUsage)
 	{
-		InstanceUsage->GetPrivate()->SetPhysicsAsset(PhysicsAsset);
+		GetPublic()->CustomizableObjectInstanceUsage->GetPrivate()->SetPhysicsAsset(PhysicsAsset);
 	}
 }
 
 
 void UCustomizableSkeletalComponentPrivate::SetPendingSetSkeletalMesh(bool bIsActive)
 {
-	if (InstanceUsage)
+	if (GetPublic()->CustomizableObjectInstanceUsage)
 	{
-		InstanceUsage->GetPrivate()->SetPendingSetSkeletalMesh(bIsActive);
+		GetPublic()->CustomizableObjectInstanceUsage->GetPrivate()->SetPendingSetSkeletalMesh(bIsActive);
 	}
 }
 
 
 USkeletalMesh* UCustomizableSkeletalComponentPrivate::GetAttachedSkeletalMesh() const
 {
-	if (InstanceUsage)
+	if (GetPublic()->CustomizableObjectInstanceUsage)
 	{
-		return InstanceUsage->GetPrivate()->GetAttachedSkeletalMesh();
+		return GetPublic()->CustomizableObjectInstanceUsage->GetPrivate()->GetAttachedSkeletalMesh();
 	}
 
 	return nullptr;
@@ -161,18 +204,18 @@ bool UCustomizableSkeletalComponent::GetSkipSetSkeletalMeshOnAttach() const
 
 void UCustomizableSkeletalComponent::UpdateSkeletalMeshAsync(bool bNeverSkipUpdate)
 {
-	if (GetPrivate()->InstanceUsage)
+	if (CustomizableObjectInstanceUsage)
 	{
-		GetPrivate()->InstanceUsage->UpdateSkeletalMeshAsync(bNeverSkipUpdate);
+		CustomizableObjectInstanceUsage->UpdateSkeletalMeshAsync(bNeverSkipUpdate);
 	}
 }
 
 
 void UCustomizableSkeletalComponent::UpdateSkeletalMeshAsyncResult(FInstanceUpdateDelegate Callback, bool bIgnoreCloseDist, bool bForceHighPriority)
 {
-	if (GetPrivate()->InstanceUsage)
+	if (CustomizableObjectInstanceUsage)
 	{
-		GetPrivate()->InstanceUsage->UpdateSkeletalMeshAsyncResult(Callback, bIgnoreCloseDist, bForceHighPriority);
+		CustomizableObjectInstanceUsage->UpdateSkeletalMeshAsyncResult(Callback, bIgnoreCloseDist, bForceHighPriority);
 	}
 }
 
@@ -194,9 +237,9 @@ const UCustomizableSkeletalComponentPrivate* UCustomizableSkeletalComponent::Get
 #if WITH_EDITOR
 void UCustomizableSkeletalComponentPrivate::EditorUpdateComponent()
 {
-	if (InstanceUsage)
+	if (GetPublic()->CustomizableObjectInstanceUsage)
 	{
-		InstanceUsage->GetPrivate()->EditorUpdateComponent();
+		GetPublic()->CustomizableObjectInstanceUsage->GetPrivate()->EditorUpdateComponent();
 	}
 }
 #endif
@@ -240,3 +283,22 @@ void UCustomizableSkeletalComponent::OnAttachmentChanged()
 	}
 }
 
+
+void UCustomizableSkeletalComponent::PostInitProperties()
+{
+	Super::PostInitProperties();
+
+	if (!HasAllFlags(RF_ClassDefaultObject))
+	{
+		if (!Private)
+		{
+			Private = NewObject<UCustomizableSkeletalComponentPrivate>(this, FName("Private"), RF_Public);
+		}
+		else if (Private->GetOuter() != this)
+		{
+			Private = Cast<UCustomizableSkeletalComponentPrivate>(StaticDuplicateObject(Private, this, FName("Private")));
+		}
+	
+		GetPrivate()->CreateCustomizableObjectInstanceUsage();
+	}
+}
diff --git a/Engine/Plugins/Mutable/Source/CustomizableObject/Private/MuCO/CustomizableSkeletalMeshActor.cpp b/Engine/Plugins/Mutable/Source/CustomizableObject/Private/MuCO/CustomizableSkeletalMeshActor.cpp
--- a/Engine/Plugins/Mutable/Source/CustomizableObject/Private/MuCO/CustomizableSkeletalMeshActor.cpp
+++ b/Engine/Plugins/Mutable/Source/CustomizableObject/Private/MuCO/CustomizableSkeletalMeshActor.cpp
@@ -1,112 +1,75 @@
 // Copyright Epic Games, Inc. All Rights Reserved.
 
 #include "MuCO/CustomizableSkeletalMeshActor.h"
-#include "MuCO/CustomizableSkeletalMeshActorPrivate.h"
 
 #include "Components/SkeletalMeshComponent.h"
 #include "MuCO/CustomizableSkeletalComponent.h"
 #include "MuCO/CustomizableObject.h"
 #include "MuCO/CustomizableObjectPrivate.h"
 #include "MuCO/CustomizableObjectSystem.h"
-#include "MuCO/CustomizableSkeletalComponentPrivate.h"
-#include "Engine/SkeletalMesh.h"
-#include "MuCO/CustomizableObjectCustomVersion.h"
-#include "MuCO/CustomizableObjectPrivate.h"
-#include "MuCO/UnrealPortabilityHelpers.h"
-
 
 #include UE_INLINE_GENERATED_CPP_BY_NAME(CustomizableSkeletalMeshActor)
 
 #define LOCTEXT_NAMESPACE "CustomizableObject"
 
 
-UCustomizableObjectInstance* ACustomizableSkeletalMeshActor::GetInstance()
+ACustomizableSkeletalMeshActor::ACustomizableSkeletalMeshActor(const FObjectInitializer& ObjectInitializer)
+	: Super(ObjectInitializer)
 {
-	for (const UCustomizableSkeletalComponent* Component : CustomizableSkeletalComponents)
+	CustomizableSkeletalComponent = CreateDefaultSubobject<UCustomizableSkeletalComponent>(TEXT("CustomizableSkeletalComponent0"));
+	CustomizableSkeletalComponents.Add(CustomizableSkeletalComponent);
+	if (USkeletalMeshComponent* SkeletalMeshComp = GetSkeletalMeshComponent()) 
 	{
-		if (UCustomizableObjectInstance* COInstance = Component->CustomizableObjectInstance)
-		{
-			return COInstance;
-		}
+		SkeletalMeshComponents.Add(SkeletalMeshComp);
+		bool Success = CustomizableSkeletalComponents[0]->AttachToComponent(SkeletalMeshComp, FAttachmentTransformRules::KeepRelativeTransform);
 	}
-
-	return nullptr;
 }
 
 
-ACustomizableSkeletalMeshActor::ACustomizableSkeletalMeshActor(const FObjectInitializer& Initializer)
+void ACustomizableSkeletalMeshActor::AttachNewComponent()
 {
-	Private = CreateDefaultSubobject<UCustomizableSkeletalMeshActorPrivate>(TEXT("Private"));
+	int32 CurrentIndex = CustomizableSkeletalComponents.Num();
+	FString CustomizableComponentName = FString::Printf(TEXT("CustomizableSkeletalComponent%d"), CurrentIndex);
+	FString SkeletalMeshComponentName = FString::Printf(TEXT("SkeletalMeshComponent%d"), CurrentIndex);
 
-	// Old assets used to create the first UCustomizableSkeletalComponent as a Subobject. To be able to deserialize them we need to create it.
-	// Creating instead an Object instead of a Suobject will not work. Only Component UCustomizableSkeletalComponent 0 is a Subobject.
-	{
-		UCustomizableSkeletalComponent* CustomizableSkeletalComponent = CreateDefaultSubobject<UCustomizableSkeletalComponent>(TEXT("CustomizableSkeletalComponent0"));
-		CustomizableSkeletalComponents.Add(CustomizableSkeletalComponent);
-		if (USkeletalMeshComponent* SkeletalMeshComp = GetSkeletalMeshComponent()) 
+	//USkeletalMeshComponent* SkeletalMesh = CreateDefaultSubobject<USkeletalMeshComponent>(FName(*SkeletalMeshComponentName));
+	USkeletalMeshComponent* SkeletalMesh = NewObject<USkeletalMeshComponent>(this, USkeletalMeshComponent::StaticClass(), FName(*SkeletalMeshComponentName));
+	USkeletalMeshComponent* RootSkeletalMesh = GetSkeletalMeshComponent();
+	
+	if (SkeletalMesh && RootSkeletalMesh)
+	{		
+		bool Success = SkeletalMesh->AttachToComponent(RootSkeletalMesh, FAttachmentTransformRules::KeepRelativeTransform);
+		if (Success)
 		{
-			SkeletalMeshComponents.Add(SkeletalMeshComp);
-			CustomizableSkeletalComponents[0]->AttachToComponent(SkeletalMeshComp, FAttachmentTransformRules::KeepRelativeTransform);
+			//UCustomizableSkeletalComponent* CustomizableSkeletalComponent = CreateDefaultSubobject<UCustomizableSkeletalComponent>(FName(*CustomizableComponentName));
+			UCustomizableSkeletalComponent* NewCustomizableSkeletalComponent = NewObject<UCustomizableSkeletalComponent>(this, UCustomizableSkeletalComponent::StaticClass(), FName(*CustomizableComponentName));
+
+			if (NewCustomizableSkeletalComponent)
+			{
+				Success = NewCustomizableSkeletalComponent->AttachToComponent(SkeletalMesh, FAttachmentTransformRules::KeepRelativeTransform);
+
+				if (Success)
+				{
+					SkeletalMeshComponents.Add(SkeletalMesh);
+					CustomizableSkeletalComponents.Add(NewCustomizableSkeletalComponent);
+				}
+			}
 		}
 	}
 }
 
 
-void UCustomizableSkeletalMeshActorPrivate::Init(UCustomizableObjectInstance* Instance)
+UCustomizableObjectInstance* ACustomizableSkeletalMeshActor::GetComponentsCommonInstance()
 {
-	UCustomizableObject* Object = Instance->GetCustomizableObject();
-	if (!Object)
-	{
-		return;
-	}
-	
-	for (int32 ObjectComponentIndex = 0; ObjectComponentIndex < Object->GetComponentCount(); ++ObjectComponentIndex)
+	for (const UCustomizableSkeletalComponent* Component : CustomizableSkeletalComponents)
 	{
-		const FName& ComponentName = Object->GetComponentName(ObjectComponentIndex);
-		const FString ComponentNameString = ComponentName.ToString();
-
-		USkeletalMeshComponent* Component;
-		if (ObjectComponentIndex == 0)
-		{
-			Component = GetPublic()->SkeletalMeshComponents[0];
-			// Renaming the component will make it disappear. Keep it as it is.
-		}
-		else
-		{
-			const FString SkeletalMeshComponentName = FString::Printf(TEXT("SkeletalMeshComponent %s"), *ComponentNameString);
-			Component = NewObject<USkeletalMeshComponent>(GetPublic(), USkeletalMeshComponent::StaticClass(), FName(*SkeletalMeshComponentName));
-			Component->CreationMethod = EComponentCreationMethod::Native;
-			Component->AttachToComponent(GetPublic()->GetRootComponent(), FAttachmentTransformRules::KeepRelativeTransform);
-			Component->RegisterComponent();
-		}
-					
-		if (!Component)
-		{
-			continue;
-		}
-
-		UCustomizableSkeletalComponent* CustomizableComponent;
-		if (ObjectComponentIndex == 0)
-		{
-			CustomizableComponent = GetPublic()->CustomizableSkeletalComponents[0];
-			// Renaming the component will make it disappear. Keep it as it is.
-		}
-		else
+		if (UCustomizableObjectInstance* COInstance = Component->CustomizableObjectInstance)
 		{
-			const FString CustomizableComponentName = FString::Printf(TEXT("CustomizableSkeletalComponent %s"), *ComponentNameString);
-			CustomizableComponent = NewObject<UCustomizableSkeletalComponent>(GetPublic(), UCustomizableSkeletalComponent::StaticClass(), FName(CustomizableComponentName));
+			return COInstance;
 		}
-
-		CustomizableComponent->AttachToComponent(Component, FAttachmentTransformRules::KeepRelativeTransform);
-		CustomizableComponent->CustomizableObjectInstance = Instance;
-		CustomizableComponent->SetComponentName(ComponentName);
-		CustomizableComponent->RegisterComponent();
-		
-		GetPublic()->SkeletalMeshComponents.Add(Component);
-		GetPublic()->CustomizableSkeletalComponents.Add(CustomizableComponent);
 	}
 
-	Instance->UpdateSkeletalMeshAsync();
+	return nullptr;
 }
 
 
@@ -126,7 +89,7 @@ void ACustomizableSkeletalMeshActor::EnableDebugMaterial(bool bEnableDebugMateri
 	bRemoveDebugMaterial = bDebugMaterialEnabled && !bEnableDebugMaterial;
 	bDebugMaterialEnabled = bEnableDebugMaterial;
 
-	if (UCustomizableObjectInstance* COInstance = GetInstance())
+	if (UCustomizableObjectInstance* COInstance = GetComponentsCommonInstance())
 	{
 		//Bind Instance Update delegate to Actor
 		COInstance->UpdatedDelegate.AddUniqueDynamic(this, &ACustomizableSkeletalMeshActor::SwitchComponentsMaterials);
@@ -135,20 +98,6 @@ void ACustomizableSkeletalMeshActor::EnableDebugMaterial(bool bEnableDebugMateri
 }
 
 
-UCustomizableSkeletalMeshActorPrivate* ACustomizableSkeletalMeshActor::GetPrivate()
-{
-	check(Private);
-	return Private;
-}
-
-
-const UCustomizableSkeletalMeshActorPrivate* ACustomizableSkeletalMeshActor::GetPrivate() const
-{
-	check(Private);
-	return Private;
-}
-
-
 void ACustomizableSkeletalMeshActor::SwitchComponentsMaterials(UCustomizableObjectInstance* Instance)
 {
 	if (!DebugMaterial)
@@ -158,7 +107,7 @@ void ACustomizableSkeletalMeshActor::SwitchComponentsMaterials(UCustomizableObje
 
 	if (bDebugMaterialEnabled || bRemoveDebugMaterial)
 	{
-		UCustomizableObjectInstance* COInstance = GetInstance();
+		UCustomizableObjectInstance* COInstance = GetComponentsCommonInstance();
 
 		if (!COInstance)
 		{
@@ -200,16 +149,4 @@ void ACustomizableSkeletalMeshActor::SwitchComponentsMaterials(UCustomizableObje
 	}
 }
 
-ACustomizableSkeletalMeshActor* UCustomizableSkeletalMeshActorPrivate::GetPublic()
-{
-	return CastChecked<ACustomizableSkeletalMeshActor>(GetOuter());
-}
-
-
-TArray<TObjectPtr<UCustomizableSkeletalComponent>>& UCustomizableSkeletalMeshActorPrivate::GetComponents()
-{
-	return GetPublic()->CustomizableSkeletalComponents;
-}
-
-
 #undef LOCTEXT_NAMESPACE
diff --git a/Engine/Plugins/Mutable/Source/CustomizableObject/Public/MuCO/CustomizableObjectInstanceUsage.h b/Engine/Plugins/Mutable/Source/CustomizableObject/Public/MuCO/CustomizableObjectInstanceUsage.h
--- a/Engine/Plugins/Mutable/Source/CustomizableObject/Public/MuCO/CustomizableObjectInstanceUsage.h
+++ b/Engine/Plugins/Mutable/Source/CustomizableObject/Public/MuCO/CustomizableObjectInstanceUsage.h
@@ -111,6 +111,9 @@ private:
 	/** Only used if the UsedComponentIndex is INDEX_NONE. */
 	UPROPERTY()
 	FName UsedComponentName;
+	
+	// Used to replace the SkeletalMesh of the parent component by the ReferenceSkeletalMesh or the generated SkeletalMesh 
+	bool bUsedPendingSetSkeletalMesh = false;
 
 	// Used to avoid replacing the SkeletalMesh of the parent component by the ReferenceSkeletalMesh if bPendingSetSkeletalMesh is true
 	UPROPERTY()
@@ -119,7 +122,7 @@ private:
 	UPROPERTY()
 	bool bUsedSkipSetSkeletalMeshOnAttach = false;
 	
-	UPROPERTY( VisibleAnywhere, Category=CO)
+	UPROPERTY()
 	TObjectPtr<UCustomizableObjectInstanceUsagePrivate> Private;
 	
 	friend UCustomizableObjectInstanceUsagePrivate;
diff --git a/Engine/Plugins/Mutable/Source/CustomizableObject/Public/MuCO/CustomizableSkeletalComponent.h b/Engine/Plugins/Mutable/Source/CustomizableObject/Public/MuCO/CustomizableSkeletalComponent.h
--- a/Engine/Plugins/Mutable/Source/CustomizableObject/Public/MuCO/CustomizableSkeletalComponent.h
+++ b/Engine/Plugins/Mutable/Source/CustomizableObject/Public/MuCO/CustomizableSkeletalComponent.h
@@ -49,6 +49,9 @@ private:
 
 	UPROPERTY()
 	bool bSkipSkipSetSkeletalMeshOnAttach = false;
+
+	UPROPERTY(Transient)
+	TObjectPtr<UCustomizableObjectInstanceUsage> CustomizableObjectInstanceUsage;
 	
 	UPROPERTY()
 	TObjectPtr<UCustomizableSkeletalComponentPrivate> Private;
@@ -61,6 +64,10 @@ protected:
 #if WITH_EDITOR
 	virtual void PostEditChangeProperty(FPropertyChangedEvent& PropertyChangedEvent) override;
 #endif
+	virtual void Serialize(FArchive& Ar) override;
+
+	// UObject interface
+	virtual void PostInitProperties() override;
 
 	// USceneComponent interface
 	virtual void OnAttachmentChanged() override;
diff --git a/Engine/Plugins/Mutable/Source/CustomizableObject/Public/MuCO/CustomizableSkeletalMeshActor.h b/Engine/Plugins/Mutable/Source/CustomizableObject/Public/MuCO/CustomizableSkeletalMeshActor.h
--- a/Engine/Plugins/Mutable/Source/CustomizableObject/Public/MuCO/CustomizableSkeletalMeshActor.h
+++ b/Engine/Plugins/Mutable/Source/CustomizableObject/Public/MuCO/CustomizableSkeletalMeshActor.h
@@ -8,59 +8,60 @@
 
 class UCustomizableObjectInstance;
 class UCustomizableSkeletalComponent;
-class UCustomizableSkeletalMeshActorPrivate;
 
 class UObject;
 
 UCLASS(ClassGroup = CustomizableObject, Blueprintable, ComponentWrapperClass, ConversionRoot, meta = (ChildCanTick))
 class CUSTOMIZABLEOBJECT_API ACustomizableSkeletalMeshActor : public ASkeletalMeshActor
 {
-	GENERATED_BODY()
-
-	friend UCustomizableSkeletalMeshActorPrivate;
-
+	GENERATED_UCLASS_BODY()
 public:
-	ACustomizableSkeletalMeshActor(const FObjectInitializer& Initializer);
 	
-	UE_DEPRECATED(5.6, "Plase replace with custom code.")
-	UFUNCTION(BlueprintCallable, Category = CustomizableSkeletalMeshActor, meta=(DeprecatedFunction))
+	class UCustomizableSkeletalComponent* GetCustomizableSkeletalComponent(int32 ObjectComponentIndex = 0) const { return CustomizableSkeletalComponents[ObjectComponentIndex]; }
+	
+	class USkeletalMeshComponent* GetSkeletalMeshComponentAt(int32 ObjectComponentIndex = 0) const { return SkeletalMeshComponents[ObjectComponentIndex]; }
+
+	void AttachNewComponent();
+
+	int32 GetNumComponents() { return CustomizableSkeletalComponents.Num(); }
+
+	/** Sets which material will be used for debugging purposes */
+	UFUNCTION(BlueprintCallable, Category = CustomizableSkeletalMeshActor)
 	void SetDebugMaterial(UMaterialInterface* InDebugMaterial);
 
-	UE_DEPRECATED(5.6, "Plase replace with custom code.")
-	UFUNCTION(BlueprintCallable, Category = CustomizableSkeletalMeshActor, meta=(DeprecatedFunction))
+	/** Enables the debug material to all the componets of the actor */
+	UFUNCTION(BlueprintCallable, Category = CustomizableSkeletalMeshActor)
 	void EnableDebugMaterial(bool bEnableDebugMaterial);
 
-	UCustomizableSkeletalMeshActorPrivate* GetPrivate();
-	const UCustomizableSkeletalMeshActorPrivate* GetPrivate() const;
-
 private:
-	void Init(UCustomizableObjectInstance* Instance);
-	
-	// DEPRECATED
+
+	/** Changes materials to all actor componets */
 	UFUNCTION()
 	void SwitchComponentsMaterials(UCustomizableObjectInstance* Instance);
 
-	/** @return the instance represented in all the actor components. */
-	UCustomizableObjectInstance* GetInstance();
+	/** return: the instance represented in all the actor components */
+	UCustomizableObjectInstance* GetComponentsCommonInstance();
 
+private:
+	
 	/** TODO: There are as many as components in the CO (not in the COI). */
-	UPROPERTY(VisibleAnywhere, Category = CustomizableSkeletalMeshActor)
-	TArray<TObjectPtr<UCustomizableSkeletalComponent>> CustomizableSkeletalComponents;
+	UPROPERTY(Category = CustomizableSkeletalMesh, VisibleAnywhere, BlueprintReadOnly, meta = (AllowPrivateAccess = "true"))
+	TArray< TObjectPtr<class UCustomizableSkeletalComponent> > CustomizableSkeletalComponents;
 
 	/** TODO: There are as many as components in the CO (not in the COI). */
-	UPROPERTY(VisibleAnywhere, Category = CustomizableSkeletalMeshActor)
-	TArray<TObjectPtr<USkeletalMeshComponent>> SkeletalMeshComponents;
+	UPROPERTY(Category = CustomizableSkeletalMesh, VisibleAnywhere, BlueprintReadOnly, meta = (AllowPrivateAccess = "true"))
+	TArray< TObjectPtr<class USkeletalMeshComponent>> SkeletalMeshComponents;
+
+	// TODO: This is a temporary fix to not break the demos, we should update the demos to support the arrays of components
+	UPROPERTY(Category = CustomizableSkeletalMesh, VisibleAnywhere, BlueprintReadOnly, meta = (AllowPrivateAccess = "true"))
+	TObjectPtr<class UCustomizableSkeletalComponent> CustomizableSkeletalComponent;
 
-	// DEPRECATED
-	UPROPERTY()
+	/** Material used for debugging */
 	TObjectPtr<UMaterialInterface> DebugMaterial;
 
-	// DEPRECATED
+	/** bool that determines if the debug material has been enabled */
 	bool bDebugMaterialEnabled = false;
 
-	// DEPRECATED
+	/** bool used to remove the debug material whet it has been disabled */
 	bool bRemoveDebugMaterial = false;
-
-	UPROPERTY()
-	TObjectPtr<UCustomizableSkeletalMeshActorPrivate> Private = nullptr;
 };
diff --git a/Engine/Plugins/Mutable/Source/CustomizableObjectEditor/Private/MuCOE/CustomizableObjectInstanceFactory.cpp b/Engine/Plugins/Mutable/Source/CustomizableObjectEditor/Private/MuCOE/CustomizableObjectInstanceFactory.cpp
--- a/Engine/Plugins/Mutable/Source/CustomizableObjectEditor/Private/MuCOE/CustomizableObjectInstanceFactory.cpp
+++ b/Engine/Plugins/Mutable/Source/CustomizableObjectEditor/Private/MuCOE/CustomizableObjectInstanceFactory.cpp
@@ -12,7 +12,6 @@
 #include "MuCO/CustomizableObjectPrivate.h"
 #include "MuCO/CustomizableSkeletalComponent.h"
 #include "MuCO/CustomizableSkeletalMeshActor.h"
-#include "MuCO/CustomizableSkeletalMeshActorPrivate.h"
 #include "MuCOE/Nodes/CustomizableObjectNodeObject.h"
 #include "MuCO/UnrealPortabilityHelpers.h"
 
@@ -49,26 +48,63 @@ void UCustomizableObjectInstanceFactory::PostSpawnActor(UObject* Asset, AActor*
 	{
 		return;
 	}
+	
 
-	NewCSMActor->GetPrivate()->Init(Instance);
-}
+	for (int32 ObjectComponentIndex = 0; ObjectComponentIndex < Object->GetComponentCount(); ++ObjectComponentIndex)
+	{
+		const FName& ComponentName = Object->GetPrivate()->GetComponentName(FCustomizableObjectComponentIndex(ObjectComponentIndex));
+		
+		USkeletalMesh* SkeletalMesh = Instance->GetComponentMeshSkeletalMesh(ComponentName);
+		if (!SkeletalMesh)
+		{
+			SkeletalMesh = Object->GetComponentMeshReferenceSkeletalMesh(ComponentName);
+		}
+		
+		if (ObjectComponentIndex > 0 && ObjectComponentIndex > NewCSMActor->GetNumComponents() - 1)
+		{
+			NewCSMActor->AttachNewComponent();
+		}
+
+		if (USkeletalMeshComponent* SkeletalMeshComp = NewCSMActor->GetSkeletalMeshComponentAt(ObjectComponentIndex))
+		{
+			SkeletalMeshComp->UnregisterComponent();
+			UE_MUTABLE_SETSKINNEDASSET(SkeletalMeshComp, SkeletalMesh);
 
+			if (ObjectComponentIndex == 0 && NewCSMActor->GetWorld()->IsGameWorld())
+			{
+				NewCSMActor->ReplicatedMesh = SkeletalMesh;
+			}
+
+			if (UCustomizableSkeletalComponent* CustomSkeletalComp = NewCSMActor->GetCustomizableSkeletalComponent(ObjectComponentIndex))
+			{
+				CustomSkeletalComp->UnregisterComponent();
+				CustomSkeletalComp->CustomizableObjectInstance = Instance;
+				CustomSkeletalComp->SetComponentName(ComponentName);
+				CustomSkeletalComp->GetPrivate()->SetSkeletalMesh(SkeletalMesh);
+				CustomSkeletalComp->UpdateSkeletalMeshAsync();
+				CustomSkeletalComp->RegisterComponent();
+			}
+
+			SkeletalMeshComp->RegisterComponent();
+		}
+	}
+}
 
 UObject* UCustomizableObjectInstanceFactory::GetAssetFromActorInstance(AActor* ActorInstance)
 {
 	if (ACustomizableSkeletalMeshActor* CSMActor = CastChecked<ACustomizableSkeletalMeshActor>(ActorInstance))
 	{
-		const TArray<TObjectPtr<UCustomizableSkeletalComponent>>& Components = CSMActor->GetPrivate()->GetComponents();
-		if (Components.IsValidIndex(0))
+		if (CSMActor->GetNumComponents() > 0)
 		{
-			return Components[0];
+			if (UCustomizableSkeletalComponent* CustomSkeletalComp = CSMActor->GetCustomizableSkeletalComponent(0))
+			{
+				return CustomSkeletalComp->CustomizableObjectInstance;
+			}
 		}
 	}
-	
 	return nullptr;
 }
 
-
 bool UCustomizableObjectInstanceFactory::CanCreateActorFrom(const FAssetData& AssetData, FText& OutErrorMsg)
 {
     if (!AssetData.IsValid() ||
diff --git a/Engine/Plugins/Mutable/Source/CustomizableObjectEditor/Public/MuCOE/CustomizableObjectInstanceFactory.h b/Engine/Plugins/Mutable/Source/CustomizableObjectEditor/Public/MuCOE/CustomizableObjectInstanceFactory.h
--- a/Engine/Plugins/Mutable/Source/CustomizableObjectEditor/Public/MuCOE/CustomizableObjectInstanceFactory.h
+++ b/Engine/Plugins/Mutable/Source/CustomizableObjectEditor/Public/MuCOE/CustomizableObjectInstanceFactory.h
@@ -17,9 +17,11 @@ class UCustomizableObjectInstanceFactory : public UActorFactory
 	GENERATED_UCLASS_BODY()
 
 protected:
-	// UActorFactory interface
-	virtual void PostSpawnActor(UObject* Asset, AActor* NewActor) override;
-	virtual UObject* GetAssetFromActorInstance(AActor* ActorInstance) override;
-	virtual bool CanCreateActorFrom(const FAssetData& AssetData, FText& OutErrorMsg) override;
-	virtual FQuat AlignObjectToSurfaceNormal(const FVector& InSurfaceNormal, const FQuat& ActorRotation) const override;
+	//~ Begin UActorFactory Interface
+	void PostSpawnActor(UObject* Asset, AActor* NewActor) override;
+	UObject* GetAssetFromActorInstance(AActor* ActorInstance) override;
+	bool CanCreateActorFrom(const FAssetData& AssetData, FText& OutErrorMsg) override;
+	//~ End UActorFactory Interface
+
+	FQuat AlignObjectToSurfaceNormal(const FVector& InSurfaceNormal, const FQuat& ActorRotation) const;
 };
