[mutable] Fix for ASAN exception when trying to access already freed memory.
- Cached name of some mutable structures due to them not being available during the module shutdown.

#rnx

diff --git a/Engine/Plugins/Mutable/Source/CustomizableObjectEditor/Private/MuCOE/CustomizableObjectEditorModule.cpp b/Engine/Plugins/Mutable/Source/CustomizableObjectEditor/Private/MuCOE/CustomizableObjectEditorModule.cpp
--- a/Engine/Plugins/Mutable/Source/CustomizableObjectEditor/Private/MuCOE/CustomizableObjectEditorModule.cpp
+++ b/Engine/Plugins/Mutable/Source/CustomizableObjectEditor/Private/MuCOE/CustomizableObjectEditorModule.cpp
@@ -277,10 +277,14 @@ void FCustomizableObjectEditorModule::StartupModule()
 	RegisterCustomDetails(PropertyModule, UCustomizableObjectInstance::StaticClass(), FOnGetDetailCustomizationInstance::CreateStatic(&FCustomizableInstanceDetails::MakeInstance));
 	RegisterCustomDetails(PropertyModule, UCustomSettings::StaticClass(), FOnGetDetailCustomizationInstance::CreateStatic(&FCustomizableObjectCustomSettingsDetails::MakeInstance));
 
+	// We need to cache those two FNames as if we want to get them on "ShutdownModule" we get "None" FNames and an ASAN error in Linux.
+	MeshReshapeBoneReferenceUStructName = FMeshReshapeBoneReference::StaticStruct()->GetFName();
+	BoneToRemoveUStructName = FBoneToRemove::StaticStruct()->GetFName();
+
 	// Custom properties
 	PropertyModule.RegisterCustomPropertyTypeLayout("CustomizableObjectIdentifier", FOnGetPropertyTypeCustomizationInstance::CreateStatic(&FCustomizableObjectIdentifierCustomization::MakeInstance));
-	PropertyModule.RegisterCustomPropertyTypeLayout(FMeshReshapeBoneReference::StaticStruct()->GetFName(), FOnGetPropertyTypeCustomizationInstance::CreateStatic(&FMeshReshapeBonesReferenceCustomization::MakeInstance));
-	PropertyModule.RegisterCustomPropertyTypeLayout(FBoneToRemove::StaticStruct()->GetFName(), FOnGetPropertyTypeCustomizationInstance::CreateStatic(&FCustomizableObjectLODReductionSettings::MakeInstance));
+	PropertyModule.RegisterCustomPropertyTypeLayout(MeshReshapeBoneReferenceUStructName, FOnGetPropertyTypeCustomizationInstance::CreateStatic(&FMeshReshapeBonesReferenceCustomization::MakeInstance));
+	PropertyModule.RegisterCustomPropertyTypeLayout(BoneToRemoveUStructName, FOnGetPropertyTypeCustomizationInstance::CreateStatic(&FCustomizableObjectLODReductionSettings::MakeInstance));
 	PropertyModule.RegisterCustomPropertyTypeLayout(NAME_StrProperty, FOnGetPropertyTypeCustomizationInstance::CreateStatic(&FCustomizableObjectStateParameterSelector::MakeInstance), MakeShared<FStatePropertyTypeIdentifier>());
 	PropertyModule.RegisterCustomPropertyTypeLayout(FCustomizableObjectVariation::StaticStruct()->GetFName(), FOnGetPropertyTypeCustomizationInstance::CreateStatic(&FCustomizableObjectVariationCustomization::MakeInstance));
 	PropertyModule.RegisterCustomPropertyTypeLayout(FCustomizableObjectTextureVariation::StaticStruct()->GetFName(), FOnGetPropertyTypeCustomizationInstance::CreateStatic(&FCustomizableObjectVariationCustomization::MakeInstance));
@@ -335,8 +339,8 @@ void FCustomizableObjectEditorModule::ShutdownModule()
 
 		// Unregister Custom properties
 		PropertyModule.UnregisterCustomPropertyTypeLayout("CustomizableObjectIdentifier");
-		PropertyModule.UnregisterCustomPropertyTypeLayout(FMeshReshapeBoneReference::StaticStruct()->GetFName());
-		PropertyModule.UnregisterCustomPropertyTypeLayout(FBoneToRemove::StaticStruct()->GetFName());
+		PropertyModule.UnregisterCustomPropertyTypeLayout(MeshReshapeBoneReferenceUStructName);
+		PropertyModule.UnregisterCustomPropertyTypeLayout(BoneToRemoveUStructName);
 		PropertyModule.UnregisterCustomPropertyTypeLayout(NAME_StrProperty);
 
 		PropertyModule.NotifyCustomizationModuleChanged();
diff --git a/Engine/Plugins/Mutable/Source/CustomizableObjectEditor/Private/MuCOE/CustomizableObjectEditorModule.h b/Engine/Plugins/Mutable/Source/CustomizableObjectEditor/Private/MuCOE/CustomizableObjectEditorModule.h
--- a/Engine/Plugins/Mutable/Source/CustomizableObjectEditor/Private/MuCOE/CustomizableObjectEditorModule.h
+++ b/Engine/Plugins/Mutable/Source/CustomizableObjectEditor/Private/MuCOE/CustomizableObjectEditorModule.h
@@ -17,7 +17,7 @@ class ICustomizableObjectEditor;
 class ICustomizableObjectInstanceEditor;
 class ICustomizableObjectDebugger;
 class FBakeOperationCompletedDelegate;
-
+class FName;
 
 /** Get a list of packages that are used by the compilation but are not directly referenced.
   * List includes:
@@ -83,7 +83,7 @@ private:
 	IConsoleCommand* LaunchCOIECommand = nullptr;
 
 	FTSTicker::FDelegateHandle WarningsTickerHandle;
-	
+
 	static void OpenCOIE(const TArray<FString>& Arguments);
 
 	// Used to ask the user if they want to recompile uncompiled PIE COs
@@ -94,4 +94,8 @@ private:
 
 	bool HandleSettingsSaved();
 	void RegisterSettings();
+
+	FName MeshReshapeBoneReferenceUStructName;
+
+	FName BoneToRemoveUStructName;
 };
\ No newline at end of file
