Mutable: Work around Internal directory being included from Public in CustomizableObject. This will need to be properly resolved.

#rnx

diff --git a/Engine/Plugins/Mutable/Source/CustomizableObject/CustomizableObject.Build.cs b/Engine/Plugins/Mutable/Source/CustomizableObject/CustomizableObject.Build.cs
--- a/Engine/Plugins/Mutable/Source/CustomizableObject/CustomizableObject.Build.cs
+++ b/Engine/Plugins/Mutable/Source/CustomizableObject/CustomizableObject.Build.cs
@@ -13,6 +13,8 @@ public class CustomizableObject : ModuleRules
 		bAllowConfidentialPlatformDefines = true;
 		//bUseUnity = false;
 
+		// TODO: CustomizableObject/Public/MuCO/CustomizableObject.h is including CustomizableObject/Internal/MuCO/CustomizableObjectCompilerTypes.h
+		PublicIncludePaths.Add(System.IO.Path.Combine(ModuleDirectory, "Internal"));
 
 		PrivateDependencyModuleNames.AddRange(new string[] {
             "InputCore",
