[Mutable] Added LODs per Component.

LOD are no longer managed by the Customizable Object Instance. Now each component has its own LOD parameters:
- Num LODs available.
- Num Redisent LODs.
- Num Streaming LODs.
- Etc.

Changes overview:
- All LOD related variables are now indexed by [REDACTED_PERSON] (FName).
- Unified LOD variables to be uint8.
- Removed unused bUseLODAsBias.
- Removed unused LOD related functions.

#rnx

diff --git a/Engine/Plugins/Mutable/Source/CustomizableObject/Internal/MuCO/CustomizableObjectCompilerTypes.h b/Engine/Plugins/Mutable/Source/CustomizableObject/Internal/MuCO/CustomizableObjectCompilerTypes.h
--- a/Engine/Plugins/Mutable/Source/CustomizableObject/Internal/MuCO/CustomizableObjectCompilerTypes.h
+++ b/Engine/Plugins/Mutable/Source/CustomizableObject/Internal/MuCO/CustomizableObjectCompilerTypes.h
@@ -110,10 +110,7 @@ struct FCompilationOptions
 
 	// Used to reduce the number of notifications when compiling objects
 	bool bSilentCompilation = true;
-
-	// Used to reduce texture size on higher mesh LODs. Only active if LOD strategy is set to Automatic LODs from Mesh
-	bool bUseLODAsBias = false;
-
+	
 	/** Force a very big number on the mips to skip during compilation. Useful to debug special cooks of the data. */
 	bool bForceLargeLODBias = false;
 	int32 DebugBias = 0;
diff --git a/Engine/Plugins/Mutable/Source/CustomizableObject/Internal/MuCO/CustomizableObjectInstanceDescriptor.h b/Engine/Plugins/Mutable/Source/CustomizableObject/Internal/MuCO/CustomizableObjectInstanceDescriptor.h
--- a/Engine/Plugins/Mutable/Source/CustomizableObject/Internal/MuCO/CustomizableObjectInstanceDescriptor.h
+++ b/Engine/Plugins/Mutable/Source/CustomizableObject/Internal/MuCO/CustomizableObjectInstanceDescriptor.h
@@ -71,19 +71,10 @@ struct CUSTOMIZABLEOBJECT_API FCustomizableObjectInstanceDescriptor
 	
 	/** Update all parameters to be up to date with the Mutable Core parameters. */
 	void ReloadParameters();
-    
-	int32 GetMinLod() const;
-
-	void SetMinLod(int32 InMinLOD);
-
-	int32 GetMaxLod() const { return MAX_int32; }; // DEPRECATED
-
-	void SetMaxLod(int32 InMaxLOD) {}; // DEPRECATED
-
-	void SetRequestedLODLevels(const TArray<uint16>& InRequestedLODLevels);
+	
+	void SetRequestedLODLevels(const TMap<FName, uint8>& InRequestedLODLevels);
 
-	const TArray<uint16>& GetRequestedLODLevels() const;
-	uint16 GetRequestedLODs(FCustomizableObjectComponentIndex) const;
+	const TMap<FName, uint8>& GetRequestedLODLevels() const;
 
 	// ------------------------------------------------------------
 	// FParameters
@@ -339,10 +330,11 @@ struct CUSTOMIZABLEOBJECT_API FCustomizableObjectInstanceDescriptor
 	bool bBuildParameterRelevancy = false;
 
 	/** These are the LODs Mutable can generate, they MUST NOT be used in an update (Mutable thread). */
-	int32 MinLOD = 0;
+	TMap<FName, uint8> MinLOD;
 
-	/** Array of RequestedLODs per component to generate, they MUST NOT be used in an update (Mutable thread). */
-	TArray<uint16> RequestedLODLevels;
+	/** RequestedLODs per component to generate.
+	 * Key is the Component name. Value represents the first LOD to generate. Zero means all LODs. */
+	TMap<FName, uint8> RequestedLODLevels;
 	
 	// Friends
 	friend FDescriptorHash;
diff --git a/Engine/Plugins/Mutable/Source/CustomizableObject/Internal/MuCO/CustomizableObjectInstancePrivate.h b/Engine/Plugins/Mutable/Source/CustomizableObject/Internal/MuCO/CustomizableObjectInstancePrivate.h
--- a/Engine/Plugins/Mutable/Source/CustomizableObject/Internal/MuCO/CustomizableObjectInstancePrivate.h
+++ b/Engine/Plugins/Mutable/Source/CustomizableObject/Internal/MuCO/CustomizableObjectInstancePrivate.h
@@ -302,8 +302,8 @@ private:
 
 	bool BuildSkeletonData(const TSharedRef<FUpdateContextPrivate>& OperationData, USkeletalMesh& SkeletalMesh, const FMutableRefSkeletalMeshData& RefSkeletalMeshData, UCustomizableObject& CustomizableObject, FCustomizableObjectInstanceComponentIndex InstanceComponentIndex);
 	void BuildMeshSockets(const TSharedRef<FUpdateContextPrivate>& OperationData, USkeletalMesh* SkeletalMesh, const FModelResources& ModelResources, const FMutableRefSkeletalMeshData& RefSkeletalMeshData, TSharedPtr<const mu::FMesh> MutableMesh);
-	void BuildOrCopyElementData(const TSharedRef<FUpdateContextPrivate>& OperationData, USkeletalMesh* SkeletalMesh, UCustomizableObjectInstance* CustomizableObjectInstance, FCustomizableObjectInstanceComponentIndex InstanceComponentIndex);
-	void BuildOrCopyMorphTargetsData(const TSharedRef<FUpdateContextPrivate>& OperationData, USkeletalMesh* SkeletalMesh, const USkeletalMesh* SrcSkeletalMesh, UCustomizableObjectInstance* CustomizableObjectInstance, FCustomizableObjectInstanceComponentIndex InstanceComponentIndex);
+	void BuildOrCopyElementData(const TSharedRef<FUpdateContextPrivate>& OperationData, USkeletalMesh* SkeletalMesh, UCustomizableObject& CustomizableObjectInstance, FCustomizableObjectInstanceComponentIndex InstanceComponentIndex);
+	void BuildOrCopyMorphTargetsData(const TSharedRef<FUpdateContextPrivate>& OperationData, USkeletalMesh* SkeletalMesh, const USkeletalMesh* SrcSkeletalMesh, UCustomizableObject& CustomizableObject, FCustomizableObjectInstanceComponentIndex InstanceComponentIndex);
 	bool BuildOrCopyRenderData(const TSharedRef<FUpdateContextPrivate>& OperationData, USkeletalMesh* SkeletalMesh, const USkeletalMesh* SrcSkeletalMesh, UCustomizableObjectInstance* CustomizableObjectInstance, FCustomizableObjectInstanceComponentIndex InstanceComponentIndex);
 	void BuildOrCopyClothingData(const TSharedRef<FUpdateContextPrivate>& OperationData, USkeletalMesh* SkeletalMesh, const USkeletalMesh* SrcSkeletalMesh, UCustomizableObjectInstance* CustomizableObjectInstance, FCustomizableObjectInstanceComponentIndex InstanceComponentIndex);
 
@@ -351,15 +351,6 @@ public:
 	// If any components are using this instance, they will store the min of their distances to the player here every frame for LOD purposes
 	float MinSquareDistFromComponentToPlayer = FLT_MAX;
 	float LastMinSquareDistFromComponentToPlayer = FLT_MAX; // The same as the previous dist for last frame
-												
-	// This is the LODs that the Customizable Object has
-	int32 NumLODsAvailable = INT32_MAX;
-
-	// First SkeletalMesh LOD we can generate on the running platform
-	uint8 FirstLODAvailable = 0;
-
-	// First Non-streamable SkeletalMesh LOD
-	uint8 FirstResidentLOD = MAX_MESH_LOD_COUNT;
 	
 	// To be indexed with object component index
 	UPROPERTY(Transient)
diff --git a/Engine/Plugins/Mutable/Source/CustomizableObject/Internal/MuCO/CustomizableObjectPrivate.h b/Engine/Plugins/Mutable/Source/CustomizableObject/Internal/MuCO/CustomizableObjectPrivate.h
--- a/Engine/Plugins/Mutable/Source/CustomizableObject/Internal/MuCO/CustomizableObjectPrivate.h
+++ b/Engine/Plugins/Mutable/Source/CustomizableObject/Internal/MuCO/CustomizableObjectPrivate.h
@@ -1198,15 +1198,15 @@ struct FModelResources
 	
 	/** Max number of LODs in the compiled Model. */
 	UPROPERTY()
-	uint8 NumLODs = 0;
+	TMap<FName, uint8> NumLODs;
 
-	/** Max number of LODs to stream. Mutable will always generate at least on LOD. */
+	/** Max number of LODs to stream. Mutable will always generate at least one LOD. */
 	UPROPERTY()
-	uint8 NumLODsToStream = 0;
+	TMap<FName, uint8> NumLODsToStream;
 
 	/** First LOD available, some platforms may remove lower LODs when cooking, this MinLOD represents the first LOD we can generate */
 	UPROPERTY()
-	uint8 FirstLODAvailable = 0;
+	TMap<FName, uint8> FirstLODAvailable;
 
 	/** Name of all possible components. Index is the ObjectComponentIndex. */
 	UPROPERTY()
@@ -1454,8 +1454,8 @@ public:
 	void GetLowPriorityTextureNames(TArray<FString>& OutTextureNames);
 
 	/** Return the MinLOD index to generate based on the active LODSettings (PerPlatformMinLOD or PerQualityLevelMinLOD) */
-	int32 GetMinLODIndex() const;
-	
+	uint8 GetMinLODIndex(const FName& ComponentName) const;
+
 #if WITH_EDITOR
 	/** See ICustomizableObjectEditorModule::IsCompilationOutOfDate. */
 	bool IsCompilationOutOfDate(bool bSkipIndirectReferences, TArray<FName>& OutOfDatePackages, TArray<FName>& AddedPackages, TArray<FName>& RemovedPackages, bool& bReleaseVersionDiff) const;
@@ -1670,6 +1670,8 @@ public:
 
 		ModelResourcesExtensionData,
 
+		LODsPerComponent,
+
 		// -----<new versions can be added above this line>--------
 		LastCustomizableObjectVersion
 	};
diff --git a/Engine/Plugins/Mutable/Source/CustomizableObject/Internal/MuCO/CustomizableObjectSystemPrivate.h b/Engine/Plugins/Mutable/Source/CustomizableObject/Internal/MuCO/CustomizableObjectSystemPrivate.h
--- a/Engine/Plugins/Mutable/Source/CustomizableObject/Internal/MuCO/CustomizableObjectSystemPrivate.h
+++ b/Engine/Plugins/Mutable/Source/CustomizableObject/Internal/MuCO/CustomizableObjectSystemPrivate.h
@@ -88,17 +88,16 @@ public:
 	EQueuePriorityType Priority = EQueuePriorityType::Med;
 
 	// These are the LODs that would be applied if this candidate is chosen
-	int32 MinLOD = 0;
+	TMap<FName, uint8> MinLOD;
 
-	/** Array of RequestedLODs per component to generate if this candidate is chosen */
-	TArray<uint16> RequestedLODLevels;
+	TMap<FName, uint8> RequestedLODLevels;
 
 	FMutableUpdateCandidate(UCustomizableObjectInstance* InCustomizableObjectInstance);
-
-
-	FMutableUpdateCandidate(const UCustomizableObjectInstance* InCustomizableObjectInstance, const int32 InMinLOD,
-		const TArray<uint16>& InRequestedLODLevels) :
-		CustomizableObjectInstance(const_cast<UCustomizableObjectInstance*>(InCustomizableObjectInstance)), MinLOD(InMinLOD),
+	
+	FMutableUpdateCandidate(const UCustomizableObjectInstance* InCustomizableObjectInstance, const TMap<FName, uint8>& InMinLOD,
+		const TMap<FName, uint8>& InRequestedLODLevels) :
+		CustomizableObjectInstance(const_cast<UCustomizableObjectInstance*>(InCustomizableObjectInstance)),
+		MinLOD(InMinLOD),
 		RequestedLODLevels(InRequestedLODLevels) {}
 
 	bool HasBeenIssued() const;
@@ -521,15 +520,14 @@ public:
 	FUpdateContextPrivate(UCustomizableObjectInstance& InInstance);
 	~FUpdateContextPrivate();
 
-	int32 GetMinLOD() const;
+	const TMap<FName, uint8>& GetMinLOD() const;
 
-	void SetMinLOD(int32 MinLOD);
+	void SetMinLOD(const TMap<FName, uint8>& MinLOD);
 
 	/** Return an array of LODs per object component. */
-	const TArray<uint16>& GetRequestedLODs() const;
-	uint16 GetRequestedLODs(FCustomizableObjectComponentIndex) const;
-
-	void SetRequestedLODs(TArray<uint16>& RequestedLODs);
+	const TMap<FName, uint8>& GetRequestedLODs() const;
+	
+	void SetRequestedLODs(const TMap<FName, uint8>& RequestedLODs);
 
 	const FCustomizableObjectInstanceDescriptor& GetCapturedDescriptor() const;
 
@@ -546,6 +544,9 @@ public:
 	 *It is weak because we don't want to lock it in case it becomes irrelevant in the game while operations are pending and it needs to be destroyed. */
 	TWeakObjectPtr<UCustomizableObjectInstance> Instance;
 
+	/** Customizable Object we are operating on. It can be destroyed between Game Thread tasks. */
+	TWeakObjectPtr<UCustomizableObject> Object;
+
 	/** Return the object component index associated with a component in this instance. */
 	FCustomizableObjectComponentIndex GetObjectComponentIndex(FCustomizableObjectInstanceComponentIndex InstanceComponentIndex);
 
@@ -604,39 +605,14 @@ public:
 	// TODO: Redundant to store it here because it is implicit in the size of some arrays or in the MutableInstance while it is valid.
 	uint8 NumInstanceComponents = 0;
 
-private:
-
-	// Index with ObjectComponent index
-	// \TODO: Wasteful: it only needs to be as big as NumInstanceComponents
-	TArray<uint8> NumLODsAvailablePerComponent;
-
-public:
-
-	inline void InitNumLODsAvailablePerComponent(int32 Size)
-	{
-		NumLODsAvailablePerComponent.SetNumZeroed(Size);
-	}
-
-	inline uint8 GetLODsAvailable(FCustomizableObjectComponentIndex Index) const
-	{
-		if (NumLODsAvailablePerComponent.IsValidIndex(Index.GetValue()))
-		{
-			return NumLODsAvailablePerComponent[Index.GetValue()];
-		}
-		return 0;
-	}
-
-	inline void SetNumLODsAvailable(FCustomizableObjectComponentIndex Index, uint8 LODCount)
-	{
-		if (NumLODsAvailablePerComponent.IsValidIndex(Index.GetValue()))
-		{
-			NumLODsAvailablePerComponent[Index.GetValue()] = LODCount;
-		}
-	}
-
+	/** List of component names. Index is the ObjectComponentIndex. */
+	TArray<FName> ComponentNames;
+	
+	TMap<FName, uint8> NumLODsAvailable;
 
-	uint8 FirstLODAvailable = 0;
-	uint8 FirstResidentLOD = 0;
+	TMap<FName, uint8> FirstLODAvailable;
+	
+	TMap<FName, uint8> FirstResidentLOD;
 
 	TMap<mu::FResourceID, FTexturePlatformData*> ImageToPlatformDataMap;
 
diff --git a/Engine/Plugins/Mutable/Source/CustomizableObject/Internal/MuCO/DescriptorHash.h b/Engine/Plugins/Mutable/Source/CustomizableObject/Internal/MuCO/DescriptorHash.h
--- a/Engine/Plugins/Mutable/Source/CustomizableObject/Internal/MuCO/DescriptorHash.h
+++ b/Engine/Plugins/Mutable/Source/CustomizableObject/Internal/MuCO/DescriptorHash.h
@@ -27,8 +27,8 @@ private:
 	uint32 Hash = 0;
 
 public:
-	int32 MinLOD = 0;
+	TMap<FName, uint8> MinLODs;
 
 	// Array of bitmasks that indicate which LODs of each component have been requested
-	TArray<uint16> RequestedLODsPerComponent;
+	TMap<FName, uint8> RequestedLODsPerComponent;
 };
diff --git a/Engine/Plugins/Mutable/Source/CustomizableObject/Internal/MuCO/LogBenchmarkUtil.h b/Engine/Plugins/Mutable/Source/CustomizableObject/Internal/MuCO/LogBenchmarkUtil.h
--- a/Engine/Plugins/Mutable/Source/CustomizableObject/Internal/MuCO/LogBenchmarkUtil.h
+++ b/Engine/Plugins/Mutable/Source/CustomizableObject/Internal/MuCO/LogBenchmarkUtil.h
@@ -127,7 +127,7 @@ public:
 	CUSTOMIZABLEOBJECT_API static bool IsBenchmarkingReportingEnabled();
 	
 	/** Get stats. */
-	void GetInstancesStats(int32& OutNumInstances, int32& OutNumBuiltInstances, int32& OutNumInstancesLOD0, int32& OutNumInstancesLOD1, int32& OutNumInstancesLOD2, int32& OutNumAllocatedSkeletalMeshes) const;
+	void GetInstancesStats(int32& OutNumInstances, int32& OutNumBuiltInstances, int32& OutNumAllocatedSkeletalMeshes) const;
 
 	/** Add Mutable created Texture to track. */
 	void AddTexture(UTexture2D& Texture);
@@ -146,9 +146,6 @@ public:
 	DECLARE_BENCHMARK_STAT(NumAllocatedSkeletalMeshes, int32);
 
 	DECLARE_BENCHMARK_STAT(NumInstances, int32);
-	DECLARE_BENCHMARK_STAT(NumInstancesLOD0, int32);
-	DECLARE_BENCHMARK_STAT(NumInstancesLOD1, int32);
-	DECLARE_BENCHMARK_STAT(NumInstancesLOD2, int32);
 
 	DECLARE_BENCHMARK_STAT(NumPendingInstanceUpdates, int32)
 	DECLARE_BENCHMARK_STAT(NumBuiltInstances, int32)
diff --git a/Engine/Plugins/Mutable/Source/CustomizableObject/Private/MuCO/CustomizableInstanceLODManagement.cpp b/Engine/Plugins/Mutable/Source/CustomizableObject/Private/MuCO/CustomizableInstanceLODManagement.cpp
--- a/Engine/Plugins/Mutable/Source/CustomizableObject/Private/MuCO/CustomizableInstanceLODManagement.cpp
+++ b/Engine/Plugins/Mutable/Source/CustomizableObject/Private/MuCO/CustomizableInstanceLODManagement.cpp
@@ -229,7 +229,8 @@ void UCustomizableInstanceLODManagement::UpdateInstanceDistsAndLODs(FMutableInst
 			SortedInstances[i]->SetIsPlayerOrNearIt(true);
 		}
 
-		TArray<uint16> RequestedLODs;
+		TMap<FName, uint8> MinLODs;
+		TMap<FName, uint8> RequestedLODs;
 		for (int32 i = NumGeneratedInstancesLimitLODs + NumGeneratedInstancesLimitLOD1 + NumGeneratedInstancesLimitLOD2; i < SortedInstances.Num(); ++i)
 		{
 			SortedInstances[i]->SetIsDiscardedBecauseOfTooManyInstances(true);
@@ -250,16 +251,26 @@ void UCustomizableInstanceLODManagement::UpdateInstanceDistsAndLODs(FMutableInst
 			for (int32 i = 0; i < NumGeneratedInstancesLimitLODs && i < SortedInstances.Num(); ++i)
 			{
 				const UCustomizableObject* CO = SortedInstances[i]->GetCustomizableObject();
-				const int32 NumObjectComponents = CO ? CO->GetComponentCount() : 0;
-				RequestedLODs.Init(0, NumObjectComponents);
-
+				
 				if (!bAlreadyReachedFixedLOD && SortedInstances[i]->GetMinSquareDistToPlayer() < DistanceForFixedLODSquared)
 				{
-					SortedInstances[i]->SetRequestedLODs(0, INDEX_NONE, RequestedLODs, InOutRequestedUpdates);
+					for (int32 ComponentIndex = 0; ComponentIndex < CO->GetComponentCount(); ++ComponentIndex)
+					{
+						MinLODs.Add(CO->GetComponentName(ComponentIndex), 0);
+						RequestedLODs.Add(CO->GetComponentName(ComponentIndex), 0);
+					}
+
+					SortedInstances[i]->SetRequestedLODs(MinLODs, RequestedLODs, InOutRequestedUpdates);
 				}
 				else
 				{
-					SortedInstances[i]->SetRequestedLODs(MAX_int32, INDEX_NONE, RequestedLODs, InOutRequestedUpdates);
+					for (int32 ComponentIndex = 0; ComponentIndex < CO->GetComponentCount(); ++ComponentIndex)
+					{
+						MinLODs.Add(CO->GetComponentName(ComponentIndex), UINT8_MAX);
+						RequestedLODs.Add(CO->GetComponentName(ComponentIndex), 0);
+					}
+
+					SortedInstances[i]->SetRequestedLODs(MinLODs, RequestedLODs, InOutRequestedUpdates);
 					bAlreadyReachedFixedLOD = true;
 				}
 			}
@@ -267,16 +278,26 @@ void UCustomizableInstanceLODManagement::UpdateInstanceDistsAndLODs(FMutableInst
 			for (int32 i = NumGeneratedInstancesLimitLODs; i < NumGeneratedInstancesLimitLODs + NumGeneratedInstancesLimitLOD1 && i < SortedInstances.Num(); ++i)
 			{
 				const UCustomizableObject* CO = SortedInstances[i]->GetCustomizableObject();
-				const int32 NumObjectComponents = CO ? CO->GetComponentCount() : 0;
-				RequestedLODs.Init(0, NumObjectComponents);
 
 				if (!bAlreadyReachedFixedLOD && SortedInstances[i]->GetMinSquareDistToPlayer() < DistanceForFixedLODSquared)
 				{
-					SortedInstances[i]->SetRequestedLODs(1, INDEX_NONE, RequestedLODs, InOutRequestedUpdates);
+					for (int32 ComponentIndex = 0; ComponentIndex < CO->GetComponentCount(); ++ComponentIndex)
+					{
+						MinLODs.Add(CO->GetComponentName(ComponentIndex), 1);
+						RequestedLODs.Add(CO->GetComponentName(ComponentIndex), 0);
+					}
+					
+					SortedInstances[i]->SetRequestedLODs(MinLODs, RequestedLODs, InOutRequestedUpdates);
 				}
 				else
 				{
-					SortedInstances[i]->SetRequestedLODs(MAX_int32, INDEX_NONE, RequestedLODs, InOutRequestedUpdates);
+					for (int32 ComponentIndex = 0; ComponentIndex < CO->GetComponentCount(); ++ComponentIndex)
+					{
+						MinLODs.Add(CO->GetComponentName(ComponentIndex), UINT8_MAX);
+						RequestedLODs.Add(CO->GetComponentName(ComponentIndex), 0);
+					}
+					
+					SortedInstances[i]->SetRequestedLODs(MinLODs, RequestedLODs, InOutRequestedUpdates);
 					bAlreadyReachedFixedLOD = true;
 				}
 			}
@@ -284,9 +305,13 @@ void UCustomizableInstanceLODManagement::UpdateInstanceDistsAndLODs(FMutableInst
 			for (int32 i = NumGeneratedInstancesLimitLODs + NumGeneratedInstancesLimitLOD1; i < NumGeneratedInstancesLimitLODs + NumGeneratedInstancesLimitLOD1 + NumGeneratedInstancesLimitLOD2 && i < SortedInstances.Num(); ++i)
 			{
 				const UCustomizableObject* CO = SortedInstances[i]->GetCustomizableObject();
-				const int32 NumObjectComponents = CO ? CO->GetComponentCount() : 0;
-				RequestedLODs.Init(0, NumObjectComponents);
-				SortedInstances[i]->SetRequestedLODs(2, INDEX_NONE, RequestedLODs, InOutRequestedUpdates);
+				for (int32 ComponentIndex = 0; ComponentIndex < CO->GetComponentCount(); ++ComponentIndex)
+				{
+					MinLODs.Add(CO->GetComponentName(ComponentIndex), 2);
+					RequestedLODs.Add(CO->GetComponentName(ComponentIndex), 0);
+				}
+
+				SortedInstances[i]->SetRequestedLODs(MinLODs, RequestedLODs, InOutRequestedUpdates);
 			}
 
 			for (int32 i = NumGeneratedInstancesLimitLODs + NumGeneratedInstancesLimitLOD1 + NumGeneratedInstancesLimitLOD2; i < SortedInstances.Num(); ++i)
@@ -300,9 +325,13 @@ void UCustomizableInstanceLODManagement::UpdateInstanceDistsAndLODs(FMutableInst
 			for (int32 i = 0; i < SortedInstances.Num(); ++i)
 			{
 				const UCustomizableObject* CO = SortedInstances[i]->GetCustomizableObject();
-				const int32 NumObjectComponents = CO ? CO->GetComponentCount() : 0;
-				RequestedLODs.Init(0, NumObjectComponents);
-				SortedInstances[i]->SetRequestedLODs(0, INDEX_NONE, RequestedLODs, InOutRequestedUpdates);
+				for (int32 ComponentIndex = 0; ComponentIndex < CO->GetComponentCount(); ++ComponentIndex)
+				{
+					MinLODs.Add(CO->GetComponentName(ComponentIndex), 2);
+					RequestedLODs.Add(CO->GetComponentName(ComponentIndex), 0);
+				}
+				
+				SortedInstances[i]->SetRequestedLODs(MinLODs, RequestedLODs, InOutRequestedUpdates);
 			}
 		}
 	}
@@ -310,10 +339,10 @@ void UCustomizableInstanceLODManagement::UpdateInstanceDistsAndLODs(FMutableInst
 	{
 		struct FLODTracker
 		{
-			int32 MinLOD = MAX_int32;
+			TMap<FName, uint8> MinLOD;
 
 			bool bInitialized = false;
-			TArray<uint16> RequestedLODPerComponent;
+			TMap<FName, uint8> RequestedLODPerComponent;
 		};
 
 		TMap<TObjectPtr<UCustomizableObjectInstance>, FLODTracker> InstancesMinLOD;
@@ -370,7 +399,13 @@ void UCustomizableInstanceLODManagement::UpdateInstanceDistsAndLODs(FMutableInst
 
 				if (!LODTracker.bInitialized)
 				{
-					LODTracker.RequestedLODPerComponent.Init(MAX_uint16, COI->GetCustomizableObject()->GetComponentCount());
+					const UCustomizableObject* CO = COI->GetCustomizableObject();
+
+					for (int32 ComponentIndex = 0; ComponentIndex < CO->GetComponentCount(); ++ComponentIndex)
+					{
+						LODTracker.RequestedLODPerComponent.Add(CO->GetComponentName(ComponentIndex), MAX_uint8);
+					}
+
 					LODTracker.bInitialized = true;
 				}
 
@@ -403,12 +438,13 @@ void UCustomizableInstanceLODManagement::UpdateInstanceDistsAndLODs(FMutableInst
 					}
 
 					// Use the component minLOD to set the minimum LOD to generate
-					LODTracker.MinLOD = FMath::Min(LODTracker.MinLOD, Parent->bOverrideMinLod ? Parent->MinLodModel : 0);
+					uint8& MinLOD = LODTracker.MinLOD.FindOrAdd(ComponentName, UINT8_MAX);
+					MinLOD = FMath::Min(MinLOD, static_cast<uint8>(Parent->bOverrideMinLod ? Parent->MinLodModel : 0));
 
 					// If the parent component have a SkeletalMesh use the RequestedLODLevel of the component as reference to know which LODs mutable should generate.
-					if (UE_MUTABLE_GETSKELETALMESHASSET(Parent) && LODTracker.RequestedLODPerComponent.IsValidIndex(ObjectComponentIndex))
+					if (UE_MUTABLE_GETSKELETALMESHASSET(Parent) && LODTracker.RequestedLODPerComponent.Contains(ComponentName))
 					{
-						uint16& RequestedLOD = LODTracker.RequestedLODPerComponent[ObjectComponentIndex];
+						uint8& RequestedLOD = LODTracker.RequestedLODPerComponent[ComponentName];
 						RequestedLOD = FMath::Min((int32)RequestedLOD, Parent->GetPredictedLODLevel());
 					}
 				}
@@ -430,10 +466,7 @@ void UCustomizableInstanceLODManagement::UpdateInstanceDistsAndLODs(FMutableInst
 					continue;
 				}
 
-				// Limit MinLOD
-				It.Value.MinLOD = FMath::Min(It.Value.MinLOD, CustomizableObject->GetNumLODs() - 1);
-
-				It.Key->SetRequestedLODs(It.Value.MinLOD, 0, It.Value.RequestedLODPerComponent, InOutRequestedUpdates);
+				It.Key->SetRequestedLODs(It.Value.MinLOD, It.Value.RequestedLODPerComponent, InOutRequestedUpdates);
 			}
 		}
 	}
diff --git a/Engine/Plugins/Mutable/Source/CustomizableObject/Private/MuCO/CustomizableObject.cpp b/Engine/Plugins/Mutable/Source/CustomizableObject/Private/MuCO/CustomizableObject.cpp
--- a/Engine/Plugins/Mutable/Source/CustomizableObject/Private/MuCO/CustomizableObject.cpp
+++ b/Engine/Plugins/Mutable/Source/CustomizableObject/Private/MuCO/CustomizableObject.cpp
@@ -1463,16 +1463,6 @@ UCustomizableObjectInstance* UCustomizableObject::CreateInstance()
 }
 
 
-int32 UCustomizableObject::GetNumLODs() const
-{
-	if (IsCompiled())
-	{
-		return GetPrivate()->GetModelResources().NumLODs;
-	}
-
-	return 0;
-}
-
 int32 UCustomizableObject::GetComponentCount() const
 {
 	if (IsCompiled())
@@ -2106,9 +2096,9 @@ void UCustomizableObjectPrivate::GetLowPriorityTextureNames(TArray<FString>& Out
 }
 
 
-int32 UCustomizableObjectPrivate::GetMinLODIndex() const
+uint8 UCustomizableObjectPrivate::GetMinLODIndex(const FName& ComponentName) const
 {
-	int32 MinLODIdx = 0;
+	uint8 MinLODIdx = 0;
 
 	if (GEngine && GEngine->UseSkeletalMeshMinLODPerQualityLevels)
 	{
@@ -2122,7 +2112,13 @@ int32 UCustomizableObjectPrivate::GetMinLODIndex() const
 		MinLODIdx = GetPublic()->LODSettings.MinLOD.GetValue();
 	}
 
-	return FMath::Max(MinLODIdx, static_cast<int32>(GetModelResources().FirstLODAvailable));
+	uint8 FirstLODAvailable = 0;
+	if (const uint8* Result = GetModelResources().FirstLODAvailable.Find(ComponentName))
+	{
+		FirstLODAvailable = *Result;
+	}
+
+	return FMath::Max(MinLODIdx, FirstLODAvailable);
 }
 
 
diff --git a/Engine/Plugins/Mutable/Source/CustomizableObject/Private/MuCO/CustomizableObjectInstance.cpp b/Engine/Plugins/Mutable/Source/CustomizableObject/Private/MuCO/CustomizableObjectInstance.cpp
--- a/Engine/Plugins/Mutable/Source/CustomizableObject/Private/MuCO/CustomizableObjectInstance.cpp
+++ b/Engine/Plugins/Mutable/Source/CustomizableObject/Private/MuCO/CustomizableObjectInstance.cpp
@@ -240,13 +240,7 @@ void UCustomizableInstancePrivate::InitCustomizableObjectData(const UCustomizabl
 	{
 		return;
 	}
-
-	// Init LOD Data
-	const FModelResources& ModelResources = InCustomizableObject->GetPrivate()->GetModelResources();
-	NumLODsAvailable = ModelResources.NumLODs;
-	FirstLODAvailable = ModelResources.FirstLODAvailable;
-	FirstResidentLOD = FMath::Clamp(ModelResources.NumLODsToStream, FirstLODAvailable, NumLODsAvailable);
-
+	
 	// Init Component Data
 	FCustomizableInstanceComponentData TemplateComponentData;
 	TemplateComponentData.LastMeshIdPerLOD.Init(MAX_uint64, MAX_MESH_LOD_COUNT);
@@ -1936,8 +1930,9 @@ bool UCustomizableInstancePrivate::DoComponentsNeedUpdate(UCustomizableObjectIns
 	for (int32 InstanceComponentIndex = 0; InstanceComponentIndex < NumInstanceComponents; ++InstanceComponentIndex)
 	{
 		const FInstanceUpdateData::FComponent& Component = OperationData->InstanceUpdateData.Components[InstanceComponentIndex];
-		
-		for (int32 LODIndex = OperationData->GetMinLOD(); LODIndex < Component.LODCount; ++LODIndex)
+		const FName ComponentName = OperationData->ComponentNames[Component.Id.GetValue()];
+
+		for (int32 LODIndex = OperationData->GetMinLOD()[ComponentName]; LODIndex < Component.LODCount; ++LODIndex)
 		{
 			const FInstanceUpdateData::FLOD& LOD = OperationData->InstanceUpdateData.LODs[Component.FirstLOD + LODIndex];
 
@@ -2009,11 +2004,11 @@ bool UCustomizableInstancePrivate::DoComponentsNeedUpdate(UCustomizableObjectIns
 		bool bIsReferenced = LOD.Mesh && LOD.Mesh->IsReference();
 		if (!bIsReferenced)
 		{
-			if (ComponentWithMesh[InstanceComponentIndex] && MeshIDs[InstanceComponentIndex * MAX_MESH_LOD_COUNT + OperationData->GetLODsAvailable(ObjectComponentIndex) - 1] == MAX_uint64)
+			if (ComponentWithMesh[InstanceComponentIndex] && MeshIDs[InstanceComponentIndex * MAX_MESH_LOD_COUNT + OperationData->NumLODsAvailable[ComponentName] - 1] == MAX_uint64)
 			{
 				UE_LOG(LogMutable, Error, TEXT("Failed to generate SkeletalMesh for CO Instance [%s]. CO [%s] is missing geometry for LOD [%d] Object Component [%d]."),
 					*Public->GetName(), *CustomizableObject->GetName(),
-					OperationData->GetLODsAvailable(ObjectComponentIndex) - 1, ObjectComponentIndex.GetValue());
+					OperationData->NumLODsAvailable[ComponentName] - 1, ObjectComponentIndex.GetValue());
 				bHasInvalidMesh = true;
 				continue;
 			}
@@ -2040,8 +2035,8 @@ bool UCustomizableInstancePrivate::DoComponentsNeedUpdate(UCustomizableObjectIns
 		// than the new FirstLODAvailable. This is an artifact of our LOD management system setting MinLOD.
 		if (!OperationData->MeshChangedPerInstanceComponent[InstanceComponentIndex] && bHadSkeletalMesh) 
 		{
-			int32 MinLOD = OperationData->GetMinLOD();
-			for (int32 LODIndex = OperationData->FirstLODAvailable; LODIndex < MinLOD; ++LODIndex)
+			int32 MinLOD = OperationData->GetMinLOD()[ComponentName];
+			for (int32 LODIndex = OperationData->FirstLODAvailable[ComponentName]; LODIndex < MinLOD; ++LODIndex)
 			{
 				TIndirectArray<FSkeletalMeshLODRenderData>& LODRenderData = (*SkeletalMeshPtr)->GetResourceForRendering()->LODRenderData;
 				if (LODRenderData[LODIndex].RenderSections.Num() != LODRenderData[MinLOD].RenderSections.Num())
@@ -2444,7 +2439,7 @@ bool UCustomizableInstancePrivate::UpdateSkeletalMesh_PostBeginUpdate0(UCustomiz
 		TObjectPtr<USkeletalMesh>* OldSkeletalMeshPtr = OldSkeletalMeshes.Find(ComponentName);
 		USkeletalMesh* OldSkeletalMesh = OldSkeletalMeshPtr ? *OldSkeletalMeshPtr : nullptr;
 
-		BuildOrCopyElementData(OperationData, SkeletalMesh, Public, InstanceComponentIndex);
+		BuildOrCopyElementData(OperationData, SkeletalMesh, *CustomizableObject, InstanceComponentIndex);
 		bool const bCopyRenderDataSuccess = BuildOrCopyRenderData(OperationData, SkeletalMesh, OldSkeletalMesh, Public, InstanceComponentIndex);
 		if (!bCopyRenderDataSuccess)
 		{
@@ -2452,7 +2447,7 @@ bool UCustomizableInstancePrivate::UpdateSkeletalMesh_PostBeginUpdate0(UCustomiz
 			return false;
 		}
 
-		BuildOrCopyMorphTargetsData(OperationData, SkeletalMesh, OldSkeletalMesh, Public, InstanceComponentIndex);
+		BuildOrCopyMorphTargetsData(OperationData, SkeletalMesh, OldSkeletalMesh, *CustomizableObject, InstanceComponentIndex);
 		BuildOrCopyClothingData(OperationData, SkeletalMesh, OldSkeletalMesh, Public, InstanceComponentIndex);
 
 		FSkeletalMeshRenderData* RenderData = SkeletalMesh->GetResourceForRendering();
@@ -2475,7 +2470,7 @@ bool UCustomizableInstancePrivate::UpdateSkeletalMesh_PostBeginUpdate0(UCustomiz
 
 		if (UCustomizableObjectSkeletalMesh* StreamableMesh = Cast<UCustomizableObjectSkeletalMesh>(SkeletalMesh))
 		{
-			StreamableMesh->InitMutableStreamingData(OperationData, *CustomizableObject, Component->FirstLOD, Component->LODCount);
+			StreamableMesh->InitMutableStreamingData(OperationData, ComponentName, Component->FirstLOD, Component->LODCount);
 		}
 	}
 
@@ -3181,7 +3176,7 @@ UCustomizableInstancePrivate* UCustomizableObjectInstance::GetPrivate() const
 FMutableUpdateCandidate::FMutableUpdateCandidate(UCustomizableObjectInstance* InCustomizableObjectInstance): CustomizableObjectInstance(InCustomizableObjectInstance)
 {
 	const FCustomizableObjectInstanceDescriptor& Descriptor = InCustomizableObjectInstance->GetPrivate()->GetDescriptor();
-	MinLOD = Descriptor.GetMinLod();
+	MinLOD = Descriptor.MinLOD;
 	RequestedLODLevels = Descriptor.GetRequestedLODLevels();
 }
 
@@ -3575,6 +3570,8 @@ void UCustomizableInstancePrivate::InitSkeletalMeshData(const TSharedRef<FUpdate
 
 	check(SkeletalMesh);
 
+	const FName ComponentName = OperationData->ComponentNames[ObjectComponentIndex.GetValue()];
+
 	SkeletalMesh->NeverStream = !OperationData->bStreamMeshLODs;
 
 	SkeletalMesh->SetImportedBounds(RefSkeletalMeshData.Bounds);
@@ -3586,9 +3583,9 @@ void UCustomizableInstancePrivate::InitSkeletalMeshData(const TSharedRef<FUpdate
 	{
 		SkeletalMesh->SetSupportRayTracing(false);
 	}
-
+	
 	// Set Min LOD
-	SkeletalMesh->SetMinLod(FMath::Max(CustomizableObject.LODSettings.MinLOD.GetDefault(), (int32)FirstLODAvailable));
+	SkeletalMesh->SetMinLod(FMath::Max(CustomizableObject.LODSettings.MinLOD.GetDefault(), static_cast<int32>(OperationData->FirstLODAvailable[ComponentName])));
 	SkeletalMesh->SetQualityLevelMinLod(CustomizableObject.LODSettings.MinQualityLevelLOD);
 
 	SkeletalMesh->SetHasVertexColors(false);
@@ -3644,12 +3641,12 @@ void UCustomizableInstancePrivate::InitSkeletalMeshData(const TSharedRef<FUpdate
 	{
 		MUTABLE_CPUPROFILER_SCOPE(InitSkeletalMesh_AddLODData);
 		SkeletalMesh->AllocateResourceForRendering();
-
+		
 		FSkeletalMeshRenderData* RenderData = SkeletalMesh->GetResourceForRendering();
-		int32 NumLODsAvailablePerComponent = OperationData->GetLODsAvailable(ObjectComponentIndex);
-		RenderData->NumInlinedLODs = NumLODsAvailablePerComponent - OperationData->FirstResidentLOD;
-		RenderData->NumNonOptionalLODs = NumLODsAvailablePerComponent - OperationData->FirstLODAvailable;
-		RenderData->CurrentFirstLODIdx = OperationData->FirstResidentLOD;
+		int32 NumLODsAvailablePerComponent = OperationData->NumLODsAvailable[ComponentName];
+		RenderData->NumInlinedLODs = NumLODsAvailablePerComponent - OperationData->FirstResidentLOD[ComponentName];
+		RenderData->NumNonOptionalLODs = NumLODsAvailablePerComponent - OperationData->FirstLODAvailable[ComponentName];
+		RenderData->CurrentFirstLODIdx = OperationData->FirstResidentLOD[ComponentName];
 		RenderData->LODBiasModifier = 0;
 		
 		if (bEnableRayTracingFix)
@@ -3662,8 +3659,8 @@ void UCustomizableInstancePrivate::InitSkeletalMeshData(const TSharedRef<FUpdate
 			RenderData->LODRenderData.Add(new FSkeletalMeshLODRenderData());
 			
 			FSkeletalMeshLODRenderData& LODRenderData = RenderData->LODRenderData[LODIndex];
-			LODRenderData.bIsLODOptional = LODIndex < OperationData->FirstLODAvailable;
-			LODRenderData.bStreamedDataInlined = LODIndex >= OperationData->FirstResidentLOD;
+			LODRenderData.bIsLODOptional = LODIndex < OperationData->FirstLODAvailable[ComponentName];
+			LODRenderData.bStreamedDataInlined = LODIndex >= OperationData->FirstResidentLOD[ComponentName];
 
 			const FMutableRefLODData& LODData = RefSkeletalMeshData.LODData[LODIndex];
 			FSkeletalMeshLODInfo& LODInfo = SkeletalMesh->AddLODInfo();
@@ -3982,7 +3979,7 @@ void UCustomizableInstancePrivate::BuildMeshSockets(const TSharedRef<FUpdateCont
 }
 
 
-void UCustomizableInstancePrivate::BuildOrCopyElementData(const TSharedRef<FUpdateContextPrivate>& OperationData, USkeletalMesh* SkeletalMesh, UCustomizableObjectInstance* CustomizableObjectInstance, 
+void UCustomizableInstancePrivate::BuildOrCopyElementData(const TSharedRef<FUpdateContextPrivate>& OperationData, USkeletalMesh* SkeletalMesh, UCustomizableObject& CustomizableObject, 
 	FCustomizableObjectInstanceComponentIndex InstanceComponentIndex)
 {
 	MUTABLE_CPUPROFILER_SCOPE(UCustomizableInstancePrivate::BuildOrCopyElementData);
@@ -3993,7 +3990,9 @@ void UCustomizableInstancePrivate::BuildOrCopyElementData(const TSharedRef<FUpda
 		return;
 	}
 
-	for (int32 LODIndex = FirstLODAvailable; LODIndex < Component->LODCount; ++LODIndex)
+	const FName ComponentName = OperationData->ComponentNames[Component->Id.GetValue()];
+
+	for (int32 LODIndex = OperationData->FirstLODAvailable[ComponentName]; LODIndex < Component->LODCount; ++LODIndex)
 	{
 		const FInstanceUpdateData::FLOD& LOD = OperationData->InstanceUpdateData.LODs[Component->FirstLOD+LODIndex];
 
@@ -4014,7 +4013,7 @@ void UCustomizableInstancePrivate::BuildOrCopyElementData(const TSharedRef<FUpda
 	}
 }
 
-void UCustomizableInstancePrivate::BuildOrCopyMorphTargetsData(const TSharedRef<FUpdateContextPrivate>& OperationData, USkeletalMesh* SkeletalMesh, const USkeletalMesh* LastUpdateSkeletalMesh, UCustomizableObjectInstance* CustomizableObjectInstance, FCustomizableObjectInstanceComponentIndex InstanceComponentIndex)
+void UCustomizableInstancePrivate::BuildOrCopyMorphTargetsData(const TSharedRef<FUpdateContextPrivate>& OperationData, USkeletalMesh* SkeletalMesh, const USkeletalMesh* LastUpdateSkeletalMesh, UCustomizableObject& CustomizableObject, FCustomizableObjectInstanceComponentIndex InstanceComponentIndex)
 {
 	MUTABLE_CPUPROFILER_SCOPE(UCustomizableInstancePrivate::BuildOrCopyMorphTargetsData);
 
@@ -4030,6 +4029,8 @@ void UCustomizableInstancePrivate::BuildOrCopyMorphTargetsData(const TSharedRef<
 	}
 	
 	FCustomizableObjectComponentIndex ObjectComponentIndex = OperationData->GetObjectComponentIndex(InstanceComponentIndex);
+	const FName ComponentName = OperationData->ComponentNames[ObjectComponentIndex.GetValue()];
+
 	FInstanceUpdateData::FRealTimeMorphsComponentData* ComponentMorphTargets = OperationData->GetComponentMorphTargets(ObjectComponentIndex);
 	if (!ComponentMorphTargets)
 	{
@@ -4055,8 +4056,8 @@ void UCustomizableInstancePrivate::BuildOrCopyMorphTargetsData(const TSharedRef<
 	}
 
 	// Copy MorphTargets from the FirstGeneratedLOD to the LODs below
-	const int32 FirstGeneratedLOD = FMath::Max((int32)OperationData->GetRequestedLODs(ObjectComponentIndex), OperationData->GetMinLOD());
-	for (int32 LODIndex = OperationData->FirstLODAvailable; LODIndex < FirstGeneratedLOD; ++LODIndex)
+	const int32 FirstGeneratedLOD = FMath::Max(OperationData->GetRequestedLODs()[ComponentName], OperationData->GetMinLOD()[ComponentName]);
+	for (int32 LODIndex = OperationData->FirstLODAvailable[ComponentName]; LODIndex < FirstGeneratedLOD; ++LODIndex)
 	{
 		MUTABLE_CPUPROFILER_SCOPE(CopyMorphTargetsData);
 
@@ -4133,10 +4134,12 @@ void UCustomizableInstancePrivate::BuildOrCopyClothingData(const TSharedRef<FUpd
 		return;
 	}
 
+	const FName ComponentName = OperationData->ComponentNames[Component->Id.GetValue()];
+
 	{
 		MUTABLE_CPUPROFILER_SCOPE(DiscoverSectionsWithCloth);
 
-		for (int32 LODIndex = OperationData->FirstLODAvailable; LODIndex < Component->LODCount; ++LODIndex)
+		for (int32 LODIndex = OperationData->FirstLODAvailable[ComponentName]; LODIndex < Component->LODCount; ++LODIndex)
 		{
 			const FInstanceUpdateData::FLOD& LOD = OperationData->InstanceUpdateData.LODs[Component->FirstLOD+LODIndex];
 
@@ -4273,7 +4276,7 @@ void UCustomizableInstancePrivate::BuildOrCopyClothingData(const TSharedRef<FUpd
 			// Only initilialize once, multiple sections with cloth could point to the same cloth asset.
 			if (!DstAssetData.LodMap.Num())
 			{
-				DstAssetData.LodMap.Init(INDEX_NONE, OperationData->GetLODsAvailable(Component->Id));
+				DstAssetData.LodMap.Init(INDEX_NONE, OperationData->NumLODsAvailable[ComponentName]);
 
 				DstAssetData.UsedBoneNames = SrcAssetData.UsedBoneNames;
 				DstAssetData.UsedBoneIndices = SrcAssetData.UsedBoneIndices;
@@ -4884,7 +4887,7 @@ void UCustomizableInstancePrivate::BuildOrCopyClothingData(const TSharedRef<FUpd
 	{
 		int32 NumSectionsWithClothProcessed = 0;
 
-		for (int32 LODIndex = OperationData->GetMinLOD(); LODIndex < Component->LODCount; ++LODIndex)
+		for (int32 LODIndex = OperationData->GetMinLOD()[ComponentName]; LODIndex < Component->LODCount; ++LODIndex)
 		{
 			const FInstanceUpdateData::FLOD& LOD = OperationData->InstanceUpdateData.LODs[Component->FirstLOD+LODIndex];
 
@@ -4933,7 +4936,7 @@ void UCustomizableInstancePrivate::BuildOrCopyClothingData(const TSharedRef<FUpd
 	{
 		MUTABLE_CPUPROFILER_SCOPE(InitClothRenderData)
 		// Based on FSkeletalMeshLODModel::GetClothMappingData().
-		for (int32 LODIndex = OperationData->GetMinLOD(); LODIndex < OperationData->GetLODsAvailable(Component->Id); ++LODIndex)
+		for (int32 LODIndex = OperationData->GetMinLOD()[ComponentName]; LODIndex < OperationData->NumLODsAvailable[ComponentName]; ++LODIndex)
 		{
 			FSkeletalMeshLODRenderData& LODModel = RenderResource->LODRenderData[LODIndex];
 	
@@ -5089,8 +5092,9 @@ bool UCustomizableInstancePrivate::BuildOrCopyRenderData(const TSharedRef<FUpdat
 	}
 
 	const FModelResources& ModelResources = CustomizableObject->GetPrivate()->GetModelResources();
+	const FName ComponentName = ModelResources.ComponentNamesPerObjectComponent[Component->Id.GetValue()];
 
-	const int32 FirstGeneratedLOD = FMath::Max((int32)OperationData->GetRequestedLODs(Component->Id), OperationData->GetMinLOD());
+	const int32 FirstGeneratedLOD = FMath::Max(OperationData->GetRequestedLODs()[ComponentName], OperationData->GetMinLOD()[ComponentName]);
 
 
 	for (int32 LODIndex = FirstGeneratedLOD; LODIndex < Component->LODCount; ++LODIndex)
@@ -5289,7 +5293,7 @@ bool UCustomizableInstancePrivate::BuildOrCopyRenderData(const TSharedRef<FUpdat
 	}
 
 	// Copy LODRenderData from the FirstGeneratedLOD to the LODs below
-	for (int32 LODIndex = OperationData->FirstLODAvailable; LODIndex < FirstGeneratedLOD; ++LODIndex)
+	for (int32 LODIndex = OperationData->FirstLODAvailable[ComponentName]; LODIndex < FirstGeneratedLOD; ++LODIndex)
 	{
 		TRACE_CPUPROFILER_EVENT_SCOPE_TEXT(*FString::Printf(TEXT("CopyRenderData: From LOD %d to LOD %d"), FirstGeneratedLOD, LODIndex));
 
@@ -5422,7 +5426,9 @@ UE::Tasks::FTask UCustomizableInstancePrivate::LoadAdditionalAssetsAndData(const
 		{
 			const FInstanceUpdateData::FComponent& Component = Components[InstanceComponentIndex];
 
-			for (int32 LODIndex = OperationData->FirstLODAvailable; LODIndex < Component.LODCount; ++LODIndex)
+			const FName ComponentName = OperationData->ComponentNames[Component.Id.GetValue()];
+
+			for (int32 LODIndex = OperationData->FirstLODAvailable[ComponentName]; LODIndex < Component.LODCount; ++LODIndex)
 			{
 				const FInstanceUpdateData::FLOD& LOD = OperationData->InstanceUpdateData.LODs[Component.FirstLOD + LODIndex];
 
@@ -6341,7 +6347,7 @@ void UCustomizableInstancePrivate::BuildMaterials(const TSharedRef<FUpdateContex
 
 		MUTABLE_CPUPROFILER_SCOPE(BuildMaterials_LODLoop);
 
-		const int32 FirstGeneratedLOD = FMath::Max((int32)OperationData->GetRequestedLODs(ObjectComponentIndex), OperationData->GetMinLOD());
+		const int32 FirstGeneratedLOD = FMath::Max(OperationData->GetRequestedLODs()[ComponentName], OperationData->GetMinLOD()[ComponentName]);
 
 		for (int32 LODIndex = FirstGeneratedLOD; LODIndex < Component.LODCount; LODIndex++)
 		{
@@ -6928,7 +6934,7 @@ void UCustomizableInstancePrivate::BuildMaterials(const TSharedRef<FUpdateContex
 		if(!bUseOverrideMaterialsOnly)
 		{
 			// Copy data from the FirstGeneratedLOD into the LODs below.
-			for (int32 LODIndex = OperationData->FirstLODAvailable; LODIndex < FirstGeneratedLOD; ++LODIndex)
+			for (int32 LODIndex = OperationData->FirstLODAvailable[ComponentName]; LODIndex < FirstGeneratedLOD; ++LODIndex)
 			{
 				SkeletalMesh->GetLODInfo(LODIndex)->LODMaterialMap = SkeletalMesh->GetLODInfo(FirstGeneratedLOD)->LODMaterialMap;
 
@@ -7122,29 +7128,7 @@ int32 UCustomizableObjectInstance::GetNumComponents() const
 }
 
 
-int32 UCustomizableObjectInstance::GetMinLODToLoad() const
-{
-	return Descriptor.MinLOD;	
-}
-
-
-int32 UCustomizableObjectInstance::GetCurrentMinLOD() const
-{
-	return GetPrivate()->CommittedDescriptor.GetMinLod();
-}
-
-
-#if !UE_BUILD_SHIPPING
-static bool bIgnoreMinLOD = false;
-FAutoConsoleVariableRef CVarMutableIgnoreMinMaxLOD(
-	TEXT("Mutable.IgnoreMinMaxLOD"),
-	bIgnoreMinLOD,
-	TEXT("The limits on the number of LODs to generate will be ignored."));
-#endif
-
-
-void UCustomizableObjectInstance::SetRequestedLODs(int32 InMinLOD, int32 , const TArray<uint16>& InRequestedLODsPerObjectComponent, 
-	                                               FMutableInstanceUpdateMap& InOutRequestedUpdates)
+void UCustomizableObjectInstance::SetRequestedLODs(const TMap<FName, uint8>& InMinLODs, const TMap<FName, uint8>& InRequestedLODs, FMutableInstanceUpdateMap& InOutRequestedUpdates)
 {
 	check(PrivateData);
 
@@ -7172,67 +7156,70 @@ void UCustomizableObjectInstance::SetRequestedLODs(int32 InMinLOD, int32 , const
 		return;
 	}
 	
-#if !UE_BUILD_SHIPPING
-	// Ignore Min LOD limits. Mainly used for debug
-	if (bIgnoreMinLOD)
-	{
-		InMinLOD = 0;
-	}
-#endif
-
 	FMutableUpdateCandidate MutableUpdateCandidate(this);
 
 	// Clamp Min LOD
-	const int32 MinLODIdx = CustomizableObject->GetPrivate()->GetMinLODIndex();
-	const int32 MaxLODIdx = PrivateData->NumLODsAvailable - 1;
-	InMinLOD = FMath::Clamp(InMinLOD, MinLODIdx, MaxLODIdx);
-
-	const bool bMinLODChanged = Descriptor.MinLOD != InMinLOD;
-
-	PrivateData->SetCOInstanceFlags(InMinLOD > GetCurrentMinLOD() ? PendingLODsDowngrade : ECONone);
-
+	const FModelResources& ModelResources = CustomizableObject->GetPrivate()->GetModelResources();
+	
+	bool bMinLODChanged = false;
+	
 	// Save the new LODs
-	MutableUpdateCandidate.MinLOD = InMinLOD;
+	MutableUpdateCandidate.MinLOD = InMinLODs;
 	MutableUpdateCandidate.RequestedLODLevels = Descriptor.GetRequestedLODLevels();
 
-	bool bUpdateRequestedLODs = false;
-	if (UCustomizableObjectSystem::GetInstance()->IsOnlyGenerateRequestedLODsEnabled())
+	const TMap<FName, uint8>& GeneratedLODsPerComponent = GetPrivate()->CommittedDescriptorHash.RequestedLODsPerComponent;
+	
+	for (const FName& ComponentName : CustomizableObject->GetPrivate()->GetModelResources().ComponentNamesPerObjectComponent)
 	{
-		const uint16 FirstNonStreamedLODIndex = FMath::Clamp(PrivateData->FirstResidentLOD, 0, MaxLODIdx);
-
-		const int32 ObjectComponentCount = CustomizableObject->GetComponentCount();
-		if (ObjectComponentCount != MutableUpdateCandidate.RequestedLODLevels.Num())
+		uint8 InMinLOD = 0;
+		if (const uint8* Result = InMinLODs.Find(ComponentName))
 		{
-			MutableUpdateCandidate.RequestedLODLevels.Init(FirstNonStreamedLODIndex, ObjectComponentCount);
+			InMinLOD = *Result;
 		}
 
-		const TArray<uint16>& GeneratedLODsPerComponent = GetPrivate()->CommittedDescriptorHash.RequestedLODsPerComponent;
+		const uint8 MinLODIdx = CustomizableObject->GetPrivate()->GetMinLODIndex(ComponentName);
+		const int32 MaxLODIdx = ModelResources.NumLODs[ComponentName] - 1;
+		InMinLOD = FMath::Clamp(InMinLOD, MinLODIdx, MaxLODIdx);
 
-		const bool bIgnoreGeneratedLODs = GeneratedLODsPerComponent.Num() != ObjectComponentCount;
-
-		if (bMinLODChanged || bIgnoreGeneratedLODs || Descriptor.GetRequestedLODLevels() != InRequestedLODsPerObjectComponent)
+		uint8 DescriptorMinLOD = 0;
+		if (const uint8* Result = Descriptor.MinLOD.Find(ComponentName))
 		{
-			check(InRequestedLODsPerObjectComponent.Num() == ObjectComponentCount);
-
-			bUpdateRequestedLODs = bIgnoreGeneratedLODs;
-			for (int32 ObjectComponentIndex = 0; ObjectComponentIndex < ObjectComponentCount; ++ObjectComponentIndex)
+			DescriptorMinLOD = *Result;
+		}
+		
+		bMinLODChanged |= DescriptorMinLOD != InMinLOD;
+		
+		if (UCustomizableObjectSystem::GetInstance()->IsOnlyGenerateRequestedLODsEnabled())
+		{
+			uint8 CurrentMinLOD = 0;
+			if (const uint8* Result = GetPrivate()->CommittedDescriptor.MinLOD.Find(ComponentName))
 			{
-				uint16 PredictedLOD = FMath::Min(InRequestedLODsPerObjectComponent[ObjectComponentIndex], FirstNonStreamedLODIndex);
+				CurrentMinLOD = *Result;
+			}
+			
+			PrivateData->SetCOInstanceFlags(InMinLOD > CurrentMinLOD ? PendingLODsDowngrade : ECONone);
 
-				if (!bIgnoreGeneratedLODs)
-				{
-					PredictedLOD = FMath::Min(PredictedLOD, GeneratedLODsPerComponent[ObjectComponentIndex]);
-				}
+			const uint8 FirstNonStreamedLODIndex = FMath::Clamp(CustomizableObject->GetPrivate()->GetModelResources().NumLODsToStream[ComponentName], 0, MaxLODIdx);
+
+			MutableUpdateCandidate.RequestedLODLevels.Add(ComponentName, FirstNonStreamedLODIndex);
 
-				bUpdateRequestedLODs |= (PredictedLOD != MutableUpdateCandidate.RequestedLODLevels[ObjectComponentIndex]);
+			uint8 PredictedLOD = FirstNonStreamedLODIndex;
+			if (const uint8* Result = InRequestedLODs.Find(ComponentName))
+			{
+				PredictedLOD = FMath::Min(*Result, PredictedLOD);
+			}
 
-				// Save new RequestedLODs
-				MutableUpdateCandidate.RequestedLODLevels[ObjectComponentIndex] = PredictedLOD;
+			if (const uint8* Result = GeneratedLODsPerComponent.Find(ComponentName))
+			{
+				PredictedLOD = FMath::Min(*Result, PredictedLOD);
 			}
+				
+			// Save new RequestedLODs
+			MutableUpdateCandidate.RequestedLODLevels[ComponentName] = PredictedLOD;
 		}
 	}
 
-	if (bMinLODChanged || bUpdateRequestedLODs)
+	if (bMinLODChanged || GeneratedLODsPerComponent != MutableUpdateCandidate.RequestedLODLevels)
 	{
 		// TODO: Remove this flag as it will become redundant with the new InOutRequestedUpdates system
 		PrivateData->SetCOInstanceFlags(PendingLODsUpdate);
@@ -7242,12 +7229,6 @@ void UCustomizableObjectInstance::SetRequestedLODs(int32 InMinLOD, int32 , const
 }
 
 
-const TArray<uint16>& UCustomizableObjectInstance::GetRequestedLODsPerComponent() const
-{
-	return Descriptor.RequestedLODLevels;
-}
-
-
 #if WITH_EDITOR
 void UCustomizableObjectInstance::Bake(const FBakingConfiguration& InBakingConfiguration)
 {
diff --git a/Engine/Plugins/Mutable/Source/CustomizableObject/Private/MuCO/CustomizableObjectInstanceDescriptor.cpp b/Engine/Plugins/Mutable/Source/CustomizableObject/Private/MuCO/CustomizableObjectInstanceDescriptor.cpp
--- a/Engine/Plugins/Mutable/Source/CustomizableObject/Private/MuCO/CustomizableObjectInstanceDescriptor.cpp
+++ b/Engine/Plugins/Mutable/Source/CustomizableObject/Private/MuCO/CustomizableObjectInstanceDescriptor.cpp
@@ -760,8 +760,8 @@ void FCustomizableObjectInstanceDescriptor::ReloadParameters()
 	}
 
 	SetState(FMath::Clamp(GetState(), 0, CustomizableObject->GetStateCount() - 1));
-	
-	RequestedLODLevels.Init(MAX_uint16, CustomizableObject->GetComponentCount());
+
+	RequestedLODLevels.Empty();
 
 	TArray<FCustomizableObjectBoolParameterValue> OldBoolParameters = BoolParameters;
 	TArray<FCustomizableObjectIntParameterValue> OldIntParameters = IntParameters;
@@ -1175,45 +1175,18 @@ void FCustomizableObjectInstanceDescriptor::ReloadParameters()
 }
 
 
-int32 FCustomizableObjectInstanceDescriptor::GetMinLod() const
-{
-	return MinLOD;
-}
-
-
-void FCustomizableObjectInstanceDescriptor::SetMinLod(int32 InMinLOD)
-{
-	MinLOD = InMinLOD;
-}
-
-
-void FCustomizableObjectInstanceDescriptor::SetRequestedLODLevels(const TArray<uint16>& InRequestedLODLevels)
+void FCustomizableObjectInstanceDescriptor::SetRequestedLODLevels(const TMap<FName, uint8>& InRequestedLODLevels)
 {
 	RequestedLODLevels = InRequestedLODLevels;
 }
 
 
-const TArray<uint16>& FCustomizableObjectInstanceDescriptor::GetRequestedLODLevels() const
+const TMap<FName, uint8>& FCustomizableObjectInstanceDescriptor::GetRequestedLODLevels() const
 {
 	return RequestedLODLevels;
 }
 
 
-uint16 FCustomizableObjectInstanceDescriptor::GetRequestedLODs(FCustomizableObjectComponentIndex Index) const
-{
-	uint16 Result = 0;
-	if (RequestedLODLevels.IsValidIndex(Index.GetValue()))
-	{
-		Result = RequestedLODLevels[Index.GetValue()];
-	}
-	else
-	{
-		ensure(false);
-	}
-	return Result;
-}
-
-
 const TArray<FCustomizableObjectBoolParameterValue>& FCustomizableObjectInstanceDescriptor::GetBoolParameters() const
 {
 	return BoolParameters;
diff --git a/Engine/Plugins/Mutable/Source/CustomizableObject/Private/MuCO/CustomizableObjectSkeletalMesh.cpp b/Engine/Plugins/Mutable/Source/CustomizableObject/Private/MuCO/CustomizableObjectSkeletalMesh.cpp
--- a/Engine/Plugins/Mutable/Source/CustomizableObject/Private/MuCO/CustomizableObjectSkeletalMesh.cpp
+++ b/Engine/Plugins/Mutable/Source/CustomizableObject/Private/MuCO/CustomizableObjectSkeletalMesh.cpp
@@ -31,8 +31,10 @@ UCustomizableObjectSkeletalMesh* UCustomizableObjectSkeletalMesh::CreateSkeletal
 }
 
 
-void UCustomizableObjectSkeletalMesh::InitMutableStreamingData(const TSharedRef<FUpdateContextPrivate>& InOperationData, UCustomizableObject& CustomizableObject, const int32 FirstLOD, const int32 LODCount)
+void UCustomizableObjectSkeletalMesh::InitMutableStreamingData(const TSharedRef<FUpdateContextPrivate>& InOperationData, const FName& ComponentName, const int32 FirstLOD, const int32 LODCount)
 {
+	UCustomizableObject& CustomizableObject = *InOperationData->Object;
+
 	// Debug info
 	CustomizableObjectPathName = GetNameSafe(&CustomizableObject);
 
@@ -44,7 +46,7 @@ void UCustomizableObjectSkeletalMesh::InitMutableStreamingData(const TSharedRef<
 
 	MeshIDs.Init(MAX_uint64, MAX_MESH_LOD_COUNT);
 
-	for (int32 LODIndex = InOperationData->FirstLODAvailable; LODIndex < LODCount; ++LODIndex)
+	for (int32 LODIndex = InOperationData->FirstLODAvailable[ComponentName]; LODIndex < LODCount; ++LODIndex)
 	{
 		const FInstanceUpdateData::FLOD& LOD = InOperationData->InstanceUpdateData.LODs[FirstLOD + LODIndex];
 		MeshIDs[LODIndex] = LOD.MeshID;
diff --git a/Engine/Plugins/Mutable/Source/CustomizableObject/Private/MuCO/CustomizableObjectSystem.cpp b/Engine/Plugins/Mutable/Source/CustomizableObject/Private/MuCO/CustomizableObjectSystem.cpp
--- a/Engine/Plugins/Mutable/Source/CustomizableObject/Private/MuCO/CustomizableObjectSystem.cpp
+++ b/Engine/Plugins/Mutable/Source/CustomizableObject/Private/MuCO/CustomizableObjectSystem.cpp
@@ -168,18 +168,25 @@ static FAutoConsoleVariableSink CVarMutableSink(FConsoleCommandDelegate::CreateS
 FUpdateContextPrivate::FUpdateContextPrivate(UCustomizableObjectInstance& InInstance, const FCustomizableObjectInstanceDescriptor& Descriptor)
 {
 	check(IsInGameThread());
-	check(InInstance.GetCustomizableObject());
-
+	
 	Instance = &InInstance;
 
+	UCustomizableObject* LocalObject = InInstance.GetCustomizableObject();
+	check(LocalObject);
+
+	const FModelResources& ModelResources = LocalObject->GetPrivate()->GetModelResources();
+
+	Object = LocalObject;
+
 	CapturedDescriptor = Descriptor;
-	InInstance.GetCustomizableObject()->GetPrivate()->ApplyStateForcedValuesToParameters(CapturedDescriptor);
+	LocalObject->GetPrivate()->ApplyStateForcedValuesToParameters(CapturedDescriptor);
 
 	CapturedDescriptorHash = FDescriptorHash(Descriptor);
-	NumObjectComponents = InInstance.GetCustomizableObject()->GetComponentCount();
-	FirstLODAvailable = InInstance.GetCustomizableObject()->GetPrivate()->GetMinLODIndex();
-	FirstResidentLOD = InInstance.GetPrivate()->FirstResidentLOD;
-
+	NumObjectComponents = LocalObject->GetComponentCount();
+	FirstLODAvailable = ModelResources.FirstLODAvailable;
+	FirstResidentLOD = ModelResources.NumLODsToStream;
+	ComponentNames = ModelResources.ComponentNamesPerObjectComponent;
+	
 	MutableSystem = UCustomizableObjectSystem::GetInstance()->GetPrivate()->MutableSystem;	
 	check(MutableSystem);
 
@@ -206,32 +213,26 @@ FUpdateContextPrivate::~FUpdateContextPrivate()
 }
 
 
-int32 FUpdateContextPrivate::GetMinLOD() const
+const TMap<FName, uint8>& FUpdateContextPrivate::GetMinLOD() const
 {
-	return CapturedDescriptor.GetMinLod();
+	return CapturedDescriptor.MinLOD;
 }
 
 
-void FUpdateContextPrivate::SetMinLOD(int32 MinLOD)
+void FUpdateContextPrivate::SetMinLOD(const TMap<FName, uint8>& MinLOD)
 {
-	CapturedDescriptor.SetMinLod(MinLOD);
-	CapturedDescriptorHash.MinLOD = MinLOD;
+	CapturedDescriptor.MinLOD = MinLOD;
+	CapturedDescriptorHash.MinLODs = MinLOD;
 }
 
 
-const TArray<uint16>& FUpdateContextPrivate::GetRequestedLODs() const
+const TMap<FName, uint8>& FUpdateContextPrivate::GetRequestedLODs() const
 {
 	return CapturedDescriptor.GetRequestedLODLevels();
 }
 
 
-uint16 FUpdateContextPrivate::GetRequestedLODs(FCustomizableObjectComponentIndex Index) const
-{
-	return CapturedDescriptor.GetRequestedLODs(Index);
-}
-
-
-void FUpdateContextPrivate::SetRequestedLODs(TArray<uint16>& RequestedLODs)
+void FUpdateContextPrivate::SetRequestedLODs(const TMap<FName, uint8>& RequestedLODs)
 {
 	CapturedDescriptor.SetRequestedLODLevels(RequestedLODs);
 	CapturedDescriptorHash.RequestedLODsPerComponent = RequestedLODs;
@@ -1246,11 +1247,6 @@ void UCustomizableObjectSystemPrivate::EnqueueUpdateSkeletalMesh(const TSharedRe
 	}		
 	case EUpdateRequired::Update:
 	{
-		if (InstancePrivate->HasCOInstanceFlags(PendingLODsUpdate))
-		{
-			UE_LOG(LogMutable, Verbose, TEXT("Min LOD change: %d -> %d"), Instance->GetCurrentMinLOD(), Instance->GetMinLODToLoad());
-		}
-
 		if (const FMutablePendingInstanceUpdate* QueueElem = MutablePendingInstanceWork.GetUpdate(Instance))
 		{
 			if (Context->GetCapturedDescriptorHash().IsSubset(QueueElem->Context->GetCapturedDescriptorHash()))
@@ -1728,50 +1724,36 @@ namespace impl
 			return;
 		}
 
-		Operation->InitNumLODsAvailablePerComponent(Operation->NumObjectComponents);
-
-		int32 LODsInFirstValidComponent = -1;
-		for (int32 InstanceComponentIndex = 0; InstanceComponentIndex < Operation->NumInstanceComponents; ++InstanceComponentIndex)
-		{
-			int32 NumLODs = Operation->MutableInstance->GetLODCount(InstanceComponentIndex);
-			FCustomizableObjectComponentIndex Index(Operation->MutableInstance->GetComponentId(InstanceComponentIndex));
-			Operation->SetNumLODsAvailable(Index, NumLODs);
-			if (LODsInFirstValidComponent <= 0)
-			{
-				LODsInFirstValidComponent = NumLODs;
-			}
-		}
+		TMap<FName, uint8> MinLODs = Operation->bStreamMeshLODs ? TMap<FName, uint8>() : Operation->GetMinLOD();
+		TMap<FName, uint8> RequestedLODs = Operation->GetRequestedLODs();
 		
-		int32 CurrentMinLOD = Operation->bStreamMeshLODs ? 0 : Operation->GetMinLOD();
-		CurrentMinLOD = FMath::Clamp(CurrentMinLOD, Operation->FirstLODAvailable, LODsInFirstValidComponent - 1);
-		Operation->SetMinLOD(CurrentMinLOD);
-
-		if (Operation->bStreamMeshLODs)
-		{
-			Operation->FirstResidentLOD = FMath::Clamp(Operation->FirstResidentLOD, Operation->FirstLODAvailable, LODsInFirstValidComponent - 1);
-		}
-		else 
+		for (int32 InstanceComponentIndex = 0; InstanceComponentIndex < Operation->NumInstanceComponents; ++InstanceComponentIndex)
 		{
-			Operation->FirstResidentLOD = Operation->FirstLODAvailable;
-		}
-		
-		// Initialize RequestedLODs to zero if not set
-		TArray<uint16> RequestedLODs = Operation->GetRequestedLODs();
-		RequestedLODs.SetNumZeroed(Operation->NumObjectComponents);
+			const int32 ObjectComponentIndex = Operation->MutableInstance->GetComponentId(InstanceComponentIndex);
+			const FName& ComponentName = Operation->ComponentNames[ObjectComponentIndex];
 
-		for (int32 ComponentIndex = 0; ComponentIndex < Operation->NumObjectComponents; ++ComponentIndex)
-		{
+			uint8& MinLOD = MinLODs.FindOrAdd(ComponentName, 0);
+			uint8& RequestedLOD = RequestedLODs.FindOrAdd(ComponentName, 0);
+			uint8& NumLODsAvailable = Operation->NumLODsAvailable.FindOrAdd(ComponentName, 0);
+			uint8& FirstResidentLOD = Operation->FirstResidentLOD.FindOrAdd(ComponentName, 0);
+			uint8& FirstLODAvailable = Operation->FirstLODAvailable.FindOrAdd(ComponentName, 0);
+			
+			NumLODsAvailable = Operation->MutableInstance->GetLODCount(InstanceComponentIndex);
+			MinLOD = FMath::Clamp(MinLOD, FirstLODAvailable, NumLODsAvailable - 1);
+			
 			if (Operation->bStreamMeshLODs)
 			{
-				RequestedLODs[ComponentIndex] = CurrentMinLOD;
+				FirstResidentLOD = FMath::Clamp(FirstResidentLOD, FirstLODAvailable, NumLODsAvailable - 1);
+				RequestedLOD = MinLOD;
 			}
 			else
 			{
-				// Clamp value to the valid LOD range.
-				RequestedLODs[ComponentIndex] = FMath::Min(RequestedLODs[ComponentIndex], (uint16)(LODsInFirstValidComponent - 1));
+				FirstResidentLOD = FirstLODAvailable;
+				RequestedLOD = FMath::Min(RequestedLOD, static_cast<uint8>(FirstLODAvailable - 1));
 			}
 		}
-
+		
+		Operation->SetMinLOD(MinLODs);
 		Operation->SetRequestedLODs(RequestedLODs);
 	}
 	
@@ -1927,7 +1909,9 @@ namespace impl
 			TArray<TArray<FMorphTargetLODModel>>& MorphsData = ComponentMorphTargetsData.RealTimeMorphsLODData;
 			const int32 NumMorphs = MorphTargetNames.Num();
 
-			const int32 NumLODsInThisComponent = OperationData->GetLODsAvailable(ComponentMorphTargetsData.ObjectComponentIndex);
+			const FName& ComponentName = OperationData->ComponentNames[ComponentMorphTargetsData.ObjectComponentIndex.GetValue()];
+
+			const int32 NumLODsInThisComponent = OperationData->NumLODsAvailable[ComponentName];
 
 			MorphsData.SetNum(MorphTargetNames.Num());
 			for (int32 MorphIndex = 0; MorphIndex < NumMorphs; ++MorphIndex)
@@ -1940,7 +1924,7 @@ namespace impl
 
 			int32 NumInvalidVertexMorphNamesFound = 0;
 			
-			const int32 FirstGeneratedLOD = FMath::Max((int32)OperationData->GetRequestedLODs(ComponentMorphTargetsData.ObjectComponentIndex), OperationData->GetMinLOD());
+			const int32 FirstGeneratedLOD = FMath::Max(OperationData->GetRequestedLODs()[ComponentName], OperationData->GetMinLOD()[ComponentName]);
 			for (int32 LODIndex = FirstGeneratedLOD; LODIndex < NumLODsInThisComponent; ++LODIndex)
 			{
 				const FInstanceUpdateData::FLOD& LOD = UpdateData.LODs[Component.FirstLOD+LODIndex];
@@ -2242,8 +2226,10 @@ namespace impl
 			FInstanceUpdateData::FComponent& Component = OperationData->InstanceUpdateData.Components[InstanceComponentIndex];
 			for (int32 MutableLODIndex = 0; MutableLODIndex < Component.LODCount; ++MutableLODIndex)
 			{
+				const FName& ComponentName = ModelResources.ComponentNamesPerObjectComponent[Component.Id.GetValue()];
+				
 				// Skip LODs outside the range we want to generate
-				if (MutableLODIndex < OperationData->GetMinLOD())
+				if (MutableLODIndex < OperationData->GetMinLOD()[ComponentName])
 				{
 					continue;
 				}
@@ -2542,8 +2528,9 @@ namespace impl
 				}
 
 				FCustomizableObjectComponentIndex ObjectComponentIndex = Component.Id;
+				const FName ComponentName = OperationData->ComponentNames[ObjectComponentIndex.GetValue()];
 
-				Component.LODCount = OperationData->GetLODsAvailable(ObjectComponentIndex);
+				Component.LODCount = OperationData->NumLODsAvailable[ComponentName];
 				if (!Component.LODCount)
 				{
 					// It happens in degenerated cases with empty components.
@@ -2556,14 +2543,14 @@ namespace impl
 					const int32 UpdateDataLODIndex = OperationData->InstanceUpdateData.LODs.AddDefaulted();
 
 					// Skip LODs outside the range we want to generate
-					if (MutableLODIndex < OperationData->GetMinLOD())
+					if (MutableLODIndex < OperationData->GetMinLOD()[ComponentName])
 					{
 						continue;
 					}
 
 					FInstanceUpdateData::FLOD& LOD = OperationData->InstanceUpdateData.LODs.Last();
 
-					const bool bGenerateLOD = OperationData->GetRequestedLODs(ObjectComponentIndex) <= MutableLODIndex;
+					const bool bGenerateLOD = OperationData->GetRequestedLODs()[ComponentName] <= MutableLODIndex;
 
 					// Mesh
 					{
@@ -3266,16 +3253,15 @@ namespace impl
 	void Task_Mutable_GetMeshID(const TSharedRef<FUpdateContextPrivate>& Operation)
 	{
 		MUTABLE_CPUPROFILER_SCOPE(Task_Mutable_GetMeshID);
-
-
-		UCustomizableObjectSystem* System = UCustomizableObjectSystem::GetInstanceChecked(); // Save since UCustomizableObjectSystem::BeginDestroy always waits for all tasks to finish
-		UCustomizableObjectSystemPrivate* SystemPrivate = System->GetPrivate();
 		
 		CreateMutableInstance(Operation);
 		FixLODs(Operation);
 
 		Operation->InitMeshDescriptors(Operation->NumObjectComponents);
 
+		const TMap<FName, uint8>& MinLOD = Operation->GetMinLOD();
+		const TMap<FName, uint8>& NumLODsAvailable = Operation->NumLODsAvailable;
+		
 		for (int32 InstanceComponentIndex = 0; InstanceComponentIndex < Operation->NumInstanceComponents; ++InstanceComponentIndex)
 		{
 			FCustomizableObjectComponentIndex ObjectComponentIndex(Operation->MutableInstance->GetComponentId(InstanceComponentIndex));
@@ -3286,10 +3272,12 @@ namespace impl
 			}
 
 			MeshId->Init(MAX_uint64, MAX_MESH_LOD_COUNT);
+
+			const FName& ComponentName = Operation->ComponentNames[ObjectComponentIndex.GetValue()];
 			
-			for (int32 LODIndex = Operation->GetMinLOD(); LODIndex < Operation->GetLODsAvailable(ObjectComponentIndex); ++LODIndex)
+			for (int32 LODIndex = MinLOD[ComponentName]; LODIndex < NumLODsAvailable[ComponentName]; ++LODIndex)
 			{
-				const bool bGenerateLOD = Operation->GetRequestedLODs(ObjectComponentIndex) <= LODIndex;
+				const bool bGenerateLOD = Operation->GetRequestedLODs()[ComponentName] <= LODIndex;
 				if (bGenerateLOD)
 				{
 					(*MeshId)[LODIndex] = Operation->MutableInstance->GetMeshId(InstanceComponentIndex, LODIndex);
@@ -3540,8 +3528,11 @@ namespace impl
 
 		if (bRequestAllLODs)
 		{
-			TArray<uint16> RequestedLODs = Operation->GetRequestedLODs();
-			RequestedLODs.Init(0, Operation->NumObjectComponents);
+			TMap<FName, uint8> RequestedLODs = Operation->GetRequestedLODs();
+			for (const FName& ComponentName : Operation->Object->GetPrivate()->GetModelResources().ComponentNamesPerObjectComponent)
+			{
+				RequestedLODs.Add(ComponentName, 0);
+			}
 
 			Operation->SetRequestedLODs(RequestedLODs);
 		}
@@ -3836,7 +3827,7 @@ int32 UCustomizableObjectSystem::TickInternal(const bool bBlocking)
 						UCustomizableInstancePrivate* CustomizableInstancePrivate = LODUpdateCandidate.CustomizableObjectInstance->GetPrivate();
 
 						FDescriptorHash LODUpdateDescriptorHash = CustomizableInstancePrivate->CommittedDescriptorHash;
-						LODUpdateDescriptorHash.MinLOD = LODUpdateCandidate.MinLOD;
+						LODUpdateDescriptorHash.MinLODs = LODUpdateCandidate.MinLOD;
 						LODUpdateDescriptorHash.RequestedLODsPerComponent = LODUpdateCandidate.RequestedLODLevels;
 
 						if (CustomizableInstancePrivate->MinSquareDistFromComponentToPlayer < MaxSquareDistanceFound &&
@@ -4112,11 +4103,8 @@ int32 UCustomizableObjectSystem::GetNumInstances() const
 {
 	int32 NumInstances;
 	int32 NumBuiltInstances;
-	int32 NumInstancesLOD0;
-	int32 NumInstancesLOD1;
-	int32 NumInstancesLOD2;
 	int32 NumAllocatedSkeletalMeshes;
-	GetPrivate()->LogBenchmarkUtil.GetInstancesStats(NumInstances, NumBuiltInstances, NumInstancesLOD0, NumInstancesLOD1, NumInstancesLOD2, NumAllocatedSkeletalMeshes);
+	GetPrivate()->LogBenchmarkUtil.GetInstancesStats(NumInstances, NumBuiltInstances, NumAllocatedSkeletalMeshes);
 
 	return NumBuiltInstances;
 }
diff --git a/Engine/Plugins/Mutable/Source/CustomizableObject/Private/MuCO/DescriptorHash.cpp b/Engine/Plugins/Mutable/Source/CustomizableObject/Private/MuCO/DescriptorHash.cpp
--- a/Engine/Plugins/Mutable/Source/CustomizableObject/Private/MuCO/DescriptorHash.cpp
+++ b/Engine/Plugins/Mutable/Source/CustomizableObject/Private/MuCO/DescriptorHash.cpp
@@ -52,7 +52,7 @@ FDescriptorHash::FDescriptorHash(const FCustomizableObjectInstanceDescriptor& De
 	Hash = HashCombine(Hash, GetTypeHash(Descriptor.State));
 	Hash = HashCombine(Hash, GetTypeHash(Descriptor.GetBuildParameterRelevancy()));
 	
-	MinLOD = Descriptor.MinLOD;
+	MinLODs = Descriptor.MinLOD;
 
 	RequestedLODsPerComponent = Descriptor.RequestedLODLevels;
 }
@@ -61,7 +61,7 @@ FDescriptorHash::FDescriptorHash(const FCustomizableObjectInstanceDescriptor& De
 bool FDescriptorHash::operator==(const FDescriptorHash& Other) const
 {
 	return Hash == Other.Hash &&
-		MinLOD == Other.MinLOD &&
+		MinLODs == Other.MinLODs &&
 		RequestedLODsPerComponent == Other.RequestedLODsPerComponent;
 }
 
@@ -74,11 +74,25 @@ bool FDescriptorHash::operator!=(const FDescriptorHash& Other) const
 
 bool FDescriptorHash::IsSubset(const FDescriptorHash& Other) const
 {
-	if (Hash != Other.Hash || MinLOD < Other.MinLOD)
+	if (Hash != Other.Hash)
 	{
 		return false;
 	}
 
+	for (const TTuple<FName, uint8>& Pair : MinLODs)
+	{
+		const uint8* Result = Other.MinLODs.Find(Pair.Key);
+		if (!Result)
+		{
+			return false;
+		}
+
+		if (Pair.Value < *Result)
+		{
+			return false;
+		}
+	}
+
 	if (RequestedLODsPerComponent == Other.RequestedLODsPerComponent)
 	{
 		return true;
@@ -89,10 +103,15 @@ bool FDescriptorHash::IsSubset(const FDescriptorHash& Other) const
 		return false;
 	}
 	
-	for (int32 ComponentIndex = 0; ComponentIndex < RequestedLODsPerComponent.Num(); ++ComponentIndex)
+	for (const TTuple<FName, uint8>& Pair : RequestedLODsPerComponent)
 	{
-		// It is a subset if the Requested LOD is the same or greater.
-		if (RequestedLODsPerComponent[ComponentIndex] < Other.RequestedLODsPerComponent[ComponentIndex])
+		const uint8* PairOther = Other.RequestedLODsPerComponent.Find(Pair.Key);
+		if (!PairOther)
+		{
+			return false;
+		}
+
+		if (Pair.Value < *PairOther)
 		{
 			return false;
 		}
@@ -107,14 +126,19 @@ FString FDescriptorHash::ToString() const
 	TStringBuilder<150> Builder;
 
 	Builder.Appendf(TEXT("(Hash=%u,"), Hash);
-	Builder.Appendf(TEXT("MinLOD=%i,"), MinLOD);
-	Builder.Appendf(TEXT("RequestredLODLevels=["));
+	
+	Builder.Appendf(TEXT("MinLODs=["));
+	for (const TTuple<FName, uint8>& MinLOD : MinLODs)
+	{
+		Builder.Appendf(TEXT("(%s, %i),"), *MinLOD.Key.ToString(), MinLOD.Value);
+	}
+	Builder.Appendf(TEXT("])"));
 
-	for (const uint16 RequestedLODs : RequestedLODsPerComponent)
+	Builder.Appendf(TEXT("RequestedLODLevels=["));
+	for (const TTuple<FName, uint8>& RequestedLODs : RequestedLODsPerComponent)
 	{
-		Builder.Appendf(TEXT("%i,"), RequestedLODs);
+		Builder.Appendf(TEXT("(%s, %i),"), *RequestedLODs.Key.ToString(), RequestedLODs.Value);
 	}
-	
 	Builder.Appendf(TEXT("])"));
 
 	return Builder.ToString();
diff --git a/Engine/Plugins/Mutable/Source/CustomizableObject/Private/MuCO/LogBenchmarkUtil.cpp b/Engine/Plugins/Mutable/Source/CustomizableObject/Private/MuCO/LogBenchmarkUtil.cpp
--- a/Engine/Plugins/Mutable/Source/CustomizableObject/Private/MuCO/LogBenchmarkUtil.cpp
+++ b/Engine/Plugins/Mutable/Source/CustomizableObject/Private/MuCO/LogBenchmarkUtil.cpp
@@ -77,13 +77,10 @@ FLogBenchmarkUtil::~FLogBenchmarkUtil()
 }
 
 
-void FLogBenchmarkUtil::GetInstancesStats(int32& OutNumInstances, int32& OutNumBuiltInstances, int32& OutNumInstancesLOD0, int32& OutNumInstancesLOD1, int32& OutNumInstancesLOD2, int32& OutNumAllocatedSkeletalMeshes) const
+void FLogBenchmarkUtil::GetInstancesStats(int32& OutNumInstances, int32& OutNumBuiltInstances, int32& OutNumAllocatedSkeletalMeshes) const
 {
 	OutNumInstances = 0;
 	OutNumBuiltInstances = 0;
-	OutNumInstancesLOD0 = 0;
-	OutNumInstancesLOD1 = 0;
-	OutNumInstancesLOD2 = 0;
 	OutNumAllocatedSkeletalMeshes = 0;
 	
 	for (TObjectIterator<UCustomizableObjectInstance> Instance; Instance; ++Instance)
@@ -95,11 +92,7 @@ void FLogBenchmarkUtil::GetInstancesStats(int32& OutNumInstances, int32& OutNumB
 		}
 		
 		++OutNumInstances;
-
-		OutNumInstancesLOD0 += Instance->GetCurrentMinLOD() == 0;
-		OutNumInstancesLOD1 += Instance->GetCurrentMinLOD() == 1;
-		OutNumInstancesLOD2 += Instance->GetCurrentMinLOD() >= 2;
-
+		
 		if (Instance->GetPrivate()->SkeletalMeshStatus == ESkeletalMeshStatus::Success)
 		{
 			++OutNumBuiltInstances;
@@ -148,18 +141,12 @@ void FLogBenchmarkUtil::UpdateStats()
 	// Instances and Skeletal Mesh
 	int32 LocalNumInstances = 0;
 	int32 LocalNumBuiltInstances = 0;
-	int32 LocalNumInstancesLOD0 = 0;
-	int32 LocalNumInstancesLOD1 = 0;
-	int32 LocalNumInstancesLOD2 = 0;
 	int32 LocalNumAllocatedSkeletalMeshes = 0;
 
-	GetInstancesStats(LocalNumInstances, LocalNumBuiltInstances, LocalNumInstancesLOD0, LocalNumInstancesLOD1, LocalNumInstancesLOD2, LocalNumAllocatedSkeletalMeshes);
+	GetInstancesStats(LocalNumInstances, LocalNumBuiltInstances, LocalNumAllocatedSkeletalMeshes);
 
 	NumInstances = LocalNumInstances;
 	NumBuiltInstances = LocalNumBuiltInstances;
-	NumInstancesLOD0 = LocalNumInstancesLOD0;
-	NumInstancesLOD1 = LocalNumInstancesLOD1;
-	NumInstancesLOD2 = LocalNumInstancesLOD2;
 	NumAllocatedSkeletalMeshes = LocalNumAllocatedSkeletalMeshes;
 
 	NumPendingInstanceUpdates = System->GetNumPendingInstances();
diff --git a/Engine/Plugins/Mutable/Source/CustomizableObject/Public/MuCO/CustomizableObject.h b/Engine/Plugins/Mutable/Source/CustomizableObject/Public/MuCO/CustomizableObject.h
--- a/Engine/Plugins/Mutable/Source/CustomizableObject/Public/MuCO/CustomizableObject.h
+++ b/Engine/Plugins/Mutable/Source/CustomizableObject/Public/MuCO/CustomizableObject.h
@@ -594,8 +594,6 @@ public:
 	UFUNCTION(BlueprintCallable, Category = CustomizableObject)
 	bool IsCompiled() const;
 
-	int32 GetNumLODs() const;
-
 	bool IsChildObject() const;
 
 private:
diff --git a/Engine/Plugins/Mutable/Source/CustomizableObject/Public/MuCO/CustomizableObjectInstance.h b/Engine/Plugins/Mutable/Source/CustomizableObject/Public/MuCO/CustomizableObjectInstance.h
--- a/Engine/Plugins/Mutable/Source/CustomizableObject/Public/MuCO/CustomizableObjectInstance.h
+++ b/Engine/Plugins/Mutable/Source/CustomizableObject/Public/MuCO/CustomizableObjectInstance.h
@@ -649,22 +649,13 @@ public:
 	*/
 	UFUNCTION(BlueprintCallable, Category = CustomizableObjectInstance)
 	int32 GetNumComponents() const;
-
-	/** Return the min LOD that will be used in the next Instance update. */
-	int32 GetMinLODToLoad() const;
 	
-	/** Return the min LOD this Instance is using. */
-	int32 GetCurrentMinLOD() const;
-
 	/** Sets an array of LODs to generate per component. Mutable will generate those plus the currently generated LODs (if any).
 	  * Requires mutable.EnableOnlyGenerateRequestedLODs and CurrentInstanceLODManagement->IsOnlyGenerateRequestedLODLevelsEnabled() to be true.
-	  * @param InMinLOD - MinLOD to generate.
-	  * @param InMaxLOD - MaxLOD to generate - DEPRECATED.
-	  * @param InRequestedLODsPerComponent - Array with bitmasks of requested LODs per component with range from [0 .. CO->GetComponentCount()].
+	  * @param InMinLODs - MinLOD to generate per Component.
+	  * @param InRequestedLODs - Requested LODs per component to generate. Key is the Component name. Value represents the first LOD to generate. Zero means all LODs.
 	  * @param InOutRequestedUpdates - Map from Instance to Update data that stores a request for the Instance to be updated, which will be either processed or discarded by priority (to be rerequested the next tick) */
-	void SetRequestedLODs(int32 InMinLOD, int32 InMaxLOD, const TArray<uint16>& InRequestedLODsPerObjectComponent, FMutableInstanceUpdateMap& InOutRequestedUpdates);
-
-	const TArray<uint16>& GetRequestedLODsPerComponent() const;
+	void SetRequestedLODs(const TMap<FName, uint8>& InMinLODs, const TMap<FName, uint8>& InRequestedLODs, FMutableInstanceUpdateMap& InOutRequestedUpdates);
 
 #if WITH_EDITOR
 	/**
diff --git a/Engine/Plugins/Mutable/Source/CustomizableObject/Public/MuCO/CustomizableObjectSkeletalMesh.h b/Engine/Plugins/Mutable/Source/CustomizableObject/Public/MuCO/CustomizableObjectSkeletalMesh.h
--- a/Engine/Plugins/Mutable/Source/CustomizableObject/Public/MuCO/CustomizableObjectSkeletalMesh.h
+++ b/Engine/Plugins/Mutable/Source/CustomizableObject/Public/MuCO/CustomizableObjectSkeletalMesh.h
@@ -40,8 +40,7 @@ private:
 
 	static UCustomizableObjectSkeletalMesh* CreateSkeletalMesh(const UCustomizableObjectInstance& InInstance);
 
-	void InitMutableStreamingData(const TSharedRef<FUpdateContextPrivate>& InOperationData, UCustomizableObject& CustomizableObject,
-		const int32 FirstLOD, const int32 LODCount);
+	void InitMutableStreamingData(const TSharedRef<FUpdateContextPrivate>& InOperationData, const FName& ComponentName, const int32 FirstLOD, const int32 LODCount);
 
 	//~ Begin UStreamableRenderAsset Interface.
 	virtual bool StreamIn(int32 NewMipCount, bool bHighPrio) override;
diff --git a/Engine/Plugins/Mutable/Source/CustomizableObjectEditor/Private/MuCOE/CustomizableInstanceDetails.cpp b/Engine/Plugins/Mutable/Source/CustomizableObjectEditor/Private/MuCOE/CustomizableInstanceDetails.cpp
--- a/Engine/Plugins/Mutable/Source/CustomizableObjectEditor/Private/MuCOE/CustomizableInstanceDetails.cpp
+++ b/Engine/Plugins/Mutable/Source/CustomizableObjectEditor/Private/MuCOE/CustomizableInstanceDetails.cpp
@@ -2291,8 +2291,8 @@ FReply FCustomizableInstanceDetails::OnPasteAllParameters()
 
 	UCustomizableObject* CustomizableObject = Descriptor.GetCustomizableObject();
 	
-	const int32 MinLOD = Descriptor.GetMinLod();
-	const TArray<uint16> RequestedLODLevels = Descriptor.GetRequestedLODLevels();
+	const TMap<FName, uint8> MinLOD = Descriptor.MinLOD;
+	const TMap<FName, uint8> RequestedLODLevels = Descriptor.GetRequestedLODLevels();
 	
 	FScopedTransaction LocalTransaction(LOCTEXT("OnPasteAllParameters", "Paste All Parameters"));
 	CustomInstance->Modify();
@@ -2302,7 +2302,7 @@ FReply FCustomizableInstanceDetails::OnPasteAllParameters()
 		Descriptor.SetCustomizableObject(CustomizableObject);
 		
 		// Keep current LOD
-		Descriptor.SetMinLod(MinLOD);
+		Descriptor.MinLOD = MinLOD;
 		Descriptor.SetRequestedLODLevels(RequestedLODLevels);
 		
 		UpdateInstance();
diff --git a/Engine/Plugins/Mutable/Source/CustomizableObjectEditor/Private/MuCOE/CustomizableObjectBenchmarkingUtils.cpp b/Engine/Plugins/Mutable/Source/CustomizableObjectEditor/Private/MuCOE/CustomizableObjectBenchmarkingUtils.cpp
--- a/Engine/Plugins/Mutable/Source/CustomizableObjectEditor/Private/MuCOE/CustomizableObjectBenchmarkingUtils.cpp
+++ b/Engine/Plugins/Mutable/Source/CustomizableObjectEditor/Private/MuCOE/CustomizableObjectBenchmarkingUtils.cpp
@@ -32,8 +32,11 @@ bool CustomizableObjectBenchmarkingUtils::GenerateDeterministicSetOfInstances(TO
 		if (UCustomizableObjectInstance* GeneratedInstance = TargetCustomizableObject->CreateInstance())
 		{
 			// Force generation of all LODS
-			TArray<uint16> RequestedLodLevels{};
-			RequestedLodLevels.Init(0, TargetCustomizableObject->GetComponentCount());
+			TMap<FName, uint8> RequestedLodLevels;
+			for (const FName ComponentName : TargetCustomizableObject->GetPrivate()->GetModelResources().ComponentNamesPerObjectComponent)
+			{
+				RequestedLodLevels.Add(ComponentName, 0);
+			}
 			GeneratedInstance->GetPrivate()->GetDescriptor().SetRequestedLODLevels(RequestedLodLevels);
 				
 			// Randomize instance values
diff --git a/Engine/Plugins/Mutable/Source/CustomizableObjectEditor/Private/MuCOE/CustomizableObjectCompiler.cpp b/Engine/Plugins/Mutable/Source/CustomizableObjectEditor/Private/MuCOE/CustomizableObjectCompiler.cpp
--- a/Engine/Plugins/Mutable/Source/CustomizableObjectEditor/Private/MuCOE/CustomizableObjectCompiler.cpp
+++ b/Engine/Plugins/Mutable/Source/CustomizableObjectEditor/Private/MuCOE/CustomizableObjectCompiler.cpp
@@ -992,11 +992,23 @@ void FCustomizableObjectCompiler::CompileInternal(bool bAsync)
 		{
 			ModelResources.ReleaseVersion = VersionBridge->GetCurrentVersionAsString();
 		}
-		
-		ModelResources.NumLODs = GenerationContext.NumLODsInRoot;
-		ModelResources.NumLODsToStream = GenerationContext.bEnableLODStreaming ? GenerationContext.NumMaxLODsToStream : 0;
-		ModelResources.FirstLODAvailable = GenerationContext.FirstLODAvailable;
 
+		ModelResources.NumLODs = GenerationContext.NumLODs;
+
+		if (GenerationContext.bEnableLODStreaming)
+		{
+			ModelResources.NumLODsToStream = GenerationContext.NumMaxLODsToStream;
+		}
+		else
+		{
+			for (TTuple<FName, uint8>& Pair : ModelResources.NumLODsToStream)
+			{
+				Pair.Value = 0;
+			}
+		}
+
+		ModelResources.FirstLODAvailable = GenerationContext.FirstLODAvailable;
+		
 		ModelResources.ParticipatingObjects = MoveTemp(GenerationContext.ParticipatingObjects);
 		
 		if (CurrentOptions.bGatherReferences)
diff --git a/Engine/Plugins/Mutable/Source/CustomizableObjectEditor/Private/MuCOE/GenerateMutableSource/GenerateMutableSource.cpp b/Engine/Plugins/Mutable/Source/CustomizableObjectEditor/Private/MuCOE/GenerateMutableSource/GenerateMutableSource.cpp
--- a/Engine/Plugins/Mutable/Source/CustomizableObjectEditor/Private/MuCOE/GenerateMutableSource/GenerateMutableSource.cpp
+++ b/Engine/Plugins/Mutable/Source/CustomizableObjectEditor/Private/MuCOE/GenerateMutableSource/GenerateMutableSource.cpp
@@ -889,15 +889,12 @@ mu::Ptr<mu::NodeObject> GenerateMutableSource(const UEdGraphPin * Pin, FMutableG
 			const ITargetPlatform* TargetPlatform = GenerationContext.Options.TargetPlatform;
 
 			int32 NumExtraLODsToBuildAfterFirstLOD = 0;
-
-			const int32* AuxNumExtraLODs = State.NumExtraLODsToBuildPerPlatform.Find(TargetPlatform->PlatformName());
-
-			if (AuxNumExtraLODs)
+			if (const int32* AuxNumExtraLODs = State.NumExtraLODsToBuildPerPlatform.Find(TargetPlatform->PlatformName()))
 			{
 				NumExtraLODsToBuildAfterFirstLOD = *AuxNumExtraLODs;
 			}
 
-			ObjectNode->SetStateProperties(StateIndex, State.TextureCompressionStrategy, State.bBuildOnlyFirstLOD, GenerationContext.FirstLODAvailable, NumExtraLODsToBuildAfterFirstLOD);
+			ObjectNode->SetStateProperties(StateIndex, State.TextureCompressionStrategy, State.bBuildOnlyFirstLOD, NumExtraLODsToBuildAfterFirstLOD);
 
 			// UI Data
 			FMutableStateData StateUIData;
@@ -1301,13 +1298,15 @@ void PopulateReferenceSkeletalMeshesData(FMutableGraphGenerationContext& Generat
 {
 	const FString PlatformName = GenerationContext.Options.TargetPlatform->IniPlatformName();
 	
-	const uint32 LODCount = GenerationContext.NumLODsInRoot;
 	const uint32 ComponentCount = GenerationContext.ComponentInfos.Num();
 
 	GenerationContext.ReferenceSkeletalMeshesData.AddDefaulted(ComponentCount);
-	for(uint32 ComponentIndex = 0; ComponentIndex < ComponentCount; ++ComponentIndex)
+	for (uint32 ComponentIndex = 0; ComponentIndex < ComponentCount; ++ComponentIndex)
 	{
-		USkeletalMesh* RefSkeletalMesh = GenerationContext.ComponentInfos[ComponentIndex].RefSkeletalMesh;
+		const FMutableComponentInfo& ComponentInfo = GenerationContext.ComponentInfos[ComponentIndex];
+		const FName& ComponentName = ComponentInfo.ComponentName;
+		
+		USkeletalMesh* RefSkeletalMesh = ComponentInfo.RefSkeletalMesh;
 		check(RefSkeletalMesh);
 
 		FMutableRefSkeletalMeshData& Data = GenerationContext.ReferenceSkeletalMeshesData[ComponentIndex];
@@ -1320,11 +1319,12 @@ void PopulateReferenceSkeletalMeshesData(FMutableGraphGenerationContext& Generat
 		Data.SkeletalMeshLODSettings = RefSkeletalMesh->GetLODSettings();
 
 		// Gather LODData, this may include per LOD settings such as render data config or LODDataInfoArray
-		Data.LODData.AddDefaulted(GenerationContext.NumLODsInRoot);
+		Data.LODData.AddDefaulted(GenerationContext.NumLODs[ComponentName]);
 		
 		const uint32 RefSkeletalMeshLODCount = RefSkeletalMesh->GetLODNum();
 
-		for(uint32 LODIndex = 0; LODIndex < LODCount; ++LODIndex)
+		const uint32 LODCount = GenerationContext.NumLODs[ComponentName];
+		for (uint32 LODIndex = 0; LODIndex < LODCount; ++LODIndex)
 		{
 			FMutableRefLODData& LODData = Data.LODData[LODIndex];
 			if(LODIndex < RefSkeletalMeshLODCount)
@@ -1531,7 +1531,7 @@ int32 GetMaxTextureSize(const UTexture2D& ReferenceTexture, const UTextureLODSet
 }
 
 
-int32 GetTextureSizeInGame(const UTexture2D& Texture, const UTextureLODSettings& LODSettings, uint8 SurfaceLODBias)
+int32 GetTextureSizeInGame(const UTexture2D& Texture, const UTextureLODSettings& LODSettings)
 {
 	const int32 SourceSize = (uint32)FMath::Max3(Texture.Source.GetSizeX(), Texture.Source.GetSizeY(), (int64)1);
 	const int32 NumMipsSource = FMath::CeilLogTwo(SourceSize) + 1;
@@ -1548,7 +1548,7 @@ int32 GetTextureSizeInGame(const UTexture2D& Texture, const UTextureLODSettings&
 	// Calculate the LODBias specific for this texture 
 	const int32 TextureLODBias = LODSettings.CalculateLODBias(MaxTextureSizeAllowed, MaxTextureSizeAllowed, 0, Texture.LODGroup, Texture.LODBias, 0, Texture.MipGenSettings, Texture.IsCurrentlyVirtualTextured());
 
-	return MaxTextureSizeAllowed >> (TextureLODBias + SurfaceLODBias);
+	return MaxTextureSizeAllowed >> TextureLODBias;
 }
 
 
diff --git a/Engine/Plugins/Mutable/Source/CustomizableObjectEditor/Private/MuCOE/GenerateMutableSource/GenerateMutableSource.h b/Engine/Plugins/Mutable/Source/CustomizableObjectEditor/Private/MuCOE/GenerateMutableSource/GenerateMutableSource.h
--- a/Engine/Plugins/Mutable/Source/CustomizableObjectEditor/Private/MuCOE/GenerateMutableSource/GenerateMutableSource.h
+++ b/Engine/Plugins/Mutable/Source/CustomizableObjectEditor/Private/MuCOE/GenerateMutableSource/GenerateMutableSource.h
@@ -787,10 +787,11 @@ public:
 	/** If this is set, we are genreating materials for a "passthrough" component, with a fixed mesh. */
 	mu::Ptr<mu::NodeMesh> ComponentMeshOverride;
 
-	uint8 NumLODsInRoot = 0;
+	TMap<FName, uint8> NumLODs;
 
-	uint8 FirstLODAvailable = MAX_MESH_LOD_COUNT;
-	uint8 NumMaxLODsToStream = MAX_MESH_LOD_COUNT;
+	TMap<FName, uint8> FirstLODAvailable;
+	
+	TMap<FName, uint8> NumMaxLODsToStream;
 
 	bool bEnableLODStreaming = true;
 
@@ -978,7 +979,7 @@ uint32 ComputeLODBiasForTexture(const FMutableGraphGenerationContext& Generation
 int32 GetMaxTextureSize(const UTexture2D& ReferenceTexture, const UTextureLODSettings& LODSettings);
 
 // Max texture size of the texture with per platform MaxTextureSize and LODBias applied.
-int32 GetTextureSizeInGame(const UTexture2D& Texture, const UTextureLODSettings& LODSettings, uint8 SurfaceLODBias = 0);
+int32 GetTextureSizeInGame(const UTexture2D& Texture, const UTextureLODSettings& LODSettings);
 
 TSharedPtr<mu::FImage> GenerateImageConstant(UTexture*, FMutableGraphGenerationContext&, bool bIsReference);
 TSharedPtr<mu::FMesh> GenerateMeshConstant(USkeletalMesh*, FMutableGraphGenerationContext&, bool bIsReference);
diff --git a/Engine/Plugins/Mutable/Source/CustomizableObjectEditor/Private/MuCOE/GenerateMutableSource/GenerateMutableSourceComponent.cpp b/Engine/Plugins/Mutable/Source/CustomizableObjectEditor/Private/MuCOE/GenerateMutableSource/GenerateMutableSourceComponent.cpp
--- a/Engine/Plugins/Mutable/Source/CustomizableObjectEditor/Private/MuCOE/GenerateMutableSource/GenerateMutableSourceComponent.cpp
+++ b/Engine/Plugins/Mutable/Source/CustomizableObjectEditor/Private/MuCOE/GenerateMutableSource/GenerateMutableSourceComponent.cpp
@@ -44,7 +44,7 @@ void GenerateMutableSourceComponentMesh(FMutableGraphGenerationContext& Generati
 {
 	int32 FirstLOD = -1;
 
-	const int32 NumLODsInRoot = GenerationContext.NumLODsInRoot;
+	const int32 NumLODsInRoot = GenerationContext.NumLODs[GenerationContext.CurrentMeshComponent];
 	for (int32 CurrentLOD = 0; CurrentLOD < NumLODsInRoot; ++CurrentLOD)
 	{
 		GenerationContext.CurrentLOD = CurrentLOD;
@@ -68,7 +68,7 @@ void GenerateMutableSourceComponentMesh(FMutableGraphGenerationContext& Generati
 			continue;
 		}
 
-		if (GenerationContext.CurrentLOD < GenerationContext.FirstLODAvailable)
+		if (GenerationContext.CurrentLOD < GenerationContext.FirstLODAvailable[GenerationContext.CurrentMeshComponent])
 		{
 			continue;
 		} 
@@ -242,14 +242,6 @@ mu::Ptr<mu::NodeComponent> GenerateMutableSourceComponent(const UEdGraphPin* Pin
 		}
 		ComponentNode->Id = GenerationContext.ComponentNames.Add(TypedComponentPassthroughMesh->ComponentName);
 
-		// While we don't handle the LODs per component, make sure we have some LODs for the passthrough mesh to appear.
-		if (GenerationContext.NumLODsInRoot == 0)
-		{
-			GenerationContext.NumLODsInRoot = SkeletalMesh->GetLODNum();
-			GenerationContext.NumMaxLODsToStream = FMath::Clamp(GenerationContext.NumMaxLODsToStream, 0, GenerationContext.NumLODsInRoot - 1);
-			GenerationContext.FirstLODAvailable = FMath::Clamp(GenerationContext.FirstLODAvailable, 0, GenerationContext.NumLODsInRoot - 1);
-		}
-
 		// Create a LOD for each pass-through mesh LOD.
 		const FSkeletalMeshModel* Model = SkeletalMesh->GetImportedModel();
 		int32 SkeletalMeshLODCount = Model->LODModels.Num();
@@ -461,78 +453,78 @@ void FirstPass(UCustomizableObjectNodeComponentMesh& Node, FMutableGraphGenerati
 	}
 
 	// Fill the basic LOD Settings
-	const int32 NumLODs = Node.LODPins.Num();
+	uint8 NumLODs = Node.LODPins.Num();
 
-	if (!GenerationContext.NumLODsInRoot)
+	// NumLODsInRoot
+	int32 MaxRefMeshLODs = Node.ReferenceSkeletalMesh->GetLODNum();
+	if (MaxRefMeshLODs < NumLODs)
 	{
-		// NumLODsInRoot
-		int32 MaxRefMeshLODs = Node.ReferenceSkeletalMesh->GetLODNum();
-		if (MaxRefMeshLODs < NumLODs)
-		{
-			FString Msg = FString::Printf(TEXT("The object has %d LODs but the reference mesh only %d. Resulting objects will have %d LODs."),
-				NumLODs, MaxRefMeshLODs, MaxRefMeshLODs);
-			GenerationContext.Log(FText::FromString(Msg), &Node, EMessageSeverity::Warning);
-			GenerationContext.NumLODsInRoot = MaxRefMeshLODs;
-		}
-		else
-		{
-			GenerationContext.NumLODsInRoot = FMath::Max(GenerationContext.NumLODsInRoot, static_cast<uint8>(NumLODs));
-		}
+		FString Msg = FString::Printf(TEXT("The object has %d LODs but the reference mesh only %d. Resulting objects will have %d LODs."),
+			NumLODs, MaxRefMeshLODs, MaxRefMeshLODs);
+		GenerationContext.Log(FText::FromString(Msg), &Node, EMessageSeverity::Warning);
+		NumLODs = MaxRefMeshLODs;
+	}
 	
-		const FMutableLODSettings& LODSettings = GenerationContext.Object->LODSettings;
+	const FMutableLODSettings& LODSettings = GenerationContext.Object->LODSettings;
 
-		// Find the MinLOD available for the target platform
-		if (RefSkeletalMesh->IsMinLodQualityLevelEnable())
-		{
-			FSupportedQualityLevelArray SupportedQualityLevels = LODSettings.MinQualityLevelLOD.GetSupportedQualityLevels(*GenerationContext.Options.TargetPlatform->GetPlatformInfo().IniPlatformName.ToString());
+	uint8 FirstLODAvailable = MAX_MESH_LOD_COUNT;
+	
+	// Find the MinLOD available for the target platform
+	if (RefSkeletalMesh->IsMinLodQualityLevelEnable()) // Engine global setting. Current implementation has nothing to do with the USkeletalMesh.
+	{
+		FSupportedQualityLevelArray SupportedQualityLevels = LODSettings.MinQualityLevelLOD.GetSupportedQualityLevels(*GenerationContext.Options.TargetPlatform->GetPlatformInfo().IniPlatformName.ToString());
 			
-			int32 MinValue = GenerationContext.NumLODsInRoot - 1;
-			for (int32& QL : SupportedQualityLevels)
+		int32 MinValue = NumLODs - 1;
+		for (int32& QL : SupportedQualityLevels)
+		{
+			// check if have data for the supported quality level or set to default.
+			if (LODSettings.MinQualityLevelLOD.IsQualityLevelValid(QL))
 			{
-				// check if have data for the supported quality level or set to default.
-				if (LODSettings.MinQualityLevelLOD.IsQualityLevelValid(QL))
-				{
-					MinValue = FMath::Min(LODSettings.MinQualityLevelLOD.GetValueForQualityLevel(QL), MinValue);
-				}
-				else 
-				{
-					MinValue = LODSettings.MinQualityLevelLOD.GetDefault();
-					break;
-				}
+				MinValue = FMath::Min(LODSettings.MinQualityLevelLOD.GetValueForQualityLevel(QL), MinValue);
+			}
+			else 
+			{
+				MinValue = LODSettings.MinQualityLevelLOD.GetDefault();
+				break;
 			}
-
-			GenerationContext.FirstLODAvailable = FMath::Max(0, MinValue);
-		}
-		else
-		{
-			GenerationContext.FirstLODAvailable = LODSettings.MinLOD.GetValueForPlatform(*GenerationContext.Options.TargetPlatform->IniPlatformName());
 		}
 
-		GenerationContext.FirstLODAvailable = FMath::Clamp(GenerationContext.FirstLODAvailable, 0, GenerationContext.NumLODsInRoot - 1);
+		FirstLODAvailable = FMath::Max(0, MinValue);
+	}
+	else
+	{
+		FirstLODAvailable = LODSettings.MinLOD.GetValueForPlatform(*GenerationContext.Options.TargetPlatform->IniPlatformName());
+	}
+
+	FirstLODAvailable = FMath::Clamp(FirstLODAvailable, 0, NumLODs - 1);
 
-		// Find the streaming settings for the target platform
-		if (LODSettings.bOverrideLODStreamingSettings)
-		{
-			GenerationContext.bEnableLODStreaming = LODSettings.bEnableLODStreaming.GetValueForPlatform(*GenerationContext.Options.TargetPlatform->IniPlatformName());
-			GenerationContext.NumMaxLODsToStream = LODSettings.NumMaxStreamedLODs.GetValueForPlatform(*GenerationContext.Options.TargetPlatform->IniPlatformName());
-		}
-		else
+	uint8 NumMaxLODsToStream = MAX_MESH_LOD_COUNT;
+
+	// Find the streaming settings for the target platform
+	if (LODSettings.bOverrideLODStreamingSettings)
+	{
+		GenerationContext.bEnableLODStreaming = LODSettings.bEnableLODStreaming.GetValueForPlatform(*GenerationContext.Options.TargetPlatform->IniPlatformName());
+		NumMaxLODsToStream = LODSettings.NumMaxStreamedLODs.GetValueForPlatform(*GenerationContext.Options.TargetPlatform->IniPlatformName());
+	}
+	else
+	{
+		for (int32 MeshIndex = 0; MeshIndex < GenerationContext.ComponentInfos.Num(); ++MeshIndex)
 		{
-			for (int32 MeshIndex = 0; MeshIndex < GenerationContext.ComponentInfos.Num(); ++MeshIndex)
-			{
-				RefSkeletalMesh = GenerationContext.ComponentInfos[MeshIndex].RefSkeletalMesh;
-				check(RefSkeletalMesh);
+			RefSkeletalMesh = GenerationContext.ComponentInfos[MeshIndex].RefSkeletalMesh;
+			check(RefSkeletalMesh);
 
-				GenerationContext.bEnableLODStreaming = GenerationContext.bEnableLODStreaming &&
-					RefSkeletalMesh->GetEnableLODStreaming(GenerationContext.Options.TargetPlatform);
+			GenerationContext.bEnableLODStreaming = GenerationContext.bEnableLODStreaming &&
+				RefSkeletalMesh->GetEnableLODStreaming(GenerationContext.Options.TargetPlatform);
 
-				GenerationContext.NumMaxLODsToStream = FMath::Min(static_cast<int32>(GenerationContext.NumMaxLODsToStream),
-					RefSkeletalMesh->GetMaxNumStreamedLODs(GenerationContext.Options.TargetPlatform));
-			}
+			NumMaxLODsToStream = FMath::Min(NumMaxLODsToStream, static_cast<uint8>(RefSkeletalMesh->GetMaxNumStreamedLODs(GenerationContext.Options.TargetPlatform)));
 		}
-
-		GenerationContext.NumMaxLODsToStream = FMath::Clamp(GenerationContext.NumMaxLODsToStream, 0, GenerationContext.NumLODsInRoot - 1);
 	}
+
+	NumMaxLODsToStream = FMath::Clamp(NumMaxLODsToStream, 0, NumLODs - 1);
+
+	GenerationContext.NumLODs.Add(Node.ComponentName, NumLODs);
+	GenerationContext.FirstLODAvailable.Add(Node.ComponentName, FirstLODAvailable);
+	GenerationContext.NumMaxLODsToStream.Add(Node.ComponentName, NumMaxLODsToStream);
 	
 	// Add a new entry to the list of Component Infos
 	FMutableComponentInfo ComponentInfo(Node.ComponentName, RefSkeletalMesh);
diff --git a/Engine/Plugins/Mutable/Source/CustomizableObjectEditor/Private/MuCOE/GenerateMutableSource/GenerateMutableSourceGroupProjector.cpp b/Engine/Plugins/Mutable/Source/CustomizableObjectEditor/Private/MuCOE/GenerateMutableSource/GenerateMutableSourceGroupProjector.cpp
--- a/Engine/Plugins/Mutable/Source/CustomizableObjectEditor/Private/MuCOE/GenerateMutableSource/GenerateMutableSourceGroupProjector.cpp
+++ b/Engine/Plugins/Mutable/Source/CustomizableObjectEditor/Private/MuCOE/GenerateMutableSource/GenerateMutableSourceGroupProjector.cpp
@@ -510,7 +510,6 @@ TOptional<FGroupProjectorTempData> GenerateMutableGroupProjector(UCustomizableOb
 		SwitchNode->SetParameter(EnumParameterNode);
 		SwitchNode->SetOptionCount(ArrayOptionTexture.Num());
 
-		const uint32 AdditionalLODBias = GenerationContext.Options.bUseLODAsBias ? GenerationContext.FirstLODAvailable : 0;
 		for (int32 SelectorIndex = 0; SelectorIndex < ArrayOptionTexture.Num(); ++SelectorIndex)
 		{
 			if (const TObjectPtr<UTexture2D>& Texture = ArrayOptionTexture[SelectorIndex].OptionTexture)
@@ -520,7 +519,7 @@ TOptional<FGroupProjectorTempData> GenerateMutableGroupProjector(UCustomizableOb
 				mu::Ptr<mu::NodeImageConstant> ImageNode = new mu::NodeImageConstant();
 				ImageNode->SetValue(ImageConstant);
 
-				const uint32 MipsToSkip = ComputeLODBiasForTexture(GenerationContext, *Texture, ProjParamNode->ReferenceTexture) + AdditionalLODBias;
+				const uint32 MipsToSkip = ComputeLODBiasForTexture(GenerationContext, *Texture, ProjParamNode->ReferenceTexture);
 				SwitchNode->SetOption(SelectorIndex, ResizeTextureByNumMips(ImageNode, MipsToSkip));
 
 				// Calculate the number of mips to tag as high res for this image.
@@ -550,9 +549,6 @@ TOptional<FGroupProjectorTempData> GenerateMutableGroupProjector(UCustomizableOb
 			TextureSize = FMath::RoundUpToPowerOfTwo(TextureSize);
 		}
 
-		// Apply additional LODBias if necessary
-		GroupProjectorTempData.TextureSize = FMath::Max(TextureSize >> AdditionalLODBias, 1);
-
 		GroupProjectorTempData.NodeImagePtr = SwitchNode;
 	}
 
diff --git a/Engine/Plugins/Mutable/Source/CustomizableObjectEditor/Private/MuCOE/GenerateMutableSource/GenerateMutableSourceImage.cpp b/Engine/Plugins/Mutable/Source/CustomizableObjectEditor/Private/MuCOE/GenerateMutableSource/GenerateMutableSourceImage.cpp
--- a/Engine/Plugins/Mutable/Source/CustomizableObjectEditor/Private/MuCOE/GenerateMutableSource/GenerateMutableSourceImage.cpp
+++ b/Engine/Plugins/Mutable/Source/CustomizableObjectEditor/Private/MuCOE/GenerateMutableSource/GenerateMutableSourceImage.cpp
@@ -619,8 +619,7 @@ mu::NodeImagePtr GenerateMutableSourceImage(const UEdGraphPin* Pin, FMutableGrap
 			GenerationContext.AddParticipatingObject(*TypedNodeProject->ReferenceTexture);
 
 			const UTextureLODSettings& TextureLODSettings = GenerationContext.Options.TargetPlatform->GetTextureLODSettings();
-			const int32 FirstLODAvailable = GenerationContext.Options.bUseLODAsBias ? GenerationContext.FirstLODAvailable : 0;
-			MaxReferenceTextureSizeInGame = GetTextureSizeInGame(*TypedNodeProject->ReferenceTexture, TextureLODSettings, FirstLODAvailable);
+			MaxReferenceTextureSizeInGame = GetTextureSizeInGame(*TypedNodeProject->ReferenceTexture, TextureLODSettings);
 		}
 
 		FUintVector2 TextureSize(TypedNodeProject->TextureSizeX, TypedNodeProject->TextureSizeY);
@@ -853,8 +852,7 @@ mu::NodeImagePtr GenerateMutableSourceImage(const UEdGraphPin* Pin, FMutableGrap
 			GenerationContext.AddParticipatingObject(*TypedNodeTransform->ReferenceTexture);
 
 			const UTextureLODSettings& TextureLODSettings = GenerationContext.Options.TargetPlatform->GetTextureLODSettings();
-			const uint8 SurfaceLODBias = GenerationContext.Options.bUseLODAsBias ? GenerationContext.FirstLODAvailable : 0;
-			MaxReferenceTextureSizeInGame = GetTextureSizeInGame(*TypedNodeTransform->ReferenceTexture, TextureLODSettings, SurfaceLODBias);
+			MaxReferenceTextureSizeInGame = GetTextureSizeInGame(*TypedNodeTransform->ReferenceTexture, TextureLODSettings);
 		}
 
 		// Max TextureSize allowed
diff --git a/Engine/Plugins/Mutable/Source/CustomizableObjectEditor/Private/MuCOE/GenerateMutableSource/GenerateMutableSourceModifier.cpp b/Engine/Plugins/Mutable/Source/CustomizableObjectEditor/Private/MuCOE/GenerateMutableSource/GenerateMutableSourceModifier.cpp
--- a/Engine/Plugins/Mutable/Source/CustomizableObjectEditor/Private/MuCOE/GenerateMutableSource/GenerateMutableSourceModifier.cpp
+++ b/Engine/Plugins/Mutable/Source/CustomizableObjectEditor/Private/MuCOE/GenerateMutableSource/GenerateMutableSourceModifier.cpp
@@ -248,7 +248,7 @@ mu::Ptr<mu::NodeModifier> GenerateMutableSourceModifier(const UEdGraphPin * Pin,
 		SurfNode->RequiredTags = TypedNodeExt->RequiredTags;
 
 		// Is this enough? Should we try to narrow down with potential mesh sections modified by this?
-		int32 LODCount = GenerationContext.NumLODsInRoot;
+		int32 LODCount = GenerationContext.NumLODs[GenerationContext.CurrentMeshComponent];
 		SurfNode->LODs.SetNum(LODCount);
 
 		for (int32 LODIndex=0; LODIndex<LODCount; ++LODIndex)
@@ -326,7 +326,7 @@ mu::Ptr<mu::NodeModifier> GenerateMutableSourceModifier(const UEdGraphPin * Pin,
 
 					if (ProjectorInfo)
 					{
-						ensure(LOD > GenerationContext.FirstLODAvailable);
+						ensure(LOD > GenerationContext.FirstLODAvailable[GenerationContext.CurrentMeshComponent]);
 						check(ProjectorInfo->SurfNode->Images[ImageIndex].Image == ProjectorInfo->ImageNode);
 						ImageNode = ProjectorInfo->ImageNode;
 
@@ -400,7 +400,7 @@ mu::Ptr<mu::NodeModifier> GenerateMutableSourceModifier(const UEdGraphPin * Pin,
 		if (const UEdGraphPin* ConnectedPin = FollowInputPin(*TypedNodeRem->RemoveMeshPin()))
 		{
 			// Is this enough? Should we try to narrow down with potential mesh sections modified by this?
-			int32 LODCount = GenerationContext.NumLODsInRoot;
+			int32 LODCount = GenerationContext.NumLODs[GenerationContext.CurrentMeshComponent];
 			SurfNode->LODs.SetNum(LODCount);
 
 			SurfNode->FaceCullStrategy = TypedNodeRem->FaceCullStrategy;
@@ -458,7 +458,7 @@ mu::Ptr<mu::NodeModifier> GenerateMutableSourceModifier(const UEdGraphPin * Pin,
 		SurfNode->RequiredTags = TypedNodeEdit->RequiredTags;
 
 		// Is this enough? Should we try to narrow down with potential mesh sections modified by this?
-		int32 LODCount = GenerationContext.NumLODsInRoot;
+		int32 LODCount = GenerationContext.NumLODs[GenerationContext.CurrentMeshComponent];
 		SurfNode->LODs.SetNum(LODCount);
 
 		for (int32 LODIndex = 0; LODIndex < LODCount; ++LODIndex)
diff --git a/Engine/Plugins/Mutable/Source/CustomizableObjectEditor/Private/MuCOE/GenerateMutableSource/GenerateMutableSourceSurface.cpp b/Engine/Plugins/Mutable/Source/CustomizableObjectEditor/Private/MuCOE/GenerateMutableSource/GenerateMutableSourceSurface.cpp
--- a/Engine/Plugins/Mutable/Source/CustomizableObjectEditor/Private/MuCOE/GenerateMutableSource/GenerateMutableSourceSurface.cpp
+++ b/Engine/Plugins/Mutable/Source/CustomizableObjectEditor/Private/MuCOE/GenerateMutableSource/GenerateMutableSourceSurface.cpp
@@ -620,8 +620,7 @@ mu::Ptr<mu::NodeSurface> GenerateMutableSourceSurface(const UEdGraphPin * Pin, F
 						if (ReferenceTexture)
 						{
 							// Apply base LODBias. It will be propagated to most images.
-							const uint32 SurfaceLODBias = GenerationContext.Options.bUseLODAsBias ? GenerationContext.FirstLODAvailable : 0;
-							const uint32 BaseLODBias = ComputeLODBiasForTexture(GenerationContext, *ReferenceTexture) + SurfaceLODBias;
+							const uint32 BaseLODBias = ComputeLODBiasForTexture(GenerationContext, *ReferenceTexture);
 							mu::NodeImagePtr LastImage = ResizeTextureByNumMips(ImageNode, BaseLODBias);
 
 							if (ReferenceTexture->MipGenSettings != TextureMipGenSettings::TMGS_NoMipmaps)
@@ -837,7 +836,7 @@ mu::Ptr<mu::NodeSurface> GenerateMutableSourceSurface(const UEdGraphPin * Pin, F
 						if (bShareProjectionTexturesBetweenLODs && bIsGroupProjectorImage)
 						{
 							// Add to the GroupProjectorLODCache to potentially reuse this projection texture in higher LODs
-							ensure(LOD == GenerationContext.FirstLODAvailable);
+							ensure(LOD == GenerationContext.FirstLODAvailable[GenerationContext.CurrentMeshComponent]);
 							float* AlternateProjectionResFactor = TextureNameToProjectionResFactor.Find(ImageName);
 							GenerationContext.GroupProjectorLODCache.Add(MaterialImageId,
 								FGroupProjectorImageInfo(ImageNode, ImageName, ImageName, TypedNodeMat,
@@ -847,7 +846,7 @@ mu::Ptr<mu::NodeSurface> GenerateMutableSourceSurface(const UEdGraphPin * Pin, F
 				}
 				else
 				{
-					ensure(LOD > GenerationContext.FirstLODAvailable);
+					ensure(LOD > GenerationContext.FirstLODAvailable[GenerationContext.CurrentMeshComponent]);
 					check(ProjectorInfo->SurfNode->Images[ImageIndex].Image == ProjectorInfo->ImageNode);
 					SurfNode->Images[ImageIndex].Image = ProjectorInfo->ImageNode;
 					SurfNode->Images[ImageIndex].Name = ProjectorInfo->TextureName;
@@ -994,14 +993,14 @@ mu::Ptr<mu::NodeSurface> GenerateMutableSourceSurface(const UEdGraphPin * Pin, F
 
 						if (ProjectorInfo)
 						{
-							ensure(LOD == GenerationContext.FirstLODAvailable);
+							ensure(LOD == GenerationContext.FirstLODAvailable[GenerationContext.CurrentMeshComponent]);
 							ProjectorInfo->ImageResizeNode = NodeImageResize;
 							ProjectorInfo->bIsAlternateResolutionResized = true;
 						}
 					}
 					else
 					{
-						ensure(LOD > GenerationContext.FirstLODAvailable);
+						ensure(LOD > GenerationContext.FirstLODAvailable[GenerationContext.CurrentMeshComponent]);
 						check(ProjectorInfo->bIsAlternateResolutionResized);
 						SurfNode2->Images[ImageIndex].Image = ProjectorInfo->ImageResizeNode;
 					}
diff --git a/Engine/Plugins/Mutable/Source/MutableTools/Private/MuT/CodeOptimiser.h b/Engine/Plugins/Mutable/Source/MutableTools/Private/MuT/CodeOptimiser.h
--- a/Engine/Plugins/Mutable/Source/MutableTools/Private/MuT/CodeOptimiser.h
+++ b/Engine/Plugins/Mutable/Source/MutableTools/Private/MuT/CodeOptimiser.h
@@ -139,13 +139,13 @@ namespace mu
     {
     public:
 
-        LODCountReducerAST( Ptr<ASTOp>& root, int lodCount );
+        LODCountReducerAST( Ptr<ASTOp>& root, uint8 NumExtraLODsToBuildAfterFirstLOD );
 
     protected:
 
         Ptr<ASTOp> Visit( Ptr<ASTOp> at, bool& processChildren ) override;
 
-        int m_lodCount;
+        uint8 NumExtraLODs = 0;
     };
 
 
diff --git a/Engine/Plugins/Mutable/Source/MutableTools/Private/MuT/CodeOptimiser_Parameters.cpp b/Engine/Plugins/Mutable/Source/MutableTools/Private/MuT/CodeOptimiser_Parameters.cpp
--- a/Engine/Plugins/Mutable/Source/MutableTools/Private/MuT/CodeOptimiser_Parameters.cpp
+++ b/Engine/Plugins/Mutable/Source/MutableTools/Private/MuT/CodeOptimiser_Parameters.cpp
@@ -11,6 +11,7 @@
 #include "MuT/ASTOpMeshMorph.h"
 #include "MuT/ASTOpInstanceAdd.h"
 #include "MuT/ASTOpLayoutFromMesh.h"
+#include "MuT/ASTOpSwitch.h"
 #include "MuT/ASTOpParameter.h"
 #include "MuT/CodeOptimiser.h"
 #include "MuT/Compiler.h"
@@ -1975,9 +1976,47 @@ namespace mu
     //---------------------------------------------------------------------------------------------
     //---------------------------------------------------------------------------------------------
     //---------------------------------------------------------------------------------------------
-    LODCountReducerAST::LODCountReducerAST( Ptr<ASTOp>& root, int32 lodCount )
+
+	/** Recursively search for the first operation of the given type. */
+	class FFindMesh : public Visitor_TopDown_Unique_Const<uint8_t>
+    {
+    public:
+    	Ptr<ASTOpInstanceAdd> Result;
+    	
+    	FFindMesh(const ASTOpList& Roots)
+    	{
+    		Traverse(Roots, false);
+    	}
+
+    private:
+	    virtual bool Visit(const Ptr<ASTOp>& Node) override
+    	{
+	    	if (Result)
+	    	{
+	    		return false;
+	    	}
+
+		    const EOpType OpType = Node->GetOpType();
+
+	    	if (OpType == EOpType::IN_ADDMESH)
+	    	{
+	    		Result = static_cast<ASTOpInstanceAdd*>(Node.get());
+	    		return false;
+	    	}
+	    	
+	    	if (GetOpDataType(OpType) != DT_INSTANCE)
+	    	{
+	    		return false;
+	    	}
+	    	
+    		return true;
+    	}
+    };
+
+
+	LODCountReducerAST::LODCountReducerAST( Ptr<ASTOp>& root, uint8 NumExtraLODsToBuildAfterFirstLOD )
     {
-        m_lodCount = lodCount;
+        NumExtraLODs = NumExtraLODsToBuildAfterFirstLOD;
         Traverse( root );
     }
 
@@ -1990,10 +2029,23 @@ namespace mu
         {
 			ASTOpAddLOD* TypedOp = static_cast<ASTOpAddLOD*>(At.get());
 
-            if (TypedOp->lods.Num()>(size_t)m_lodCount)
+        	// Search for the first LOD that has a valid mesh.
+        	const int32 FirstLOD = TypedOp->lods.IndexOfByPredicate([](const ASTChild& Element)
+        	{
+        		TArray<mu::Ptr<ASTOp>> Roots;
+        		Roots.Add(Element.child());
+        		
+        		FFindMesh SearchMesh(Roots);
+        		
+       			return SearchMesh.Result.get() && SearchMesh.Result->value.child().get();
+        	});
+
+        	const int32 NumLODs = FirstLOD + NumExtraLODs + 1;
+        	
+            if (TypedOp->lods.Num() > NumLODs)
             {
                 Ptr<ASTOpAddLOD> newAt = mu::Clone<ASTOpAddLOD>(At);
-                while( newAt->lods.Num()>(size_t)m_lodCount )
+                while (newAt->lods.Num() > NumLODs)
                 {
                     newAt->lods.Pop();
                 }
@@ -2020,7 +2072,7 @@ namespace mu
 			FStateOptimizationOptions StateOptimization = States[s].nodeState.Optimisation;
 			if (StateOptimization.bOnlyFirstLOD)
 			{
-				LODCountReducerAST(States[s].root, StateOptimization.FirstLOD + 1 + StateOptimization.NumExtraLODsToBuildAfterFirstLOD);
+				LODCountReducerAST(States[s].root, StateOptimization.NumExtraLODsToBuildAfterFirstLOD);
 			}
 
 			// Apply texture compression strategy
diff --git a/Engine/Plugins/Mutable/Source/MutableTools/Private/MuT/NodeObjectNew.cpp b/Engine/Plugins/Mutable/Source/MutableTools/Private/MuT/NodeObjectNew.cpp
--- a/Engine/Plugins/Mutable/Source/MutableTools/Private/MuT/NodeObjectNew.cpp
+++ b/Engine/Plugins/Mutable/Source/MutableTools/Private/MuT/NodeObjectNew.cpp
@@ -98,7 +98,6 @@ namespace mu
     void NodeObjectNew::SetStateProperties( int32 StateIndex, 
 		ETextureCompressionStrategy TextureCompressionStrategy, 
 		bool bOnlyFirstLOD, 
-		uint8 FirstLOD, 
 		uint8 NumExtraLODsToBuildAfterFirstLOD )
     {
         check(StateIndex >=0 && StateIndex<GetStateCount() );
@@ -106,7 +105,6 @@ namespace mu
 		FStateOptimizationOptions& Data = States[StateIndex].Optimisation;
         Data.TextureCompressionStrategy = TextureCompressionStrategy;
 		Data.bOnlyFirstLOD = bOnlyFirstLOD;
-		Data.FirstLOD = FirstLOD;
 		Data.NumExtraLODsToBuildAfterFirstLOD = NumExtraLODsToBuildAfterFirstLOD;
     }
 
diff --git a/Engine/Plugins/Mutable/Source/MutableTools/Public/MuT/Compiler.h b/Engine/Plugins/Mutable/Source/MutableTools/Public/MuT/Compiler.h
--- a/Engine/Plugins/Mutable/Source/MutableTools/Public/MuT/Compiler.h
+++ b/Engine/Plugins/Mutable/Source/MutableTools/Public/MuT/Compiler.h
@@ -128,14 +128,12 @@ namespace mu
     //!
     struct FStateOptimizationOptions
     {
-		uint8 FirstLOD = 0;
 		uint8 NumExtraLODsToBuildAfterFirstLOD = 0;
 		bool bOnlyFirstLOD = false;
 		ETextureCompressionStrategy TextureCompressionStrategy = ETextureCompressionStrategy::None;
 
         void Serialise( FOutputArchive& arch ) const
         {
-			arch << FirstLOD;
 			arch << bOnlyFirstLOD;
             arch << TextureCompressionStrategy;
 			arch << NumExtraLODsToBuildAfterFirstLOD;
@@ -143,7 +141,6 @@ namespace mu
 
         void Unserialise( FInputArchive& arch )
         {
-			arch >> FirstLOD;
             arch >> bOnlyFirstLOD;
 			arch >> TextureCompressionStrategy;
 			arch >> NumExtraLODsToBuildAfterFirstLOD;
diff --git a/Engine/Plugins/Mutable/Source/MutableTools/Public/MuT/NodeObjectNew.h b/Engine/Plugins/Mutable/Source/MutableTools/Public/MuT/NodeObjectNew.h
--- a/Engine/Plugins/Mutable/Source/MutableTools/Public/MuT/NodeObjectNew.h
+++ b/Engine/Plugins/Mutable/Source/MutableTools/Public/MuT/NodeObjectNew.h
@@ -87,7 +87,6 @@ namespace mu
 		void SetStateProperties(int32 StateIndex,
 			ETextureCompressionStrategy TextureCompressionStrategy,
 			bool bOnlyFirstLOD,
-			uint8 FirstLOD,
 			uint8 NumExtraLODsToBuildAfterFirstLOD);
 
 		//! Connect a node that produces ExtensionData to be added to the final Instance, and provide a name to associate with the data
