[Mutable] Type mismatch in NumLODs of ModelResources.

#rnx

diff --git a/Engine/Plugins/Mutable/Source/CustomizableObject/Internal/MuCO/CustomizableObjectPrivate.h b/Engine/Plugins/Mutable/Source/CustomizableObject/Internal/MuCO/CustomizableObjectPrivate.h
--- a/Engine/Plugins/Mutable/Source/CustomizableObject/Internal/MuCO/CustomizableObjectPrivate.h
+++ b/Engine/Plugins/Mutable/Source/CustomizableObject/Internal/MuCO/CustomizableObjectPrivate.h
@@ -1198,7 +1198,7 @@ struct FModelResources
 	
 	/** Max number of LODs in the compiled Model. */
 	UPROPERTY()
-	TMap<FName, uint8> NumLODs;
+	TMap<FName, uint8> NumLODsAvailable;
 
 	/** Max number of LODs to stream. Mutable will always generate at least one LOD. */
 	UPROPERTY()
diff --git a/Engine/Plugins/Mutable/Source/CustomizableObject/Private/MuCO/CustomizableObject.cpp b/Engine/Plugins/Mutable/Source/CustomizableObject/Private/MuCO/CustomizableObject.cpp
--- a/Engine/Plugins/Mutable/Source/CustomizableObject/Private/MuCO/CustomizableObject.cpp
+++ b/Engine/Plugins/Mutable/Source/CustomizableObject/Private/MuCO/CustomizableObject.cpp
@@ -775,7 +775,7 @@ void FModelResources::Serialize(FObjectAndNameAsStringProxyArchive& MemoryWriter
 	MemoryWriter << ClothingAssetsData;
 	MemoryWriter << ClothSharedConfigsData;
 
-	MemoryWriter << NumLODs;
+	MemoryWriter << NumLODsAvailable;
 	MemoryWriter << NumLODsToStream;
 	MemoryWriter << FirstLODAvailable;
 
@@ -937,7 +937,7 @@ bool FModelResources::Unserialize(FObjectAndNameAsStringProxyArchive& MemoryRead
 	MemoryReader << ClothingAssetsData;
 	MemoryReader << ClothSharedConfigsData;
 
-	MemoryReader << NumLODs;
+	MemoryReader << NumLODsAvailable;
 	MemoryReader << NumLODsToStream;
 	MemoryReader << FirstLODAvailable;
 
diff --git a/Engine/Plugins/Mutable/Source/CustomizableObject/Private/MuCO/CustomizableObjectInstance.cpp b/Engine/Plugins/Mutable/Source/CustomizableObject/Private/MuCO/CustomizableObjectInstance.cpp
--- a/Engine/Plugins/Mutable/Source/CustomizableObject/Private/MuCO/CustomizableObjectInstance.cpp
+++ b/Engine/Plugins/Mutable/Source/CustomizableObject/Private/MuCO/CustomizableObjectInstance.cpp
@@ -7178,7 +7178,7 @@ void UCustomizableObjectInstance::SetRequestedLODs(const TMap<FName, uint8>& InM
 		}
 
 		const uint8 MinLODIdx = CustomizableObject->GetPrivate()->GetMinLODIndex(ComponentName);
-		const int32 MaxLODIdx = ModelResources.NumLODs[ComponentName] - 1;
+		const int32 MaxLODIdx = ModelResources.NumLODsAvailable[ComponentName] - 1;
 		InMinLOD = FMath::Clamp(InMinLOD, MinLODIdx, MaxLODIdx);
 
 		uint8 DescriptorMinLOD = 0;
diff --git a/Engine/Plugins/Mutable/Source/CustomizableObjectEditor/Private/MuCOE/CustomizableObjectCompiler.cpp b/Engine/Plugins/Mutable/Source/CustomizableObjectEditor/Private/MuCOE/CustomizableObjectCompiler.cpp
--- a/Engine/Plugins/Mutable/Source/CustomizableObjectEditor/Private/MuCOE/CustomizableObjectCompiler.cpp
+++ b/Engine/Plugins/Mutable/Source/CustomizableObjectEditor/Private/MuCOE/CustomizableObjectCompiler.cpp
@@ -993,7 +993,7 @@ void FCustomizableObjectCompiler::CompileInternal(bool bAsync)
 			ModelResources.ReleaseVersion = VersionBridge->GetCurrentVersionAsString();
 		}
 
-		ModelResources.NumLODs = GenerationContext.NumLODs;
+		ModelResources.NumLODsAvailable = GenerationContext.NumLODs;
 
 		if (GenerationContext.bEnableLODStreaming)
 		{
