[mutable] Ensure that the optimizer can query parameter operations for "is constant color" without an assert.
Also apply the coding standards a little bit more.
#rnx

diff --git a/Engine/Plugins/Mutable/Source/MutableTools/Private/MuT/AST.cpp b/Engine/Plugins/Mutable/Source/MutableTools/Private/MuT/AST.cpp
--- a/Engine/Plugins/Mutable/Source/MutableTools/Private/MuT/AST.cpp
+++ b/Engine/Plugins/Mutable/Source/MutableTools/Private/MuT/AST.cpp
@@ -869,14 +869,16 @@ bool ASTOp::IsImagePlainConstant(FVector4f&) const
 {
 	// Some image operations don't have this implemented and hit here.
 	// \TODO: Optimize for those cases.
-    //check(false);
+	//ensure(false);
     return false;
 }
 
 
 bool ASTOp::IsColourConstant(FVector4f&) const
 {
-    check(false);
+	// Some operations don't have this implemented and hit here.
+	// \TODO: It is an opportunity to implement it and help with CO optimization in some cases.
+    //ensure(false);
     return false;
 }
 
diff --git a/Engine/Plugins/Mutable/Source/MutableTools/Private/MuT/ASTOpParameter.cpp b/Engine/Plugins/Mutable/Source/MutableTools/Private/MuT/ASTOpParameter.cpp
--- a/Engine/Plugins/Mutable/Source/MutableTools/Private/MuT/ASTOpParameter.cpp
+++ b/Engine/Plugins/Mutable/Source/MutableTools/Private/MuT/ASTOpParameter.cpp
@@ -23,44 +23,50 @@ namespace mu
 
 	void ASTOpParameter::ForEachChild(const TFunctionRef<void(ASTChild&)> f)
 	{
-		for (auto& c : ranges)
+		for (FRangeData& c : Ranges)
 		{
 			f(c.rangeSize);
 		}
 	}
 
 
-	bool ASTOpParameter::IsEqual(const ASTOp& otherUntyped) const
+	bool ASTOpParameter::IsEqual(const ASTOp& OtherUntyped) const
 	{
-		if (otherUntyped.GetOpType() == GetOpType())
+		if (OtherUntyped.GetOpType() == GetOpType())
 		{
-			const ASTOpParameter* other = static_cast<const ASTOpParameter*>(&otherUntyped);
-			return type == other->type &&
-				parameter == other->parameter &&
-				ranges == other->ranges;
+			const ASTOpParameter* Other = static_cast<const ASTOpParameter*>(&OtherUntyped);
+			return Type == Other->Type &&
+				Parameter == Other->Parameter &&
+				Ranges == Other->Ranges;
 		}
 		return false;
 	}
 
 
-	mu::Ptr<ASTOp> ASTOpParameter::Clone(MapChildFuncRef mapChild) const
+	mu::Ptr<ASTOp> ASTOpParameter::Clone(MapChildFuncRef MapChild) const
 	{
 		Ptr<ASTOpParameter> n = new ASTOpParameter();
-		n->type = type;
-		n->parameter = parameter;
-		for (const auto& c : ranges)
+		n->Type = Type;
+		n->Parameter = Parameter;
+		for (const auto& c : Ranges)
 		{
-			n->ranges.Emplace(n.get(), mapChild(c.rangeSize.child()), c.rangeName, c.rangeUID);
+			n->Ranges.Emplace(n.get(), MapChild(c.rangeSize.child()), c.rangeName, c.rangeUID);
 		}
 		return n;
 	}
 
 
+	EOpType ASTOpParameter::GetOpType() const 
+	{ 
+		return Type; 
+	}
+
+
 	uint64 ASTOpParameter::Hash() const
 	{
-		uint64 res = std::hash<uint64>()(uint64(type));
-		hash_combine(res, parameter.Type);
-		hash_combine(res, parameter.Name.Len());
+		uint64 res = std::hash<uint64>()(uint64(Type));
+		hash_combine(res, Parameter.Type);
+		hash_combine(res, Parameter.Name.Len());
 		return res;
 	}
 
@@ -79,7 +85,7 @@ namespace mu
 			OP::ParameterArgs args;
 			FMemory::Memzero(&args, sizeof(args));
 
-			LinkedParameterIndex = Program.Parameters.Find(parameter);
+			LinkedParameterIndex = Program.Parameters.Find(Parameter);
 
 			// If this fails, it means an ASTOpParameter was created at code generation time, but not registered into the
 			// parameters map in the CodeGenerator_FirstPass.
@@ -87,7 +93,7 @@ namespace mu
 
 			args.variable = (OP::ADDRESS)LinkedParameterIndex;
 
-			for (const auto& d : ranges)
+			for (const auto& d : Ranges)
 			{
 				OP::ADDRESS sizeAt = 0;
 				uint16 rangeId = 0;
@@ -98,7 +104,7 @@ namespace mu
 			linkedAddress = (OP::ADDRESS)Program.OpAddress.Num();
 
 			Program.OpAddress.Add((uint32)Program.ByteCode.Num());
-			AppendCode(Program.ByteCode, type);
+			AppendCode(Program.ByteCode, Type);
 			AppendCode(Program.ByteCode, args);
 		}
 	}
@@ -132,18 +138,23 @@ namespace mu
 	}
 
 
-	ASTOp::FBoolEvalResult ASTOpParameter::EvaluateBool(ASTOpList& /*facts*/,
-		FEvaluateBoolCache*) const
+	ASTOp::FBoolEvalResult ASTOpParameter::EvaluateBool(ASTOpList& /*facts*/, FEvaluateBoolCache*) const
 	{
-		check(type == EOpType::BO_PARAMETER);
+		check(Type == EOpType::BO_PARAMETER);
 		return BET_UNKNOWN;
 	}
 
 
 	FImageDesc ASTOpParameter::GetImageDesc(bool, FGetImageDescContext*) const
 	{
-		check(type == EOpType::IM_PARAMETER);
+		check(Type == EOpType::IM_PARAMETER);
 		return FImageDesc();
 	}
 
+
+	bool ASTOpParameter::IsColourConstant(FVector4f&) const
+	{
+		return false;
+	}
+
 }
diff --git a/Engine/Plugins/Mutable/Source/MutableTools/Private/MuT/ASTOpParameter.h b/Engine/Plugins/Mutable/Source/MutableTools/Private/MuT/ASTOpParameter.h
--- a/Engine/Plugins/Mutable/Source/MutableTools/Private/MuT/ASTOpParameter.h
+++ b/Engine/Plugins/Mutable/Source/MutableTools/Private/MuT/ASTOpParameter.h
@@ -15,21 +15,19 @@ namespace mu
 	struct FProgram;
 
 
-	//---------------------------------------------------------------------------------------------
-	//! Parameter operation.
-	//---------------------------------------------------------------------------------------------
+	/** Parameter operation. */
 	class ASTOpParameter final : public ASTOp
 	{
 	public:
 
 		//! Type of parameter
-		EOpType type;
+		EOpType Type = EOpType::NONE;
 
 		//!
-		FParameterDesc parameter;
+		FParameterDesc Parameter;
 
 		//** Ranges adding dimensions to this parameter. */
-		TArray<FRangeData> ranges;
+		TArray<FRangeData> Ranges;
 
 		/** Index of the parameter in the program parameter list. Generated ar link time. */
 		int32 LinkedParameterIndex = -1;
@@ -38,16 +36,17 @@ namespace mu
 
 		~ASTOpParameter() override;
 
-		EOpType GetOpType() const override { return type; }
-		uint64 Hash() const override;
-		void ForEachChild(const TFunctionRef<void(ASTChild&)>) override;
-		bool IsEqual(const ASTOp& otherUntyped) const override;
-		Ptr<ASTOp> Clone(MapChildFuncRef mapChild) const override;
-		void Assert() override;
-		void Link(FProgram& program, FLinkerOptions*) override;
-		int EvaluateInt(ASTOpList& facts, bool& unknown) const override;
-		FBoolEvalResult EvaluateBool(ASTOpList& /*facts*/, FEvaluateBoolCache* = nullptr) const override;
-		FImageDesc GetImageDesc(bool, FGetImageDescContext*) const override;
+		virtual EOpType GetOpType() const override;
+		virtual uint64 Hash() const override;
+		virtual void ForEachChild(const TFunctionRef<void(ASTChild&)>) override;
+		virtual bool IsEqual(const ASTOp&) const override;
+		virtual Ptr<ASTOp> Clone(MapChildFuncRef ) const override;
+		virtual void Assert() override;
+		virtual void Link(FProgram&, FLinkerOptions*) override;
+		virtual int EvaluateInt(ASTOpList& Facts, bool& bOutUnknown) const override;
+		virtual FBoolEvalResult EvaluateBool(ASTOpList& Facts, FEvaluateBoolCache* = nullptr) const override;
+		virtual FImageDesc GetImageDesc(bool, FGetImageDescContext*) const override;
+		virtual bool IsColourConstant(FVector4f& OutColour) const override;
 
 	};
 
diff --git a/Engine/Plugins/Mutable/Source/MutableTools/Private/MuT/ASTOpSwitch.cpp b/Engine/Plugins/Mutable/Source/MutableTools/Private/MuT/ASTOpSwitch.cpp
--- a/Engine/Plugins/Mutable/Source/MutableTools/Private/MuT/ASTOpSwitch.cpp
+++ b/Engine/Plugins/Mutable/Source/MutableTools/Private/MuT/ASTOpSwitch.cpp
@@ -420,7 +420,7 @@ namespace mu
 		{
 			Ptr<ASTOp> Branch = def.child();
 
-			auto typedCondition = static_cast<const ASTOpFixed*>(variable.child().get());
+			const ASTOpFixed* typedCondition = static_cast<const ASTOpFixed*>(variable.child().get());
 			for (int32 o = 0; o < cases.Num(); ++o)
 			{
 				if (cases[o].branch &&
@@ -441,12 +441,12 @@ namespace mu
 			// If all the branches for the possible values are the same op remove the instruction
 			const ASTOpParameter* ParamOp = static_cast<const ASTOpParameter*>(variable.child().get());
 			check(ParamOp);
-			check(!ParamOp->parameter.PossibleValues.IsEmpty());
+			check(!ParamOp->Parameter.PossibleValues.IsEmpty());
 
 			bool bFirstValue = true;
 			bool bAllSame = true;
 			Ptr<ASTOp> SameBranch = nullptr;
-			for (const FParameterDesc::FIntValueDesc& Value : ParamOp->parameter.PossibleValues)
+			for (const FParameterDesc::FIntValueDesc& Value : ParamOp->Parameter.PossibleValues)
 			{
 				// Look for the switch branch it would take
 				Ptr<ASTOp> Branch = def.child();
diff --git a/Engine/Plugins/Mutable/Source/MutableTools/Private/MuT/CodeGenerator.cpp b/Engine/Plugins/Mutable/Source/MutableTools/Private/MuT/CodeGenerator.cpp
--- a/Engine/Plugins/Mutable/Source/MutableTools/Private/MuT/CodeGenerator.cpp
+++ b/Engine/Plugins/Mutable/Source/MutableTools/Private/MuT/CodeGenerator.cpp
@@ -377,8 +377,8 @@ namespace mu
         }
 
 		Ptr<ASTOpParameter> op = new ASTOpParameter();
-		op->type = EOpType::NU_PARAMETER;
-		op->parameter = param;
+		op->Type = EOpType::NU_PARAMETER;
+		op->Parameter = param;
 
 		FirstPass.ParameterNodes.Add(InNode, op);
 
diff --git a/Engine/Plugins/Mutable/Source/MutableTools/Private/MuT/CodeGenerator_Bool.cpp b/Engine/Plugins/Mutable/Source/MutableTools/Private/MuT/CodeGenerator_Bool.cpp
--- a/Engine/Plugins/Mutable/Source/MutableTools/Private/MuT/CodeGenerator_Bool.cpp
+++ b/Engine/Plugins/Mutable/Source/MutableTools/Private/MuT/CodeGenerator_Bool.cpp
@@ -98,15 +98,15 @@ namespace mu
 			param.DefaultValue.Set<FParamBoolType>(node.DefaultValue);
 
 			op = new ASTOpParameter();
-			op->type = EOpType::BO_PARAMETER;
-			op->parameter = param;
+			op->Type = EOpType::BO_PARAMETER;
+			op->Parameter = param;
 
 			// Generate the code for the ranges
 			for (int32 a = 0; a < node.Ranges.Num(); ++a)
 			{
 				FRangeGenerationResult rangeResult;
 				GenerateRange(rangeResult, Options, node.Ranges[a]);
-				op->ranges.Emplace(op.get(), rangeResult.sizeOp, rangeResult.rangeName, rangeResult.rangeUID);
+				op->Ranges.Emplace(op.get(), rangeResult.sizeOp, rangeResult.rangeName, rangeResult.rangeUID);
 			}
 
 			FirstPass.ParameterNodes.Add(node.m_pNode,op);
diff --git a/Engine/Plugins/Mutable/Source/MutableTools/Private/MuT/CodeGenerator_Color.cpp b/Engine/Plugins/Mutable/Source/MutableTools/Private/MuT/CodeGenerator_Color.cpp
--- a/Engine/Plugins/Mutable/Source/MutableTools/Private/MuT/CodeGenerator_Color.cpp
+++ b/Engine/Plugins/Mutable/Source/MutableTools/Private/MuT/CodeGenerator_Color.cpp
@@ -158,15 +158,15 @@ namespace mu
 			param.DefaultValue.Set<FParamColorType>(Value);
 			
 			op = new ASTOpParameter();
-			op->type = EOpType::CO_PARAMETER;
-			op->parameter = param;
+			op->Type = EOpType::CO_PARAMETER;
+			op->Parameter = param;
 
 			// Generate the code for the ranges
 			for (int32 a = 0; a < node.Ranges.Num(); ++a)
 			{
 				FRangeGenerationResult rangeResult;
 				GenerateRange(rangeResult, Options, node.Ranges[a]);
-				op->ranges.Emplace(op.get(), rangeResult.sizeOp, rangeResult.rangeName, rangeResult.rangeUID);
+				op->Ranges.Emplace(op.get(), rangeResult.sizeOp, rangeResult.rangeName, rangeResult.rangeUID);
 			}
 
 			FirstPass.ParameterNodes.Add(Typed, op);
diff --git a/Engine/Plugins/Mutable/Source/MutableTools/Private/MuT/CodeGenerator_FirstPass.cpp b/Engine/Plugins/Mutable/Source/MutableTools/Private/MuT/CodeGenerator_FirstPass.cpp
--- a/Engine/Plugins/Mutable/Source/MutableTools/Private/MuT/CodeGenerator_FirstPass.cpp
+++ b/Engine/Plugins/Mutable/Source/MutableTools/Private/MuT/CodeGenerator_FirstPass.cpp
@@ -619,20 +619,20 @@ namespace mu
 			Private->Type==NodeObjectGroup::CS_ONE_OR_NONE )
         {
             Ptr<ASTOpParameter> op = new ASTOpParameter();
-            op->type = EOpType::NU_PARAMETER;
+            op->Type = EOpType::NU_PARAMETER;
 
-            op->parameter.Name = Private->Name;
+            op->Parameter.Name = Private->Name;
 			const TCHAR* CStr = ToCStr(Private->Uid);
-			op->parameter.UID.ImportTextItem(CStr, 0, nullptr, nullptr);
-            op->parameter.Type = EParameterType::Int;
-            op->parameter.DefaultValue.Set<FParamIntType>(Private->DefaultValue);
+			op->Parameter.UID.ImportTextItem(CStr, 0, nullptr, nullptr);
+            op->Parameter.Type = EParameterType::Int;
+            op->Parameter.DefaultValue.Set<FParamIntType>(Private->DefaultValue);
 
             if (Private->Type==NodeObjectGroup::CS_ONE_OR_NONE )
             {
                 FParameterDesc::FIntValueDesc nullValue;
                 nullValue.Value = -1;
                 nullValue.Name = "None";
-                op->parameter.PossibleValues.Add( nullValue );
+                op->Parameter.PossibleValues.Add( nullValue );
             }
 
 			ParameterNodes.Add(InNode, op);
@@ -644,7 +644,7 @@ namespace mu
         // Parse the child objects
 		for ( int32 t=0; t< Private->m_children.Num(); ++t )
         {
-            if ( mu::Ptr<const NodeObject> pChildNode = Private->m_children[t] )
+            if ( mu::Ptr<const NodeObject> ChildNode = Private->m_children[t] )
             {
                 // Overwrite the implicit condition
                 Ptr<ASTOp> paramOp = 0;
@@ -652,7 +652,7 @@ namespace mu
                 {
                     case NodeObjectGroup::CS_TOGGLE_EACH:
                     {           
-						if (pChildNode->GetType() == NodeObjectGroup::GetStaticType())
+						if (ChildNode->GetType() == NodeObjectGroup::GetStaticType())
 						{
 							FString Msg = FString::Printf(TEXT("The Group Node [%s] has type Toggle and its direct child is a Group node, which is not allowed. Change the type or add a Child Object node in between them."),
 								*Private->Name);
@@ -662,15 +662,15 @@ namespace mu
 						{
 							// Create a new boolean parameter
 							Ptr<ASTOpParameter> op = new ASTOpParameter();
-							op->type = EOpType::BO_PARAMETER;
+							op->Type = EOpType::BO_PARAMETER;
 
-							op->parameter.Name = pChildNode->GetName();
-							const TCHAR* CStr = ToCStr(pChildNode->GetUid());
-							op->parameter.UID.ImportTextItem(CStr, 0, nullptr, nullptr);
-							op->parameter.Type = EParameterType::Bool;
-							op->parameter.DefaultValue.Set<FParamBoolType>(false);
+							op->Parameter.Name = ChildNode->GetName();
+							bool bParseOk = FGuid::Parse(ChildNode->GetUid(), op->Parameter.UID);
+							check(bParseOk);
+							op->Parameter.Type = EParameterType::Bool;
+							op->Parameter.DefaultValue.Set<FParamBoolType>(false);
 
-							ParameterNodes.Add(pChildNode, op);
+							ParameterNodes.Add(ChildNode, op);
 
 							paramOp = op;
 						}
@@ -694,8 +694,8 @@ namespace mu
                         // Add the option to the enumeration parameter
                         FParameterDesc::FIntValueDesc value;
                         value.Value = (int16)t;
-                        value.Name = pChildNode->GetName();
-                        enumOp->parameter.PossibleValues.Add( value );
+                        value.Name = ChildNode->GetName();
+                        enumOp->Parameter.PossibleValues.Add( value );
 
                         check(enumOp);
 
@@ -727,7 +727,7 @@ namespace mu
                 data.ObjectCondition = paramOp;
                 CurrentCondition.Add( data );
 
-				Generate_Generic(pChildNode.get());
+				Generate_Generic(ChildNode.get());
 
                 CurrentCondition.Pop();
             }
diff --git a/Engine/Plugins/Mutable/Source/MutableTools/Private/MuT/CodeGenerator_Image.cpp b/Engine/Plugins/Mutable/Source/MutableTools/Private/MuT/CodeGenerator_Image.cpp
--- a/Engine/Plugins/Mutable/Source/MutableTools/Private/MuT/CodeGenerator_Image.cpp
+++ b/Engine/Plugins/Mutable/Source/MutableTools/Private/MuT/CodeGenerator_Image.cpp
@@ -281,20 +281,20 @@ namespace mu
         if ( !it )
         {
             op = new ASTOpParameter();
-            op->type = EOpType::IM_PARAMETER;
+            op->Type = EOpType::IM_PARAMETER;
 
-			op->parameter.Name = node.Name;
-        	const TCHAR* CStr = ToCStr(node.UID);
-        	op->parameter.UID.ImportTextItem(CStr, 0, nullptr, nullptr);
-        	op->parameter.Type = EParameterType::Image;
-        	op->parameter.DefaultValue.Set<FParamImageType>(node.DefaultValue);
+			op->Parameter.Name = node.Name;
+			bool bParseOk = FGuid::Parse(node.UID, op->Parameter.UID);
+			check(bParseOk);
+			op->Parameter.Type = EParameterType::Image;
+        	op->Parameter.DefaultValue.Set<FParamImageType>(node.DefaultValue);
 
 			// Generate the code for the ranges
 			for (int32 a = 0; a < node.Ranges.Num(); ++a)
 			{
 				FRangeGenerationResult rangeResult;
 				GenerateRange(rangeResult, Options, node.Ranges[a]);
-				op->ranges.Emplace(op.get(), rangeResult.sizeOp, rangeResult.rangeName, rangeResult.rangeUID);
+				op->Ranges.Emplace(op.get(), rangeResult.sizeOp, rangeResult.rangeName, rangeResult.rangeUID);
 			}
 
 			FirstPass.ParameterNodes.Add(node.m_pNode, op);
diff --git a/Engine/Plugins/Mutable/Source/MutableTools/Private/MuT/CodeGenerator_Matrix.cpp b/Engine/Plugins/Mutable/Source/MutableTools/Private/MuT/CodeGenerator_Matrix.cpp
--- a/Engine/Plugins/Mutable/Source/MutableTools/Private/MuT/CodeGenerator_Matrix.cpp
+++ b/Engine/Plugins/Mutable/Source/MutableTools/Private/MuT/CodeGenerator_Matrix.cpp
@@ -56,21 +56,21 @@ namespace mu
 		{
 			FParameterDesc param;
 			param.Name = Typed->Name;
-			const TCHAR* CStr = ToCStr(Typed->Uid);
-			param.UID.ImportTextItem(CStr, 0, nullptr, nullptr);
+			bool bParseOk = FGuid::Parse(Typed->Uid, param.UID);
+			check(bParseOk);
 			param.Type = EParameterType::Matrix;
 			param.DefaultValue.Set<FParamMatrixType>(Typed->DefaultValue);
 
 			op = new ASTOpParameter();
-			op->type = EOpType::MA_PARAMETER;
-			op->parameter = param;
+			op->Type = EOpType::MA_PARAMETER;
+			op->Parameter = param;
 
 			// Generate the code for the ranges
 			for (int32 a = 0; a < Typed->Ranges.Num(); ++a)
 			{
 				FRangeGenerationResult rangeResult;
 				GenerateRange(rangeResult, Options, Typed->Ranges[a]);
-				op->ranges.Emplace(op.get(), rangeResult.sizeOp, rangeResult.rangeName, rangeResult.rangeUID);
+				op->Ranges.Emplace(op.get(), rangeResult.sizeOp, rangeResult.rangeName, rangeResult.rangeUID);
 			}
 
 			FirstPass.ParameterNodes.Add(Typed.get(), op);
diff --git a/Engine/Plugins/Mutable/Source/MutableTools/Private/MuT/CodeGenerator_Projector.cpp b/Engine/Plugins/Mutable/Source/MutableTools/Private/MuT/CodeGenerator_Projector.cpp
--- a/Engine/Plugins/Mutable/Source/MutableTools/Private/MuT/CodeGenerator_Projector.cpp
+++ b/Engine/Plugins/Mutable/Source/MutableTools/Private/MuT/CodeGenerator_Projector.cpp
@@ -114,23 +114,23 @@ namespace mu
 		{
 			FParameterDesc param;
 			param.Name = node.Name;
-			const TCHAR* CStr = ToCStr(node.UID);
-			param.UID.ImportTextItem(CStr, 0, nullptr, nullptr);
+			bool bParseOk = FGuid::Parse(node.UID, param.UID);
+			check(bParseOk);
 			param.Type = EParameterType::Projector;
 
 			FProjector p = ProjectorFromNode(node);
 			param.DefaultValue.Set<FParamProjectorType>(p);
 
 			op = new ASTOpParameter();
-			op->type = EOpType::PR_PARAMETER;
-			op->parameter = param;
+			op->Type = EOpType::PR_PARAMETER;
+			op->Parameter = param;
 
 			// Generate the code for the ranges
 			for (int32 a = 0; a < node.Ranges.Num(); ++a)
 			{
 				FRangeGenerationResult rangeResult;
 				GenerateRange(rangeResult, Options, node.Ranges[a]);
-				op->ranges.Emplace(op.get(), rangeResult.sizeOp, rangeResult.rangeName, rangeResult.rangeUID);
+				op->Ranges.Emplace(op.get(), rangeResult.sizeOp, rangeResult.rangeName, rangeResult.rangeUID);
 			}
 
 			FirstPass.ParameterNodes.Add(node.m_pNode, op);
@@ -140,7 +140,7 @@ namespace mu
 			op = *it;
 		}
 
-		const FProjector& Projector = op->parameter.DefaultValue.Get<FParamProjectorType>();
+		const FProjector& Projector = op->Parameter.DefaultValue.Get<FParamProjectorType>();
 		result.type = Projector.type;
 		result.op = op;
 	}
diff --git a/Engine/Plugins/Mutable/Source/MutableTools/Private/MuT/CodeGenerator_Scalar.cpp b/Engine/Plugins/Mutable/Source/MutableTools/Private/MuT/CodeGenerator_Scalar.cpp
--- a/Engine/Plugins/Mutable/Source/MutableTools/Private/MuT/CodeGenerator_Scalar.cpp
+++ b/Engine/Plugins/Mutable/Source/MutableTools/Private/MuT/CodeGenerator_Scalar.cpp
@@ -142,21 +142,21 @@ namespace mu
 		{
 			FParameterDesc param;
 			param.Name = node.Name;
-			const TCHAR* CStr = ToCStr(node.UID);
-			param.UID.ImportTextItem(CStr, 0, nullptr, nullptr);
+			bool bParseOk = FGuid::Parse(node.UID, param.UID);
+			check(bParseOk);
 			param.Type = EParameterType::Float;
 			param.DefaultValue.Set<FParamFloatType>(node.DefaultValue);
 
 			op = new ASTOpParameter();
-			op->type = EOpType::SC_PARAMETER;
-			op->parameter = param;
+			op->Type = EOpType::SC_PARAMETER;
+			op->Parameter = param;
 
 			// Generate the code for the ranges
 			for (int32 a = 0; a < node.Ranges.Num(); ++a)
 			{
 				FRangeGenerationResult rangeResult;
 				GenerateRange(rangeResult, Options, node.Ranges[a]);
-				op->ranges.Emplace(op.get(), rangeResult.sizeOp, rangeResult.rangeName, rangeResult.rangeUID);
+				op->Ranges.Emplace(op.get(), rangeResult.sizeOp, rangeResult.rangeName, rangeResult.rangeUID);
 			}
 
 			FirstPass.ParameterNodes.Add(node.m_pNode, op);
@@ -182,8 +182,8 @@ namespace mu
 		{
 			FParameterDesc param;
 			param.Name = node.Name;
-			const TCHAR* CStr = ToCStr(node.UID);
-			param.UID.ImportTextItem(CStr, 0, nullptr, nullptr);
+			bool bParseOk = FGuid::Parse(node.UID, param.UID);
+			check(bParseOk);
 			param.Type = EParameterType::Int;
 			param.DefaultValue.Set<FParamIntType>(node.DefaultValue);
 
@@ -195,15 +195,15 @@ namespace mu
 			}
 
 			op = new ASTOpParameter();
-			op->type = EOpType::NU_PARAMETER;
-			op->parameter = param;
+			op->Type = EOpType::NU_PARAMETER;
+			op->Parameter = param;
 
 			// Generate the code for the ranges
 			for (int32 a = 0; a < node.Ranges.Num(); ++a)
 			{
 				FRangeGenerationResult rangeResult;
 				GenerateRange(rangeResult, Options, node.Ranges[a]);
-				op->ranges.Emplace(op.get(), rangeResult.sizeOp, rangeResult.rangeName, rangeResult.rangeUID);
+				op->Ranges.Emplace(op.get(), rangeResult.sizeOp, rangeResult.rangeName, rangeResult.rangeUID);
 			}
 
 			FirstPass.ParameterNodes.Add(node.m_pNode, op);
diff --git a/Engine/Plugins/Mutable/Source/MutableTools/Private/MuT/CodeGenerator_String.cpp b/Engine/Plugins/Mutable/Source/MutableTools/Private/MuT/CodeGenerator_String.cpp
--- a/Engine/Plugins/Mutable/Source/MutableTools/Private/MuT/CodeGenerator_String.cpp
+++ b/Engine/Plugins/Mutable/Source/MutableTools/Private/MuT/CodeGenerator_String.cpp
@@ -87,21 +87,21 @@ namespace mu
 		{
 			FParameterDesc param;
 			param.Name = node.Name;
-			const TCHAR* CStr = ToCStr(node.UID);
-			param.UID.ImportTextItem(CStr, 0, nullptr, nullptr);
+			bool bParseOk = FGuid::Parse(node.UID, param.UID);
+			check(bParseOk);
 			param.Type = EParameterType::Float;
 			param.DefaultValue.Set<FParamStringType>(node.DefaultValue);
 
 			op = new ASTOpParameter();
-			op->type = EOpType::ST_PARAMETER;
-			op->parameter = param;
+			op->Type = EOpType::ST_PARAMETER;
+			op->Parameter = param;
 
 			// Generate the code for the ranges
 			for (int32 a = 0; a < node.Ranges.Num(); ++a)
 			{
 				FRangeGenerationResult rangeResult;
 				GenerateRange(rangeResult, Options, node.Ranges[a]);
-				op->ranges.Emplace(op.get(), rangeResult.sizeOp, rangeResult.rangeName, rangeResult.rangeUID);
+				op->Ranges.Emplace(op.get(), rangeResult.sizeOp, rangeResult.rangeName, rangeResult.rangeUID);
 			}
 
 			FirstPass.ParameterNodes.Add(node.m_pNode, op);
diff --git a/Engine/Plugins/Mutable/Source/MutableTools/Private/MuT/CodeOptimiser_Parameters.cpp b/Engine/Plugins/Mutable/Source/MutableTools/Private/MuT/CodeOptimiser_Parameters.cpp
--- a/Engine/Plugins/Mutable/Source/MutableTools/Private/MuT/CodeOptimiser_Parameters.cpp
+++ b/Engine/Plugins/Mutable/Source/MutableTools/Private/MuT/CodeOptimiser_Parameters.cpp
@@ -155,7 +155,7 @@ namespace mu
                     {
 						const ASTOpParameter* typed = static_cast<const ASTOpParameter*>(At.get());
                         const TArray<FString>& params = State->nodeState.RuntimeParams;
-                        if ( params.Find( typed->parameter.Name)
+                        if ( params.Find( typed->Parameter.Name)
                              !=
                              INDEX_NONE )
                         {
@@ -1655,7 +1655,7 @@ namespace mu
                 case EOpType::IM_PARAMETER:
                 {
 					const ASTOpParameter* TypedOp = static_cast<const ASTOpParameter*>(At.get());
-                    Parameters.Add(TypedOp->parameter.Name);
+                    Parameters.Add(TypedOp->Parameter.Name);
 
                     // Not interested in the parameters from the parameters decorators.
                     return false;
diff --git a/Engine/Plugins/Mutable/Source/MutableTools/Private/MuT/Compiler.cpp b/Engine/Plugins/Mutable/Source/MutableTools/Private/MuT/Compiler.cpp
--- a/Engine/Plugins/Mutable/Source/MutableTools/Private/MuT/Compiler.cpp
+++ b/Engine/Plugins/Mutable/Source/MutableTools/Private/MuT/Compiler.cpp
@@ -328,7 +328,7 @@ namespace mu
 			int32 ParameterIndex = 0;
 			for ( const TPair<Ptr<const Node>, Ptr<ASTOpParameter>>& Entry : gen.FirstPass.ParameterNodes )
 			{
-				Parameters[ParameterIndex] = Entry.Value->parameter;
+				Parameters[ParameterIndex] = Entry.Value->Parameter;
 				++ParameterIndex;
 			}
 
