[Mutable] Fix ModelResources type mismatch.

#rnx

diff --git a/Engine/Plugins/Mutable/Source/CustomizableObject/Internal/MuCO/CustomizableObjectPrivate.h b/Engine/Plugins/Mutable/Source/CustomizableObject/Internal/MuCO/CustomizableObjectPrivate.h
--- a/Engine/Plugins/Mutable/Source/CustomizableObject/Internal/MuCO/CustomizableObjectPrivate.h
+++ b/Engine/Plugins/Mutable/Source/CustomizableObject/Internal/MuCO/CustomizableObjectPrivate.h
@@ -1198,15 +1198,15 @@ struct FModelResources
 	
 	/** Max number of LODs in the compiled Model. */
 	UPROPERTY()
-	TMap<FName, uint8> NumLODsAvailable;
+	TMap<FName, uint8> ComponentLODs;
 
 	/** Max number of LODs to stream. Mutable will always generate at least one LOD. */
 	UPROPERTY()
-	TMap<FName, uint8> NumLODsToStream;
+	TMap<FName, uint8> ComponentNumLODsToStream;
 
 	/** First LOD available, some platforms may remove lower LODs when cooking, this MinLOD represents the first LOD we can generate */
 	UPROPERTY()
-	TMap<FName, uint8> FirstLODAvailable;
+	TMap<FName, uint8> ComponentFirstLODAvailable;
 
 	/** Name of all possible components. Index is the ObjectComponentIndex. */
 	UPROPERTY()
@@ -1672,6 +1672,8 @@ public:
 
 		LODsPerComponent,
 
+		LODsPerComponentTypeMismatch,
+
 		// -----<new versions can be added above this line>--------
 		LastCustomizableObjectVersion
 	};
diff --git a/Engine/Plugins/Mutable/Source/CustomizableObject/Private/MuCO/CustomizableObject.cpp b/Engine/Plugins/Mutable/Source/CustomizableObject/Private/MuCO/CustomizableObject.cpp
--- a/Engine/Plugins/Mutable/Source/CustomizableObject/Private/MuCO/CustomizableObject.cpp
+++ b/Engine/Plugins/Mutable/Source/CustomizableObject/Private/MuCO/CustomizableObject.cpp
@@ -775,9 +775,9 @@ void FModelResources::Serialize(FObjectAndNameAsStringProxyArchive& MemoryWriter
 	MemoryWriter << ClothingAssetsData;
 	MemoryWriter << ClothSharedConfigsData;
 
-	MemoryWriter << NumLODsAvailable;
-	MemoryWriter << NumLODsToStream;
-	MemoryWriter << FirstLODAvailable;
+	MemoryWriter << ComponentLODs;
+	MemoryWriter << ComponentNumLODsToStream;
+	MemoryWriter << ComponentFirstLODAvailable;
 
 	MemoryWriter << ComponentNamesPerObjectComponent;
 	MemoryWriter << ReleaseVersion;
@@ -937,9 +937,9 @@ bool FModelResources::Unserialize(FObjectAndNameAsStringProxyArchive& MemoryRead
 	MemoryReader << ClothingAssetsData;
 	MemoryReader << ClothSharedConfigsData;
 
-	MemoryReader << NumLODsAvailable;
-	MemoryReader << NumLODsToStream;
-	MemoryReader << FirstLODAvailable;
+	MemoryReader << ComponentLODs;
+	MemoryReader << ComponentNumLODsToStream;
+	MemoryReader << ComponentFirstLODAvailable;
 
 	MemoryReader << ComponentNamesPerObjectComponent;
 	MemoryReader << ReleaseVersion;
@@ -2113,7 +2113,7 @@ uint8 UCustomizableObjectPrivate::GetMinLODIndex(const FName& ComponentName) con
 	}
 
 	uint8 FirstLODAvailable = 0;
-	if (const uint8* Result = GetModelResources().FirstLODAvailable.Find(ComponentName))
+	if (const uint8* Result = GetModelResources().ComponentFirstLODAvailable.Find(ComponentName))
 	{
 		FirstLODAvailable = *Result;
 	}
diff --git a/Engine/Plugins/Mutable/Source/CustomizableObject/Private/MuCO/CustomizableObjectInstance.cpp b/Engine/Plugins/Mutable/Source/CustomizableObject/Private/MuCO/CustomizableObjectInstance.cpp
--- a/Engine/Plugins/Mutable/Source/CustomizableObject/Private/MuCO/CustomizableObjectInstance.cpp
+++ b/Engine/Plugins/Mutable/Source/CustomizableObject/Private/MuCO/CustomizableObjectInstance.cpp
@@ -7178,7 +7178,7 @@ void UCustomizableObjectInstance::SetRequestedLODs(const TMap<FName, uint8>& InM
 		}
 
 		const uint8 MinLODIdx = CustomizableObject->GetPrivate()->GetMinLODIndex(ComponentName);
-		const int32 MaxLODIdx = ModelResources.NumLODsAvailable[ComponentName] - 1;
+		const int32 MaxLODIdx = ModelResources.ComponentLODs[ComponentName] - 1;
 		InMinLOD = FMath::Clamp(InMinLOD, MinLODIdx, MaxLODIdx);
 
 		uint8 DescriptorMinLOD = 0;
@@ -7199,7 +7199,7 @@ void UCustomizableObjectInstance::SetRequestedLODs(const TMap<FName, uint8>& InM
 			
 			PrivateData->SetCOInstanceFlags(InMinLOD > CurrentMinLOD ? PendingLODsDowngrade : ECONone);
 
-			const uint8 FirstNonStreamedLODIndex = FMath::Clamp(CustomizableObject->GetPrivate()->GetModelResources().NumLODsToStream[ComponentName], 0, MaxLODIdx);
+			const uint8 FirstNonStreamedLODIndex = FMath::Clamp(CustomizableObject->GetPrivate()->GetModelResources().ComponentNumLODsToStream[ComponentName], 0, MaxLODIdx);
 
 			MutableUpdateCandidate.RequestedLODLevels.Add(ComponentName, FirstNonStreamedLODIndex);
 
diff --git a/Engine/Plugins/Mutable/Source/CustomizableObject/Private/MuCO/CustomizableObjectSystem.cpp b/Engine/Plugins/Mutable/Source/CustomizableObject/Private/MuCO/CustomizableObjectSystem.cpp
--- a/Engine/Plugins/Mutable/Source/CustomizableObject/Private/MuCO/CustomizableObjectSystem.cpp
+++ b/Engine/Plugins/Mutable/Source/CustomizableObject/Private/MuCO/CustomizableObjectSystem.cpp
@@ -183,8 +183,8 @@ FUpdateContextPrivate::FUpdateContextPrivate(UCustomizableObjectInstance& InInst
 
 	CapturedDescriptorHash = FDescriptorHash(Descriptor);
 	NumObjectComponents = LocalObject->GetComponentCount();
-	FirstLODAvailable = ModelResources.FirstLODAvailable;
-	FirstResidentLOD = ModelResources.NumLODsToStream;
+	FirstLODAvailable = ModelResources.ComponentFirstLODAvailable;
+	FirstResidentLOD = ModelResources.ComponentNumLODsToStream;
 	ComponentNames = ModelResources.ComponentNamesPerObjectComponent;
 	
 	MutableSystem = UCustomizableObjectSystem::GetInstance()->GetPrivate()->MutableSystem;	
diff --git a/Engine/Plugins/Mutable/Source/CustomizableObjectEditor/Private/MuCOE/CustomizableObjectCompiler.cpp b/Engine/Plugins/Mutable/Source/CustomizableObjectEditor/Private/MuCOE/CustomizableObjectCompiler.cpp
--- a/Engine/Plugins/Mutable/Source/CustomizableObjectEditor/Private/MuCOE/CustomizableObjectCompiler.cpp
+++ b/Engine/Plugins/Mutable/Source/CustomizableObjectEditor/Private/MuCOE/CustomizableObjectCompiler.cpp
@@ -993,21 +993,21 @@ void FCustomizableObjectCompiler::CompileInternal(bool bAsync)
 			ModelResources.ReleaseVersion = VersionBridge->GetCurrentVersionAsString();
 		}
 
-		ModelResources.NumLODsAvailable = GenerationContext.NumLODs;
+		ModelResources.ComponentLODs = GenerationContext.NumLODs;
 
 		if (GenerationContext.bEnableLODStreaming)
 		{
-			ModelResources.NumLODsToStream = GenerationContext.NumMaxLODsToStream;
+			ModelResources.ComponentNumLODsToStream = GenerationContext.NumMaxLODsToStream;
 		}
 		else
 		{
-			for (TTuple<FName, uint8>& Pair : ModelResources.NumLODsToStream)
+			for (TTuple<FName, uint8>& Pair : ModelResources.ComponentNumLODsToStream)
 			{
 				Pair.Value = 0;
 			}
 		}
 
-		ModelResources.FirstLODAvailable = GenerationContext.FirstLODAvailable;
+		ModelResources.ComponentFirstLODAvailable = GenerationContext.FirstLODAvailable;
 		
 		ModelResources.ParticipatingObjects = MoveTemp(GenerationContext.ParticipatingObjects);
 		
