[Mutable] Fixed bug where textures of passthrough components did not appear in the Texture Analyzer.

#rnx

diff --git a/Engine/Plugins/Mutable/Source/CustomizableObject/Internal/MuCO/CustomizableObjectSystemPrivate.h b/Engine/Plugins/Mutable/Source/CustomizableObject/Internal/MuCO/CustomizableObjectSystemPrivate.h
--- a/Engine/Plugins/Mutable/Source/CustomizableObject/Internal/MuCO/CustomizableObjectSystemPrivate.h
+++ b/Engine/Plugins/Mutable/Source/CustomizableObject/Internal/MuCO/CustomizableObjectSystemPrivate.h
@@ -300,6 +300,12 @@ struct FGeneratedMaterial
 	// Index of the material to instantiate (UCustomizableObject::ReferencedMaterials)
 	uint32 MaterialIndex = 0;
 
+#if WITH_EDITORONLY_DATA
+	// Name of the component that contains this material
+	UPROPERTY(Category = NoCategory, VisibleAnywhere)
+	FName ComponentName;
+#endif
+
 	bool operator==(const FGeneratedMaterial& Other) const { return SurfaceId == Other.SurfaceId && MaterialIndex == Other.MaterialIndex; };
 };
 
diff --git a/Engine/Plugins/Mutable/Source/CustomizableObject/Private/MuCO/CustomizableObjectInstance.cpp b/Engine/Plugins/Mutable/Source/CustomizableObject/Private/MuCO/CustomizableObjectInstance.cpp
--- a/Engine/Plugins/Mutable/Source/CustomizableObject/Private/MuCO/CustomizableObjectInstance.cpp
+++ b/Engine/Plugins/Mutable/Source/CustomizableObject/Private/MuCO/CustomizableObjectInstance.cpp
@@ -6841,6 +6841,10 @@ void UCustomizableInstancePrivate::BuildMaterials(const TSharedRef<FUpdateContex
 					Material.MaterialIndex = Surface.MaterialIndex;
 					Material.MaterialInterface = MaterialInterface;
 
+#if WITH_EDITORONLY_DATA
+					Material.ComponentName = ComponentName;
+#endif
+
 					UMaterialInstanceDynamic* MaterialInstance = nullptr;
 					
 					if (const int32 OldMaterialIndex = OldGeneratedMaterials.Find(Material); bReuseMaterials && OldMaterialIndex != INDEX_NONE)
diff --git a/Engine/Plugins/Mutable/Source/CustomizableObjectEditor/Private/MuCOE/CustomizableObjectEditor.h b/Engine/Plugins/Mutable/Source/CustomizableObjectEditor/Private/MuCOE/CustomizableObjectEditor.h
--- a/Engine/Plugins/Mutable/Source/CustomizableObjectEditor/Private/MuCOE/CustomizableObjectEditor.h
+++ b/Engine/Plugins/Mutable/Source/CustomizableObjectEditor/Private/MuCOE/CustomizableObjectEditor.h
@@ -36,10 +36,8 @@ class UCustomizableObjectNodeProjectorParameter;
 class UCustomizableSkeletalComponent;
 class UDebugSkelMeshComponent;
 class UEdGraph;
-class UMaterialInterface;
 class UPoseAsset;
 class UStaticMeshComponent;
-class UTexture;
 struct FAssetData;
 struct FFrame;
 struct FPropertyChangedEvent;
@@ -79,75 +77,6 @@ public:
 };
 
 
-/** Statistics for the Texture Analyzer */
-UCLASS(Transient, MinimalAPI, meta = (DisplayName = "Texture Stats"))
-class UCustomizableObjectEditorTextureStats : public UObject
-{
-	GENERATED_BODY()
-
-public:
-
-	/** Texture - double click to open */
-	UPROPERTY(VisibleAnywhere, AssetRegistrySearchable, Category = "Stats", meta = (DisplayName = "Texture", ColumnWidth = "40", NoResetToDefault))
-	FString TextureName;
-
-	/** Texture - double click to open */
-	UPROPERTY(VisibleAnywhere, AssetRegistrySearchable, Category = "Stats", meta = (DisplayName = "Name", ColumnWidth = "50", NoResetToDefault))
-	FString TextureParameterName;
-
-	/** Material - double click to open */
-	UPROPERTY(VisibleAnywhere, AssetRegistrySearchable, Category = "Stats", meta = (DisplayName = "Material", ColumnWidth = "50", NoResetToDefault))
-	FString MaterialName;
-
-	/** Parent Material - double click to open */
-	UPROPERTY(VisibleAnywhere, AssetRegistrySearchable, Category = "Stats", meta = (DisplayName = "Parent", ColumnWidth = "50", NoResetToDefault))
-	FString MaterialParameterName;
-
-	/** Used to open the texture in the editor*/
-	UPROPERTY(meta = (NoResetToDefault))
-	TObjectPtr<UTexture> Texture;
-
-	/** Used to open the material in the editor*/
-	UPROPERTY(meta = (NoResetToDefault))
-	TObjectPtr<UMaterialInterface> Material;
-
-	/** Used to open the parent material in the editor*/
-	UPROPERTY(meta = (NoResetToDefault))
-	TObjectPtr<UMaterialInterface> ParentMaterial;
-
-	/** Resolution of the texture */
-	UPROPERTY(VisibleAnywhere, AssetRegistrySearchable, Category = "Stats", meta = (DisplayName = "Resolution X", ColumnWidth = "40", DisplayRight = "true", NoResetToDefault))
-	int32 ResolutionX;
-	UPROPERTY(VisibleAnywhere, AssetRegistrySearchable, Category = "Stats", meta = (DisplayName = "Resolution Y", ColumnWidth = "40", DisplayRight = "true", NoResetToDefault))
-	int32 ResolutionY;
-
-	/** The memory used in KB */
-	UPROPERTY(VisibleAnywhere, AssetRegistrySearchable, Category = "Stats", meta = (DisplayName = "Size Kb", ColumnWidth = "90", NoResetToDefault))
-	FString Size;
-
-	/** The texture format, e.g. PF_DXT1 */
-	UPROPERTY(VisibleAnywhere, AssetRegistrySearchable, Category = "Stats", meta = (ColumnWidth = "96", NoResetToDefault))
-	TEnumAsByte<EPixelFormat> Format;
-
-	/** LOD Bias for this texture. (Texture LODBias + Texture group) */
-	UPROPERTY(VisibleAnywhere, AssetRegistrySearchable, Category = "Stats", meta = (DisplayName = "LOD Bias", ColumnWidth = "70", NoResetToDefault))
-	int32 LODBias;
-	
-	/** Says if the texture is being streamed */
-	UPROPERTY(VisibleAnywhere, AssetRegistrySearchable, Category = "Stats", meta = (DisplayName = "Streamed", ColumnWidth = "70", NoResetToDefault))
-	FString IsStreamed;
-
-	/** The Level of detail group of the texture */
-	UPROPERTY(VisibleAnywhere, AssetRegistrySearchable, Category = "Stats", meta = (DisplayName = "Texture Group", ColumnWidth = "70", NoResetToDefault))
-	TEnumAsByte<enum TextureGroup> LODGroup;
-	
-	/** The Component of the texture */
-	UPROPERTY(VisibleAnywhere, AssetRegistrySearchable, Category = "Stats", meta = (DisplayName = "Component", ColumnWidth = "70", NoResetToDefault))
-	FName ComponentName;
-
-};
-
-
 /**
  * CustomizableObject Editor class
  */
diff --git a/Engine/Plugins/Mutable/Source/CustomizableObjectEditor/Private/MuCOE/SCustomizableObjectEditorTextureAnalyzer.cpp b/Engine/Plugins/Mutable/Source/CustomizableObjectEditor/Private/MuCOE/SCustomizableObjectEditorTextureAnalyzer.cpp
--- a/Engine/Plugins/Mutable/Source/CustomizableObjectEditor/Private/MuCOE/SCustomizableObjectEditorTextureAnalyzer.cpp
+++ b/Engine/Plugins/Mutable/Source/CustomizableObjectEditor/Private/MuCOE/SCustomizableObjectEditorTextureAnalyzer.cpp
@@ -4,13 +4,18 @@
 
 #include "Engine/SkeletalMesh.h"
 #include "Engine/SkinnedAssetCommon.h"
+#include "Engine/Texture.h"
+#include "Engine/TextureDefines.h"
 #include "Framework/Views/ITypedTableView.h"
 #include "IPropertyTable.h"
 #include "IPropertyTableCell.h"
 #include "IPropertyTableColumn.h"
 #include "Materials/MaterialInstance.h"
+#include "Materials/MaterialInterface.h"
+#include "Materials/MaterialInstanceDynamic.h"
 #include "Modules/ModuleManager.h"
 #include "MuCO/CustomizableObjectInstance.h"
+#include "MuCO/CustomizableObjectInstancePrivate.h"
 #include "MuCOE/CustomizableObjectEditor.h"
 #include "MuCOE/CustomizableObjectInstanceEditor.h"
 #include "PropertyEditorModule.h"
@@ -133,59 +138,48 @@ void SCustomizableObjecEditorTextureAnalyzer::FillTextureAnalyzerTable(UCustomiz
 		const UCustomizableObject* CustomizableObject = PreviewInstance->GetCustomizableObject();
 		check(CustomizableObject)
 
-		for (int32 ObjectComponentIndex = 0; ObjectComponentIndex < CustomizableObject->GetComponentCount(); ++ObjectComponentIndex)
+		for (const FGeneratedMaterial& GeneratedMaterial : PreviewInstance->GetPrivate()->GeneratedMaterials)
 		{
-			const FName ComponentName = CustomizableObject->GetPrivate()->GetComponentName(FCustomizableObjectComponentIndex(ObjectComponentIndex));
-			
-			if (const USkeletalMesh* SkeletalMesh = PreviewInstance->GetComponentMeshSkeletalMesh(ComponentName))
+			UMaterialInstanceDynamic* Material = Cast<UMaterialInstanceDynamic>(GeneratedMaterial.MaterialInterface);
+			if (!Material)
 			{
-				for (int MaterialIndex = 0; MaterialIndex < SkeletalMesh->GetMaterials().Num(); ++MaterialIndex)
+				continue;
+			}
+
+			for (const FGeneratedTexture& GeneratedTexture : GeneratedMaterial.Textures)
+			{
+				UTexture* Texture = GeneratedTexture.Texture;
+				if (!Texture)
 				{
-					UMaterialInstance* Material = Cast<UMaterialInstance>(SkeletalMesh->GetMaterials()[MaterialIndex].MaterialInterface);
-
-					// We just want the Transient materials
-					if (Material && Material->GetPathName().Contains("Transient"))
-					{
-						for (int32 TextureIndex = 0; TextureIndex < Material->TextureParameterValues.Num(); ++TextureIndex)
-						{
-							UTexture* Texture = Material->TextureParameterValues[TextureIndex].ParameterValue;
-							if (!Texture) continue;
-
-							FString TextPath = Texture->GetPathName();
-
-							// We just want the Transient textures
-							if (TextPath.Contains("Transient"))
-							{
-								//New Object for the table
-								UCustomizableObjectEditorTextureStats* Entry = NewObject<UCustomizableObjectEditorTextureStats>();
-								TabTextures.Add(Entry);
-
-								//Texture Info
-								Entry->Texture = Texture;
-								Entry->TextureName = Texture->GetName();
-								Entry->TextureParameterName = Material->TextureParameterValues[TextureIndex].ParameterInfo.Name.ToString();
-								Entry->ResolutionX = Cast<UTexture2D>(Texture)->GetSizeX();
-								Entry->ResolutionY = Cast<UTexture2D>(Texture)->GetSizeY();
-								Entry->LODBias = Texture->GetCachedLODBias();
-								uint32 textureSize = Texture->CalcTextureMemorySizeEnum(TMC_AllMips);
-								Entry->Size = FString::Printf(TEXT("%.2f"), textureSize / 1024.0);
-								Entry->Format = Cast<UTexture2D>(Texture)->GetPixelFormat();
-								Entry->IsStreamed = Texture->IsCurrentlyVirtualTextured() ? "Virtual Streamed" : (!Texture->IsStreamable() ? "Not Streamed" : "Streamed");
-								Entry->LODGroup = Texture->LODGroup;
-
-								//Material Info
-								Entry->Material = Material;
-								Entry->MaterialName = Material->GetName();
-								Entry->ParentMaterial = Material->Parent;
-								Entry->MaterialParameterName = Material->Parent ? Material->Parent->GetName() : FString();
-								Entry->ComponentName = ComponentName;
-
-								//Total size
-								TotalSize += textureSize;
-							}
-						}
-					}
+					continue;
 				}
+
+				//New Object for the table
+				UCustomizableObjectEditorTextureStats* Entry = NewObject<UCustomizableObjectEditorTextureStats>();
+				TabTextures.Add(Entry);
+
+				//Texture Info
+				Entry->Texture = Texture;
+				Entry->TextureName = Texture->GetName();
+				Entry->TextureParameterName = GeneratedTexture.Name;
+				Entry->ResolutionX = Cast<UTexture2D>(Texture)->GetSizeX();
+				Entry->ResolutionY = Cast<UTexture2D>(Texture)->GetSizeY();
+				Entry->LODBias = Texture->GetCachedLODBias();
+				uint32 TextureSize = Texture->CalcTextureMemorySizeEnum(TMC_AllMips);
+				Entry->Size = TextureSize / 1024.0f;
+				Entry->Format = Cast<UTexture2D>(Texture)->GetPixelFormat();
+				Entry->IsStreamed = Texture->IsCurrentlyVirtualTextured() ? "Virtual Streamed" : (!Texture->IsStreamable() ? "Not Streamed" : "Streamed");
+				Entry->LODGroup = Texture->LODGroup;
+
+				//Material Info
+				Entry->Material = Material;
+				Entry->MaterialName = Material->GetName();
+				Entry->ParentMaterial = Material->Parent;
+				Entry->MaterialParameterName = Material->Parent ? Material->Parent->GetName() : FString();
+				Entry->ComponentName = GeneratedMaterial.ComponentName;
+
+				//Total size
+				TotalSize += TextureSize;
 			}
 		}
 
@@ -207,7 +201,7 @@ void SCustomizableObjecEditorTextureAnalyzer::OnTextureTableSelectionChanged(TSh
 	if (UCustomizableObjectEditorTextureStats* Current = Cast<UCustomizableObjectEditorTextureStats>(Cell->GetObject().Get()))
 	{
 		//Open Texture Editor Tab
-		if ((SelectedColumn == "Texture" || SelectedColumn == "Name") && Current->Texture)
+		if ((SelectedColumn == "Texture" || SelectedColumn == "Parameter Name") && Current->Texture)
 		{
 			GEditor->GetEditorSubsystem<UAssetEditorSubsystem>()->OpenEditorForAsset(Current->Texture);
 		}
diff --git a/Engine/Plugins/Mutable/Source/CustomizableObjectEditor/Private/MuCOE/SCustomizableObjectEditorTextureAnalyzer.h b/Engine/Plugins/Mutable/Source/CustomizableObjectEditor/Private/MuCOE/SCustomizableObjectEditorTextureAnalyzer.h
--- a/Engine/Plugins/Mutable/Source/CustomizableObjectEditor/Private/MuCOE/SCustomizableObjectEditorTextureAnalyzer.h
+++ b/Engine/Plugins/Mutable/Source/CustomizableObjectEditor/Private/MuCOE/SCustomizableObjectEditorTextureAnalyzer.h
@@ -4,6 +4,8 @@
 
 #include "UObject/GCObject.h"
 #include "Widgets/SCompoundWidget.h"
+#include "SCustomizableObjectEditorTextureAnalyzer.generated.h"
+
 
 class FCustomizableObjectEditor;
 class FCustomizableObjectInstanceEditor;
@@ -12,10 +14,84 @@ class IPropertyTable;
 class STextBlock;
 class SWidget;
 class UCustomizableObjectInstance;
+class UMaterialInterface;
 class UObject;
+class UTexture;
+
+enum TextureGroup : int;
+
 struct FGeometry;
 
 
+/** Statistics for the Texture Analyzer */
+UCLASS(Transient, MinimalAPI, meta = (DisplayName = "Texture Stats"))
+class UCustomizableObjectEditorTextureStats : public UObject
+{
+	GENERATED_BODY()
+
+public:
+
+	/** Texture - double click to open */
+	UPROPERTY(VisibleAnywhere, AssetRegistrySearchable, Category = "Stats", meta = (DisplayName = "Parameter Name", ColumnWidth = "50", NoResetToDefault))
+	FString TextureParameterName;
+
+	/** Texture - double click to open */
+	UPROPERTY(VisibleAnywhere, AssetRegistrySearchable, Category = "Stats", meta = (DisplayName = "Texture", ColumnWidth = "40", NoResetToDefault))
+	FString TextureName;
+
+	/** Material - double click to open */
+	UPROPERTY(VisibleAnywhere, AssetRegistrySearchable, Category = "Stats", meta = (DisplayName = "Material", ColumnWidth = "50", NoResetToDefault))
+	FString MaterialName;
+
+	/** Parent Material - double click to open */
+	UPROPERTY(VisibleAnywhere, AssetRegistrySearchable, Category = "Stats", meta = (DisplayName = "Parent", ColumnWidth = "50", NoResetToDefault))
+	FString MaterialParameterName;
+
+	/** Used to open the texture in the editor*/
+	UPROPERTY(meta = (NoResetToDefault))
+	TObjectPtr<UTexture> Texture;
+
+	/** Used to open the material in the editor*/
+	UPROPERTY(meta = (NoResetToDefault))
+	TObjectPtr<UMaterialInterface> Material;
+
+	/** Used to open the parent material in the editor*/
+	UPROPERTY(meta = (NoResetToDefault))
+	TObjectPtr<UMaterialInterface> ParentMaterial;
+
+	/** Resolution of the texture */
+	UPROPERTY(VisibleAnywhere, AssetRegistrySearchable, Category = "Stats", meta = (DisplayName = "Resolution X", ColumnWidth = "40", DisplayRight = "true", NoResetToDefault))
+	int32 ResolutionX;
+	UPROPERTY(VisibleAnywhere, AssetRegistrySearchable, Category = "Stats", meta = (DisplayName = "Resolution Y", ColumnWidth = "40", DisplayRight = "true", NoResetToDefault))
+	int32 ResolutionY;
+
+	/** The memory used in KB */
+	UPROPERTY(VisibleAnywhere, AssetRegistrySearchable, Category = "Stats", meta = (DisplayName = "Size Kb", ColumnWidth = "90", NoResetToDefault, Units = "Kb"))
+	float Size;
+
+	/** The texture format, e.g. PF_DXT1 */
+	UPROPERTY(VisibleAnywhere, AssetRegistrySearchable, Category = "Stats", meta = (ColumnWidth = "96", NoResetToDefault))
+	TEnumAsByte<EPixelFormat> Format;
+
+	/** LOD Bias for this texture. (Texture LODBias + Texture group) */
+	UPROPERTY(VisibleAnywhere, AssetRegistrySearchable, Category = "Stats", meta = (DisplayName = "LOD Bias", ColumnWidth = "70", NoResetToDefault))
+	int32 LODBias;
+
+	/** Says if the texture is being streamed */
+	UPROPERTY(VisibleAnywhere, AssetRegistrySearchable, Category = "Stats", meta = (DisplayName = "Streamed", ColumnWidth = "70", NoResetToDefault))
+	FString IsStreamed;
+
+	/** The Level of detail group of the texture */
+	UPROPERTY(VisibleAnywhere, AssetRegistrySearchable, Category = "Stats", meta = (DisplayName = "Texture Group", ColumnWidth = "70", NoResetToDefault))
+	TEnumAsByte<TextureGroup> LODGroup;
+
+	/** The Component of the texture */
+	UPROPERTY(VisibleAnywhere, AssetRegistrySearchable, Category = "Stats", meta = (DisplayName = "Component", ColumnWidth = "70", NoResetToDefault))
+	FName ComponentName;
+
+};
+
+
 class SCustomizableObjecEditorTextureAnalyzer : public SCompoundWidget, public FGCObject
 {
 public:
