[mutable] Fixed crash when editing a transform parameter though a numeric input field widget in the CO Editor.
#tests Editor
#rnx

diff --git a/Engine/Plugins/Mutable/Source/CustomizableObject/Private/MuCO/CustomizableObjectInstanceDescriptor.cpp b/Engine/Plugins/Mutable/Source/CustomizableObject/Private/MuCO/CustomizableObjectInstanceDescriptor.cpp
--- a/Engine/Plugins/Mutable/Source/CustomizableObject/Private/MuCO/CustomizableObjectInstanceDescriptor.cpp
+++ b/Engine/Plugins/Mutable/Source/CustomizableObject/Private/MuCO/CustomizableObjectInstanceDescriptor.cpp
@@ -1406,7 +1406,7 @@ float FCustomizableObjectInstanceDescriptor::GetFloatParameterSelectedOption(con
 	const int32 ParameterIndexInObject = CustomizableObject->FindParameter(FloatParamName);
 	const int32 FloatParamIndex = FindTypedParameterIndex(FloatParamName, EMutableParameterType::Float);
 
-	if (ParameterIndexInObject >= 0 && FloatParamIndex >= 0)
+	if (ParameterIndexInObject >= 0 && FloatParameters.IsValidIndex(FloatParamIndex))
 	{
 		// Due to optimizations a parameter may lose its multidimensionality (if it becomes constant). 
 		// In that case it means there is only one possible option so it is ok to set it as if wasn't multidimensional.
diff --git a/Engine/Plugins/Mutable/Source/CustomizableObjectEditor/Private/MuCOE/CustomizableInstanceDetails.cpp b/Engine/Plugins/Mutable/Source/CustomizableObjectEditor/Private/MuCOE/CustomizableInstanceDetails.cpp
--- a/Engine/Plugins/Mutable/Source/CustomizableObjectEditor/Private/MuCOE/CustomizableInstanceDetails.cpp
+++ b/Engine/Plugins/Mutable/Source/CustomizableObjectEditor/Private/MuCOE/CustomizableInstanceDetails.cpp
@@ -1432,74 +1432,92 @@ void FCustomizableInstanceDetails::OnSetColorFromColorPicker(FLinearColor NewCol
 
 TSharedRef<SWidget> FCustomizableInstanceDetails::GenerateTransformWidget(const FString& ParamName)
 {
-	auto OnLocationChanged = [this, ParamName](double Value, EAxis::Type Axis)
+	auto OnLocationChanged = [WeakThis = SharedThis(this).ToWeakPtr(), ParamName](double Value, EAxis::Type Axis)
 	{
-		FTransform Transform = GetTransformParameterValue(ParamName);
-		FVector Location = Transform.GetLocation();
-		if (!FMath::IsNearlyEqualByULP(Value, Location.GetComponentForAxis(Axis)))
+		if (TSharedPtr<FCustomizableInstanceDetails> SharedThis = WeakThis.Pin())
 		{
-			Location.SetComponentForAxis(Axis, Value);
-			Transform.SetLocation(Location);
-			OnTransformParameterChanged(Transform, ParamName);
+			FTransform Transform = SharedThis->GetTransformParameterValue(ParamName);
+			FVector Location = Transform.GetLocation();
+			if (!FMath::IsNearlyEqualByULP(Value, Location.GetComponentForAxis(Axis)))
+			{
+				Location.SetComponentForAxis(Axis, Value);
+				Transform.SetLocation(Location);
+				SharedThis->OnTransformParameterChanged(Transform, ParamName);
+			}
 		}
 	};
-	auto OnLocationCommitted = [this, ParamName](double Value, ETextCommit::Type Type, EAxis::Type Axis)
+	auto OnLocationCommitted = [WeakThis = SharedThis(this).ToWeakPtr(), ParamName](double Value, ETextCommit::Type Type, EAxis::Type Axis)
 	{
-		FTransform Transform = GetTransformParameterValue(ParamName);
-		FVector Location = Transform.GetLocation();
-		if (!FMath::IsNearlyEqualByULP(Value, Location.GetComponentForAxis(Axis)))
+		if (TSharedPtr<FCustomizableInstanceDetails> SharedThis = WeakThis.Pin())
 		{
-			Location.SetComponentForAxis(Axis, Value);
-			Transform.SetLocation(Location);
-			OnTransformParameterCommitted(Transform, Type, ParamName);
+			FTransform Transform = SharedThis->GetTransformParameterValue(ParamName);
+			FVector Location = Transform.GetLocation();
+			if (!FMath::IsNearlyEqualByULP(Value, Location.GetComponentForAxis(Axis)))
+			{
+				Location.SetComponentForAxis(Axis, Value);
+				Transform.SetLocation(Location);
+				SharedThis->OnTransformParameterCommitted(Transform, Type, ParamName);
+			}
 		}
 	};
 
-	auto OnRotationChanged = [this, ParamName](double Value, EAxis::Type Axis)
+	auto OnRotationChanged = [WeakThis = SharedThis(this).ToWeakPtr(), ParamName](double Value, EAxis::Type Axis)
 	{
-		FTransform Transform = GetTransformParameterValue(ParamName);
-		FRotator Rotation = Transform.Rotator();
-		Rotation.SetComponentForAxis(Axis, Value);
-		
-		if (!Transform.Rotator().Equals(Rotation))
+		if (TSharedPtr<FCustomizableInstanceDetails> SharedThis = WeakThis.Pin())
 		{
-			Transform.SetRotation(Rotation.Quaternion());
-			OnTransformParameterChanged(Transform, ParamName);
+			FTransform Transform = SharedThis->GetTransformParameterValue(ParamName);
+			FRotator Rotation = Transform.Rotator();
+			Rotation.SetComponentForAxis(Axis, Value);
+		
+			if (!Transform.Rotator().Equals(Rotation))
+			{
+				Transform.SetRotation(Rotation.Quaternion());
+				SharedThis->OnTransformParameterChanged(Transform, ParamName);
+			}
 		}
 	};
-	auto OnRotationCommitted = [this, ParamName](double Value, ETextCommit::Type Type, EAxis::Type Axis)
+	auto OnRotationCommitted = [WeakThis = SharedThis(this).ToWeakPtr(), ParamName](double Value, ETextCommit::Type Type, EAxis::Type Axis)
 	{
-		FTransform Transform = GetTransformParameterValue(ParamName);
-		FRotator Rotation = Transform.Rotator();
-		Rotation.SetComponentForAxis(Axis, Value);
-		
-		if (!Transform.Rotator().Equals(Rotation))
+		if (TSharedPtr<FCustomizableInstanceDetails> SharedThis = WeakThis.Pin())
 		{
-			Transform.SetRotation(Rotation.Quaternion());
-			OnTransformParameterCommitted(Transform, Type, ParamName);
+			FTransform Transform = SharedThis->GetTransformParameterValue(ParamName);
+			FRotator Rotation = Transform.Rotator();
+			Rotation.SetComponentForAxis(Axis, Value);
+		
+			if (!Transform.Rotator().Equals(Rotation))
+			{
+				Transform.SetRotation(Rotation.Quaternion());
+				SharedThis->OnTransformParameterCommitted(Transform, Type, ParamName);
+			}
 		}
 	};
 	
-	auto OnScaleChanged = [this, ParamName](double Value, EAxis::Type Axis)
+	auto OnScaleChanged = [WeakThis = SharedThis(this).ToWeakPtr(), ParamName](double Value, EAxis::Type Axis)
 	{
-		FTransform Transform = GetTransformParameterValue(ParamName);
-		FVector Scale = Transform.GetScale3D();
-		if (!FMath::IsNearlyEqualByULP(Value, Scale.GetComponentForAxis(Axis)))
+		if (TSharedPtr<FCustomizableInstanceDetails> SharedThis = WeakThis.Pin())
 		{
-			Scale.SetComponentForAxis(Axis, Value);
-			Transform.SetScale3D(Scale);
-			OnTransformParameterChanged(Transform, ParamName);
+			FTransform Transform = SharedThis->GetTransformParameterValue(ParamName);
+			FVector Scale = Transform.GetScale3D();
+			if (!FMath::IsNearlyEqualByULP(Value, Scale.GetComponentForAxis(Axis)))
+			{
+				Scale.SetComponentForAxis(Axis, Value);
+				Transform.SetScale3D(Scale);
+				SharedThis->OnTransformParameterChanged(Transform, ParamName);
+			}
 		}
 	};
-	auto OnScaleCommitted = [this, ParamName](double Value, ETextCommit::Type Type, EAxis::Type Axis)
+	auto OnScaleCommitted = [WeakThis = SharedThis(this).ToWeakPtr(), ParamName](double Value, ETextCommit::Type Type, EAxis::Type Axis)
 	{
-		FTransform Transform = GetTransformParameterValue(ParamName);
-		FVector Scale = Transform.GetScale3D();
-		if (!FMath::IsNearlyEqualByULP(Value, Scale.GetComponentForAxis(Axis)))
+		if (TSharedPtr<FCustomizableInstanceDetails> SharedThis = WeakThis.Pin())
 		{
-			Scale.SetComponentForAxis(Axis, Value);
-			Transform.SetScale3D(Scale);
-			OnTransformParameterCommitted(Transform, Type, ParamName);
+			FTransform Transform = SharedThis->GetTransformParameterValue(ParamName);
+			FVector Scale = Transform.GetScale3D();
+			if (!FMath::IsNearlyEqualByULP(Value, Scale.GetComponentForAxis(Axis)))
+			{
+				Scale.SetComponentForAxis(Axis, Value);
+				Transform.SetScale3D(Scale);
+				SharedThis->OnTransformParameterCommitted(Transform, Type, ParamName);
+			}
 		}
 	};
 	
@@ -1536,9 +1554,33 @@ TSharedRef<SWidget> FCustomizableInstanceDetails::GenerateTransformWidget(const
 			SNew(SNumericVectorInputBox<double>)
 			.bColorAxisLabels(true)
 			.AllowSpin(true)
-			.X_Lambda([this, ParamName]() { return GetTransformParameterValue(ParamName).GetLocation().X; })
-			.Y_Lambda([this, ParamName]() { return GetTransformParameterValue(ParamName).GetLocation().Y; })
-			.Z_Lambda([this, ParamName]() { return GetTransformParameterValue(ParamName).GetLocation().Z; })
+			.X_Lambda([WeakThis = SharedThis(this).ToWeakPtr(), ParamName]()
+			{
+				if (TSharedPtr<FCustomizableInstanceDetails> SharedThis = WeakThis.Pin())
+				{
+					return SharedThis->GetTransformParameterValue(ParamName).GetLocation().X;
+				}
+
+				return 0.0;
+			})
+			.Y_Lambda([WeakThis = SharedThis(this).ToWeakPtr(), ParamName]()
+			{
+				if (TSharedPtr<FCustomizableInstanceDetails> SharedThis = WeakThis.Pin())
+				{
+					return SharedThis->GetTransformParameterValue(ParamName).GetLocation().Y;
+				}
+
+				return 0.0;
+			})
+			.Z_Lambda([WeakThis = SharedThis(this).ToWeakPtr(), ParamName]()
+			{
+				if (TSharedPtr<FCustomizableInstanceDetails> SharedThis = WeakThis.Pin())
+				{
+					return SharedThis->GetTransformParameterValue(ParamName).GetLocation().Z;
+				}
+
+				return 0.0;
+			})
 			.OnXChanged_Lambda([OnLocationChanged](float Value) { OnLocationChanged(Value, EAxis::X); })
 			.OnYChanged_Lambda([OnLocationChanged](float Value) { OnLocationChanged(Value, EAxis::Y); })
 			.OnZChanged_Lambda([OnLocationChanged](float Value) { OnLocationChanged(Value, EAxis::Z); })
@@ -1570,9 +1612,33 @@ TSharedRef<SWidget> FCustomizableInstanceDetails::GenerateTransformWidget(const
 			SNew(SNumericRotatorInputBox<double>)
 			.bColorAxisLabels(true)
 			.AllowSpin(true)
-			.Roll_Lambda([this, ParamName]() { return GetTransformParameterValue(ParamName).Rotator().Roll; })
-			.Pitch_Lambda([this, ParamName]() { return GetTransformParameterValue(ParamName).Rotator().Pitch; })
-			.Yaw_Lambda([this, ParamName]() { return GetTransformParameterValue(ParamName).Rotator().Yaw; })
+			.Roll_Lambda([WeakThis = SharedThis(this).ToWeakPtr(), ParamName]()
+			{
+				if (TSharedPtr<FCustomizableInstanceDetails> SharedThis = WeakThis.Pin())
+				{
+					return SharedThis->GetTransformParameterValue(ParamName).Rotator().Roll;
+				}
+				
+				return 0.0;
+			})
+			.Pitch_Lambda([WeakThis = SharedThis(this).ToWeakPtr(), ParamName]()
+			{
+				if (TSharedPtr<FCustomizableInstanceDetails> SharedThis = WeakThis.Pin())
+				{
+					return SharedThis->GetTransformParameterValue(ParamName).Rotator().Pitch;
+				}
+				
+				return 0.0;
+			})
+			.Yaw_Lambda([WeakThis = SharedThis(this).ToWeakPtr(), ParamName]()
+			{
+				if (TSharedPtr<FCustomizableInstanceDetails> SharedThis = WeakThis.Pin())
+				{
+					return SharedThis->GetTransformParameterValue(ParamName).Rotator().Yaw;
+				}
+
+				return 0.0;
+			})
 			.OnRollChanged_Lambda([OnRotationChanged](float Value) { OnRotationChanged(Value, EAxis::X); })
 			.OnPitchChanged_Lambda([OnRotationChanged](float Value) { OnRotationChanged(Value, EAxis::Y); })
 			.OnYawChanged_Lambda([OnRotationChanged](float Value) { OnRotationChanged(Value, EAxis::Z); })
@@ -1604,9 +1670,33 @@ TSharedRef<SWidget> FCustomizableInstanceDetails::GenerateTransformWidget(const
 			SNew(SNumericVectorInputBox<double>)
 			.bColorAxisLabels(true)
 			.AllowSpin(false)
-			.X_Lambda([this, ParamName]() { return GetTransformParameterValue(ParamName).GetScale3D().X; })
-			.Y_Lambda([this, ParamName]() { return GetTransformParameterValue(ParamName).GetScale3D().Y; })
-			.Z_Lambda([this, ParamName]() { return GetTransformParameterValue(ParamName).GetScale3D().Z; })
+			.X_Lambda([WeakThis = SharedThis(this).ToWeakPtr(), ParamName]()
+			{
+				if (TSharedPtr<FCustomizableInstanceDetails> SharedThis = WeakThis.Pin())
+				{
+					return SharedThis->GetTransformParameterValue(ParamName).GetScale3D().X;
+				}
+				
+				return 0.0;
+			})
+			.Y_Lambda([WeakThis = SharedThis(this).ToWeakPtr(), ParamName]()
+			{
+				if (TSharedPtr<FCustomizableInstanceDetails> SharedThis = WeakThis.Pin())
+				{
+					return SharedThis->GetTransformParameterValue(ParamName).GetScale3D().Y;
+				}
+				
+				return 0.0;
+			})
+			.Z_Lambda([WeakThis = SharedThis(this).ToWeakPtr(), ParamName]()
+			{
+				if (TSharedPtr<FCustomizableInstanceDetails> SharedThis = WeakThis.Pin())
+				{
+					return SharedThis->GetTransformParameterValue(ParamName).GetScale3D().Z;
+				}
+				
+				return 0.0;
+			})
 
 			.OnXChanged_Lambda([OnScaleChanged](float Value) { OnScaleChanged(Value, EAxis::X); })
 			.OnYChanged_Lambda([OnScaleChanged](float Value) { OnScaleChanged(Value, EAxis::Y); })
