[Mutable] Fixed crash when running game in Standalone mode.
- Restricted access to CustomizableObjectEditorModule if the game is being played
- Fixed typos

#rnx

diff --git a/Engine/Plugins/Mutable/Source/CustomizableObject/Internal/MuCO/ICustomizableObjectEditorModule.h b/Engine/Plugins/Mutable/Source/CustomizableObject/Internal/MuCO/ICustomizableObjectEditorModule.h
--- a/Engine/Plugins/Mutable/Source/CustomizableObject/Internal/MuCO/ICustomizableObjectEditorModule.h
+++ b/Engine/Plugins/Mutable/Source/CustomizableObject/Internal/MuCO/ICustomizableObjectEditorModule.h
@@ -35,11 +35,18 @@ class CUSTOMIZABLEOBJECT_API ICustomizableObjectEditorModule : public IModuleInt
 public:
 	static ICustomizableObjectEditorModule* Get()
 	{
+		// Prevent access to this module if the game is being played (in Standalone mode for example)
+		if (IsRunningGame())
+		{
+			return nullptr;
+		}
+		
 		return FModuleManager::LoadModulePtr<ICustomizableObjectEditorModule>("CustomizableObjectEditor");
 	}
 	
 	static ICustomizableObjectEditorModule& GetChecked()
 	{
+		check(!IsRunningGame())		// This module is editor-only. DO NOT try to access it during gameplay
 		return FModuleManager::LoadModuleChecked<ICustomizableObjectEditorModule>("CustomizableObjectEditor");
 	}
 
diff --git a/Engine/Plugins/Mutable/Source/CustomizableObject/Private/MuCO/CustomizableObject.cpp b/Engine/Plugins/Mutable/Source/CustomizableObject/Private/MuCO/CustomizableObject.cpp
--- a/Engine/Plugins/Mutable/Source/CustomizableObject/Private/MuCO/CustomizableObject.cpp
+++ b/Engine/Plugins/Mutable/Source/CustomizableObject/Private/MuCO/CustomizableObject.cpp
@@ -2379,11 +2379,6 @@ const FModelResources& UCustomizableObjectPrivate::GetModelResources(bool bIsCoo
 #if WITH_EDITOR
 bool UCustomizableObjectPrivate::IsCompilationOutOfDate(bool bSkipIndirectReferences, TArray<FName>& OutOfDatePackages, TArray<FName>& AddedPackages, TArray<FName>& RemovedPackages, bool& bReleaseVersionDiff) const
 {
-	if (IsRunningGame())
-	{
-		return false;
-	}
-
 	if (const ICustomizableObjectEditorModule* Module = ICustomizableObjectEditorModule::Get())
 	{
 		return Module->IsCompilationOutOfDate(*GetPublic(), bSkipIndirectReferences, OutOfDatePackages, AddedPackages, RemovedPackages, bReleaseVersionDiff);
diff --git a/Engine/Plugins/Mutable/Source/CustomizableObject/Private/MuCO/CustomizableObjectSystem.cpp b/Engine/Plugins/Mutable/Source/CustomizableObject/Private/MuCO/CustomizableObjectSystem.cpp
--- a/Engine/Plugins/Mutable/Source/CustomizableObject/Private/MuCO/CustomizableObjectSystem.cpp
+++ b/Engine/Plugins/Mutable/Source/CustomizableObject/Private/MuCO/CustomizableObjectSystem.cpp
@@ -1254,7 +1254,7 @@ void UCustomizableObjectSystemPrivate::EnqueueUpdateSkeletalMesh(const TSharedRe
 				Context->bOptimizedUpdate = true;
 				Context->UpdateResult = EUpdateResult::ErrorOptimized;
 				FinishUpdateGlobal(Context);			
-				return; // The the requested update is equal to the last enqueued update.
+				return; // The requested update is equal to the last enqueued update.
 			}
 		}	
 
@@ -1276,8 +1276,8 @@ void UCustomizableObjectSystemPrivate::EnqueueUpdateSkeletalMesh(const TSharedRe
 			{
 				Context->bOptimizedUpdate = true;
 
-				// The user may have changed the AttachParent and we need to recustomize it.
-				// In case nothing need to be recustomized, the update will be considered ErrorOptimized.
+				// The user may have changed the AttachParent and we need to re-customize it.
+				// In case nothing need to be re-customized, the update will be considered ErrorOptimized.
 				UpdateSkeletalMesh(Context);
 				Context->UpdateResult = Context->AttachedParentUpdated.IsEmpty() ? EUpdateResult::ErrorOptimized : EUpdateResult::Success;
 
@@ -1287,8 +1287,8 @@ void UCustomizableObjectSystemPrivate::EnqueueUpdateSkeletalMesh(const TSharedRe
 			{
 				Context->bOptimizedUpdate = false;
 
-				// The user may have changed the AttachParent and we need to recustomize it.
-				// In case nothing need to be recustomized, the update will be considered ErrorOptimized.
+				// The user may have changed the AttachParent and we need to re-customize it.
+				// In case nothing need to be re-customized, the update will be considered ErrorOptimized.
 				UpdateSkeletalMesh(Context);
 				Context->UpdateResult = EUpdateResult::Success;
 
