[Mutable] Fix Components PImpl pattern + Actor improvements.

- Created ACustomizableSkeletalMeshActorPrivate.
- Renamed CustomizableObjectInstanceUsage due to old assets having the Usage serialized (must be transient).
- Prepared Mutable Actor to use dynamic components.
- Deprecated Actor debug material functions.

diff --git a/Engine/Plugins/Mutable/Source/CustomizableObject/Internal/MuCO/CustomizableObjectInstanceUsagePrivate.h b/Engine/Plugins/Mutable/Source/CustomizableObject/Internal/MuCO/CustomizableObjectInstanceUsagePrivate.h
--- a/Engine/Plugins/Mutable/Source/CustomizableObject/Internal/MuCO/CustomizableObjectInstanceUsagePrivate.h
+++ b/Engine/Plugins/Mutable/Source/CustomizableObject/Internal/MuCO/CustomizableObjectInstanceUsagePrivate.h
@@ -55,9 +55,8 @@ public:
 	
 	bool GetPendingSetSkeletalMesh() const;
 
-	UCustomizableSkeletalComponent* GetCustomizableSkeletalComponent();
-
-	void SetCustomizableSkeletalComponent(UCustomizableSkeletalComponent* Component);
+	/** Return a valid pointer if this Usage is being used inside a UCustomizableSkeletalComponent. Nullptr if it is being used in Standalone mode. */
+	UCustomizableSkeletalComponent* GetCustomizableSkeletalComponent() const;
 	
 	UCustomizableObjectInstanceUsage* GetPublic();
 
@@ -65,4 +64,7 @@ public:
 
 	// Returns true if the NetMode of the associated UCustomizableSkeletalComponent (or the associated SkeletalMeshComponent if the former does not exist) is equal to InNetMode
 	bool IsNetMode(ENetMode InNetMode) const;
+	
+	// Used to replace the SkeletalMesh of the parent component by the ReferenceSkeletalMesh or the generated SkeletalMesh 
+	bool bUsedPendingSetSkeletalMesh = false;
 };
diff --git a/Engine/Plugins/Mutable/Source/CustomizableObject/Internal/MuCO/CustomizableSkeletalComponentPrivate.h b/Engine/Plugins/Mutable/Source/CustomizableObject/Internal/MuCO/CustomizableSkeletalComponentPrivate.h
--- a/Engine/Plugins/Mutable/Source/CustomizableObject/Internal/MuCO/CustomizableSkeletalComponentPrivate.h
+++ b/Engine/Plugins/Mutable/Source/CustomizableObject/Internal/MuCO/CustomizableSkeletalComponentPrivate.h
@@ -18,8 +18,6 @@ class CUSTOMIZABLEOBJECT_API UCustomizableSkeletalComponentPrivate : public UObj
 public:
 	UCustomizableSkeletalComponentPrivate();
 	
-	void CreateCustomizableObjectInstanceUsage();
-
 	/** Common end point of all updates. Even those which failed. */
 	void Callbacks() const;
 	
@@ -42,5 +40,8 @@ public:
 	UCustomizableSkeletalComponent* GetPublic();
 	
 	const UCustomizableSkeletalComponent* GetPublic() const;
+
+	UPROPERTY()
+	TObjectPtr<UCustomizableObjectInstanceUsage> InstanceUsage;
 };
 
diff --git a/Engine/Plugins/Mutable/Source/CustomizableObject/Internal/MuCO/CustomizableSkeletalMeshActorPrivate.h b/Engine/Plugins/Mutable/Source/CustomizableObject/Internal/MuCO/CustomizableSkeletalMeshActorPrivate.h
new file mode 100644
--- /dev/null
+++ b/Engine/Plugins/Mutable/Source/CustomizableObject/Internal/MuCO/CustomizableSkeletalMeshActorPrivate.h
@@ -0,0 +1,28 @@
+﻿// Copyright Epic Games, Inc. All Rights Reserved.
+
+#pragma once
+
+#include "UObject/Object.h"
+
+#include "UObject/ObjectPtr.h"
+#include "Containers/Array.h"
+
+#include "CustomizableSkeletalMeshActorPrivate.generated.h"
+
+class UCustomizableObjectInstance;
+class UCustomizableSkeletalComponent;
+class ACustomizableSkeletalMeshActor;
+
+
+UCLASS()
+class CUSTOMIZABLEOBJECT_API UCustomizableSkeletalMeshActorPrivate : public UObject
+{
+	GENERATED_BODY()
+
+public:
+	ACustomizableSkeletalMeshActor* GetPublic();
+
+	void Init(UCustomizableObjectInstance* Instance);
+	
+	TArray<TObjectPtr<UCustomizableSkeletalComponent>>& GetComponents();
+};
diff --git a/Engine/Plugins/Mutable/Source/CustomizableObject/Private/MuCO/CustomizableObjectInstanceUsage.cpp b/Engine/Plugins/Mutable/Source/CustomizableObject/Private/MuCO/CustomizableObjectInstanceUsage.cpp
--- a/Engine/Plugins/Mutable/Source/CustomizableObject/Private/MuCO/CustomizableObjectInstanceUsage.cpp
+++ b/Engine/Plugins/Mutable/Source/CustomizableObject/Private/MuCO/CustomizableObjectInstanceUsage.cpp
@@ -37,11 +37,11 @@ void UCustomizableObjectInstanceUsagePrivate::Callbacks() const
 		Extension->OnCustomizableObjectInstanceUsageUpdated(GetPublic());
 	}
 	
-	if (GetPublic()->CustomizableSkeletalComponent)
+	if (UCustomizableSkeletalComponent* CustomizableSkeletalComponent = GetCustomizableSkeletalComponent())
 	{
-		GetPublic()->CustomizableSkeletalComponent->UpdatedDelegate.ExecuteIfBound();
+		CustomizableSkeletalComponent->UpdatedDelegate.ExecuteIfBound();
 
-		if (GetPublic()->UpdatedDelegate.IsBound() && GetPublic()->CustomizableSkeletalComponent->UpdatedDelegate.IsBound())
+		if (GetPublic()->UpdatedDelegate.IsBound() && CustomizableSkeletalComponent->UpdatedDelegate.IsBound())
 		{
 			UE_LOG(LogMutable, Error, TEXT("The UpdatedDelegate is bound both in the UCustomizableObjectInstanceUsage and in its parent CustomizableSkeletalComponent. Only one should be bound."));
 			ensure(false);
@@ -60,7 +60,7 @@ UCustomizableObjectInstanceUsage::UCustomizableObjectInstanceUsage()
 
 void UCustomizableObjectInstanceUsage::SetCustomizableObjectInstance(UCustomizableObjectInstance* CustomizableObjectInstance)
 {
-	if (CustomizableSkeletalComponent)
+	if (UCustomizableSkeletalComponent* CustomizableSkeletalComponent = GetPrivate()->GetCustomizableSkeletalComponent())
 	{
 		CustomizableSkeletalComponent->CustomizableObjectInstance = CustomizableObjectInstance;
 	}
@@ -73,7 +73,7 @@ void UCustomizableObjectInstanceUsage::SetCustomizableObjectInstance(UCustomizab
 
 UCustomizableObjectInstance* UCustomizableObjectInstanceUsage::GetCustomizableObjectInstance() const
 {
-	if (CustomizableSkeletalComponent)
+	if (UCustomizableSkeletalComponent* CustomizableSkeletalComponent = GetPrivate()->GetCustomizableSkeletalComponent())
 	{
 		return CustomizableSkeletalComponent->CustomizableObjectInstance;
 	}
@@ -86,7 +86,7 @@ UCustomizableObjectInstance* UCustomizableObjectInstanceUsage::GetCustomizableOb
 
 void UCustomizableObjectInstanceUsage::SetComponentIndex(int32 ComponentIndex)
 {
-	if (CustomizableSkeletalComponent)
+	if (UCustomizableSkeletalComponent* CustomizableSkeletalComponent = GetPrivate()->GetCustomizableSkeletalComponent())
 	{
 		CustomizableSkeletalComponent->ComponentIndex = ComponentIndex;
 	}
@@ -99,7 +99,7 @@ void UCustomizableObjectInstanceUsage::SetComponentIndex(int32 ComponentIndex)
 
 int32 UCustomizableObjectInstanceUsage::GetComponentIndex() const
 {
-	if (CustomizableSkeletalComponent)
+	if (UCustomizableSkeletalComponent* CustomizableSkeletalComponent = GetPrivate()->GetCustomizableSkeletalComponent())
 	{
 		return CustomizableSkeletalComponent->ComponentIndex;
 	}
@@ -112,7 +112,7 @@ int32 UCustomizableObjectInstanceUsage::GetComponentIndex() const
 
 void UCustomizableObjectInstanceUsage::SetComponentName(const FName& Name)
 {
-	if (CustomizableSkeletalComponent)
+	if (UCustomizableSkeletalComponent* CustomizableSkeletalComponent = GetPrivate()->GetCustomizableSkeletalComponent())
 	{
 		CustomizableSkeletalComponent->SetComponentName(Name);
 	}
@@ -126,7 +126,7 @@ void UCustomizableObjectInstanceUsage::SetComponentName(const FName& Name)
 
 FName UCustomizableObjectInstanceUsage::GetComponentName() const
 {
-	if (CustomizableSkeletalComponent)
+	if (UCustomizableSkeletalComponent* CustomizableSkeletalComponent = GetPrivate()->GetCustomizableSkeletalComponent())
 	{
 		return CustomizableSkeletalComponent->GetComponentName();
 	}
@@ -146,42 +146,35 @@ FName UCustomizableObjectInstanceUsage::GetComponentName() const
 
 void UCustomizableObjectInstanceUsagePrivate::SetPendingSetSkeletalMesh(bool bIsActive)
 {
-	if (GetPublic()->CustomizableSkeletalComponent)
+	if (UCustomizableSkeletalComponent* CustomizableSkeletalComponent = GetCustomizableSkeletalComponent())
 	{
-		GetPublic()->CustomizableSkeletalComponent->GetPrivate()->PendingSetSkeletalMesh() = bIsActive;
+		CustomizableSkeletalComponent->GetPrivate()->PendingSetSkeletalMesh() = bIsActive;
 	}
 	else
 	{
-		GetPublic()->bUsedPendingSetSkeletalMesh = bIsActive;
+		bUsedPendingSetSkeletalMesh = bIsActive;
 	}
 }
 
 
 bool UCustomizableObjectInstanceUsagePrivate::GetPendingSetSkeletalMesh() const
 {
-	if (GetPublic()->CustomizableSkeletalComponent)
+	if (UCustomizableSkeletalComponent* CustomizableSkeletalComponent = GetCustomizableSkeletalComponent())
 	{
-		return GetPublic()->CustomizableSkeletalComponent->GetPrivate()->PendingSetSkeletalMesh();
+		return CustomizableSkeletalComponent->GetPrivate()->PendingSetSkeletalMesh();
 	}
 	else
 	{
-		return GetPublic()->bUsedPendingSetSkeletalMesh;
+		return bUsedPendingSetSkeletalMesh;
 	}
 }
 
 
-UCustomizableSkeletalComponent* UCustomizableObjectInstanceUsagePrivate::GetCustomizableSkeletalComponent()
+UCustomizableSkeletalComponent* UCustomizableObjectInstanceUsagePrivate::GetCustomizableSkeletalComponent() const
 {
-	return GetPublic()->CustomizableSkeletalComponent;
+	return Cast<UCustomizableSkeletalComponent>(GetPublic()->GetOuter()->GetOuter());
 }
 
-
-void UCustomizableObjectInstanceUsagePrivate::SetCustomizableSkeletalComponent(UCustomizableSkeletalComponent* Component)
-{
-	GetPublic()->CustomizableSkeletalComponent = Component;
-}
-
-
 UCustomizableObjectInstanceUsage* UCustomizableObjectInstanceUsagePrivate::GetPublic()
 {
 	UCustomizableObjectInstanceUsage* Public = StaticCast<UCustomizableObjectInstanceUsage*>(GetOuter());
@@ -202,7 +195,7 @@ const UCustomizableObjectInstanceUsage* UCustomizableObjectInstanceUsagePrivate:
 
 void UCustomizableObjectInstanceUsage::SetSkipSetReferenceSkeletalMesh(bool bSkip)
 {
-	if (CustomizableSkeletalComponent)
+	if (UCustomizableSkeletalComponent* CustomizableSkeletalComponent = GetPrivate()->GetCustomizableSkeletalComponent())
 	{
 		CustomizableSkeletalComponent->SetSkipSetReferenceSkeletalMesh(bSkip);
 	}
@@ -215,7 +208,7 @@ void UCustomizableObjectInstanceUsage::SetSkipSetReferenceSkeletalMesh(bool bSki
 
 bool UCustomizableObjectInstanceUsage::GetSkipSetReferenceSkeletalMesh() const
 {
-	if (CustomizableSkeletalComponent)
+	if (UCustomizableSkeletalComponent* CustomizableSkeletalComponent = GetPrivate()->GetCustomizableSkeletalComponent())
 	{
 		return CustomizableSkeletalComponent->GetSkipSetReferenceSkeletalMesh();
 	}
@@ -228,7 +221,7 @@ bool UCustomizableObjectInstanceUsage::GetSkipSetReferenceSkeletalMesh() const
 
 void UCustomizableObjectInstanceUsage::SetSkipSetSkeletalMeshOnAttach(bool bSkip)
 {
-	if (CustomizableSkeletalComponent)
+	if (UCustomizableSkeletalComponent* CustomizableSkeletalComponent = GetPrivate()->GetCustomizableSkeletalComponent())
 	{
 		CustomizableSkeletalComponent->SetSkipSetSkeletalMeshOnAttach(bSkip);
 	}
@@ -241,7 +234,7 @@ void UCustomizableObjectInstanceUsage::SetSkipSetSkeletalMeshOnAttach(bool bSkip
 
 bool UCustomizableObjectInstanceUsage::GetSkipSetSkeletalMeshOnAttach() const
 {
-	if (CustomizableSkeletalComponent)
+	if (UCustomizableSkeletalComponent* CustomizableSkeletalComponent = GetPrivate()->GetCustomizableSkeletalComponent())
 	{
 		return CustomizableSkeletalComponent->GetSkipSetReferenceSkeletalMesh();
 	}
@@ -254,7 +247,7 @@ bool UCustomizableObjectInstanceUsage::GetSkipSetSkeletalMeshOnAttach() const
 
 void UCustomizableObjectInstanceUsage::AttachTo(USkeletalMeshComponent* SkeletalMeshComponent)
 {
-	if (CustomizableSkeletalComponent)
+	if (GetPrivate()->GetCustomizableSkeletalComponent())
 	{
 		UE_LOG(LogMutable, Error, TEXT("Cannot change the attachment of a UCustomizableObjectInstanceUsage that has been automatically created by a CustomizableSkeletalComponent. Reattach the CustomizableSkeletalComponent instead."));
 		ensure(false);
@@ -278,7 +271,7 @@ void UCustomizableObjectInstanceUsage::AttachTo(USkeletalMeshComponent* Skeletal
 
 USkeletalMeshComponent* UCustomizableObjectInstanceUsage::GetAttachParent() const
 {
-	if (CustomizableSkeletalComponent)
+	if (UCustomizableSkeletalComponent* CustomizableSkeletalComponent = GetPrivate()->GetCustomizableSkeletalComponent())
 	{
 		return Cast<USkeletalMeshComponent>(CustomizableSkeletalComponent->GetAttachParent());
 	}
@@ -692,9 +685,9 @@ bool UCustomizableObjectInstanceUsagePrivate::IsTickable() const
 
 bool UCustomizableObjectInstanceUsagePrivate::IsNetMode(ENetMode InNetMode) const
 {
-	if (GetPublic()->CustomizableSkeletalComponent)
+	if (UCustomizableSkeletalComponent* CustomizableSkeletalComponent = GetCustomizableSkeletalComponent())
 	{
-		return GetPublic()->CustomizableSkeletalComponent->IsNetMode(InNetMode);
+		return CustomizableSkeletalComponent->IsNetMode(InNetMode);
 	}
 	else if (GetPublic()->UsedSkeletalMeshComponent.IsValid())
 	{
diff --git a/Engine/Plugins/Mutable/Source/CustomizableObject/Private/MuCO/CustomizableSkeletalComponent.cpp b/Engine/Plugins/Mutable/Source/CustomizableObject/Private/MuCO/CustomizableSkeletalComponent.cpp
--- a/Engine/Plugins/Mutable/Source/CustomizableObject/Private/MuCO/CustomizableSkeletalComponent.cpp
+++ b/Engine/Plugins/Mutable/Source/CustomizableObject/Private/MuCO/CustomizableSkeletalComponent.cpp
@@ -11,60 +11,9 @@
 #include "MuCO/CustomizableObject.h"
 
 
-static TAutoConsoleVariable<bool> CVarCustomizableComponentPrivateNullptr(
-	TEXT("mutable.CustomizableComponentPrivateNullptr"),
-	true,
-	TEXT("Rollback. If true, force create the private part after the serialization."),
-	ECVF_Scalability);
-
-
-UCustomizableSkeletalComponentPrivate::UCustomizableSkeletalComponentPrivate()
-{
-	// This object may get instantiated into a level, as a part of UCustomizableSkeletalComponent, so needs to be public to ensure
-	// it can be serialized out.
-	if (!HasAllFlags(RF_ClassDefaultObject))
-	{
-		SetFlags(RF_Public);
-	}
-}
-
-void UCustomizableSkeletalComponentPrivate::CreateCustomizableObjectInstanceUsage()
-{
-	if (GetPublic()->CustomizableObjectInstanceUsage)
-	{
-		// CustomizableObjectInstanceUsage may already exist if duplicated from an existing Customizable Skeletal Component
-		if (GetPublic()->CustomizableObjectInstanceUsage->GetPrivate()->GetCustomizableSkeletalComponent() != GetPublic())
-		{
-			GetPublic()->CustomizableObjectInstanceUsage = nullptr;
-		}
-	}
-
-	AActor* RootActor = GetPublic()->GetAttachmentRootActor();
-	bool bIsDefaultActor = RootActor ?
-		RootActor->HasAnyFlags(RF_ClassDefaultObject) :
-		false;
-
-	if (!GetPublic()->CustomizableObjectInstanceUsage && !HasAnyFlags(RF_ClassDefaultObject) && !bIsDefaultActor)
-	{
-		GetPublic()->CustomizableObjectInstanceUsage = NewObject<UCustomizableObjectInstanceUsage>(this, TEXT("InstanceUsage"), RF_Transient);
-		GetPublic()->CustomizableObjectInstanceUsage->GetPrivate()->SetCustomizableSkeletalComponent(GetPublic());
-	}
-}
-
-
 UCustomizableSkeletalComponent::UCustomizableSkeletalComponent()
 {
-}
-
-void UCustomizableSkeletalComponent::Serialize(FArchive& Ar)
-{
-	Super::Serialize(Ar);
-	
-	// Somehow we got some assets serialized with Private = nullptr. Fix them.
-	if (Ar.IsLoading() && !Private && CVarCustomizableComponentPrivateNullptr.GetValueOnAnyThread())
-	{
-		Private = NewObject<UCustomizableSkeletalComponentPrivate>(this, FName("Private"), RF_Public);
-	}
+	Private = CreateDefaultSubobject<UCustomizableSkeletalComponentPrivate>(TEXT("Private"));
 }
 
 
@@ -88,20 +37,26 @@ void UCustomizableSkeletalComponent::PostEditChangeProperty(FPropertyChangedEven
 #endif
 
 
+UCustomizableSkeletalComponentPrivate::UCustomizableSkeletalComponentPrivate()
+{
+	InstanceUsage = CreateDefaultSubobject<UCustomizableObjectInstanceUsage>(TEXT("InstanceUsage"));
+}
+
+
 void UCustomizableSkeletalComponentPrivate::Callbacks() const
 {
-	if (GetPublic()->CustomizableObjectInstanceUsage)
+	if (InstanceUsage)
 	{
-		GetPublic()->CustomizableObjectInstanceUsage->GetPrivate()->Callbacks();
+		InstanceUsage->GetPrivate()->Callbacks();
 	}
 }
 
 
 USkeletalMesh* UCustomizableSkeletalComponentPrivate::GetSkeletalMesh() const
 {
-	if (GetPublic()->CustomizableObjectInstanceUsage)
+	if (InstanceUsage)
 	{
-		return GetPublic()->CustomizableObjectInstanceUsage->GetPrivate()->GetSkeletalMesh();
+		return InstanceUsage->GetPrivate()->GetSkeletalMesh();
 	}
 
 	return nullptr;
@@ -110,36 +65,36 @@ USkeletalMesh* UCustomizableSkeletalComponentPrivate::GetSkeletalMesh() const
 
 void UCustomizableSkeletalComponentPrivate::SetSkeletalMesh(USkeletalMesh* SkeletalMesh)
 {
-	if (GetPublic()->CustomizableObjectInstanceUsage)
+	if (InstanceUsage)
 	{
-		GetPublic()->CustomizableObjectInstanceUsage->GetPrivate()->SetSkeletalMesh(SkeletalMesh);
+		InstanceUsage->GetPrivate()->SetSkeletalMesh(SkeletalMesh);
 	}
 }
 
 
 void UCustomizableSkeletalComponentPrivate::SetPhysicsAsset(UPhysicsAsset* PhysicsAsset)
 {
-	if (GetPublic()->CustomizableObjectInstanceUsage)
+	if (InstanceUsage)
 	{
-		GetPublic()->CustomizableObjectInstanceUsage->GetPrivate()->SetPhysicsAsset(PhysicsAsset);
+		InstanceUsage->GetPrivate()->SetPhysicsAsset(PhysicsAsset);
 	}
 }
 
 
 void UCustomizableSkeletalComponentPrivate::SetPendingSetSkeletalMesh(bool bIsActive)
 {
-	if (GetPublic()->CustomizableObjectInstanceUsage)
+	if (InstanceUsage)
 	{
-		GetPublic()->CustomizableObjectInstanceUsage->GetPrivate()->SetPendingSetSkeletalMesh(bIsActive);
+		InstanceUsage->GetPrivate()->SetPendingSetSkeletalMesh(bIsActive);
 	}
 }
 
 
 USkeletalMesh* UCustomizableSkeletalComponentPrivate::GetAttachedSkeletalMesh() const
 {
-	if (GetPublic()->CustomizableObjectInstanceUsage)
+	if (InstanceUsage)
 	{
-		return GetPublic()->CustomizableObjectInstanceUsage->GetPrivate()->GetAttachedSkeletalMesh();
+		return InstanceUsage->GetPrivate()->GetAttachedSkeletalMesh();
 	}
 
 	return nullptr;
@@ -204,18 +159,18 @@ bool UCustomizableSkeletalComponent::GetSkipSetSkeletalMeshOnAttach() const
 
 void UCustomizableSkeletalComponent::UpdateSkeletalMeshAsync(bool bNeverSkipUpdate)
 {
-	if (CustomizableObjectInstanceUsage)
+	if (GetPrivate()->InstanceUsage)
 	{
-		CustomizableObjectInstanceUsage->UpdateSkeletalMeshAsync(bNeverSkipUpdate);
+		GetPrivate()->InstanceUsage->UpdateSkeletalMeshAsync(bNeverSkipUpdate);
 	}
 }
 
 
 void UCustomizableSkeletalComponent::UpdateSkeletalMeshAsyncResult(FInstanceUpdateDelegate Callback, bool bIgnoreCloseDist, bool bForceHighPriority)
 {
-	if (CustomizableObjectInstanceUsage)
+	if (GetPrivate()->InstanceUsage)
 	{
-		CustomizableObjectInstanceUsage->UpdateSkeletalMeshAsyncResult(Callback, bIgnoreCloseDist, bForceHighPriority);
+		GetPrivate()->InstanceUsage->UpdateSkeletalMeshAsyncResult(Callback, bIgnoreCloseDist, bForceHighPriority);
 	}
 }
 
@@ -237,9 +192,9 @@ const UCustomizableSkeletalComponentPrivate* UCustomizableSkeletalComponent::Get
 #if WITH_EDITOR
 void UCustomizableSkeletalComponentPrivate::EditorUpdateComponent()
 {
-	if (GetPublic()->CustomizableObjectInstanceUsage)
+	if (InstanceUsage)
 	{
-		GetPublic()->CustomizableObjectInstanceUsage->GetPrivate()->EditorUpdateComponent();
+		InstanceUsage->GetPrivate()->EditorUpdateComponent();
 	}
 }
 #endif
@@ -283,22 +238,3 @@ void UCustomizableSkeletalComponent::OnAttachmentChanged()
 	}
 }
 
-
-void UCustomizableSkeletalComponent::PostInitProperties()
-{
-	Super::PostInitProperties();
-
-	if (!HasAllFlags(RF_ClassDefaultObject))
-	{
-		if (!Private)
-		{
-			Private = NewObject<UCustomizableSkeletalComponentPrivate>(this, FName("Private"), RF_Public);
-		}
-		else if (Private->GetOuter() != this)
-		{
-			Private = Cast<UCustomizableSkeletalComponentPrivate>(StaticDuplicateObject(Private, this, FName("Private")));
-		}
-	
-		GetPrivate()->CreateCustomizableObjectInstanceUsage();
-	}
-}
diff --git a/Engine/Plugins/Mutable/Source/CustomizableObject/Private/MuCO/CustomizableSkeletalMeshActor.cpp b/Engine/Plugins/Mutable/Source/CustomizableObject/Private/MuCO/CustomizableSkeletalMeshActor.cpp
--- a/Engine/Plugins/Mutable/Source/CustomizableObject/Private/MuCO/CustomizableSkeletalMeshActor.cpp
+++ b/Engine/Plugins/Mutable/Source/CustomizableObject/Private/MuCO/CustomizableSkeletalMeshActor.cpp
@@ -1,75 +1,112 @@
 // Copyright Epic Games, Inc. All Rights Reserved.
 
 #include "MuCO/CustomizableSkeletalMeshActor.h"
+#include "MuCO/CustomizableSkeletalMeshActorPrivate.h"
 
 #include "Components/SkeletalMeshComponent.h"
 #include "MuCO/CustomizableSkeletalComponent.h"
 #include "MuCO/CustomizableObject.h"
 #include "MuCO/CustomizableObjectPrivate.h"
 #include "MuCO/CustomizableObjectSystem.h"
+#include "MuCO/CustomizableSkeletalComponentPrivate.h"
+#include "Engine/SkeletalMesh.h"
+#include "MuCO/CustomizableObjectCustomVersion.h"
+#include "MuCO/CustomizableObjectPrivate.h"
+#include "MuCO/UnrealPortabilityHelpers.h"
+
 
 #include UE_INLINE_GENERATED_CPP_BY_NAME(CustomizableSkeletalMeshActor)
 
 #define LOCTEXT_NAMESPACE "CustomizableObject"
 
 
-ACustomizableSkeletalMeshActor::ACustomizableSkeletalMeshActor(const FObjectInitializer& ObjectInitializer)
-	: Super(ObjectInitializer)
+UCustomizableObjectInstance* ACustomizableSkeletalMeshActor::GetInstance()
 {
-	CustomizableSkeletalComponent = CreateDefaultSubobject<UCustomizableSkeletalComponent>(TEXT("CustomizableSkeletalComponent0"));
-	CustomizableSkeletalComponents.Add(CustomizableSkeletalComponent);
-	if (USkeletalMeshComponent* SkeletalMeshComp = GetSkeletalMeshComponent()) 
+	for (const UCustomizableSkeletalComponent* Component : CustomizableSkeletalComponents)
 	{
-		SkeletalMeshComponents.Add(SkeletalMeshComp);
-		bool Success = CustomizableSkeletalComponents[0]->AttachToComponent(SkeletalMeshComp, FAttachmentTransformRules::KeepRelativeTransform);
+		if (UCustomizableObjectInstance* COInstance = Component->CustomizableObjectInstance)
+		{
+			return COInstance;
+		}
 	}
+
+	return nullptr;
 }
 
 
-void ACustomizableSkeletalMeshActor::AttachNewComponent()
+ACustomizableSkeletalMeshActor::ACustomizableSkeletalMeshActor(const FObjectInitializer& Initializer)
 {
-	int32 CurrentIndex = CustomizableSkeletalComponents.Num();
-	FString CustomizableComponentName = FString::Printf(TEXT("CustomizableSkeletalComponent%d"), CurrentIndex);
-	FString SkeletalMeshComponentName = FString::Printf(TEXT("SkeletalMeshComponent%d"), CurrentIndex);
+	Private = CreateDefaultSubobject<UCustomizableSkeletalMeshActorPrivate>(TEXT("Private"));
 
-	//USkeletalMeshComponent* SkeletalMesh = CreateDefaultSubobject<USkeletalMeshComponent>(FName(*SkeletalMeshComponentName));
-	USkeletalMeshComponent* SkeletalMesh = NewObject<USkeletalMeshComponent>(this, USkeletalMeshComponent::StaticClass(), FName(*SkeletalMeshComponentName));
-	USkeletalMeshComponent* RootSkeletalMesh = GetSkeletalMeshComponent();
-	
-	if (SkeletalMesh && RootSkeletalMesh)
-	{		
-		bool Success = SkeletalMesh->AttachToComponent(RootSkeletalMesh, FAttachmentTransformRules::KeepRelativeTransform);
-		if (Success)
+	// Old assets used to create the first UCustomizableSkeletalComponent as a Subobject. To be able to deserialize them we need to create it.
+	// Creating instead an Object instead of a Suobject will not work. Only Component UCustomizableSkeletalComponent 0 is a Subobject.
+	{
+		UCustomizableSkeletalComponent* CustomizableSkeletalComponent = CreateDefaultSubobject<UCustomizableSkeletalComponent>(TEXT("CustomizableSkeletalComponent0"));
+		CustomizableSkeletalComponents.Add(CustomizableSkeletalComponent);
+		if (USkeletalMeshComponent* SkeletalMeshComp = GetSkeletalMeshComponent()) 
 		{
-			//UCustomizableSkeletalComponent* CustomizableSkeletalComponent = CreateDefaultSubobject<UCustomizableSkeletalComponent>(FName(*CustomizableComponentName));
-			UCustomizableSkeletalComponent* NewCustomizableSkeletalComponent = NewObject<UCustomizableSkeletalComponent>(this, UCustomizableSkeletalComponent::StaticClass(), FName(*CustomizableComponentName));
-
-			if (NewCustomizableSkeletalComponent)
-			{
-				Success = NewCustomizableSkeletalComponent->AttachToComponent(SkeletalMesh, FAttachmentTransformRules::KeepRelativeTransform);
-
-				if (Success)
-				{
-					SkeletalMeshComponents.Add(SkeletalMesh);
-					CustomizableSkeletalComponents.Add(NewCustomizableSkeletalComponent);
-				}
-			}
+			SkeletalMeshComponents.Add(SkeletalMeshComp);
+			CustomizableSkeletalComponents[0]->AttachToComponent(SkeletalMeshComp, FAttachmentTransformRules::KeepRelativeTransform);
 		}
 	}
 }
 
 
-UCustomizableObjectInstance* ACustomizableSkeletalMeshActor::GetComponentsCommonInstance()
+void UCustomizableSkeletalMeshActorPrivate::Init(UCustomizableObjectInstance* Instance)
 {
-	for (const UCustomizableSkeletalComponent* Component : CustomizableSkeletalComponents)
+	UCustomizableObject* Object = Instance->GetCustomizableObject();
+	if (!Object)
 	{
-		if (UCustomizableObjectInstance* COInstance = Component->CustomizableObjectInstance)
+		return;
+	}
+	
+	for (int32 ObjectComponentIndex = 0; ObjectComponentIndex < Object->GetComponentCount(); ++ObjectComponentIndex)
+	{
+		const FName& ComponentName = Object->GetComponentName(ObjectComponentIndex);
+		const FString ComponentNameString = ComponentName.ToString();
+
+		USkeletalMeshComponent* Component;
+		if (ObjectComponentIndex == 0)
 		{
-			return COInstance;
+			Component = GetPublic()->SkeletalMeshComponents[0];
+			// Renaming the component will make it disappear. Keep it as it is.
+		}
+		else
+		{
+			const FString SkeletalMeshComponentName = FString::Printf(TEXT("SkeletalMeshComponent %s"), *ComponentNameString);
+			Component = NewObject<USkeletalMeshComponent>(GetPublic(), USkeletalMeshComponent::StaticClass(), FName(*SkeletalMeshComponentName));
+			Component->CreationMethod = EComponentCreationMethod::Native;
+			Component->AttachToComponent(GetPublic()->GetRootComponent(), FAttachmentTransformRules::KeepRelativeTransform);
+			Component->RegisterComponent();
+		}
+					
+		if (!Component)
+		{
+			continue;
 		}
+
+		UCustomizableSkeletalComponent* CustomizableComponent;
+		if (ObjectComponentIndex == 0)
+		{
+			CustomizableComponent = GetPublic()->CustomizableSkeletalComponents[0];
+			// Renaming the component will make it disappear. Keep it as it is.
+		}
+		else
+		{
+			const FString CustomizableComponentName = FString::Printf(TEXT("CustomizableSkeletalComponent %s"), *ComponentNameString);
+			CustomizableComponent = NewObject<UCustomizableSkeletalComponent>(GetPublic(), UCustomizableSkeletalComponent::StaticClass(), FName(CustomizableComponentName));
+		}
+
+		CustomizableComponent->AttachToComponent(Component, FAttachmentTransformRules::KeepRelativeTransform);
+		CustomizableComponent->CustomizableObjectInstance = Instance;
+		CustomizableComponent->SetComponentName(ComponentName);
+		CustomizableComponent->RegisterComponent();
+		
+		GetPublic()->SkeletalMeshComponents.Add(Component);
+		GetPublic()->CustomizableSkeletalComponents.Add(CustomizableComponent);
 	}
 
-	return nullptr;
+	Instance->UpdateSkeletalMeshAsync();
 }
 
 
@@ -89,7 +126,7 @@ void ACustomizableSkeletalMeshActor::EnableDebugMaterial(bool bEnableDebugMateri
 	bRemoveDebugMaterial = bDebugMaterialEnabled && !bEnableDebugMaterial;
 	bDebugMaterialEnabled = bEnableDebugMaterial;
 
-	if (UCustomizableObjectInstance* COInstance = GetComponentsCommonInstance())
+	if (UCustomizableObjectInstance* COInstance = GetInstance())
 	{
 		//Bind Instance Update delegate to Actor
 		COInstance->UpdatedDelegate.AddUniqueDynamic(this, &ACustomizableSkeletalMeshActor::SwitchComponentsMaterials);
@@ -98,6 +135,20 @@ void ACustomizableSkeletalMeshActor::EnableDebugMaterial(bool bEnableDebugMateri
 }
 
 
+UCustomizableSkeletalMeshActorPrivate* ACustomizableSkeletalMeshActor::GetPrivate()
+{
+	check(Private);
+	return Private;
+}
+
+
+const UCustomizableSkeletalMeshActorPrivate* ACustomizableSkeletalMeshActor::GetPrivate() const
+{
+	check(Private);
+	return Private;
+}
+
+
 void ACustomizableSkeletalMeshActor::SwitchComponentsMaterials(UCustomizableObjectInstance* Instance)
 {
 	if (!DebugMaterial)
@@ -107,7 +158,7 @@ void ACustomizableSkeletalMeshActor::SwitchComponentsMaterials(UCustomizableObje
 
 	if (bDebugMaterialEnabled || bRemoveDebugMaterial)
 	{
-		UCustomizableObjectInstance* COInstance = GetComponentsCommonInstance();
+		UCustomizableObjectInstance* COInstance = GetInstance();
 
 		if (!COInstance)
 		{
@@ -149,4 +200,16 @@ void ACustomizableSkeletalMeshActor::SwitchComponentsMaterials(UCustomizableObje
 	}
 }
 
+ACustomizableSkeletalMeshActor* UCustomizableSkeletalMeshActorPrivate::GetPublic()
+{
+	return CastChecked<ACustomizableSkeletalMeshActor>(GetOuter());
+}
+
+
+TArray<TObjectPtr<UCustomizableSkeletalComponent>>& UCustomizableSkeletalMeshActorPrivate::GetComponents()
+{
+	return GetPublic()->CustomizableSkeletalComponents;
+}
+
+
 #undef LOCTEXT_NAMESPACE
diff --git a/Engine/Plugins/Mutable/Source/CustomizableObject/Public/MuCO/CustomizableObjectInstanceUsage.h b/Engine/Plugins/Mutable/Source/CustomizableObject/Public/MuCO/CustomizableObjectInstanceUsage.h
--- a/Engine/Plugins/Mutable/Source/CustomizableObject/Public/MuCO/CustomizableObjectInstanceUsage.h
+++ b/Engine/Plugins/Mutable/Source/CustomizableObject/Public/MuCO/CustomizableObjectInstanceUsage.h
@@ -91,19 +91,15 @@ public:
 	FCustomizableObjectInstanceUsageUpdatedDelegate UpdatedDelegate;
 
 private:
-	// If this CustomizableSkeletalComponent is not null, it means this Usage was created by it, and all persistent properties should be obtained through it
-	UPROPERTY()
-	TObjectPtr<UCustomizableSkeletalComponent> CustomizableSkeletalComponent;
-
-	// If no CustomizableSkeletalComponent is associated, this SkeletalComponent will be used
+	// If the outer is no CustomizableSkeletalComponent, this SkeletalComponent will be used
 	UPROPERTY()
 	TWeakObjectPtr<USkeletalMeshComponent> UsedSkeletalMeshComponent;
 
-	// If no CustomizableSkeletalComponent is associated, this Instance will be used
+	// If the outer is no CustomizableSkeletalComponentd, this Instance will be used
 	UPROPERTY()
 	TObjectPtr<UCustomizableObjectInstance> UsedCustomizableObjectInstance;
 
-	// If no CustomizableSkeletalComponent is associated, this Index will be used
+	// If the outer is no CustomizableSkeletalComponentd, this Index will be used
 	// DEPRECATED
 	UPROPERTY()
 	int32 UsedComponentIndex;
@@ -111,9 +107,6 @@ private:
 	/** Only used if the UsedComponentIndex is INDEX_NONE. */
 	UPROPERTY()
 	FName UsedComponentName;
-	
-	// Used to replace the SkeletalMesh of the parent component by the ReferenceSkeletalMesh or the generated SkeletalMesh 
-	bool bUsedPendingSetSkeletalMesh = false;
 
 	// Used to avoid replacing the SkeletalMesh of the parent component by the ReferenceSkeletalMesh if bPendingSetSkeletalMesh is true
 	UPROPERTY()
@@ -122,7 +115,7 @@ private:
 	UPROPERTY()
 	bool bUsedSkipSetSkeletalMeshOnAttach = false;
 	
-	UPROPERTY()
+	UPROPERTY( VisibleAnywhere, Category=CO)
 	TObjectPtr<UCustomizableObjectInstanceUsagePrivate> Private;
 	
 	friend UCustomizableObjectInstanceUsagePrivate;
diff --git a/Engine/Plugins/Mutable/Source/CustomizableObject/Public/MuCO/CustomizableSkeletalComponent.h b/Engine/Plugins/Mutable/Source/CustomizableObject/Public/MuCO/CustomizableSkeletalComponent.h
--- a/Engine/Plugins/Mutable/Source/CustomizableObject/Public/MuCO/CustomizableSkeletalComponent.h
+++ b/Engine/Plugins/Mutable/Source/CustomizableObject/Public/MuCO/CustomizableSkeletalComponent.h
@@ -49,9 +49,6 @@ private:
 
 	UPROPERTY()
 	bool bSkipSkipSetSkeletalMeshOnAttach = false;
-
-	UPROPERTY(Transient)
-	TObjectPtr<UCustomizableObjectInstanceUsage> CustomizableObjectInstanceUsage;
 	
 	UPROPERTY()
 	TObjectPtr<UCustomizableSkeletalComponentPrivate> Private;
@@ -64,10 +61,6 @@ protected:
 #if WITH_EDITOR
 	virtual void PostEditChangeProperty(FPropertyChangedEvent& PropertyChangedEvent) override;
 #endif
-	virtual void Serialize(FArchive& Ar) override;
-
-	// UObject interface
-	virtual void PostInitProperties() override;
 
 	// USceneComponent interface
 	virtual void OnAttachmentChanged() override;
diff --git a/Engine/Plugins/Mutable/Source/CustomizableObject/Public/MuCO/CustomizableSkeletalMeshActor.h b/Engine/Plugins/Mutable/Source/CustomizableObject/Public/MuCO/CustomizableSkeletalMeshActor.h
--- a/Engine/Plugins/Mutable/Source/CustomizableObject/Public/MuCO/CustomizableSkeletalMeshActor.h
+++ b/Engine/Plugins/Mutable/Source/CustomizableObject/Public/MuCO/CustomizableSkeletalMeshActor.h
@@ -8,60 +8,59 @@
 
 class UCustomizableObjectInstance;
 class UCustomizableSkeletalComponent;
+class UCustomizableSkeletalMeshActorPrivate;
 
 class UObject;
 
 UCLASS(ClassGroup = CustomizableObject, Blueprintable, ComponentWrapperClass, ConversionRoot, meta = (ChildCanTick))
 class CUSTOMIZABLEOBJECT_API ACustomizableSkeletalMeshActor : public ASkeletalMeshActor
 {
-	GENERATED_UCLASS_BODY()
-public:
-	
-	class UCustomizableSkeletalComponent* GetCustomizableSkeletalComponent(int32 ObjectComponentIndex = 0) const { return CustomizableSkeletalComponents[ObjectComponentIndex]; }
-	
-	class USkeletalMeshComponent* GetSkeletalMeshComponentAt(int32 ObjectComponentIndex = 0) const { return SkeletalMeshComponents[ObjectComponentIndex]; }
-
-	void AttachNewComponent();
+	GENERATED_BODY()
 
-	int32 GetNumComponents() { return CustomizableSkeletalComponents.Num(); }
+	friend UCustomizableSkeletalMeshActorPrivate;
 
-	/** Sets which material will be used for debugging purposes */
-	UFUNCTION(BlueprintCallable, Category = CustomizableSkeletalMeshActor)
+public:
+	ACustomizableSkeletalMeshActor(const FObjectInitializer& Initializer);
+	
+	UE_DEPRECATED(5.6, "Plase replace with custom code.")
+	UFUNCTION(BlueprintCallable, Category = CustomizableSkeletalMeshActor, meta=(DeprecatedFunction))
 	void SetDebugMaterial(UMaterialInterface* InDebugMaterial);
 
-	/** Enables the debug material to all the componets of the actor */
-	UFUNCTION(BlueprintCallable, Category = CustomizableSkeletalMeshActor)
+	UE_DEPRECATED(5.6, "Plase replace with custom code.")
+	UFUNCTION(BlueprintCallable, Category = CustomizableSkeletalMeshActor, meta=(DeprecatedFunction))
 	void EnableDebugMaterial(bool bEnableDebugMaterial);
 
-private:
+	UCustomizableSkeletalMeshActorPrivate* GetPrivate();
+	const UCustomizableSkeletalMeshActorPrivate* GetPrivate() const;
 
-	/** Changes materials to all actor componets */
+private:
+	void Init(UCustomizableObjectInstance* Instance);
+	
+	// DEPRECATED
 	UFUNCTION()
 	void SwitchComponentsMaterials(UCustomizableObjectInstance* Instance);
 
-	/** return: the instance represented in all the actor components */
-	UCustomizableObjectInstance* GetComponentsCommonInstance();
+	/** @return the instance represented in all the actor components. */
+	UCustomizableObjectInstance* GetInstance();
 
-private:
-	
 	/** TODO: There are as many as components in the CO (not in the COI). */
-	UPROPERTY(Category = CustomizableSkeletalMesh, VisibleAnywhere, BlueprintReadOnly, meta = (AllowPrivateAccess = "true"))
-	TArray< TObjectPtr<class UCustomizableSkeletalComponent> > CustomizableSkeletalComponents;
+	UPROPERTY(VisibleAnywhere, Category = CustomizableSkeletalMeshActor)
+	TArray<TObjectPtr<UCustomizableSkeletalComponent>> CustomizableSkeletalComponents;
 
 	/** TODO: There are as many as components in the CO (not in the COI). */
-	UPROPERTY(Category = CustomizableSkeletalMesh, VisibleAnywhere, BlueprintReadOnly, meta = (AllowPrivateAccess = "true"))
-	TArray< TObjectPtr<class USkeletalMeshComponent>> SkeletalMeshComponents;
-
-	// TODO: This is a temporary fix to not break the demos, we should update the demos to support the arrays of components
-	UPROPERTY(Category = CustomizableSkeletalMesh, VisibleAnywhere, BlueprintReadOnly, meta = (AllowPrivateAccess = "true"))
-	TObjectPtr<class UCustomizableSkeletalComponent> CustomizableSkeletalComponent;
+	UPROPERTY(VisibleAnywhere, Category = CustomizableSkeletalMeshActor)
+	TArray<TObjectPtr<USkeletalMeshComponent>> SkeletalMeshComponents;
 
-	/** Material used for debugging */
+	// DEPRECATED
+	UPROPERTY()
 	TObjectPtr<UMaterialInterface> DebugMaterial;
 
-	/** bool that determines if the debug material has been enabled */
+	// DEPRECATED
 	bool bDebugMaterialEnabled = false;
 
-	/** bool used to remove the debug material whet it has been disabled */
+	// DEPRECATED
 	bool bRemoveDebugMaterial = false;
+
+	UPROPERTY()
+	TObjectPtr<UCustomizableSkeletalMeshActorPrivate> Private = nullptr;
 };
diff --git a/Engine/Plugins/Mutable/Source/CustomizableObjectEditor/Private/MuCOE/CustomizableObjectInstanceFactory.cpp b/Engine/Plugins/Mutable/Source/CustomizableObjectEditor/Private/MuCOE/CustomizableObjectInstanceFactory.cpp
--- a/Engine/Plugins/Mutable/Source/CustomizableObjectEditor/Private/MuCOE/CustomizableObjectInstanceFactory.cpp
+++ b/Engine/Plugins/Mutable/Source/CustomizableObjectEditor/Private/MuCOE/CustomizableObjectInstanceFactory.cpp
@@ -12,6 +12,7 @@
 #include "MuCO/CustomizableObjectPrivate.h"
 #include "MuCO/CustomizableSkeletalComponent.h"
 #include "MuCO/CustomizableSkeletalMeshActor.h"
+#include "MuCO/CustomizableSkeletalMeshActorPrivate.h"
 #include "MuCOE/Nodes/CustomizableObjectNodeObject.h"
 #include "MuCO/UnrealPortabilityHelpers.h"
 
@@ -48,63 +49,26 @@ void UCustomizableObjectInstanceFactory::PostSpawnActor(UObject* Asset, AActor*
 	{
 		return;
 	}
-	
-
-	for (int32 ObjectComponentIndex = 0; ObjectComponentIndex < Object->GetComponentCount(); ++ObjectComponentIndex)
-	{
-		const FName& ComponentName = Object->GetPrivate()->GetComponentName(FCustomizableObjectComponentIndex(ObjectComponentIndex));
-		
-		USkeletalMesh* SkeletalMesh = Instance->GetComponentMeshSkeletalMesh(ComponentName);
-		if (!SkeletalMesh)
-		{
-			SkeletalMesh = Object->GetComponentMeshReferenceSkeletalMesh(ComponentName);
-		}
-		
-		if (ObjectComponentIndex > 0 && ObjectComponentIndex > NewCSMActor->GetNumComponents() - 1)
-		{
-			NewCSMActor->AttachNewComponent();
-		}
-
-		if (USkeletalMeshComponent* SkeletalMeshComp = NewCSMActor->GetSkeletalMeshComponentAt(ObjectComponentIndex))
-		{
-			SkeletalMeshComp->UnregisterComponent();
-			UE_MUTABLE_SETSKINNEDASSET(SkeletalMeshComp, SkeletalMesh);
 
-			if (ObjectComponentIndex == 0 && NewCSMActor->GetWorld()->IsGameWorld())
-			{
-				NewCSMActor->ReplicatedMesh = SkeletalMesh;
-			}
-
-			if (UCustomizableSkeletalComponent* CustomSkeletalComp = NewCSMActor->GetCustomizableSkeletalComponent(ObjectComponentIndex))
-			{
-				CustomSkeletalComp->UnregisterComponent();
-				CustomSkeletalComp->CustomizableObjectInstance = Instance;
-				CustomSkeletalComp->SetComponentName(ComponentName);
-				CustomSkeletalComp->GetPrivate()->SetSkeletalMesh(SkeletalMesh);
-				CustomSkeletalComp->UpdateSkeletalMeshAsync();
-				CustomSkeletalComp->RegisterComponent();
-			}
-
-			SkeletalMeshComp->RegisterComponent();
-		}
-	}
+	NewCSMActor->GetPrivate()->Init(Instance);
 }
 
+
 UObject* UCustomizableObjectInstanceFactory::GetAssetFromActorInstance(AActor* ActorInstance)
 {
 	if (ACustomizableSkeletalMeshActor* CSMActor = CastChecked<ACustomizableSkeletalMeshActor>(ActorInstance))
 	{
-		if (CSMActor->GetNumComponents() > 0)
+		const TArray<TObjectPtr<UCustomizableSkeletalComponent>>& Components = CSMActor->GetPrivate()->GetComponents();
+		if (Components.IsValidIndex(0))
 		{
-			if (UCustomizableSkeletalComponent* CustomSkeletalComp = CSMActor->GetCustomizableSkeletalComponent(0))
-			{
-				return CustomSkeletalComp->CustomizableObjectInstance;
-			}
+			return Components[0];
 		}
 	}
+	
 	return nullptr;
 }
 
+
 bool UCustomizableObjectInstanceFactory::CanCreateActorFrom(const FAssetData& AssetData, FText& OutErrorMsg)
 {
     if (!AssetData.IsValid() ||
diff --git a/Engine/Plugins/Mutable/Source/CustomizableObjectEditor/Public/MuCOE/CustomizableObjectInstanceFactory.h b/Engine/Plugins/Mutable/Source/CustomizableObjectEditor/Public/MuCOE/CustomizableObjectInstanceFactory.h
--- a/Engine/Plugins/Mutable/Source/CustomizableObjectEditor/Public/MuCOE/CustomizableObjectInstanceFactory.h
+++ b/Engine/Plugins/Mutable/Source/CustomizableObjectEditor/Public/MuCOE/CustomizableObjectInstanceFactory.h
@@ -17,11 +17,9 @@ class UCustomizableObjectInstanceFactory : public UActorFactory
 	GENERATED_UCLASS_BODY()
 
 protected:
-	//~ Begin UActorFactory Interface
-	void PostSpawnActor(UObject* Asset, AActor* NewActor) override;
-	UObject* GetAssetFromActorInstance(AActor* ActorInstance) override;
-	bool CanCreateActorFrom(const FAssetData& AssetData, FText& OutErrorMsg) override;
-	//~ End UActorFactory Interface
-
-	FQuat AlignObjectToSurfaceNormal(const FVector& InSurfaceNormal, const FQuat& ActorRotation) const;
+	// UActorFactory interface
+	virtual void PostSpawnActor(UObject* Asset, AActor* NewActor) override;
+	virtual UObject* GetAssetFromActorInstance(AActor* ActorInstance) override;
+	virtual bool CanCreateActorFrom(const FAssetData& AssetData, FText& OutErrorMsg) override;
+	virtual FQuat AlignObjectToSurfaceNormal(const FVector& InSurfaceNormal, const FQuat& ActorRotation) const override;
 };
