[Mutable] Fix Internal include headers.

#rnx

diff --git a/Engine/Plugins/Mutable/Source/CustomizableObject/CustomizableObject.Build.cs b/Engine/Plugins/Mutable/Source/CustomizableObject/CustomizableObject.Build.cs
--- a/Engine/Plugins/Mutable/Source/CustomizableObject/CustomizableObject.Build.cs
+++ b/Engine/Plugins/Mutable/Source/CustomizableObject/CustomizableObject.Build.cs
@@ -13,9 +13,6 @@ public class CustomizableObject : ModuleRules
 		bAllowConfidentialPlatformDefines = true;
 		//bUseUnity = false;
 
-		// TODO: CustomizableObject/Public/MuCO/CustomizableObject.h is including CustomizableObject/Internal/MuCO/CustomizableObjectCompilerTypes.h
-		PublicIncludePaths.Add(System.IO.Path.Combine(ModuleDirectory, "Internal"));
-
 		PrivateDependencyModuleNames.AddRange(new string[] {
             "InputCore",
             "SlateCore",
diff --git a/Engine/Plugins/Mutable/Source/CustomizableObject/Internal/MuCO/CustomizableObjectCompilerTypes.h b/Engine/Plugins/Mutable/Source/CustomizableObject/Internal/MuCO/CustomizableObjectCompilerTypes.h
--- a/Engine/Plugins/Mutable/Source/CustomizableObject/Internal/MuCO/CustomizableObjectCompilerTypes.h
+++ b/Engine/Plugins/Mutable/Source/CustomizableObject/Internal/MuCO/CustomizableObjectCompilerTypes.h
@@ -2,6 +2,7 @@
 
 #pragma once
 
+#include "MuCO/CustomizableObjectSystem.h"
 #include "UObject/WeakObjectPtr.h"
 
 #if WITH_EDITOR
@@ -9,40 +10,11 @@
 #include "DerivedDataCacheKey.h"
 #endif
 
-#include "CustomizableObjectCompilerTypes.generated.h"
-
-
 /** Index of the maximum optimization level when compiling CustomizableObjects */
 #define UE_MUTABLE_MAX_OPTIMIZATION			2
 
 class UCustomizableObject;
 
-DECLARE_MULTICAST_DELEGATE(FPostCompileDelegate)
-
-UENUM()
-enum class ECustomizableObjectNumBoneInfluences : uint8
-{
-	// The enum values can be used as the real numeric value of number of bone influences
-	Four = 4,
-	// 
-	Eight = 8,
-	//
-	Twelve = 12 // This is essentially the same as "Unlimited", but UE ultimately limits to 12
-};
-
-//#if WITH_EDITORONLY_DATA
-
-UENUM()
-enum class ECustomizableObjectTextureCompression : uint8
-{
-	// Don't use texture compression
-	None = 0,
-	// Use Mutable's fast low-quality compression
-	Fast,
-	// Use Unreal's highest quality compression (100x slower to compress)
-	HighQuality
-};
-
 
 struct FCompilationOptions
 {
diff --git a/Engine/Plugins/Mutable/Source/CustomizableObject/Internal/MuCO/CustomizableObjectInstanceDescriptor.h b/Engine/Plugins/Mutable/Source/CustomizableObject/Internal/MuCO/CustomizableObjectInstanceDescriptor.h
deleted file mode 100644
diff --git a/Engine/Plugins/Mutable/Source/CustomizableObject/Internal/MuCO/CustomizableObjectPrivate.h b/Engine/Plugins/Mutable/Source/CustomizableObject/Internal/MuCO/CustomizableObjectPrivate.h
--- a/Engine/Plugins/Mutable/Source/CustomizableObject/Internal/MuCO/CustomizableObjectPrivate.h
+++ b/Engine/Plugins/Mutable/Source/CustomizableObject/Internal/MuCO/CustomizableObjectPrivate.h
@@ -4,6 +4,7 @@
 
 #include "MuCO/CustomizableObject.h"
 #include "MuCO/CustomizableObjectCompilerTypes.h"
+#include "MuCO/CustomizableObjectClothingTypes.h"
 #include "MuCO/CustomizableObjectStreamedResourceData.h"
 #include "MuCO/StateMachine.h"
 #include "MuCO/CustomizableObjectUIData.h"
diff --git a/Engine/Plugins/Mutable/Source/CustomizableObject/Private/MuCO/CustomizableObject.cpp b/Engine/Plugins/Mutable/Source/CustomizableObject/Private/MuCO/CustomizableObject.cpp
--- a/Engine/Plugins/Mutable/Source/CustomizableObject/Private/MuCO/CustomizableObject.cpp
+++ b/Engine/Plugins/Mutable/Source/CustomizableObject/Private/MuCO/CustomizableObject.cpp
@@ -54,6 +54,7 @@
 
 
 #include "MuCO/CustomizableObjectCustomVersion.h"
+#include "MuCO/CustomizableObjectResourceDataTypes.h"
 
 #include UE_INLINE_GENERATED_CPP_BY_NAME(CustomizableObject)
 
diff --git a/Engine/Plugins/Mutable/Source/CustomizableObject/Private/MuCO/CustomizableObjectInstance.cpp b/Engine/Plugins/Mutable/Source/CustomizableObject/Private/MuCO/CustomizableObjectInstance.cpp
--- a/Engine/Plugins/Mutable/Source/CustomizableObject/Private/MuCO/CustomizableObjectInstance.cpp
+++ b/Engine/Plugins/Mutable/Source/CustomizableObject/Private/MuCO/CustomizableObjectInstance.cpp
@@ -53,6 +53,7 @@
 #include "PhysicsEngine/SkeletalBodySetup.h"
 #include "Hash/CityHash.h"
 #include "MuCO/CustomizableObjectCustomVersion.h"
+#include "MuCO/CustomizableObjectResourceDataTypes.h"
 #include "Serialization/BulkData.h"
 
 #include UE_INLINE_GENERATED_CPP_BY_NAME(CustomizableObjectInstance)
diff --git a/Engine/Plugins/Mutable/Source/CustomizableObject/Private/MuCO/LogBenchmarkUtil.cpp b/Engine/Plugins/Mutable/Source/CustomizableObject/Private/MuCO/LogBenchmarkUtil.cpp
--- a/Engine/Plugins/Mutable/Source/CustomizableObject/Private/MuCO/LogBenchmarkUtil.cpp
+++ b/Engine/Plugins/Mutable/Source/CustomizableObject/Private/MuCO/LogBenchmarkUtil.cpp
@@ -19,6 +19,7 @@
 #include "MuCO/CustomizableObjectSystem.h"
 #include "MuCO/CustomizableObjectInstancePrivate.h"
 #include "MuCO/CustomizableObjectInstanceUsage.h"
+#include "MuCO/CustomizableObjectInstanceUsagePrivate.h"
 
 
 extern ENGINE_API UEngine* GEngine;
diff --git a/Engine/Plugins/Mutable/Source/CustomizableObject/Public/MuCO/CustomizableObject.h b/Engine/Plugins/Mutable/Source/CustomizableObject/Public/MuCO/CustomizableObject.h
--- a/Engine/Plugins/Mutable/Source/CustomizableObject/Public/MuCO/CustomizableObject.h
+++ b/Engine/Plugins/Mutable/Source/CustomizableObject/Public/MuCO/CustomizableObject.h
@@ -3,11 +3,7 @@
 #pragma once
 
 #include "RHIDefinitions.h"
-#include "MuCO/CustomizableObjectCompilerTypes.h"
-#include "MuCO/CustomizableObjectClothingTypes.h"
-#include "MuCO/CustomizableObjectResourceData.h"
-#include "MuCO/CustomizableObjectResourceDataTypes.h"
-#include "MuCO/CustomizableObjectStreamedResourceData.h"
+#include "SkeletalMeshTypes.h"
 #include "MuCO/CustomizableObjectParameterTypeDefinitions.h"
 #include "MuCO/CustomizableObjectUIData.h"
 #include "MuCO/CustomizableObject_Deprecated.h"
@@ -40,6 +36,9 @@ namespace mu
 CUSTOMIZABLEOBJECT_API DECLARE_LOG_CATEGORY_EXTERN(LogMutable, Log, All);
 
 
+DECLARE_MULTICAST_DELEGATE(FPostCompileDelegate)
+
+
 USTRUCT()
 struct FFParameterOptionsTags
 {
diff --git a/Engine/Plugins/Mutable/Source/CustomizableObject/Public/MuCO/CustomizableObjectInstanceDescriptor.h b/Engine/Plugins/Mutable/Source/CustomizableObject/Public/MuCO/CustomizableObjectInstanceDescriptor.h
new file mode 100644
--- /dev/null
+++ b/Engine/Plugins/Mutable/Source/CustomizableObject/Public/MuCO/CustomizableObjectInstanceDescriptor.h
@@ -0,0 +1,347 @@
+// Copyright Epic Games, Inc. All Rights Reserved.
+
+#pragma once
+
+#include "MuCO/MultilayerProjector.h"
+#include "MuCO/CustomizableObjectParameterTypeDefinitions.h"
+#include "Math/RandomStream.h"
+
+#include "CustomizableObjectInstanceDescriptor.generated.h"
+
+class UTexture2D;
+enum class ECustomizableObjectProjectorType : uint8;
+
+class FArchive;
+class UCustomizableInstancePrivate;
+class UCustomizableObject;
+class UCustomizableObjectInstance;
+class FDescriptorHash;
+class FMutableUpdateCandidate;
+struct FCustomizableObjectComponentIndex;
+
+typedef TMap<const UCustomizableObjectInstance*, FMutableUpdateCandidate> FMutableInstanceUpdateMap;
+
+namespace mu
+{
+	class FParameters;
+	
+	template<typename Type>
+	class Ptr;
+}
+
+
+/** Internal use only! Currently in /Public due to UCustomizableObjectInstance::Descriptor.
+ *
+ * Set of parameters + state that defines a CustomizableObjectInstance.
+ *
+ * This object has the same parameters + state interface as UCustomizableObjectInstance.
+ * UCustomizableObjectInstance must share the same interface. Any public methods added here should also end up in the Instance. */
+USTRUCT()
+struct CUSTOMIZABLEOBJECT_API FCustomizableObjectInstanceDescriptor
+{
+	GENERATED_BODY()
+
+	FCustomizableObjectInstanceDescriptor() = default;
+	
+	explicit FCustomizableObjectInstanceDescriptor(UCustomizableObject& Object);
+
+	/** Serialize this object. 
+	 *
+	 * Backwards compatibility is not guaranteed.
+ 	 * Multilayer Projectors not supported.
+	 *
+  	 * @param bUseCompactDescriptor If true it assumes the compiled objects are the same on both ends of the serialisation */
+	void SaveDescriptor(FArchive &Ar, bool bUseCompactDescriptor);
+
+	/** Deserialize this object. Does not support Multilayer Projectors! */
+
+	/** Deserialize this object.
+     *
+	 * Backwards compatibility is not guaranteed.
+	 * Multilayer Projectors not supported */
+	void LoadDescriptor(FArchive &Ar);
+
+	// Could return nullptr in some rare situations, so check first
+	UCustomizableObject* GetCustomizableObject() const;
+
+	void SetCustomizableObject(UCustomizableObject* InCustomizableObject);
+	
+	bool GetBuildParameterRelevancy() const;
+	
+	void SetBuildParameterRelevancy(bool Value);
+	
+	/** Update all parameters to be up to date with the Mutable Core parameters. */
+	void ReloadParameters();
+	
+	void SetRequestedLODLevels(const TMap<FName, uint8>& InRequestedLODLevels);
+
+	const TMap<FName, uint8>& GetRequestedLODLevels() const;
+
+	// ------------------------------------------------------------
+	// FParameters
+	// ------------------------------------------------------------
+
+	const TArray<FCustomizableObjectBoolParameterValue>& GetBoolParameters() const;
+	
+	const TArray<FCustomizableObjectIntParameterValue>& GetIntParameters() const;
+
+	const TArray<FCustomizableObjectFloatParameterValue>& GetFloatParameters() const;
+	
+	const TArray<FCustomizableObjectTextureParameterValue>& GetTextureParameters() const;
+
+	const TArray<FCustomizableObjectVectorParameterValue>& GetVectorParameters() const;
+
+	const TArray<FCustomizableObjectProjectorParameterValue>& GetProjectorParameters() const;
+	
+	const TArray<FCustomizableObjectTransformParameterValue>& GetTransformParameters() const;
+	
+	/** Return true if there are any parameters. */
+	bool HasAnyParameters() const;
+
+	/** Gets the value of the int parameter with name "ParamName". */
+	const FString& GetIntParameterSelectedOption(const FString& ParamName, int32 RangeIndex = -1) const;
+
+	/** Sets the selected option of an int parameter by the option's name. */
+	void SetIntParameterSelectedOption(int32 IntParamIndex, const FString& SelectedOption, int32 RangeIndex = -1);
+
+	/** Sets the selected option of an int parameter, by the option's name */
+	void SetIntParameterSelectedOption(const FString& ParamName, const FString& SelectedOptionName, int32 RangeIndex = -1);
+
+	/** Gets the value of a float parameter with name "FloatParamName". */
+	float GetFloatParameterSelectedOption(const FString& FloatParamName, int32 RangeIndex = -1) const;
+
+	/** Sets the float value "FloatValue" of a float parameter with index "FloatParamIndex". */
+	void SetFloatParameterSelectedOption(const FString& FloatParamName, float FloatValue, int32 RangeIndex = -1);
+
+	/** Gets the value of a texture parameter with name "TextureParamName". */
+	FName GetTextureParameterSelectedOption(const FString& TextureParamName, int32 RangeIndex) const;
+
+	/** Sets the texture value "TextureValue" of a texture parameter with index "TextureParamIndex". */
+	void SetTextureParameterSelectedOption(const FString& TextureParamName, const FString& TextureValue, int32 RangeIndex);
+	
+	/** Gets the value of a color parameter with name "ColorParamName". */
+	FLinearColor GetColorParameterSelectedOption(const FString& ColorParamName) const;
+
+	/** Sets the color value "ColorValue" of a color parameter with index "ColorParamIndex". */
+	void SetColorParameterSelectedOption(const FString& ColorParamName, const FLinearColor& ColorValue);
+
+	/** Gets the value of a transform parameter with name "TransformParamName". */
+	FTransform GetTransformParameterSelectedOption(const FString& TransformParamName) const;
+
+	/** Sets the transform value "TransformValue" of a transform parameter with name "TransformParamName". */
+	void SetTransformParameterSelectedOption(const FString& TransformParamName, const FTransform& TransformValue);
+
+	/** Gets the value of the bool parameter with name "BoolParamName". */
+	bool GetBoolParameterSelectedOption(const FString& BoolParamName) const;
+
+	/** Sets the bool value "BoolValue" of a bool parameter with name "BoolParamName". */
+	void SetBoolParameterSelectedOption(const FString& BoolParamName, bool BoolValue);
+
+	/** Sets the vector value "VectorValue" of a bool parameter with index "VectorParamIndex". */
+	void SetVectorParameterSelectedOption(const FString& VectorParamName, const FLinearColor& VectorValue);
+
+	/** Sets the projector values of a projector parameter with index "ProjectorParamIndex". */
+	void SetProjectorValue(const FString& ProjectorParamName,
+		const FVector& Pos, const FVector& Direction, const FVector& Up, const FVector& Scale,
+		float Angle,
+		int32 RangeIndex = -1);
+
+	/** Set only the projector position. */
+	void SetProjectorPosition(const FString& ProjectorParamName, const FVector& Pos, int32 RangeIndex = -1);
+
+	/** Set only the projector direction. */
+	void SetProjectorDirection(const FString& ProjectorParamName, const FVector& Direction, int32 RangeIndex = -1);
+	
+	/** Set only the projector up vector. */
+	void SetProjectorUp(const FString& ProjectorParamName, const FVector& Up, int32 RangeIndex = -1);
+
+	/** Set only the projector scale. */
+	void SetProjectorScale(const FString& ProjectorParamName, const FVector& Scale, int32 RangeIndex = -1);
+
+	/** Set only the cylindrical projector angle. */
+	void SetProjectorAngle(const FString& ProjectorParamName, float Angle, int32 RangeIndex = -1);
+	
+	/** Get the projector values of a projector parameter with index "ProjectorParamIndex". */
+	void GetProjectorValue(const FString& ProjectorParamName,
+		FVector& OutPos, FVector& OutDirection, FVector& OutUp, FVector& OutScale,
+		float& OutAngle, ECustomizableObjectProjectorType& OutType,
+		int32 RangeIndex = -1) const;
+
+	/** Float version. See GetProjectorValue. */
+	void GetProjectorValueF(const FString& ProjectorParamName,
+		FVector3f& OutPos, FVector3f& OutDirection, FVector3f& OutUp, FVector3f& OutScale,
+		float& OutAngle, ECustomizableObjectProjectorType& OutType,
+		int32 RangeIndex = -1) const;
+
+	/** Get the current projector position for the parameter with the given name. */
+	FVector GetProjectorPosition(const FString & ParamName, int32 RangeIndex = -1) const;
+
+	/** Get the current projector direction vector for the parameter with the given name. */
+	FVector GetProjectorDirection(const FString & ParamName, int32 RangeIndex = -1) const;
+
+	/** Get the current projector up vector for the parameter with the given name. */
+	FVector GetProjectorUp(const FString & ParamName, int32 RangeIndex = -1) const;
+
+	/** Get the current projector scale for the parameter with the given name. */
+	FVector GetProjectorScale(const FString & ParamName, int32 RangeIndex = -1) const;
+
+	/** Get the current cylindrical projector angle for the parameter with the given name. */
+	float GetProjectorAngle(const FString& ParamName, int32 RangeIndex = -1) const;
+
+	/** Get the current projector type for the parameter with the given name. */
+	ECustomizableObjectProjectorType GetProjectorParameterType(const FString& ParamName, int32 RangeIndex = -1) const;
+
+	/** Get the current projector for the parameter with the given name. */
+	FCustomizableObjectProjector GetProjector(const FString& ParamName, int32 RangeIndex) const;
+	
+	/** Finds the parameter with name ParamName in the array of its type, returns the index if found, INDEX_NONE otherwise. */
+	int32 FindTypedParameterIndex(const FString& ParamName, EMutableParameterType Type) const;
+
+	// Parameter Ranges
+
+	/** Gets the range of values of the projector with ParamName, returns -1 if the parameter does not exist. */
+	int32 GetProjectorValueRange(const FString& ParamName) const;
+
+	/** Gets the range of values of the int with ParamName, returns -1 if the parameter does not exist. */
+	int32 GetIntValueRange(const FString& ParamName) const;
+
+	/** Gets the range of values of the float with ParamName, returns -1 if the parameter does not exist. */
+	int32 GetFloatValueRange(const FString& ParamName) const;
+
+	/** Gets the range of values of the texture with ParamName, returns -1 if the parameter does not exist. */
+	int32 GetTextureValueRange(const FString& ParamName) const;
+
+	/** Increases the range of values of the integer with ParamName, returns the index of the new integer value, -1 otherwise.
+	 * The added value is initialized with the first integer option and is the last one of the range. */
+	int32 AddValueToIntRange(const FString& ParamName);
+
+	/** Increases the range of values of the float with ParamName, returns the index of the new float value, -1 otherwise.
+	 * The added value is initialized with 0.5f and is the last one of the range. */
+	int32 AddValueToFloatRange(const FString& ParamName);
+
+	/** Increases the range of values of the float with ParamName, returns the index of the new float value, -1 otherwise. 
+	 * The added value is not initialized. */
+	int32 AddValueToTextureRange(const FString& ParamName);
+
+	/** Increases the range of values of the projector with ParamName, returns the index of the new projector value, -1 otherwise.
+	 * The added value is initialized with the default projector as set up in the editor and is the last one of the range. */
+	int32 AddValueToProjectorRange(const FString& ParamName);
+
+	/** Remove the FRangeIndex element of the integer range of values from the parameter ParamName. If Range index is -1, removes the last element.
+		Returns the index of the last valid integer, -1 if no values left. */
+	int32 RemoveValueFromIntRange(const FString& ParamName, int32 RangeIndex = -1);
+
+	/** Remove the FRangeIndex element of the float range of values from the parameter ParamName. If Range index is -1, removes the last element.
+		Returns the index of the last valid float, -1 if no values left. */
+	int32 RemoveValueFromFloatRange(const FString& ParamName, int32 RangeIndex = -1);
+
+	/** Remove the last of the texture range of values from the parameter ParamName, returns the index of the last valid float, -1 if no values left. */
+	int32 RemoveValueFromTextureRange(const FString& ParamName);
+
+	/** Remove the FRangeIndex element of the texture range of values from the parameter ParamName, returns the index of the last valid float, -1 if no values left. */
+	int32 RemoveValueFromTextureRange(const FString& ParamName, int32 RangeIndex);
+
+	/** Remove the FRangeIndex element of the projector range of values from the parameter ParamName. If Range index is -1, removes the last element.
+		Returns the index of the last valid projector, -1 if no values left.
+	*/
+	int32 RemoveValueFromProjectorRange(const FString& ParamName, int32 RangeIndex = -1);
+
+	// ------------------------------------------------------------
+   	// States
+   	// ------------------------------------------------------------
+
+	/** Get the current optimization state. */
+	int32 GetState() const;
+
+	/** Get the current optimization state. */	
+	FString GetCurrentState() const;
+
+	/** Set the current optimization state. */
+	void SetState(int32 InState);
+
+	/** Set the current optimization state. */
+	void SetCurrentState(const FString& StateName);
+
+	// ------------------------------------------------------------
+
+	void SetRandomValues();
+	
+	void SetRandomValuesFromStream(const FRandomStream& InStream);
+
+	void SetDefaultValue(int32 ParamIndex);
+	void SetDefaultValues();
+	
+	// ------------------------------------------------------------
+	// Multilayer Projectors
+	// ------------------------------------------------------------
+
+	/** @return true if ParamName belongs to a multilayer projector parameter. */
+	bool IsMultilayerProjector(const FString& ParamName) const;
+
+	// Layers
+	/** @return number of layers of the projector with name ParamName,-1 if invalid or not found. */
+	int32 NumProjectorLayers(const FName& ParamName) const;
+
+	/** Creates a new layer for the multilayer projector with name ParamName. */
+	void CreateLayer(const FName& ParamName, int32 Index);
+
+	/** Removes the layer at Index from the multilayer projector with name ParamName. */
+	void RemoveLayerAt(const FName& ParamName, int32 Index);
+
+	/** @return copy of the layer at Index for the multilayer projector with name ParamName. */
+	FMultilayerProjectorLayer GetLayer(const FName& ParamName, int32 Index) const;
+
+	/** Updates the parameters of the layer at Index from the multilayer projector with name ParamName. */
+	void UpdateLayer(const FName& ParamName, int32 Index, const FMultilayerProjectorLayer& Layer);
+
+	/** Return a Mutable Core object containing all parameters. */
+	TSharedPtr<mu::FParameters> GetParameters() const;
+
+	FString ToString() const;
+	
+	UPROPERTY()
+	TObjectPtr<UCustomizableObject> CustomizableObject = nullptr;
+
+	UPROPERTY()
+	TArray<FCustomizableObjectBoolParameterValue> BoolParameters;
+
+	UPROPERTY()
+	TArray<FCustomizableObjectIntParameterValue> IntParameters;
+	
+	UPROPERTY()
+	TArray<FCustomizableObjectFloatParameterValue> FloatParameters;
+	
+	UPROPERTY()
+	TArray<FCustomizableObjectTextureParameterValue> TextureParameters;
+	
+	UPROPERTY()
+	TArray<FCustomizableObjectVectorParameterValue> VectorParameters;
+	
+	UPROPERTY()
+	TArray<FCustomizableObjectProjectorParameterValue> ProjectorParameters;
+
+	UPROPERTY()
+	TArray<FCustomizableObjectTransformParameterValue> TransformParameters;
+
+	/** Mutable parameters optimization state. Transient UProperty to make it transactable. */
+	UPROPERTY(Transient)
+	int32 State = 0;
+	
+	/** If this is set to true, when updating the instance an additional step will be performed to calculate the list of instance parameters that are relevant for the current parameter values. */
+	bool bBuildParameterRelevancy = false;
+
+	/** These are the LODs Mutable can generate, they MUST NOT be used in an update (Mutable thread). */
+	TMap<FName, uint8> MinLOD;
+
+	/** RequestedLODs per component to generate.
+	 * Key is the Component name. Value represents the first LOD to generate. Zero means all LODs. */
+	TMap<FName, uint8> RequestedLODLevels;
+	
+	// Friends
+	friend FDescriptorHash;
+	friend UCustomizableObjectInstance;
+	friend UCustomizableInstancePrivate;
+	friend FMultilayerProjector;
+	friend FMutableUpdateCandidate;
+};
+
diff --git a/Engine/Plugins/Mutable/Source/CustomizableObject/Public/MuCO/CustomizableObjectInstanceUsage.h b/Engine/Plugins/Mutable/Source/CustomizableObject/Public/MuCO/CustomizableObjectInstanceUsage.h
--- a/Engine/Plugins/Mutable/Source/CustomizableObject/Public/MuCO/CustomizableObjectInstanceUsage.h
+++ b/Engine/Plugins/Mutable/Source/CustomizableObject/Public/MuCO/CustomizableObjectInstanceUsage.h
@@ -3,8 +3,6 @@
 #pragma once
 
 #include "MuCO/CustomizableObjectInstance.h"
-#include "MuCO/CustomizableObject.h"
-#include "MuCO/CustomizableObjectInstanceUsagePrivate.h"
 #include "Tickable.h"
 #include "Engine/EngineBaseTypes.h"
 
@@ -15,6 +13,7 @@ class FObjectPreSaveContext;
 class UObject;
 class USkeletalMesh;
 class UPhysicsAsset;
+class UCustomizableObjectInstanceUsagePrivate;
 struct FFrame;
 enum class EUpdateResult : uint8;
 
diff --git a/Engine/Plugins/Mutable/Source/CustomizableObject/Public/MuCO/CustomizableObjectSystem.h b/Engine/Plugins/Mutable/Source/CustomizableObject/Public/MuCO/CustomizableObjectSystem.h
--- a/Engine/Plugins/Mutable/Source/CustomizableObject/Public/MuCO/CustomizableObjectSystem.h
+++ b/Engine/Plugins/Mutable/Source/CustomizableObject/Public/MuCO/CustomizableObjectSystem.h
@@ -2,11 +2,12 @@
 
 #pragma once
 
+#define UE_MUTABLE_MAX_OPTIMIZATION			2
+
 #include "HAL/IConsoleManager.h"
 #include "Engine/StreamableManager.h"
 #include "AssetRegistry/AssetData.h"
 #include "MuCO/CustomizableObjectInstance.h"
-#include "MuCO/CustomizableObjectInstanceDescriptor.h"
 #include "MuR/Parameters.h"
 #include "MuR/Image.h"
 #include "MuR/Types.h"
@@ -46,6 +47,30 @@ extern TAutoConsoleVariable<bool> CVarEnableReleaseMeshResources;
 extern TAutoConsoleVariable<bool> CVarFixLowPriorityTasksOverlap;
 
 
+UENUM()
+enum class ECustomizableObjectTextureCompression : uint8
+{
+	// Don't use texture compression
+	None = 0,
+	// Use Mutable's fast low-quality compression
+	Fast,
+	// Use Unreal's highest quality compression (100x slower to compress)
+	HighQuality
+};
+
+
+UENUM()
+enum class ECustomizableObjectNumBoneInfluences : uint8
+{
+	// The enum values can be used as the real numeric value of number of bone influences
+	Four = 4,
+	// 
+	Eight = 8,
+	//
+	Twelve = 12 // This is essentially the same as "Unlimited", but UE ultimately limits to 12
+};
+
+
 #if WITH_EDITOR
 
 // Struct used to keep a copy of the EditorSettings needed to compile Customizable Objects.
diff --git a/Engine/Plugins/Mutable/Source/CustomizableObject/Public/MuCO/CustomizableObject_Deprecated.h b/Engine/Plugins/Mutable/Source/CustomizableObject/Public/MuCO/CustomizableObject_Deprecated.h
--- a/Engine/Plugins/Mutable/Source/CustomizableObject/Public/MuCO/CustomizableObject_Deprecated.h
+++ b/Engine/Plugins/Mutable/Source/CustomizableObject/Public/MuCO/CustomizableObject_Deprecated.h
@@ -2,7 +2,7 @@
 
 #pragma once
 
-#include "MuCO/CustomizableObjectCompilerTypes.h"
+#include "CustomizableObjectSystem.h"
 
 #include "CustomizableObject_Deprecated.generated.h"
 
@@ -16,7 +16,7 @@ struct FCompilationOptions_DEPRECATED
 	ECustomizableObjectTextureCompression TextureCompression = ECustomizableObjectTextureCompression::Fast;
 
 	UPROPERTY()
-	int32 OptimizationLevel = UE_MUTABLE_MAX_OPTIMIZATION;
+	int32 OptimizationLevel = 2; // UE_MUTABLE_MAX_OPTIMIZATION;
 
 	UPROPERTY()
 	bool bUseDiskCompilation = false;
diff --git a/Engine/Plugins/Mutable/Source/CustomizableObjectEditor/Private/MuCOE/GenerateMutableSource/GenerateMutableSource.cpp b/Engine/Plugins/Mutable/Source/CustomizableObjectEditor/Private/MuCOE/GenerateMutableSource/GenerateMutableSource.cpp
--- a/Engine/Plugins/Mutable/Source/CustomizableObjectEditor/Private/MuCOE/GenerateMutableSource/GenerateMutableSource.cpp
+++ b/Engine/Plugins/Mutable/Source/CustomizableObjectEditor/Private/MuCOE/GenerateMutableSource/GenerateMutableSource.cpp
@@ -55,6 +55,7 @@
 #include "AssetRegistry/AssetRegistryModule.h"
 #include "Math/NumericLimits.h"
 #include "Hash/CityHash.h"
+#include "MuCO/CustomizableObjectResourceDataTypes.h"
 #include "MuCOE/Nodes/CustomizableObjectNodeCopyMaterial.h"
 #include "MuCOE/Nodes/CustomizableObjectNodeGroupProjectorParameter.h"
 #include "UObject/ICookInfo.h"
