[Mutable] Added support for C++ generated structure assets.

#rnx

diff --git a/Engine/Plugins/Mutable/Source/CustomizableObject/Internal/MuCO/CustomizableObjectCustomVersion.h b/Engine/Plugins/Mutable/Source/CustomizableObject/Internal/MuCO/CustomizableObjectCustomVersion.h
--- a/Engine/Plugins/Mutable/Source/CustomizableObject/Internal/MuCO/CustomizableObjectCustomVersion.h
+++ b/Engine/Plugins/Mutable/Source/CustomizableObject/Internal/MuCO/CustomizableObjectCustomVersion.h
@@ -164,6 +164,8 @@ struct CUSTOMIZABLEOBJECT_API FCustomizableObjectCustomVersion
 
 		SnapToBoneComponentIndexToName,
 
+		SetDisplayNamePropertyAsPinNameOfTableNodes,
+
 		// -----<new versions can be added above this line>-------------------------------------------------
 		VersionPlusOne,
 		LatestVersion = VersionPlusOne - 1
diff --git a/Engine/Plugins/Mutable/Source/CustomizableObjectEditor/Private/MuCOE/GenerateMutableSource/GenerateMutableSourceColor.cpp b/Engine/Plugins/Mutable/Source/CustomizableObjectEditor/Private/MuCOE/GenerateMutableSource/GenerateMutableSourceColor.cpp
--- a/Engine/Plugins/Mutable/Source/CustomizableObjectEditor/Private/MuCOE/GenerateMutableSource/GenerateMutableSourceColor.cpp
+++ b/Engine/Plugins/Mutable/Source/CustomizableObjectEditor/Private/MuCOE/GenerateMutableSource/GenerateMutableSourceColor.cpp
@@ -277,8 +277,8 @@ mu::Ptr<mu::NodeColour> GenerateMutableSourceColor(const UEdGraphPin* Pin, FMuta
 
 		if (DataTable)
 		{
-			FString ColumnName = Pin->PinFriendlyName.ToString();
-			FProperty* Property = DataTable->FindTableProperty(FName(*ColumnName));
+			FString ColumnName = TypedNodeTable->GetPinColumnName(Pin);
+			FProperty* Property = TypedNodeTable->FindPinProperty(*Pin);
 
 			if (!Property)
 			{
diff --git a/Engine/Plugins/Mutable/Source/CustomizableObjectEditor/Private/MuCOE/GenerateMutableSource/GenerateMutableSourceFloat.cpp b/Engine/Plugins/Mutable/Source/CustomizableObjectEditor/Private/MuCOE/GenerateMutableSource/GenerateMutableSourceFloat.cpp
--- a/Engine/Plugins/Mutable/Source/CustomizableObjectEditor/Private/MuCOE/GenerateMutableSource/GenerateMutableSourceFloat.cpp
+++ b/Engine/Plugins/Mutable/Source/CustomizableObjectEditor/Private/MuCOE/GenerateMutableSource/GenerateMutableSourceFloat.cpp
@@ -318,8 +318,8 @@ mu::NodeScalarPtr GenerateMutableSourceFloat(const UEdGraphPin* Pin, FMutableGra
 
 		if (DataTable)
 		{
-			FString ColumnName = Pin->PinFriendlyName.ToString();
-			FProperty* Property = DataTable->FindTableProperty(FName(*ColumnName));
+			FString ColumnName = TypedNodeTable->GetPinColumnName(Pin);
+			FProperty* Property = TypedNodeTable->FindPinProperty(*Pin);
 
 			if (!Property)
 			{
diff --git a/Engine/Plugins/Mutable/Source/CustomizableObjectEditor/Private/MuCOE/GenerateMutableSource/GenerateMutableSourceImage.cpp b/Engine/Plugins/Mutable/Source/CustomizableObjectEditor/Private/MuCOE/GenerateMutableSource/GenerateMutableSourceImage.cpp
--- a/Engine/Plugins/Mutable/Source/CustomizableObjectEditor/Private/MuCOE/GenerateMutableSource/GenerateMutableSourceImage.cpp
+++ b/Engine/Plugins/Mutable/Source/CustomizableObjectEditor/Private/MuCOE/GenerateMutableSource/GenerateMutableSourceImage.cpp
@@ -1018,8 +1018,8 @@ mu::NodeImagePtr GenerateMutableSourceImage(const UEdGraphPin* Pin, FMutableGrap
 
 			if (DataTable)
 			{
-				FString ColumnName = Pin->PinFriendlyName.ToString();
-				FProperty* Property = DataTable->FindTableProperty(FName(*ColumnName));
+				FString ColumnName = TypedNodeTable->GetPinColumnName(Pin);
+				FProperty* Property = TypedNodeTable->FindPinProperty(*Pin);
 
 				if (!Property)
 				{
@@ -1073,7 +1073,7 @@ mu::NodeImagePtr GenerateMutableSourceImage(const UEdGraphPin* Pin, FMutableGrap
 							}
 							else
 							{
-								GenerationContext.GeneratedTableImages.Add({ ColumnName, Pin->PinType.PinCategory, Table, TypedNodeTable });
+								GenerationContext.GeneratedTableImages.Add({ColumnName, Pin->PinType.PinCategory, Table, TypedNodeTable});
 							}
 
 							// Encoding the texture mode of the type to allow different texture modes from the same pin
diff --git a/Engine/Plugins/Mutable/Source/CustomizableObjectEditor/Private/MuCOE/GenerateMutableSource/GenerateMutableSourceMesh.cpp b/Engine/Plugins/Mutable/Source/CustomizableObjectEditor/Private/MuCOE/GenerateMutableSource/GenerateMutableSourceMesh.cpp
--- a/Engine/Plugins/Mutable/Source/CustomizableObjectEditor/Private/MuCOE/GenerateMutableSource/GenerateMutableSourceMesh.cpp
+++ b/Engine/Plugins/Mutable/Source/CustomizableObjectEditor/Private/MuCOE/GenerateMutableSource/GenerateMutableSourceMesh.cpp
@@ -4160,12 +4160,12 @@ mu::Ptr<mu::NodeMesh> GenerateMutableSourceMesh(const UEdGraphPin* Pin,
 			const UEdGraphSchema_CustomizableObject* Schema = GetDefault<UEdGraphSchema_CustomizableObject>();
 
 			// Getting the real name of the data table column
-			FString DataTableColumnName = TypedNodeTable->GetColumnNameByPin(Pin);
-			FProperty* Property = DataTable->FindTableProperty(FName(*DataTableColumnName));
+			FString ColumnName = TypedNodeTable->GetPinColumnName(Pin);
+			FProperty* Property = TypedNodeTable->FindPinProperty(*Pin);
 
 			if (!Property)
 			{
-				FString Msg = FString::Printf(TEXT("Couldn't find the column [%s] in the data table's struct."), *DataTableColumnName);
+				FString Msg = FString::Printf(TEXT("Couldn't find the column [%s] in the data table's struct."), *ColumnName);
 				GenerationContext.Log(FText::FromString(Msg), Node);
 
 				bSuccess = false;
@@ -4177,7 +4177,7 @@ mu::Ptr<mu::NodeMesh> GenerateMutableSourceMesh(const UEdGraphPin* Pin,
 
 			if (bSuccess && !DefaultSkeletalMesh && !DefaultStaticMesh && !DefaultPoseAsset)
 			{
-				FString Msg = FString::Printf(TEXT("Couldn't find a default value in the data table's struct for the column [%s]."), *DataTableColumnName);
+				FString Msg = FString::Printf(TEXT("Couldn't find a default value in the data table's struct for the column [%s]."), *ColumnName);
 				GenerationContext.Log(FText::FromString(Msg), Node);
 
 				bSuccess = false;
@@ -4199,7 +4199,7 @@ mu::Ptr<mu::NodeMesh> GenerateMutableSourceMesh(const UEdGraphPin* Pin,
 					int32 SectionIndex = 0;
 
 					// Getting the mutable table mesh column name
-					FString MutableColumnName = DataTableColumnName;
+					FString MutableColumnName = ColumnName;
 
 					if (Pin->PinType.PinCategory == Schema->PC_Mesh)
 					{
@@ -4209,18 +4209,18 @@ mu::Ptr<mu::NodeMesh> GenerateMutableSourceMesh(const UEdGraphPin* Pin,
 						if (DefaultSkeletalMesh)
 						{
 							GetLODAndSectionForAutomaticLODs(GenerationContext, *Node, *DefaultSkeletalMesh, LODIndexConnected, SectionIndexConnected, LODIndex, SectionIndex, bOnlyConnectedLOD);
-							MutableColumnName = TypedNodeTable->GenerateSkeletalMeshMutableColumName(DataTableColumnName, LODIndex, SectionIndex);
+							MutableColumnName = TypedNodeTable->GenerateSkeletalMeshMutableColumName(ColumnName, LODIndex, SectionIndex);
 						}
 						else
 						{
-							MutableColumnName = TypedNodeTable->GenerateStaticMeshMutableColumName(DataTableColumnName, SectionIndexConnected);
+							MutableColumnName = TypedNodeTable->GenerateStaticMeshMutableColumName(ColumnName, SectionIndexConnected);
 						}
 					}
 					
 					// Generating a new FMesh column if not exists
 					if (Table->FindColumn(MutableColumnName) == INDEX_NONE)
 					{
-						bSuccess = GenerateTableColumn(TypedNodeTable, Pin, Table, DataTableColumnName, Property, LODIndexConnected, SectionIndexConnected, LODIndex, SectionIndex, SurfaceMetadataId, bOnlyConnectedLOD, GenerationContext);
+						bSuccess = GenerateTableColumn(TypedNodeTable, Pin, Table, ColumnName, Property, LODIndexConnected, SectionIndexConnected, LODIndex, SectionIndex, SurfaceMetadataId, bOnlyConnectedLOD, GenerationContext);
 
 						if (!bSuccess)
 						{
diff --git a/Engine/Plugins/Mutable/Source/CustomizableObjectEditor/Private/MuCOE/GenerateMutableSource/GenerateMutableSourceSurface.cpp b/Engine/Plugins/Mutable/Source/CustomizableObjectEditor/Private/MuCOE/GenerateMutableSource/GenerateMutableSourceSurface.cpp
--- a/Engine/Plugins/Mutable/Source/CustomizableObjectEditor/Private/MuCOE/GenerateMutableSource/GenerateMutableSourceSurface.cpp
+++ b/Engine/Plugins/Mutable/Source/CustomizableObjectEditor/Private/MuCOE/GenerateMutableSource/GenerateMutableSourceSurface.cpp
@@ -367,37 +367,33 @@ mu::Ptr<mu::NodeSurface> GenerateMutableSourceSurface(const UEdGraphPin * Pin, F
 		FString AlternateResStateName;
 
 		bool bTableMaterialPinLinked = TypedNodeMat->GetMaterialAssetPin() && FollowInputPin(*TypedNodeMat->GetMaterialAssetPin()) != nullptr;
-		FString TableColumnName;
 
 		// Checking if we should not use the material of the table node even if it is linked to the material node
-		const UEdGraphPin* MaterialAssetConnectedPin = nullptr;
 		if (TypedNodeMat->GetMaterialAssetPin())
 		{
-			MaterialAssetConnectedPin = FollowInputPin(*TypedNodeMat->GetMaterialAssetPin());
-		}
-
-		if (MaterialAssetConnectedPin)
-		{
-			if (const UCustomizableObjectNodeTable* TypedNodeTable = Cast< UCustomizableObjectNodeTable >(MaterialAssetConnectedPin->GetOwningNode()))
+			if (const UEdGraphPin* MaterialAssetConnectedPin = FollowInputPin(*TypedNodeMat->GetMaterialAssetPin()))
 			{
-				TableColumnName = MaterialAssetConnectedPin->PinFriendlyName.ToString();
-
-				if (UMaterialInstance* TableMaterial = TypedNodeTable->GetColumnDefaultAssetByType<UMaterialInstance>(MaterialAssetConnectedPin))
+				if (const UCustomizableObjectNodeTable* TypedNodeTable = Cast< UCustomizableObjectNodeTable >(MaterialAssetConnectedPin->GetOwningNode()))
 				{
-					// Checking if the reference material of the Table Node has the same parent as the material of the Material Node 
-					if (!TypedNodeMat->GetMaterial() || TableMaterial->GetMaterial() != TypedNodeMat->GetMaterial()->GetMaterial())
+					if (UMaterialInstance* TableMaterial = TypedNodeTable->GetColumnDefaultAssetByType<UMaterialInstance>(MaterialAssetConnectedPin))
 					{
-						bTableMaterialPinLinked = false;
+						// Checking if the reference material of the Table Node has the same parent as the material of the Material Node 
+						if (!TypedNodeMat->GetMaterial() || TableMaterial->GetMaterial() != TypedNodeMat->GetMaterial()->GetMaterial())
+						{
+							bTableMaterialPinLinked = false;
 
-						GenerationContext.Log(LOCTEXT("DifferentParentMaterial", "The Default Material Instance of the Data Table must have the same Parent Material."), TypedNodeMat->GetMaterialNode());
+							GenerationContext.Log(LOCTEXT("DifferentParentMaterial", "The Default Material Instance of the Data Table must have the same Parent Material."), TypedNodeMat->GetMaterialNode());
+						}
 					}
-				}
-				else
-				{
-					FText Msg = FText::Format(LOCTEXT("DefaultValueNotFound", "Couldn't find a default value in the data table's struct for the column {0}. The default value is null or not a Material Instance."), FText::FromString(TableColumnName));
-					GenerationContext.Log(Msg, Node);
+					else
+					{
+						FString TableColumnName = TypedNodeTable->GetPinColumnName(MaterialAssetConnectedPin);
 
-					bTableMaterialPinLinked = false;
+						FText Msg = FText::Format(LOCTEXT("DefaultValueNotFound", "Couldn't find a default value in the data table's struct for the column {0}. The default value is null or not a Material Instance."), FText::FromString(TableColumnName));
+						GenerationContext.Log(Msg, Node);
+
+						bTableMaterialPinLinked = false;
+					}
 				}
 			}
 		}
diff --git a/Engine/Plugins/Mutable/Source/CustomizableObjectEditor/Private/MuCOE/GenerateMutableSource/GenerateMutableSourceTable.cpp b/Engine/Plugins/Mutable/Source/CustomizableObjectEditor/Private/MuCOE/GenerateMutableSource/GenerateMutableSourceTable.cpp
--- a/Engine/Plugins/Mutable/Source/CustomizableObjectEditor/Private/MuCOE/GenerateMutableSource/GenerateMutableSourceTable.cpp
+++ b/Engine/Plugins/Mutable/Source/CustomizableObjectEditor/Private/MuCOE/GenerateMutableSource/GenerateMutableSourceTable.cpp
@@ -36,6 +36,7 @@ bool FillTableColumn(const UCustomizableObjectNodeTable* TableNode, mu::Ptr<mu::
 {
 	int32 CurrentColumn;
 	UDataTable* DataTablePtr = GetDataTable(TableNode, GenerationContext);
+	FString ColumnPropertyName = ColumnProperty->GetAuthoredName();
 
 	// Getting property type
 	if (const FSoftObjectProperty* SoftObjectProperty = CastField<FSoftObjectProperty>(ColumnProperty))
@@ -70,9 +71,8 @@ bool FillTableColumn(const UCustomizableObjectNodeTable* TableNode, mu::Ptr<mu::
 			// Getting Animation Blueprint and Animation Slot
 			FString AnimBP, AnimSlot, GameplayTag, AnimBPAssetTag;
 			TArray<FGameplayTag> GameplayTags;
-			FGuid ColumnPropertyId = FStructureEditorUtils::GetGuidForProperty(ColumnProperty);
 
-			TableNode->GetAnimationColumns(ColumnPropertyId, AnimBP, AnimSlot, GameplayTag);
+			TableNode->GetAnimationColumns(ColumnName, AnimBP, AnimSlot, GameplayTag);
 
 			if (!AnimBP.IsEmpty())
 			{
@@ -87,7 +87,7 @@ bool FillTableColumn(const UCustomizableObjectNodeTable* TableNode, mu::Ptr<mu::
 							FName SlotIndex;
 
 							// Getting animation slot row value from data table
-							if (FProperty* AnimSlotProperty = DataTablePtr->FindTableProperty(FName(*AnimSlot)))
+							if (FProperty* AnimSlotProperty = TableNode->FindColumnProperty(FName(*AnimSlot)))
 							{
 								uint8* AnimSlotData = AnimSlotProperty->ContainerPtrToValuePtr<uint8>(AnimRowData, 0);
 
@@ -112,7 +112,7 @@ bool FillTableColumn(const UCustomizableObjectNodeTable* TableNode, mu::Ptr<mu::
 							if (SlotIndex.GetStringLength() != 0)
 							{
 								// Getting animation instance soft class from data table
-								if (FProperty* AnimBPProperty = DataTablePtr->FindTableProperty(FName(*AnimBP)))
+								if (FProperty* AnimBPProperty = TableNode->FindColumnProperty(FName(*AnimBP)))
 								{
 									uint8* AnimBPData = AnimBPProperty->ContainerPtrToValuePtr<uint8>(AnimRowData, 0);
 
@@ -159,7 +159,7 @@ bool FillTableColumn(const UCustomizableObjectNodeTable* TableNode, mu::Ptr<mu::
 					if (GameplayRowData)
 					{
 						// Getting animation tag row value from data table
-						if (FProperty* GameplayTagProperty = DataTablePtr->FindTableProperty(FName(*GameplayTag)))
+						if (FProperty* GameplayTagProperty = TableNode->FindColumnProperty(FName(*GameplayTag)))
 						{
 							uint8* GameplayTagData = GameplayTagProperty->ContainerPtrToValuePtr<uint8>(GameplayRowData, 0);
 
@@ -180,7 +180,7 @@ bool FillTableColumn(const UCustomizableObjectNodeTable* TableNode, mu::Ptr<mu::
 			}
 
 			// Getting reference Mesh
-			USkeletalMesh* ReferenceSkeletalMesh = TableNode->GetColumnDefaultAssetByType<USkeletalMesh>(ColumnName);
+			USkeletalMesh* ReferenceSkeletalMesh = TableNode->GetColumnDefaultAssetByType<USkeletalMesh>(ColumnPropertyName);
 
 			if (!ReferenceSkeletalMesh)
 			{
@@ -344,7 +344,7 @@ bool FillTableColumn(const UCustomizableObjectNodeTable* TableNode, mu::Ptr<mu::
 			}
 
 			// Getting reference Mesh
-			UStaticMesh* ReferenceStaticMesh = TableNode->GetColumnDefaultAssetByType<UStaticMesh>(ColumnName);
+			UStaticMesh* ReferenceStaticMesh = TableNode->GetColumnDefaultAssetByType<UStaticMesh>(ColumnPropertyName);
 
 			if (!ReferenceStaticMesh)
 			{
@@ -409,18 +409,15 @@ bool FillTableColumn(const UCustomizableObjectNodeTable* TableNode, mu::Ptr<mu::
 		{
 			GenerationContext.AddParticipatingObject(SoftObject);
 
-			// Removing encoding part
-			const FString PinName = ColumnName.Replace(TEXT("--PassThrough"), TEXT(""), ESearchCase::CaseSensitive);
-			
 			UObject* Object = GenerationContext.LoadObject(SoftObject);
-
 			UTexture* Texture = Cast<UTexture>(Object);
+
 			if (!Texture)
 			{
-				Texture = TableNode->GetColumnDefaultAssetByType<UTexture>(PinName);
+				Texture = TableNode->GetColumnDefaultAssetByType<UTexture>(ColumnPropertyName);
 
 				FString Message = Cast<UObject>(Object) ? "not a suported Texture" : "null";
-				FString WarningMessage = FString::Printf(TEXT("Texture from column [%s] row [%s] is %s. The default texture will be used instead."), *PinName, *RowName, *Message);
+				FString WarningMessage = FString::Printf(TEXT("Texture from column [%s] row [%s] is %s. The default texture will be used instead."), *ColumnName, *RowName, *Message);
 				LogRowGenerationMessage(TableNode, DataTablePtr, GenerationContext, WarningMessage, RowName);
 			}
 
@@ -435,7 +432,7 @@ bool FillTableColumn(const UCustomizableObjectNodeTable* TableNode, mu::Ptr<mu::
 				CurrentColumn = MutableTable->AddColumn(ColumnName, mu::ETableColumnType::Image);
 			}
 
-			const bool bIsPassthroughTexture = TableNode->GetColumnImageMode(PinName) == ETableTextureType::PASSTHROUGH_TEXTURE;
+			const bool bIsPassthroughTexture = TableNode->GetColumnImageMode(ColumnPropertyName) == ETableTextureType::PASSTHROUGH_TEXTURE;
 			mu::Ptr<mu::TResourceProxyMemory<mu::FImage>> Proxy = new mu::TResourceProxyMemory<mu::FImage>(GenerateImageConstant(Texture, GenerationContext, bIsPassthroughTexture));
 			MutableTable->SetCell(CurrentColumn, RowId, Proxy.get());
 		}
@@ -449,12 +446,9 @@ bool FillTableColumn(const UCustomizableObjectNodeTable* TableNode, mu::Ptr<mu::
 			// Get display name of the column of the data table (name showed in the table and struct editors)
 			// Will be used in the warnings to help to identify a column with errors.
 			FString MaterialColumnDisplayName = ColumnProperty->GetDisplayNameText().ToString();
-			
-			// Get the real name of the Property column
-			FString MaterialColumnName = ColumnProperty->GetName();
 
 			UMaterialInstance* MaterialInstance = Cast<UMaterialInstance>(Object);
-			UMaterialInstance* ReferenceMaterial = TableNode->GetColumnDefaultAssetByType<UMaterialInstance>(MaterialColumnName);
+			UMaterialInstance* ReferenceMaterial = TableNode->GetColumnDefaultAssetByType<UMaterialInstance>(ColumnPropertyName);
 
 			if (!ReferenceMaterial)
 			{
@@ -825,7 +819,7 @@ void RestrictRowNamesToSelectedOption(TArray<FName>& InOutRowNames, const UCusto
 
 void RestrictRowContentByVersion( TArray<FName>& InOutRowNames, const UDataTable& DataTable, const UCustomizableObjectNodeTable& TableNode, FMutableGraphGenerationContext& GenerationContext)
 {
-	FProperty* ColumnProperty = DataTable.FindTableProperty(TableNode.VersionColumn);
+	FProperty* ColumnProperty = TableNode.FindColumnProperty(TableNode.VersionColumn);
 
 	if (!ColumnProperty)
 	{
@@ -1002,13 +996,11 @@ void GenerateTableParameterUIData(const UDataTable* DataTable, const UCustomizab
 	// Generating Parameter UI MetaData if not exists
 	if (!GenerationContext.ParameterUIDataMap.Contains(TableNode->ParameterName))
 	{
-		// Getting Table and row names to access the information
-
 		FMutableParameterData ParameterUIData(TableNode->ParamUIMetadata, EMutableParameterType::Int);
 		ParameterUIData.IntegerParameterGroupType = TableNode->bAddNoneOption ? ECustomizableObjectGroupType::COGT_ONE_OR_NONE : ECustomizableObjectGroupType::COGT_ONE;
 		
 		FMutableParameterData& ParameterUIDataRef = GenerationContext.ParameterUIDataMap.Add(TableNode->ParameterName, ParameterUIData);
-		FProperty* MetadataColumnProperty = DataTable->FindTableProperty(TableNode->ParamUIMetadataColumn);
+		FProperty* MetadataColumnProperty = TableNode->FindColumnProperty(TableNode->ParamUIMetadataColumn);
 		bool bIsValidMetadataColumn = MetadataColumnProperty &&
 			CastField<FStructProperty>(MetadataColumnProperty) &&
 			CastField<FStructProperty>(MetadataColumnProperty)->Struct == FMutableParamUIMetadata::StaticStruct();
@@ -1021,7 +1013,7 @@ void GenerateTableParameterUIData(const UDataTable* DataTable, const UCustomizab
 			GenerationContext.Log(LogMessage, TableNode);
 		}
 
-		FProperty* ThumbnailColumnProperty = DataTable->FindTableProperty(TableNode->ThumbnailColumn);
+		FProperty* ThumbnailColumnProperty = TableNode->FindColumnProperty(TableNode->ThumbnailColumn);
 		bool bIsValidThumbnailColumn = ThumbnailColumnProperty && CastField<FSoftObjectProperty>(ThumbnailColumnProperty);
 
 		// Trigger warning only if the name is different than "None"
@@ -1057,7 +1049,7 @@ void GenerateTableParameterUIData(const UDataTable* DataTable, const UCustomizab
 					}
 				}
 
-				// Add tags
+				// Add Tags
 				if (TableNode->bGatherTags)
 				{
 					if (const UScriptStruct* Struct = DataTable->GetRowStruct())
@@ -1075,10 +1067,7 @@ void GenerateTableParameterUIData(const UDataTable* DataTable, const UCustomizab
 							{
 								if (StructProperty->Struct == TBaseStructure<FGameplayTagContainer>::Get())
 								{
-									FName ColumnName = FName(DataTableUtils::GetPropertyExportName(ColumnProperty));
-								
-									FProperty* TagColumnProperty = DataTable->FindTableProperty(ColumnName);
-									if (uint8* TagCellData = GetCellData(RowName, *DataTable, *TagColumnProperty))
+									if (uint8* TagCellData = GetCellData(RowName, *DataTable, *ColumnProperty))
 									{
 										FGameplayTagContainer* TagContainer = reinterpret_cast<FGameplayTagContainer*>(TagCellData);
 										IntegerMetadata.ParamUIMetadata.GameplayTags.AppendTags(*TagContainer);
diff --git a/Engine/Plugins/Mutable/Source/CustomizableObjectEditor/Private/MuCOE/Nodes/CustomizableObjectNodeTable.cpp b/Engine/Plugins/Mutable/Source/CustomizableObjectEditor/Private/MuCOE/Nodes/CustomizableObjectNodeTable.cpp
--- a/Engine/Plugins/Mutable/Source/CustomizableObjectEditor/Private/MuCOE/Nodes/CustomizableObjectNodeTable.cpp
+++ b/Engine/Plugins/Mutable/Source/CustomizableObjectEditor/Private/MuCOE/Nodes/CustomizableObjectNodeTable.cpp
@@ -40,12 +40,28 @@ bool UCustomizableObjectNodeTableRemapPins::Equal(const UCustomizableObjectNode&
 	const UCustomizableObjectNodeTableObjectPinData& OldPinData = Node.GetPinData<UCustomizableObjectNodeTableObjectPinData>(OldPin);
 	const UCustomizableObjectNodeTableObjectPinData& NewPinData = Node.GetPinData<UCustomizableObjectNodeTableObjectPinData>(NewPin);
 
-	// If one of these two option fails, pins are different
-	if (OldPinData.StructColumnId != NewPinData.StructColumnId || OldPin.Direction != NewPin.Direction)
+	if (OldPin.Direction != NewPin.Direction)
 	{
 		return false;
 	}
 
+	// If both pins have valid IDs, check if its id has changed instead of checking its name
+	if (OldPinData.StructColumnId.IsValid() && NewPinData.StructColumnId.IsValid())
+	{
+		if (OldPinData.StructColumnId != NewPinData.StructColumnId)
+		{
+			return false;
+		}
+	}
+	else
+	{
+		// If one of these two option fails, pins are different
+		if (OldPinData.ColumnPropertyName != NewPinData.ColumnPropertyName)
+		{
+			return false;
+		}
+	}
+
 	// In this case pin type may have changed but we consider them the same type
 	if (OldPin.PinType.PinCategory == UEdGraphSchema_CustomizableObject::PC_Image ||
 		OldPin.PinType.PinCategory == UEdGraphSchema_CustomizableObject::PC_PassThroughImage)
@@ -99,7 +115,6 @@ void UCustomizableObjectNodeTableRemapPins::RemapPins(const UCustomizableObjectN
 
 
 // Table Node -------
-
 void UCustomizableObjectNodeTable::BackwardsCompatibleFixup(int32 CustomizableObjectCustomVersion)
 {
 	Super::BackwardsCompatibleFixup(CustomizableObjectCustomVersion);
@@ -115,7 +130,7 @@ void UCustomizableObjectNodeTable::BackwardsCompatibleFixup(int32 CustomizableOb
 
 				if (OldPinData && NewPinData)
 				{
-					NewPinData->ColumnName = OldPinData->ColumnName;
+					NewPinData->ColumnName_DEPRECATED = OldPinData->ColumnName_DEPRECATED;
 					NewPinData->ImageMode = ETableTextureType::MUTABLE_TEXTURE; // Old pin type by default
 				}
 
@@ -139,12 +154,12 @@ void UCustomizableObjectNodeTable::BackwardsCompatibleFixup(int32 CustomizableOb
 				if (!PinData)
 				{
 					PinData = NewObject<UCustomizableObjectNodeTableObjectPinData>(this);
-					PinData->ColumnName = Pin->PinFriendlyName.ToString();
+					PinData->ColumnName_DEPRECATED = Pin->PinFriendlyName.ToString();
 
 					AddPinData(*Pin, *PinData);
 				}
 
-				FProperty* ColumnProperty = FindTableProperty(TableStruct, FName(*PinData->ColumnName));
+				FProperty* ColumnProperty = FindTableProperty(FName(*PinData->ColumnName_DEPRECATED));
 
 				if (!ColumnProperty)
 				{
@@ -164,7 +179,7 @@ void UCustomizableObjectNodeTable::BackwardsCompatibleFixup(int32 CustomizableOb
 						NewColumnData.AnimSlotColumnName = MeshPinData->AnimSlotColumnName_DEPRECATED;
 						NewColumnData.AnimTagColumnName = MeshPinData->AnimTagColumnName_DEPRECATED;
 
-						ColumnDataMap.Add(ColumnPropertyId, NewColumnData);
+						ColumnDataMap_DEPRECATED.Add(ColumnPropertyId, NewColumnData);
 					}
 				}
 			}
@@ -198,6 +213,49 @@ void UCustomizableObjectNodeTable::BackwardsCompatibleFixup(int32 CustomizableOb
 			}
 		}
 	}
+
+	if (CustomizableObjectCustomVersion == FCustomizableObjectCustomVersion::SetDisplayNamePropertyAsPinNameOfTableNodes)
+	{
+		// We asume that until now we only supported UserDefinedStructs (Structure assets generated trhough the Editor)
+		const UUserDefinedStruct* UserStruct = Cast<UUserDefinedStruct>(GetTableNodeStruct());
+
+		if (UserStruct)
+		{
+			for (UEdGraphPin* Pin : Pins)
+			{
+				if (!Pin)
+				{
+					continue;
+				}
+
+				UCustomizableObjectNodeTableObjectPinData* PinData = Cast<UCustomizableObjectNodeTableObjectPinData>(GetPinData(*Pin));
+
+				if (!PinData)
+				{
+					continue;
+				}
+
+				if (const FProperty* Property = FStructureEditorUtils::GetPropertyByGuid(UserStruct, PinData->StructColumnId))
+				{
+					// Using the standard naming of the property
+					PinData->ColumnPropertyName = Property->GetAuthoredName();
+					PinData->ColumnDisplayName = Property->GetDisplayNameText().ToString();
+
+					if (ColumnDataMap_DEPRECATED.Contains(PinData->StructColumnId))
+					{
+						const FTableNodeColumnData* ColumnData = ColumnDataMap_DEPRECATED.Find(PinData->StructColumnId);
+						PinColumnDataMap.Add(PinData->ColumnDisplayName, *ColumnData);
+
+						// Note: We do not need to convert the content of the map because the display name of an editor structure is the same than its property name
+					}
+				}
+				else
+				{
+					UE_LOG(LogTemp, Warning, TEXT("BackwardCompatibleFixUp Problem: Property doesn't exist"));
+				}
+			}
+		}
+	}
 }
 
 
@@ -270,7 +328,7 @@ FText UCustomizableObjectNodeTable::GetNodeTitle(ENodeTitleType::Type TitleType)
 	FFormatNamedArguments Args;
 	Args.Add(TEXT("ParamName"), FText::FromString(ParameterName));
 
-	if (TitleType == ENodeTitleType::ListView && ParameterName.IsEmpty() && (Table || Structure))
+	if (TitleType == ENodeTitleType::ListView)
 	{
 		return LOCTEXT("Mutable_Table_Title", "Table");
 	}
@@ -353,8 +411,8 @@ void UCustomizableObjectNodeTable::PostEditChangeProperty(FPropertyChangedEvent&
 void UCustomizableObjectNodeTable::AllocateDefaultPins(UCustomizableObjectNodeRemapPins* RemapPins)
 {
 	// Reset of the column data map
-	TMap<FGuid, FTableNodeColumnData> AuxOldColumnData = ColumnDataMap;
-	ColumnDataMap.Reset();
+	TMap<FString, FTableNodeColumnData> AuxOldColumnData = PinColumnDataMap;
+	PinColumnDataMap.Reset();
 
 	// Getting Struct Pointer
 	const UScriptStruct* TableStruct = GetTableNodeStruct();
@@ -364,7 +422,7 @@ void UCustomizableObjectNodeTable::AllocateDefaultPins(UCustomizableObjectNodeRe
 		return;
 	}
 
-	NumProperties = GetColumnTitles(TableStruct).Num();
+	NumProperties = GetColumnTitles().Num();
 
 	// Getting Default Struct Values
 	// A Script Struct always has at leaset one property
@@ -377,11 +435,13 @@ void UCustomizableObjectNodeTable::AllocateDefaultPins(UCustomizableObjectNodeRe
 		UStaticMesh::StaticClass(),
 		UTexture2D::StaticClass(),
 		UTexture::StaticClass(),
-		UMaterialInstance::StaticClass(),
+		UMaterialInterface::StaticClass(),
 		UPoseAsset::StaticClass()
 	};
 
 	TArray<UEdGraphPin*> OldPins(Pins);
+	TSet<FString> PinNames;
+	bool bHasRepeatedNames = false;
 	
 	for (TFieldIterator<FProperty> It(TableStruct); It; ++It)
 	{
@@ -395,10 +455,18 @@ void UCustomizableObjectNodeTable::AllocateDefaultPins(UCustomizableObjectNodeRe
 		const UEdGraphSchema_CustomizableObject* Schema = GetDefault<UEdGraphSchema_CustomizableObject>();
 
 		UEdGraphPin* OutPin = nullptr;
+		FString PinName = ColumnProperty->GetDisplayNameText().ToString();
+		FString PropertyName = ColumnProperty->GetAuthoredName();
+		FGuid PropertyID = FStructureEditorUtils::GetGuidForProperty(ColumnProperty);
+
+		// Checking if there are repeated pin names (Not supported).
+		if (PinNames.Contains(PinName))
+		{
+			bHasRepeatedNames = true;
+			continue;
+		}
 
-		FString ColumnName = DataTableUtils::GetPropertyExportName(ColumnProperty);
-		FString PinName = ColumnName;
-		FGuid ColumnPropertyId = FStructureEditorUtils::GetGuidForProperty(ColumnProperty);
+		PinNames.Add(PinName);
 
 		if (const FSoftObjectProperty* SoftObjectProperty = CastField<FSoftObjectProperty>(ColumnProperty))
 		{
@@ -424,30 +492,113 @@ void UCustomizableObjectNodeTable::AllocateDefaultPins(UCustomizableObjectNodeRe
 				{
 					if (Object->IsA(USkeletalMesh::StaticClass()) || Object->IsA(UStaticMesh::StaticClass()))
 					{
-						GenerateMeshPins(Object, ColumnName, ColumnPropertyId);
+						if (USkeletalMesh* SkeletalMesh = Cast<USkeletalMesh>(Object))
+						{
+							const int NumLODs = SkeletalMesh->GetLODNum();
+
+							for (int32 LODIndex = 0; LODIndex < NumLODs; ++LODIndex)
+							{
+								int32 NumMaterials = SkeletalMesh->GetImportedModel()->LODModels[LODIndex].Sections.Num();
 
-						if (FTableNodeColumnData* ColumnData = AuxOldColumnData.Find(ColumnPropertyId))
+								for (int32 MatIndex = 0; MatIndex < NumMaterials; ++MatIndex)
+								{
+									FString TableMeshPinName = GenerateSkeletalMeshMutableColumName(PinName, LODIndex, MatIndex);
+
+									// Pin Data
+									UCustomizableObjectNodeTableMeshPinData* PinData = NewObject<UCustomizableObjectNodeTableMeshPinData>(this);
+									PinData->ColumnPropertyName = PropertyName;
+									PinData->ColumnDisplayName = PinName;
+									PinData->StructColumnId = PropertyID;
+ 
+									// Mesh Data
+									PinData->LOD = LODIndex;
+									PinData->Material = MatIndex;
+
+									// Create a new pin for each lod and mesh section
+									UEdGraphPin* MeshPin = CustomCreatePin(EGPD_Output, Schema->PC_Mesh, FName(*TableMeshPinName), PinData);
+									MeshPin->SafeSetHidden(false);
+									MeshPin->PinToolTip = TableMeshPinName;
+
+									FSkeletalMeshModel* importModel = SkeletalMesh->GetImportedModel();
+									if (importModel && importModel->LODModels.Num() > LODIndex)
+									{
+										const uint32 NumberOfUVLayouts = importModel->LODModels[LODIndex].NumTexCoords;
+
+										for (uint32 LayoutIndex = 0; LayoutIndex < NumberOfUVLayouts; ++LayoutIndex)
+										{
+											UCustomizableObjectLayout* Layout = NewObject<UCustomizableObjectLayout>(this);
+											FString LayoutName = TableMeshPinName;
+
+											if (NumberOfUVLayouts > 1)
+											{
+												LayoutName += FString::Printf(TEXT(" UV_%d"), LayoutIndex);
+											}
+
+											if (Layout)
+											{
+												PinData->Layouts.Add(Layout);
+
+												Layout->SetLayout(LODIndex, MatIndex, LayoutIndex);
+												Layout->SetLayoutName(LayoutName);
+											}
+										}
+									}
+								}
+							}
+						}
+
+						else if (UStaticMesh* StaticMesh = Cast<UStaticMesh>(Object))
 						{
-							ColumnDataMap.Add(ColumnPropertyId, *ColumnData);
+							if (StaticMesh->GetRenderData()->LODResources.Num())
+							{
+								int32 NumMaterials = StaticMesh->GetRenderData()->LODResources[0].Sections.Num();
+
+								for (int32 MatIndex = 0; MatIndex < NumMaterials; ++MatIndex)
+								{
+									FString TableMeshPinName = GenerateStaticMeshMutableColumName(PinName, MatIndex);
+
+									UCustomizableObjectNodeTableMeshPinData* PinData = NewObject<UCustomizableObjectNodeTableMeshPinData>(this);
+									PinData->ColumnPropertyName = PropertyName;
+									PinData->ColumnDisplayName = PinName;
+									PinData->StructColumnId = PropertyID;
+
+									// Mesh Data
+									PinData->LOD = 0;
+									PinData->Material = MatIndex;
+
+									// Create a new pin for each lod and mesh section
+									UEdGraphPin* MeshPin = CustomCreatePin(EGPD_Output, Schema->PC_Mesh, FName(*TableMeshPinName), PinData);
+									MeshPin->SafeSetHidden(false);
+									MeshPin->PinToolTip = TableMeshPinName;
+								}
+							}
+						}
+
+						if (FTableNodeColumnData* ColumnData = AuxOldColumnData.Find(PinName))
+						{
+							PinColumnDataMap.Add(PinName, *ColumnData);
 						}
 					}
 
 					else if (Object->IsA(UTexture2D::StaticClass()))
 					{
 						UCustomizableObjectNodeTableImagePinData* PinData = NewObject<UCustomizableObjectNodeTableImagePinData>(this);
-						PinData->ColumnName = ColumnName;
-						PinData->StructColumnId = ColumnPropertyId;
+						PinData->ColumnPropertyName = PropertyName;
+						PinData->ColumnDisplayName = PinName;
+						PinData->StructColumnId = PropertyID;
+
+						// Texture Data
 						PinData->bIsNotTexture2D = false;
 						PinData->NodeTable = this;
 
 						FName PinCategory = PinData->ImageMode == ETableTextureType::PASSTHROUGH_TEXTURE ? Schema->PC_PassThroughImage : Schema->PC_Image;
 
-						for (UEdGraphPin* Pin : OldPins)
+						for (UEdGraphPin* OldPin : OldPins)
 						{
 							// Checking if this column already exist
-							if (UCustomizableObjectNodeTableImagePinData* OldPinData = Cast<UCustomizableObjectNodeTableImagePinData>(GetPinData(*(Pin))))
+							if (UCustomizableObjectNodeTableImagePinData* OldPinData = Cast<UCustomizableObjectNodeTableImagePinData>(GetPinData(*(OldPin))))
 							{
-								if (OldPinData->ColumnName == ColumnName)
+								if (OldPinData->ColumnPropertyName == PropertyName)
 								{
 									PinCategory = OldPinData->ImageMode == ETableTextureType::PASSTHROUGH_TEXTURE ? Schema->PC_PassThroughImage : Schema->PC_Image;
 									break;
@@ -461,8 +612,11 @@ void UCustomizableObjectNodeTable::AllocateDefaultPins(UCustomizableObjectNodeRe
 					else if (Object->IsA(UTexture::StaticClass()))
 					{
 						UCustomizableObjectNodeTableImagePinData* PinData = NewObject<UCustomizableObjectNodeTableImagePinData>(this);
-						PinData->ColumnName = ColumnName;
-						PinData->StructColumnId = ColumnPropertyId;
+						PinData->ColumnPropertyName = PropertyName;
+						PinData->ColumnDisplayName = PinName;
+						PinData->StructColumnId = PropertyID;
+
+						// Texture Data
 						PinData->ImageMode = ETableTextureType::PASSTHROUGH_TEXTURE;
 						PinData->bIsNotTexture2D = true;
 						PinData->NodeTable = this;
@@ -470,11 +624,12 @@ void UCustomizableObjectNodeTable::AllocateDefaultPins(UCustomizableObjectNodeRe
 						OutPin = CustomCreatePin(EGPD_Output, Schema->PC_PassThroughImage, FName(*PinName), PinData);
 					}
 
-					else if (Object->IsA(UMaterialInstance::StaticClass()))
+					else if (Object->IsA(UMaterialInterface::StaticClass()))
 					{
 						UCustomizableObjectNodeTableObjectPinData* PinData = NewObject<UCustomizableObjectNodeTableObjectPinData>(this);
-						PinData->ColumnName = ColumnName;
-						PinData->StructColumnId = ColumnPropertyId;
+						PinData->ColumnPropertyName = PropertyName;
+						PinData->ColumnDisplayName = PinName;
+						PinData->StructColumnId = PropertyID;
 
 						OutPin = CustomCreatePin(EGPD_Output, Schema->PC_MaterialAsset, FName(*PinName), PinData);
 					}
@@ -482,15 +637,16 @@ void UCustomizableObjectNodeTable::AllocateDefaultPins(UCustomizableObjectNodeRe
 					else if (Object->IsA(UPoseAsset::StaticClass()))
 					{
 						UCustomizableObjectNodeTableObjectPinData* PinData = NewObject<UCustomizableObjectNodeTableObjectPinData>(this);
-						PinData->ColumnName = ColumnName;
-						PinData->StructColumnId = ColumnPropertyId;
+						PinData->ColumnPropertyName = PropertyName;
+						PinData->ColumnDisplayName = PinName;
+						PinData->StructColumnId = PropertyID;
 
 						OutPin = CustomCreatePin(EGPD_Output, Schema->PC_PoseAsset, FName(*PinName), PinData);
 					}
 				}
 				else
 				{
-					const FText Text = FText::FromString(FString::Printf(TEXT("Could not find a Default Value in Structure member [%s]"), *ColumnName));
+					const FText Text = FText::FromString(FString::Printf(TEXT("Could not find a Default Value in Structure member [%s]"), *PinName));
 
 					FCustomizableObjectEditorLogger::CreateLog(Text)
 						.Category(ELoggerCategory::General)
@@ -506,8 +662,9 @@ void UCustomizableObjectNodeTable::AllocateDefaultPins(UCustomizableObjectNodeRe
 			if (StructProperty->Struct == TBaseStructure<FLinearColor>::Get())
 			{
 				UCustomizableObjectNodeTableObjectPinData* PinData = NewObject<UCustomizableObjectNodeTableObjectPinData>(this);
-				PinData->ColumnName = ColumnName;
-				PinData->StructColumnId = ColumnPropertyId;
+				PinData->ColumnPropertyName = PropertyName;
+				PinData->ColumnDisplayName = PinName;
+				PinData->StructColumnId = PropertyID;
 
 				OutPin = CustomCreatePin(EGPD_Output, Schema->PC_Color, FName(*PinName), PinData);
 			}
@@ -516,8 +673,9 @@ void UCustomizableObjectNodeTable::AllocateDefaultPins(UCustomizableObjectNodeRe
 		else if (const FNumericProperty* NumFloatProperty = CastField<FFloatProperty>(ColumnProperty))
 		{
 			UCustomizableObjectNodeTableObjectPinData* PinData = NewObject<UCustomizableObjectNodeTableObjectPinData>(this);
-			PinData->ColumnName = ColumnName;
-			PinData->StructColumnId = ColumnPropertyId;
+			PinData->ColumnPropertyName = PropertyName;
+			PinData->ColumnDisplayName = PinName;
+			PinData->StructColumnId = PropertyID;
 
 			OutPin = CustomCreatePin(EGPD_Output, Schema->PC_Float, FName(*PinName), PinData);
 		}
@@ -525,15 +683,16 @@ void UCustomizableObjectNodeTable::AllocateDefaultPins(UCustomizableObjectNodeRe
 		else if (const FNumericProperty* NumDoubleProperty = CastField<FDoubleProperty>(ColumnProperty))
 		{
 			UCustomizableObjectNodeTableObjectPinData* PinData = NewObject<UCustomizableObjectNodeTableObjectPinData>(this);
-			PinData->ColumnName = ColumnName;
-			PinData->StructColumnId = ColumnPropertyId;
+			PinData->ColumnPropertyName = PropertyName;
+			PinData->ColumnDisplayName = PinName;
+			PinData->StructColumnId = PropertyID;
 
 			OutPin = CustomCreatePin(EGPD_Output, Schema->PC_Float, FName(*PinName), PinData);
 		}
 
 		else if (const FObjectProperty* ObjectProperty = CastField<FObjectProperty>(ColumnProperty))
 		{
-			const FText Text = FText::FromString(FString::Printf(TEXT("Asset format not supported in Structure member [%s]. All assets should be Soft References."), *ColumnName));
+			const FText Text = FText::FromString(FString::Printf(TEXT("Asset format not supported in Structure member [%s]. All assets should be Soft References."), *PinName));
 
 			FCustomizableObjectEditorLogger::CreateLog(Text)
 				.Category(ELoggerCategory::General)
@@ -544,99 +703,22 @@ void UCustomizableObjectNodeTable::AllocateDefaultPins(UCustomizableObjectNodeRe
 
 		if (OutPin)
 		{
-			OutPin->PinFriendlyName = FText::FromString(PinName);
 			OutPin->PinToolTip = PinName;
 			OutPin->SafeSetHidden(false);
 		}
 	}
 
 	TableStruct->DestroyStruct(DefaultDataArray.GetData());
-}
-
-
-void UCustomizableObjectNodeTable::GenerateMeshPins(UObject* Mesh, const FString& ColumnName, const FGuid& ColumnId)
-{
-	if (USkeletalMesh* SkeletalMesh = Cast<USkeletalMesh>(Mesh))
-	{
-		const int NumLODs = SkeletalMesh->GetLODNum();
-		const UEdGraphSchema_CustomizableObject* Schema = GetDefault<UEdGraphSchema_CustomizableObject>();
-
-		for (int32 LODIndex = 0; LODIndex < NumLODs; ++LODIndex)
-		{
-			int32 NumMaterials = SkeletalMesh->GetImportedModel()->LODModels[LODIndex].Sections.Num();
-
-			for (int32 MatIndex = 0; MatIndex < NumMaterials; ++MatIndex)
-			{
-				FString TableMeshPinName = GenerateSkeletalMeshMutableColumName(ColumnName, LODIndex, MatIndex);
-
-				// Pin Data
-				UCustomizableObjectNodeTableMeshPinData* PinData = NewObject<UCustomizableObjectNodeTableMeshPinData>(this);
-				PinData->ColumnName = ColumnName;
-				PinData->StructColumnId = ColumnId;
-
-				// Mesh Data
-				PinData->MutableColumnName = TableMeshPinName;
-				PinData->LOD = LODIndex;
-				PinData->Material = MatIndex;
-
-				UEdGraphPin* MeshPin = CustomCreatePin(EGPD_Output, Schema->PC_Mesh, FName(*TableMeshPinName), PinData);
-				MeshPin->PinFriendlyName = FText::FromString(TableMeshPinName);
-				MeshPin->SafeSetHidden(false);
-
-				FSkeletalMeshModel* importModel = SkeletalMesh->GetImportedModel();
-				if (importModel && importModel->LODModels.Num() > LODIndex)
-				{
-					const uint32 NumberOfUVLayouts = importModel->LODModels[LODIndex].NumTexCoords;
-
-					for (uint32 LayoutIndex = 0; LayoutIndex < NumberOfUVLayouts; ++LayoutIndex)
-					{
-						UCustomizableObjectLayout* Layout = NewObject<UCustomizableObjectLayout>(this);
-						FString LayoutName = TableMeshPinName;
-
-						if (NumberOfUVLayouts > 1)
-						{
-							LayoutName += FString::Printf(TEXT(" UV_%d"), LayoutIndex);
-						}
 
-						if (Layout)
-						{
-							PinData->Layouts.Add(Layout);
-
-							Layout->SetLayout(LODIndex, MatIndex, LayoutIndex);
-							Layout->SetLayoutName(LayoutName);
-						}
-					}
-				}
-			}
-		}
-	}
-	
-	else if (UStaticMesh* StaticMesh = Cast<UStaticMesh>(Mesh))
+	// Repeated names log
+	if (bHasRepeatedNames)
 	{
-		const UEdGraphSchema_CustomizableObject* Schema = GetDefault<UEdGraphSchema_CustomizableObject>();
-
-		if (StaticMesh->GetRenderData()->LODResources.Num())
-		{
-			int32 NumMaterials = StaticMesh->GetRenderData()->LODResources[0].Sections.Num();
-
-			for (int32 MatIndex = 0; MatIndex < NumMaterials; ++MatIndex)
-			{
-				FString TableMeshPinName = GenerateStaticMeshMutableColumName(ColumnName, MatIndex);
-
-				UCustomizableObjectNodeTableMeshPinData* PinData = NewObject<UCustomizableObjectNodeTableMeshPinData>(this);
-				PinData->ColumnName = ColumnName;
-				PinData->StructColumnId = ColumnId;
-
-				// Mesh Data
-				PinData->MutableColumnName = TableMeshPinName;
-				PinData->LOD = 0;
-				PinData->Material = MatIndex;
-
-				UEdGraphPin* MeshPin = CustomCreatePin(EGPD_Output, Schema->PC_Mesh, FName(*TableMeshPinName), PinData);
-				MeshPin->PinFriendlyName = FText::FromString(TableMeshPinName);
-				MeshPin->SafeSetHidden(false);
-			}
-		}
+		FCustomizableObjectEditorLogger::CreateLog(FText::Format(LOCTEXT("ContainsRepeatedNamesWarning", "Table Node '{0}' contains 2 or more columns with repeated display names."
+			"The pins of these columns will not be available"), FText::FromString(ParameterName)))
+			.BaseObject()
+			.Severity(EMessageSeverity::Warning)
+			.Context(*this)
+			.Log();
 	}
 }
 
@@ -651,7 +733,7 @@ bool UCustomizableObjectNodeTable::IsNodeOutDatedAndNeedsRefresh()
 		return Pins.Num() > 0;
 	}
 
-	if (NumProperties != GetColumnTitles(TableStruct).Num())
+	if (NumProperties != GetColumnTitles().Num())
 	{
 		return true;
 	}
@@ -673,8 +755,7 @@ bool UCustomizableObjectNodeTable::IsNodeOutDatedAndNeedsRefresh()
 		if (ColumnProperty)
 		{
 			const UEdGraphSchema_CustomizableObject* Schema = GetDefault<UEdGraphSchema_CustomizableObject>();
-
-			UEdGraphPin* OutPin = nullptr;
+			const FString& PinName = ColumnProperty->GetDisplayNameText().ToString();
 
 			if (const FSoftObjectProperty* SoftObjectProperty = CastField<FSoftObjectProperty>(ColumnProperty))
 			{
@@ -705,10 +786,9 @@ bool UCustomizableObjectNodeTable::IsNodeOutDatedAndNeedsRefresh()
 
 							for (int32 MatIndex = 0; MatIndex < NumMaterials; ++MatIndex)
 							{
-								FString PropertyName = DataTableUtils::GetPropertyExportName(ColumnProperty);
-								FString PinName = GenerateSkeletalMeshMutableColumName(PropertyName, LODIndex, MatIndex);
+								FString MeshPinName = GenerateSkeletalMeshMutableColumName(PinName, LODIndex, MatIndex);
 
-								if (CheckPinUpdated(PinName, Schema->PC_Mesh))
+								if (CheckPinUpdated(MeshPinName, Schema->PC_Mesh))
 								{
 									bNeedsUpdate = true;
 								}
@@ -725,10 +805,9 @@ bool UCustomizableObjectNodeTable::IsNodeOutDatedAndNeedsRefresh()
 
 							for (int32 MatIndex = 0; MatIndex < NumMaterials; ++MatIndex)
 							{
-								FString PropertyName = DataTableUtils::GetPropertyExportName(ColumnProperty);
-								FString PinName = GenerateStaticMeshMutableColumName(PropertyName, MatIndex);
+								FString MeshPinName = GenerateStaticMeshMutableColumName(PinName, MatIndex);
 
-								if (CheckPinUpdated(PinName, Schema->PC_Mesh))
+								if (CheckPinUpdated(MeshPinName, Schema->PC_Mesh))
 								{
 									bNeedsUpdate = true;
 								}
@@ -740,8 +819,6 @@ bool UCustomizableObjectNodeTable::IsNodeOutDatedAndNeedsRefresh()
 				}
 				else if (Object->IsA(UTexture2D::StaticClass()))
 				{
-					FString PinName = DataTableUtils::GetPropertyExportName(ColumnProperty);
-					
 					if (CheckPinUpdated(PinName, Schema->PC_Image) && CheckPinUpdated(PinName, Schema->PC_PassThroughImage))
 					{
 						bNeedsUpdate = true;
@@ -751,8 +828,6 @@ bool UCustomizableObjectNodeTable::IsNodeOutDatedAndNeedsRefresh()
 				}
 				else if (Object->IsA(UTexture2DArray::StaticClass()))
 				{
-					FString PinName = DataTableUtils::GetPropertyExportName(ColumnProperty);
-
 					if (CheckPinUpdated(PinName, Schema->PC_PassThroughImage))
 					{
 						bNeedsUpdate = true;
@@ -760,10 +835,8 @@ bool UCustomizableObjectNodeTable::IsNodeOutDatedAndNeedsRefresh()
 
 					NumPins++;
 				}
-				else if (Object->IsA(UMaterialInstance::StaticClass()))
+				else if (Object->IsA(UMaterialInterface::StaticClass()))
 				{
-					FString PinName = DataTableUtils::GetPropertyExportName(ColumnProperty);
-					
 					if (CheckPinUpdated(PinName, Schema->PC_MaterialAsset))
 					{
 						bNeedsUpdate = true;
@@ -773,8 +846,6 @@ bool UCustomizableObjectNodeTable::IsNodeOutDatedAndNeedsRefresh()
 				}
 				else if (Object->IsA(UPoseAsset::StaticClass()))
 				{
-					FString PinName = DataTableUtils::GetPropertyExportName(ColumnProperty);
-					
 					if (CheckPinUpdated(PinName, Schema->PC_PoseAsset))
 					{
 						bNeedsUpdate = true;
@@ -787,8 +858,6 @@ bool UCustomizableObjectNodeTable::IsNodeOutDatedAndNeedsRefresh()
 			{
 				if (StructProperty->Struct == TBaseStructure<FLinearColor>::Get())
 				{
-					FString PinName = DataTableUtils::GetPropertyExportName(ColumnProperty);
-
 					if (CheckPinUpdated(PinName, Schema->PC_Color))
 					{
 						bNeedsUpdate = true;
@@ -799,8 +868,6 @@ bool UCustomizableObjectNodeTable::IsNodeOutDatedAndNeedsRefresh()
 			}
 			else if (const FFloatProperty* FloatNumProperty = CastField<FFloatProperty>(ColumnProperty))
 			{
-				FString PinName = DataTableUtils::GetPropertyExportName(ColumnProperty);
-
 				if (CheckPinUpdated(PinName, Schema->PC_Float))
 				{
 					bNeedsUpdate = true;
@@ -810,8 +877,6 @@ bool UCustomizableObjectNodeTable::IsNodeOutDatedAndNeedsRefresh()
 			}
 			else if (const FDoubleProperty* DoubleNumProperty = CastField<FDoubleProperty>(ColumnProperty))
 			{
-				FString PinName = DataTableUtils::GetPropertyExportName(ColumnProperty);
-
 				if (CheckPinUpdated(PinName, Schema->PC_Float))
 				{
 					bNeedsUpdate = true;
@@ -944,7 +1009,7 @@ UObject* UCustomizableObjectNodeTable::GetDefaultMeshForLayout(const UCustomizab
 			{
 				if (Layout == InLayout)
 				{
-					return GetColumnDefaultAssetByType<UObject>(MeshPinData->ColumnName);
+					return GetColumnDefaultAssetByType<UObject>(MeshPinData->ColumnPropertyName);
 				}
 			}
 		}
@@ -972,16 +1037,16 @@ TArray<UCustomizableObjectLayout*> UCustomizableObjectNodeTable::GetLayouts(cons
 }
 
 
-FString UCustomizableObjectNodeTable::GetColumnNameByPin(const UEdGraphPin* Pin) const
+FString UCustomizableObjectNodeTable::GetPinColumnName(const UEdGraphPin* Pin) const
 {
-	const UCustomizableObjectNodeTableObjectPinData* PinData = Cast<UCustomizableObjectNodeTableObjectPinData >(GetPinData(*Pin));
-	
+	const UCustomizableObjectNodeTableObjectPinData* PinData = Cast<UCustomizableObjectNodeTableObjectPinData>(GetPinData(*Pin));
+
 	if (PinData)
 	{
-		return PinData->ColumnName;
+		return PinData->ColumnDisplayName;
 	}
 
-	return Pin->PinFriendlyName.ToString();
+	return FString();
 }
 
 
@@ -1000,14 +1065,15 @@ void UCustomizableObjectNodeTable::GetPinLODAndSection(const UEdGraphPin* Pin, i
 }
 
 
-void UCustomizableObjectNodeTable::GetAnimationColumns(const FGuid& ColumnId, FString& AnimBPColumnName, FString& AnimSlotColumnName, FString& AnimTagColumnName) const
+void UCustomizableObjectNodeTable::GetAnimationColumns(const FString& ColumnDisplayName, FString& AnimBPColumnName, FString& AnimSlotColumnName, FString& AnimTagColumnName) const
 {
-	if (ColumnDataMap.Contains(ColumnId))
+	if (PinColumnDataMap.Contains(ColumnDisplayName))
 	{
-		const FTableNodeColumnData* ColumnData = ColumnDataMap.Find(ColumnId);
-		AnimBPColumnName = ColumnData->AnimInstanceColumnName;
-		AnimSlotColumnName = ColumnData->AnimSlotColumnName;
-		AnimTagColumnName = ColumnData->AnimTagColumnName;
+		const FTableNodeColumnData& ColumnData = PinColumnDataMap[ColumnDisplayName];
+
+		AnimBPColumnName = ColumnData.AnimInstanceColumnName;
+		AnimSlotColumnName = ColumnData.AnimSlotColumnName;
+		AnimTagColumnName = ColumnData.AnimTagColumnName;
 	}
 }
 
@@ -1039,11 +1105,9 @@ FSoftObjectPtr UCustomizableObjectNodeTable::GetSkeletalMeshAt(const UEdGraphPin
 	
 	const UScriptStruct* TableStruct = DataTable->GetRowStruct();
 
-	FString ColumnName = GetColumnNameByPin(Pin);
-
 	// Here we are using the DataTable function because we are only calling GetSkeletalMeshAt() at Generation Time
 	// If we want to use it elswhere, we can use our FindTablreProperty method
-	FProperty* ColumnProperty = DataTable->FindTableProperty(FName(*ColumnName));
+	FProperty* ColumnProperty = FindPinProperty(*Pin);
 
 	if (!ColumnProperty)
 	{
@@ -1079,13 +1143,13 @@ TSoftClassPtr<UAnimInstance> UCustomizableObjectNodeTable::GetAnimInstanceAt(con
 		return TSoftClassPtr<UAnimInstance>();
 	}
 
-	if (const FTableNodeColumnData* ColumnData = ColumnDataMap.Find(PinData->StructColumnId))
+	if (const FTableNodeColumnData* ColumnData = PinColumnDataMap.Find(PinData->ColumnDisplayName))
 	{
 		FString AnimColumn = ColumnData->AnimInstanceColumnName;
 
 		// Here we are using the DataTable function because we are only calling GetAnimInstanceAt() at Generation Time
 		// If we want to use it elswhere, we can use our FindTablreProperty method
-		FProperty* ColumnProperty = DataTable->FindTableProperty(FName(*AnimColumn));
+		FProperty* ColumnProperty = FindColumnProperty(FName(*AnimColumn));
 
 		if (!ColumnProperty)
 		{
@@ -1187,13 +1251,13 @@ TSharedPtr<IDetailsView> UCustomizableObjectNodeTable::CustomizePinDetails(const
 }
 
 
-ETableTextureType UCustomizableObjectNodeTable::GetColumnImageMode(const FString& ColumnName) const
+ETableTextureType UCustomizableObjectNodeTable::GetColumnImageMode(const FString& ColumnPropertyName) const
 {
 	for (const UEdGraphPin* Pin : Pins)
 	{
 		if (const UCustomizableObjectNodeTableImagePinData* PinData = Cast<UCustomizableObjectNodeTableImagePinData>(GetPinData(*(Pin))))
 		{
-			if (PinData->ColumnName == ColumnName)
+			if (PinData->ColumnPropertyName == ColumnPropertyName)
 			{
 				return PinData->ImageMode;
 			}
@@ -1207,13 +1271,9 @@ ETableTextureType UCustomizableObjectNodeTable::GetColumnImageMode(const FString
 
 ETableMeshPinType UCustomizableObjectNodeTable::GetPinMeshType(const UEdGraphPin* Pin) const
 {
-	// Getting Struct Pointer
-	const UScriptStruct* TableStruct = GetTableNodeStruct();
-
-	if (TableStruct && Pin && Pin->PinType.PinCategory == UEdGraphSchema_CustomizableObject::PC_Mesh)
+	if ( Pin && Pin->PinType.PinCategory == UEdGraphSchema_CustomizableObject::PC_Mesh)
 	{
-		FString ColumnName = GetColumnNameByPin(Pin);
-		FProperty* Property = FindTableProperty(TableStruct, FName(*ColumnName));
+		FProperty* Property = FindPinProperty(*Pin);
 
 		if (const FSoftObjectProperty* SoftObjectProperty = CastField<FSoftObjectProperty>(Property))
 		{
@@ -1267,13 +1327,16 @@ const UScriptStruct* UCustomizableObjectNodeTable::GetTableNodeStruct() const
 }
 
 
-TArray<FString> UCustomizableObjectNodeTable::GetColumnTitles(const UScriptStruct* ScriptStruct) const
+TArray<FString> UCustomizableObjectNodeTable::GetColumnTitles() const
 {
 	TArray<FString> Result;
 	Result.Add(TEXT("Name"));
-	if (ScriptStruct)
+
+	const UScriptStruct* TableStruct = GetTableNodeStruct();
+
+	if (TableStruct)
 	{
-		for (TFieldIterator<FProperty> It(ScriptStruct); It; ++It)
+		for (TFieldIterator<FProperty> It(TableStruct); It; ++It)
 		{
 			FProperty* Prop = *It;
 			check(Prop != nullptr);
@@ -1286,63 +1349,78 @@ TArray<FString> UCustomizableObjectNodeTable::GetColumnTitles(const UScriptStruc
 }
 
 
-FProperty* UCustomizableObjectNodeTable::FindTableProperty(const UScriptStruct* ScriptStruct, const FName& PropertyName) const
+FProperty* UCustomizableObjectNodeTable::FindTableProperty(const FName& PropertyName) const
 {
+	const UScriptStruct* TableStruct = GetTableNodeStruct();
+	
+	if (!TableStruct)
+	{
+		return nullptr;
+	}
+
 	FProperty* Property = nullptr;
 
-	if (ScriptStruct)
+	Property = TableStruct->FindPropertyByName(PropertyName);
+	if (Property == nullptr && TableStruct->IsA<UUserDefinedStruct>())
 	{
-		Property = ScriptStruct->FindPropertyByName(PropertyName);
-		if (Property == nullptr && ScriptStruct->IsA<UUserDefinedStruct>())
-		{
-			const FString PropertyNameStr = PropertyName.ToString();
+		const FString PropertyNameStr = PropertyName.ToString();
 
-			for (TFieldIterator<FProperty> It(ScriptStruct); It; ++It)
+		for (TFieldIterator<FProperty> It(TableStruct); It; ++It)
+		{
+			if (PropertyNameStr == TableStruct->GetAuthoredNameForField(*It))
 			{
-				if (PropertyNameStr == ScriptStruct->GetAuthoredNameForField(*It))
-				{
-					Property = *It;
-					break;
-				}
+				Property = *It;
+				break;
 			}
 		}
-		if (!DataTableUtils::IsSupportedTableProperty(Property))
-		{
-			Property = nullptr;
-		}
+	}
+
+	if (!DataTableUtils::IsSupportedTableProperty(Property))
+	{
+		Property = nullptr;
 	}
 
 	return Property;
 }
 
 
-FGuid UCustomizableObjectNodeTable::GetColumnIdByName(const FName& ColumnName) const
+FProperty* UCustomizableObjectNodeTable::FindPinProperty(const UEdGraphPin& Pin) const
 {
-	if (const UScriptStruct* TableStruct = GetTableNodeStruct())
+	if (UCustomizableObjectNodeTableObjectPinData* PinData = Cast<UCustomizableObjectNodeTableObjectPinData>(GetPinData(Pin)))
+	{
+		return FindTableProperty(FName(*PinData->ColumnPropertyName));
+	}
+
+	return nullptr;
+}
+
+
+FProperty* UCustomizableObjectNodeTable::FindColumnProperty(const FName& ColumnDisplayName) const
+{
+	const UScriptStruct* TableStruct = GetTableNodeStruct();
+
+	if (!TableStruct)
 	{
-		if(FProperty* ColumnProperty = FindTableProperty(TableStruct, ColumnName))
+		return nullptr;
+	}
+
+	FProperty* Property = nullptr;
+
+	for (TFieldIterator<FProperty> It(TableStruct); It; ++It)
+	{
+		if (ColumnDisplayName == FName(*It->GetDisplayNameText().ToString()))
 		{
-			return FStructureEditorUtils::GetGuidForProperty(ColumnProperty);
+			Property = *It;
+			break;
 		}
-		else
-		{
-			// We could not find the property by name
-			for (const UEdGraphPin* Pin : GetAllNonOrphanPins())
-			{
-				if (UCustomizableObjectNodeTableMeshPinData* PinData = Cast< UCustomizableObjectNodeTableMeshPinData>(GetPinData(*Pin)))
-				{
-					if (PinData->ColumnName != ColumnName)
-					{
-						continue;
-					}
+	}
 
-					return PinData->StructColumnId;
-				}
-			}
-		}
+	if (!DataTableUtils::IsSupportedTableProperty(Property))
+	{
+		Property = nullptr;
 	}
 
-	return FGuid();
+	return Property;
 }
 
 
diff --git a/Engine/Plugins/Mutable/Source/CustomizableObjectEditor/Private/MuCOE/Nodes/CustomizableObjectNodeTableDetails.cpp b/Engine/Plugins/Mutable/Source/CustomizableObjectEditor/Private/MuCOE/Nodes/CustomizableObjectNodeTableDetails.cpp
--- a/Engine/Plugins/Mutable/Source/CustomizableObjectEditor/Private/MuCOE/Nodes/CustomizableObjectNodeTableDetails.cpp
+++ b/Engine/Plugins/Mutable/Source/CustomizableObjectEditor/Private/MuCOE/Nodes/CustomizableObjectNodeTableDetails.cpp
@@ -269,7 +269,7 @@ void FCustomizableObjectNodeTableDetails::GenerateMeshColumnComboBoxOptions()
 			if (SoftObjectProperty->PropertyClass->IsChildOf(USkeletalMesh::StaticClass())
 				|| SoftObjectProperty->PropertyClass->IsChildOf(UStaticMesh::StaticClass()))
 			{
-				FString MeshColumnName = DataTableUtils::GetPropertyExportName(ColumnProperty);
+				FString MeshColumnName = ColumnProperty->GetDisplayNameText().ToString();
 				AnimMeshColumnOptionNames.Add(MakeShareable(new FString(MeshColumnName)));
 			}
 		}
@@ -304,21 +304,21 @@ void FCustomizableObjectNodeTableDetails::GenerateMeshSectionOptions(TArray<FLay
 			if (SoftObjectProperty->PropertyClass->IsChildOf(USkeletalMesh::StaticClass())
 				|| SoftObjectProperty->PropertyClass->IsChildOf(UStaticMesh::StaticClass()))
 			{
-				FString MeshColumnName = DataTableUtils::GetPropertyExportName(ColumnProperty);
+				FString MeshColumnName = ColumnProperty->GetAuthoredName();
 
 				for (const UEdGraphPin* Pin : Node->GetAllNonOrphanPins())
 				{
 					const UCustomizableObjectNodeTableMeshPinData* PinData = Cast<UCustomizableObjectNodeTableMeshPinData >(Node->GetPinData(*Pin));
 
-					if (!PinData || PinData->ColumnName != MeshColumnName || Node->GetPinMeshType(Pin) != ETableMeshPinType::SKELETAL_MESH)
+					if (!PinData || PinData->ColumnPropertyName != MeshColumnName || Node->GetPinMeshType(Pin) != ETableMeshPinType::SKELETAL_MESH)
 					{
 						continue;
 					}
 
-					if (PinData && PinData->ColumnName == MeshColumnName)
+					if (PinData && PinData->ColumnPropertyName == MeshColumnName)
 					{
 						FLayoutEditorMeshSection& MeshSection = OutMeshSections.AddDefaulted_GetRef();
-						MeshSection.MeshName = MakeShareable(new FString(Pin->PinFriendlyName.ToString()));
+						MeshSection.MeshName = MakeShareable(new FString(Pin->PinName.ToString()));
 
 						for (UCustomizableObjectLayout* Layout : PinData->Layouts)
 						{
@@ -357,21 +357,14 @@ void FCustomizableObjectNodeTableDetails::GenerateAnimInstanceComboBoxOptions()
 
 	const UScriptStruct* TableStruct = Node->GetTableNodeStruct();
 
-	FString ColumnName;
-	FTableNodeColumnData* MeshColumnData = nullptr;
-
-	if (TableStruct && AnimMeshColumnComboBox.IsValid())
-	{
-		ColumnName = *AnimMeshColumnComboBox->GetSelectedItem();
-		FGuid ColumnId = Node->GetColumnIdByName(FName(*ColumnName));
-
-		MeshColumnData = Node->ColumnDataMap.Find(ColumnId);
-	}
-	else
+	if (!TableStruct || !AnimMeshColumnComboBox.IsValid())
 	{
 		return;
 	}
 
+	FString ColumnName = *AnimMeshColumnComboBox->GetSelectedItem();
+	FTableNodeColumnData* MeshColumnData = Node->PinColumnDataMap.Find(ColumnName);
+
 	// Fill in name option arrays and set the selected item if any
 	for (TFieldIterator<FProperty> It(TableStruct); It; ++It)
 	{
@@ -381,7 +374,7 @@ void FCustomizableObjectNodeTableDetails::GenerateAnimInstanceComboBoxOptions()
 			{
 				if (SoftClassProperty->MetaClass->IsChildOf(UAnimInstance::StaticClass()))
 				{
-					TSharedPtr<FString> Option = MakeShareable(new FString(DataTableUtils::GetPropertyExportName(ColumnProperty)));
+					TSharedPtr<FString> Option = MakeShareable(new FString(ColumnProperty->GetDisplayNameText().ToString()));
 					AnimOptionNames.Add(Option);
 
 					if (MeshColumnData && MeshColumnData->AnimInstanceColumnName == *Option)
@@ -393,7 +386,7 @@ void FCustomizableObjectNodeTableDetails::GenerateAnimInstanceComboBoxOptions()
 
 			else if (CastField<FIntProperty>(ColumnProperty) || CastField<FNameProperty>(ColumnProperty))
 			{
-				TSharedPtr<FString> Option = MakeShareable(new FString(DataTableUtils::GetPropertyExportName(ColumnProperty)));
+				TSharedPtr<FString> Option = MakeShareable(new FString(ColumnProperty->GetDisplayNameText().ToString()));
 				AnimSlotOptionNames.Add(Option);
 
 				if (MeshColumnData && MeshColumnData->AnimSlotColumnName == *Option)
@@ -406,7 +399,7 @@ void FCustomizableObjectNodeTableDetails::GenerateAnimInstanceComboBoxOptions()
 			{
 				if (StructProperty->Struct == TBaseStructure<FGameplayTagContainer>::Get())
 				{
-					TSharedPtr<FString> Option = MakeShareable(new FString(DataTableUtils::GetPropertyExportName(ColumnProperty)));
+					TSharedPtr<FString> Option = MakeShareable(new FString(ColumnProperty->GetDisplayNameText().ToString()));
 					AnimTagsOptionNames.Add(Option);
 
 					if (MeshColumnData && MeshColumnData->AnimTagColumnName == *Option)
@@ -447,19 +440,18 @@ void FCustomizableObjectNodeTableDetails::OnAnimInstanceComboBoxSelectionChanged
 	if (bIsMeshSelectionValid && Selection.IsValid() && SelectInfo != ESelectInfo::Direct)
 	{
 		FString ColumnName = *AnimMeshColumnComboBox->GetSelectedItem();
-		FGuid ColumnId = Node->GetColumnIdByName(FName(*ColumnName));
-		FTableNodeColumnData* MeshColumnData = Node->ColumnDataMap.Find(ColumnId);
+		FTableNodeColumnData* MeshColumnData = Node->PinColumnDataMap.Find(ColumnName);
 		
 		if (MeshColumnData)
 		{
 			MeshColumnData->AnimInstanceColumnName = *Selection;
 		}
-		else if(ColumnId.IsValid())
+		else if(!ColumnName.IsEmpty())
 		{
 			FTableNodeColumnData NewMeshColumnData;
 			NewMeshColumnData.AnimInstanceColumnName = *Selection;
 
-			Node->ColumnDataMap.Add(ColumnId, NewMeshColumnData);
+			Node->PinColumnDataMap.Add(ColumnName, NewMeshColumnData);
 		}
 
 		Node->MarkPackageDirty();
@@ -474,19 +466,18 @@ void FCustomizableObjectNodeTableDetails::OnAnimSlotComboBoxSelectionChanged(TSh
 	if (bIsMeshSelectionValid && Selection.IsValid() && SelectInfo != ESelectInfo::Direct)
 	{
 		FString ColumnName = *AnimMeshColumnComboBox->GetSelectedItem();
-		FGuid ColumnId = Node->GetColumnIdByName(FName(*ColumnName));
-		FTableNodeColumnData* MeshColumnData = Node->ColumnDataMap.Find(ColumnId);
+		FTableNodeColumnData* MeshColumnData = Node->PinColumnDataMap.Find(ColumnName);
 
 		if (MeshColumnData)
 		{
 			MeshColumnData->AnimSlotColumnName = *Selection;
 		}
-		else if (ColumnId.IsValid())
+		else if (!ColumnName.IsEmpty())
 		{
 			FTableNodeColumnData NewMeshColumnData;
 			NewMeshColumnData.AnimSlotColumnName = *Selection;
 
-			Node->ColumnDataMap.Add(ColumnId, NewMeshColumnData);
+			Node->PinColumnDataMap.Add(ColumnName, NewMeshColumnData);
 		}
 
 		Node->MarkPackageDirty();
@@ -501,21 +492,18 @@ void FCustomizableObjectNodeTableDetails::OnAnimTagsComboBoxSelectionChanged(TSh
 	if (bIsMeshSelectionValid && Selection.IsValid() && SelectInfo != ESelectInfo::Direct)
 	{
 		FString ColumnName = *AnimMeshColumnComboBox->GetSelectedItem();
-		FGuid ColumnId = Node->GetColumnIdByName(FName(*ColumnName));
-		FTableNodeColumnData* MeshColumnData = Node->ColumnDataMap.Find(ColumnId);
-
-		MeshColumnData = Node->ColumnDataMap.Find(ColumnId);
+		FTableNodeColumnData* MeshColumnData = Node->PinColumnDataMap.Find(ColumnName);
 
 		if (MeshColumnData)
 		{
 			MeshColumnData->AnimTagColumnName = *Selection;
 		}
-		else if (ColumnId.IsValid())
+		else if (!ColumnName.IsEmpty())
 		{
 			FTableNodeColumnData NewMeshColumnData;
 			NewMeshColumnData.AnimTagColumnName = *Selection;
 
-			Node->ColumnDataMap.Add(ColumnId, NewMeshColumnData);
+			Node->PinColumnDataMap.Add(ColumnName, NewMeshColumnData);
 		}
 
 		Node->MarkPackageDirty();
@@ -540,8 +528,7 @@ void FCustomizableObjectNodeTableDetails::OnAnimCustomRowResetButtonClicked(EAni
 	}
 
 	FString ColumnName = *AnimMeshColumnComboBox->GetSelectedItem();
-	FGuid ColumnId = Node->GetColumnIdByName(FName(*ColumnName));
-	FTableNodeColumnData* MeshColumnData = Node->ColumnDataMap.Find(ColumnId);
+	FTableNodeColumnData* MeshColumnData = Node->PinColumnDataMap.Find(ColumnName);
 
 	if (!MeshColumnData)
 	{
@@ -606,7 +593,7 @@ TSharedPtr<FString> FCustomizableObjectNodeTableDetails::GenerateMutableMetaData
 		{
 			if (StructProperty->Struct == FMutableParamUIMetadata::StaticStruct())
 			{
-				TSharedPtr<FString> Option = MakeShareable(new FString(DataTableUtils::GetPropertyExportName(ColumnProperty)));
+				TSharedPtr<FString> Option = MakeShareable(new FString(ColumnProperty->GetDisplayNameText().ToString()));
 				MutableMetaDataColumnsOptionNames.Add(Option);
 
 				if (*Option == Node->ParamUIMetadataColumn)
@@ -653,7 +640,7 @@ void FCustomizableObjectNodeTableDetails::OnMutableMetaDataColumnComboBoxSelecti
 
 FSlateColor FCustomizableObjectNodeTableDetails::GetComboBoxTextColor(TArray<TSharedPtr<FString>>* CurrentOptions, const FName ColumnName) const
 {	
-	if (Node->FindTableProperty(Node->GetTableNodeStruct(), ColumnName) || ColumnName.IsNone())
+	if (Node->FindColumnProperty(ColumnName) || ColumnName.IsNone())
 	{
 		return FSlateColor::UseForeground();
 	}
@@ -699,7 +686,7 @@ TSharedPtr<FString> FCustomizableObjectNodeTableDetails::GenerateThumbnailColumn
 
 		if (const FSoftObjectProperty* ObjectProperty = CastField<FSoftObjectProperty>(ColumnProperty))
 		{
-			TSharedPtr<FString> Option = MakeShareable(new FString(DataTableUtils::GetPropertyExportName(ColumnProperty)));
+			TSharedPtr<FString> Option = MakeShareable(new FString(ColumnProperty->GetDisplayNameText().ToString()));
 			ThumbnailColumnOptionNames.Add(Option);
 
 			if (*Option == Node->ThumbnailColumn)
@@ -779,7 +766,7 @@ TSharedPtr<FString> FCustomizableObjectNodeTableDetails::GenerateVersionColumnCo
 			continue;
 		}
 
-		TSharedPtr<FString> Option = MakeShareable(new FString(DataTableUtils::GetPropertyExportName(ColumnProperty)));
+		TSharedPtr<FString> Option = MakeShareable(new FString(ColumnProperty->GetDisplayNameText().ToString()));
 		VersionColumnsOptionNames.Add(Option);
 
 		if (*Option == Node->VersionColumn)
@@ -824,7 +811,7 @@ void FCustomizableObjectNodeTableDetails::OnVersionColumnComboBoxSelectionChange
 
 FSlateColor FCustomizableObjectNodeTableDetails::GetVersionColumnComboBoxTextColor(TArray<TSharedPtr<FString>>* CurrentOptions) const
 {
-	if (Node->FindTableProperty(Node->GetTableNodeStruct(), Node->VersionColumn) || Node->VersionColumn.IsNone())
+	if (Node->FindColumnProperty(Node->VersionColumn) || Node->VersionColumn.IsNone())
 	{
 		return FSlateColor::UseForeground();
 	}
diff --git a/Engine/Plugins/Mutable/Source/CustomizableObjectEditor/Public/MuCOE/Nodes/CustomizableObjectNodeTable.h b/Engine/Plugins/Mutable/Source/CustomizableObjectEditor/Public/MuCOE/Nodes/CustomizableObjectNodeTable.h
--- a/Engine/Plugins/Mutable/Source/CustomizableObjectEditor/Public/MuCOE/Nodes/CustomizableObjectNodeTable.h
+++ b/Engine/Plugins/Mutable/Source/CustomizableObjectEditor/Public/MuCOE/Nodes/CustomizableObjectNodeTable.h
@@ -86,10 +86,17 @@ public:
 	UPROPERTY()
 	FGuid StructColumnId;
 
-	/** Name of the data table column related to the pin */
+	/** Name of the data table column property related to the pin */
 	UPROPERTY()
-	FString ColumnName = "";
+	FString ColumnName_DEPRECATED = "";
 
+	/** Unique Name of the struct property related to the pin. Used to data purposes: search struct porperty, set mutable column name... */
+	UPROPERTY()
+	FString ColumnPropertyName = "";
+
+	/** Name of the data table column related to the pin. Used for UI purposes.*/
+	UPROPERTY()
+	FString ColumnDisplayName = "";
 };
 
 
@@ -142,9 +149,6 @@ public:
 	UPROPERTY()
 	FString AnimTagColumnName_DEPRECATED = "";
 
-	UPROPERTY()
-	FString MutableColumnName = "";
-
 	/** LOD of the mesh related to this Mesh pin */
 	UPROPERTY()
 	int32 LOD = 0;
@@ -231,7 +235,13 @@ public:
 
 	/** Map to relate a Structure Column with its Data */
 	UPROPERTY()
-	TMap<FGuid, FTableNodeColumnData> ColumnDataMap;
+	TMap<FGuid, FTableNodeColumnData> ColumnDataMap_DEPRECATED;
+
+	/** Map to relate a Structure Column with its Data 
+	* Key: ColumnPropertyName variable of the PinData	
+	*/
+	UPROPERTY()
+	TMap<FString, FTableNodeColumnData> PinColumnDataMap;
 
 public:
 
@@ -272,31 +282,19 @@ public:
 	TArray<UCustomizableObjectLayout*> GetLayouts(const UEdGraphPin* Pin) const;
 
 	// Returns the name of the table column related to a pin
-	FString GetColumnNameByPin(const UEdGraphPin* Pin) const;
+	FString GetPinColumnName(const UEdGraphPin* Pin) const;
 
 	// Returns the LOD of the mesh associated to the input pin
 	void GetPinLODAndSection(const UEdGraphPin* Pin, int32& LODIndex, int32& SectionIndex) const;
 
 	// Get the anim blueprint and anim slot columns related to a mesh
-	void GetAnimationColumns(const FGuid& ColumnId, FString& AnimBPColumnName, FString& AnimSlotColumnName, FString& AnimTagColumnName) const;
+	void GetAnimationColumns(const FString& ColumnPropertyName, FString& AnimBPColumnName, FString& AnimSlotColumnName, FString& AnimTagColumnName) const;
 
 	template<class T> 
-	T* GetColumnDefaultAssetByType(FString ColumnName) const
+	T* GetColumnDefaultAssetByType(FString ColumnPropertyName) const
 	{
 		T* ObjectToReturn = nullptr;
-		const UScriptStruct* TableStruct = nullptr;
-
-		if (TableDataGatheringMode == ETableDataGatheringSource::ETDGM_AssetRegistry)
-		{
-			TableStruct = Structure;
-		}
-		else
-		{
-			if (Table)
-			{
-				TableStruct = Table->GetRowStruct();
-			}
-		}
+		const UScriptStruct* TableStruct = GetTableNodeStruct();
 
 		if (TableStruct)
 		{
@@ -306,7 +304,7 @@ public:
 			DeafaultDataArray.SetNumZeroed(TableStruct->GetStructureSize());
 			TableStruct->InitializeStruct(DeafaultDataArray.GetData());
 
-			FProperty* Property = FindTableProperty(TableStruct, FName(*ColumnName));
+			FProperty* Property = FindTableProperty(FName(*ColumnPropertyName));
 
 			if (Property)
 			{
@@ -344,11 +342,9 @@ public:
 	{
 		if (Pin)
 		{
-			FString ColumnName = GetColumnNameByPin(Pin);
-
-			if (ColumnName != FString())
+			if (const UCustomizableObjectNodeTableObjectPinData* PinData = Cast<UCustomizableObjectNodeTableObjectPinData >(GetPinData(*Pin)))
 			{
-				return GetColumnDefaultAssetByType<T>(ColumnName);
+				return GetColumnDefaultAssetByType<T>(PinData->ColumnPropertyName);
 			}
 		}
 
@@ -361,7 +357,7 @@ public:
 	TSoftClassPtr<UAnimInstance> GetAnimInstanceAt(const UEdGraphPin* Pin, const UDataTable* DataTable, const FName& RowName) const;
 	
 	// Returns the image mode of the column
-	ETableTextureType GetColumnImageMode(const FString& ColumnName) const;
+	ETableTextureType GetColumnImageMode(const FString& ColumnPropertyName) const;
 
 	// Returns the mesh type of the Pin
 	ETableMeshPinType GetPinMeshType(const UEdGraphPin* Pin) const;
@@ -375,13 +371,16 @@ public:
 
 	// Methods from UDataTable but modified to support ScriptStructs
 	/** Get an array of all the column titles, using the friendly display name from the property */
-	TArray<FString> GetColumnTitles(const UScriptStruct* ScriptStruct) const;
+	TArray<FString> GetColumnTitles() const;
 	
-	/** Returns the column property where PropertyName matches the name of the column property. Returns nullptr if no match is found or the match is not a supported table property */
-	FProperty* FindTableProperty(const UScriptStruct* ScriptStruct, const FName& PropertyName) const;
+	/** Returns the property using its name. */
+	FProperty* FindTableProperty(const FName& PropertyName) const;
 
-	/** Returns the column id */
-	FGuid GetColumnIdByName(const FName& ColumnName) const;
+	/** Returns the property linked to a pin */
+	FProperty* FindPinProperty(const UEdGraphPin& Pin) const;
+
+	/** Returns the property linked to a column display name. Useful for properties that do not generate a pin. */
+	FProperty* FindColumnProperty(const FName& ColumnDisplayName) const;
 
 	/** Return the list of UDataTable that will be used to compose the final UDataTable. */
 	TArray<FAssetData> GetParentTables() const;
@@ -401,9 +400,6 @@ private:
 	int32 NumProperties;
 	
 	FDelegateHandle OnTableChangedDelegateHandle;
-	
-	// Generates a mesh pin for each LOD and Material Surface of the reference SkeletalMesh
-	void GenerateMeshPins(UObject* Mesh, const FString& ColumnName, const FGuid& ColumnId);
 
 	// Checks if a pin already exists and if it has the same type as before the node refresh
 	bool CheckPinUpdated(const FString& PinName, const FName& PinType) const;
