[Mutable] Better UV Layout Material Error message.

#rnx

diff --git a/Engine/Plugins/Mutable/Source/CustomizableObjectEditor/Private/MuCOE/Nodes/CustomizableObjectNodeMaterial.cpp b/Engine/Plugins/Mutable/Source/CustomizableObjectEditor/Private/MuCOE/Nodes/CustomizableObjectNodeMaterial.cpp
--- a/Engine/Plugins/Mutable/Source/CustomizableObjectEditor/Private/MuCOE/Nodes/CustomizableObjectNodeMaterial.cpp
+++ b/Engine/Plugins/Mutable/Source/CustomizableObjectEditor/Private/MuCOE/Nodes/CustomizableObjectNodeMaterial.cpp
@@ -961,7 +961,9 @@ int32 UCustomizableObjectNodeMaterial::GetImageUVLayout(const int32 ImageIndex)
 	const int32 UVIndex = GetImageUVLayoutFromMaterial(ImageIndex);
 	if (UVIndex == -1)
 	{
-		FCustomizableObjectEditorLogger::CreateLog(LOCTEXT("UVLayoutMaterialError", "Could not deduce the UV Layout Index from the UMaterial. Nodes connected to the Texture Sample UVs pin are supported."))
+		const FText ParamName = FText::FromName(GetParameterName(EMaterialParameterType::Texture, ImageIndex));
+		const FText Msg = FText::Format(LOCTEXT("UVLayoutMaterialError", "Could not deduce the UV Layout Index in [{0}]. Use [Index] UV Layout Mode or, in the UMaterial, remove any nodes connected to the [UVs] pin in [{0}] node ."), ParamName);
+		FCustomizableObjectEditorLogger::CreateLog(Msg)
 			.Severity(EMessageSeverity::Warning)
 			.Context(*this)
 			.Log();
