[Mutable] Fixed repeated name declaration.

#rnx

diff --git a/Engine/Plugins/Mutable/Source/CustomizableObjectEditor/Private/MuCOE/Nodes/CustomizableObjectNodeCopyMaterial.cpp b/Engine/Plugins/Mutable/Source/CustomizableObjectEditor/Private/MuCOE/Nodes/CustomizableObjectNodeCopyMaterial.cpp
--- a/Engine/Plugins/Mutable/Source/CustomizableObjectEditor/Private/MuCOE/Nodes/CustomizableObjectNodeCopyMaterial.cpp
+++ b/Engine/Plugins/Mutable/Source/CustomizableObjectEditor/Private/MuCOE/Nodes/CustomizableObjectNodeCopyMaterial.cpp
@@ -13,15 +13,17 @@
 
 #define LOCTEXT_NAMESPACE "CustomizableObjectEditor"
 
+namespace CustomizableObjectNodeCopyMaterial
+{
+	/** Mesh input pin key */
+	const static FString MeshPinName = TEXT("Mesh_Input_Pin");
 
-/** Mesh input pin key */
-const static FString MeshPinName = TEXT("Mesh_Input_Pin");
-
-/** Material input pin key */
-const static FString MeshSectionPinName = TEXT("MeshSection_Input_Pin");
+	/** Material input pin key */
+	const static FString MeshSectionPinName = TEXT("MeshSection_Input_Pin");
 
-/** Material output pin key*/
-const static FString OutputPinName = TEXT("MeshSection_Output_Pin");
+	/** Material output pin key*/
+	const static FString OutputPinName = TEXT("MeshSection_Output_Pin");
+}
 
 
 FText UCustomizableObjectNodeCopyMaterial::GetNodeTitle(ENodeTitleType::Type TitleType) const
@@ -35,16 +37,16 @@ void UCustomizableObjectNodeCopyMaterial::AllocateDefaultPins(UCustomizableObjec
 	const UEdGraphSchema_CustomizableObject* Schema = GetDefault<UEdGraphSchema_CustomizableObject>();
 
 	// Input pins
-	UEdGraphPin* Pin = CustomCreatePin(EGPD_Input, Schema->PC_Mesh, FName(MeshPinName));
+	UEdGraphPin* Pin = CustomCreatePin(EGPD_Input, Schema->PC_Mesh, FName(CustomizableObjectNodeCopyMaterial::MeshPinName));
 	Pin->PinFriendlyName = LOCTEXT("MeshPin", "Mesh");
 	Pin->bDefaultValueIsIgnored = true;
 
-	Pin = CustomCreatePin(EGPD_Input, Schema->PC_Material, FName(MeshSectionPinName));
+	Pin = CustomCreatePin(EGPD_Input, Schema->PC_Material, FName(CustomizableObjectNodeCopyMaterial::MeshSectionPinName));
 	Pin->PinFriendlyName = LOCTEXT("BaseMeshSectionPin", "Base Mesh Section");;
 	Pin->bDefaultValueIsIgnored = true;
 
 	// Output pins
-	Pin = CustomCreatePin(EGPD_Output, Schema->PC_Material, FName(OutputPinName));
+	Pin = CustomCreatePin(EGPD_Output, Schema->PC_Material, FName(CustomizableObjectNodeCopyMaterial::OutputPinName));
 	Pin->PinFriendlyName = LOCTEXT("MeshSectionPin", "Mesh Section");
 }
 
@@ -59,7 +61,7 @@ void UCustomizableObjectNodeCopyMaterial::BackwardsCompatibleFixup(int32 Customi
 		if (BaseMaterialPin)
 		{
 			const FString PinFriendlyName = TEXT("Base Mesh Section");
-			BaseMaterialPin->PinName = FName(MeshSectionPinName);
+			BaseMaterialPin->PinName = FName(CustomizableObjectNodeCopyMaterial::MeshSectionPinName);
 			BaseMaterialPin->PinFriendlyName = FText::FromString(PinFriendlyName);
 		}
 
@@ -67,7 +69,7 @@ void UCustomizableObjectNodeCopyMaterial::BackwardsCompatibleFixup(int32 Customi
 		if (MaterialPin)
 		{
 			const FString PinFriendlyName = TEXT("Mesh Section");
-			MaterialPin->PinName = FName(OutputPinName);
+			MaterialPin->PinName = FName(CustomizableObjectNodeCopyMaterial::OutputPinName);
 			MaterialPin->PinFriendlyName = FText::FromString(PinFriendlyName);
 		}
 	}
@@ -76,7 +78,7 @@ void UCustomizableObjectNodeCopyMaterial::BackwardsCompatibleFixup(int32 Customi
 
 UEdGraphPin* UCustomizableObjectNodeCopyMaterial::GetMeshPin() const
 {
-	return FindPin(MeshPinName);
+	return FindPin(CustomizableObjectNodeCopyMaterial::MeshPinName);
 }
 
 
@@ -108,7 +110,7 @@ TArray<UCustomizableObjectLayout*> UCustomizableObjectNodeCopyMaterial::GetLayou
 
 UEdGraphPin* UCustomizableObjectNodeCopyMaterial::OutputPin() const
 {
-	UEdGraphPin* Pin = FindPin(OutputPinName, EEdGraphPinDirection::EGPD_Output);
+	UEdGraphPin* Pin = FindPin(CustomizableObjectNodeCopyMaterial::OutputPinName, EEdGraphPinDirection::EGPD_Output);
 
 	// Legacy name
 	if (!Pin)
@@ -135,7 +137,7 @@ bool UCustomizableObjectNodeCopyMaterial::RealMaterialDataHasChanged() const
 
 UEdGraphPin* UCustomizableObjectNodeCopyMaterial::GetMeshSectionPin() const
 {
-	UEdGraphPin* Pin = FindPin(MeshSectionPinName, EEdGraphPinDirection::EGPD_Input);
+	UEdGraphPin* Pin = FindPin(CustomizableObjectNodeCopyMaterial::MeshSectionPinName, EEdGraphPinDirection::EGPD_Input);
 
 	// Legacy name
 	if (!Pin)
