[Mutable] Fix crash when updating a COI without CO.

#rnx

diff --git a/Engine/Plugins/Mutable/Source/CustomizableObject/Internal/MuCO/CustomizableObjectSystemPrivate.h b/Engine/Plugins/Mutable/Source/CustomizableObject/Internal/MuCO/CustomizableObjectSystemPrivate.h
--- a/Engine/Plugins/Mutable/Source/CustomizableObject/Internal/MuCO/CustomizableObjectSystemPrivate.h
+++ b/Engine/Plugins/Mutable/Source/CustomizableObject/Internal/MuCO/CustomizableObjectSystemPrivate.h
@@ -524,8 +524,11 @@ public:
 	FUpdateContextPrivate(UCustomizableObjectInstance& InInstance, const FCustomizableObjectInstanceDescriptor& Descriptor);
 
 	FUpdateContextPrivate(UCustomizableObjectInstance& InInstance);
+
 	~FUpdateContextPrivate();
 
+	bool IsContextValid() const;
+	
 	const TMap<FName, uint8>& GetMinLOD() const;
 
 	void SetMinLOD(const TMap<FName, uint8>& MinLOD);
@@ -595,6 +598,9 @@ public:
 	/** true if the Update has blocked Low Priority Tasks from launching. */
 	bool bLowPriorityTasksBlocked = false;
 
+	/** The Context has been successfully created. */
+	bool bValid = false; 
+
 	/** This option comes from the operation request. It is used to reduce the number of mipmaps that mutable must generate for images.  */
 	int32 MipsToSkip = 0;
 
diff --git a/Engine/Plugins/Mutable/Source/CustomizableObject/Private/MuCO/CustomizableObjectSystem.cpp b/Engine/Plugins/Mutable/Source/CustomizableObject/Private/MuCO/CustomizableObjectSystem.cpp
--- a/Engine/Plugins/Mutable/Source/CustomizableObject/Private/MuCO/CustomizableObjectSystem.cpp
+++ b/Engine/Plugins/Mutable/Source/CustomizableObject/Private/MuCO/CustomizableObjectSystem.cpp
@@ -168,21 +168,29 @@ static FAutoConsoleVariableSink CVarMutableSink(FConsoleCommandDelegate::CreateS
 FUpdateContextPrivate::FUpdateContextPrivate(UCustomizableObjectInstance& InInstance, const FCustomizableObjectInstanceDescriptor& Descriptor)
 {
 	check(IsInGameThread());
+
+	if (!IsValid(&InInstance))
+	{
+		return;
+	}
 	
 	Instance = &InInstance;
 
-	UCustomizableObject* LocalObject = InInstance.GetCustomizableObject();
-	check(LocalObject);
+	UCustomizableObject* InObject = InInstance.GetCustomizableObject();
+	if (!IsValid(InObject))
+	{
+		return;
+	}
 
-	const FModelResources& ModelResources = LocalObject->GetPrivate()->GetModelResources();
+	Object = InObject;
 
-	Object = LocalObject;
+	const FModelResources& ModelResources = InObject->GetPrivate()->GetModelResources();
 
 	CapturedDescriptor = Descriptor;
-	LocalObject->GetPrivate()->ApplyStateForcedValuesToParameters(CapturedDescriptor);
+	InObject->GetPrivate()->ApplyStateForcedValuesToParameters(CapturedDescriptor);
 
 	CapturedDescriptorHash = FDescriptorHash(Descriptor);
-	NumObjectComponents = LocalObject->GetComponentCount();
+	NumObjectComponents = InObject->GetComponentCount();
 	FirstLODAvailable = ModelResources.ComponentFirstLODAvailable;
 	FirstResidentLOD = ModelResources.ComponentNumLODsToStream;
 	ComponentNames = ModelResources.ComponentNamesPerObjectComponent;
@@ -192,6 +200,14 @@ FUpdateContextPrivate::FUpdateContextPrivate(UCustomizableObjectInstance& InInst
 
 	UCustomizableObjectSystem* System = UCustomizableObjectSystem::GetInstance();
 	System->GetPrivate()->CacheTextureParameters(CapturedDescriptor.GetTextureParameters());
+
+	bValid = true;
+}
+
+
+bool FUpdateContextPrivate::IsContextValid() const
+{
+	return bValid;
 }
 
 
@@ -205,7 +221,7 @@ FUpdateContextPrivate::~FUpdateContextPrivate()
 {
 	check(IsInGameThread());
 
-	if (UCustomizableObjectSystem::IsCreated())
+	if (bValid && UCustomizableObjectSystem::IsCreated())
 	{
 		UCustomizableObjectSystem* System = UCustomizableObjectSystem::GetInstance();
 		System->GetPrivate()->UnCacheTextureParameters(CapturedDescriptor.GetTextureParameters());
@@ -916,7 +932,7 @@ void FinishUpdateGlobal(const TSharedRef<FUpdateContextPrivate>& Context)
 	}
 	
 	const FName ObjectName = Object ? Object->GetFName() : NAME_None;
-	const FName InstanceName = Instance ? Instance->GetFName() : NAME_None;
+	const FName InstanceName = GetFNameSafe(Instance);
 	UE_LOG(LogMutable, Verbose, TEXT("Finished Update Skeletal Mesh Async. CustomizableObject=%s Instance=%s, Frame=%d  QueueTime=%f, UpdateTime=%f"), *ObjectName.ToString(), *InstanceName.ToString(), GFrameNumber, Context->QueueTime, Context->UpdateTime);
 	
 	if (SystemPrivate && FLogBenchmarkUtil::IsBenchmarkingReportingEnabled())
@@ -1208,34 +1224,40 @@ void UCustomizableObjectSystemPrivate::EnqueueUpdateSkeletalMesh(const TSharedRe
 {
 	MUTABLE_CPUPROFILER_SCOPE(FCustomizableObjectSystemPrivate::EnqueueUpdateSkeletalMesh);
 	check(IsInGameThread());
-
-	UCustomizableObjectInstance* Instance = Context->Instance.Get();
-	check(Instance);
 	
-	UCustomizableInstancePrivate* InstancePrivate = Instance->GetPrivate();
+	UCustomizableObject* Object = Context->Object.Get();
+	UCustomizableObjectInstance* Instance = Context->Instance.Get();
 
-	UCustomizableObject* Object = Instance->GetCustomizableObject();
-	
 	const FName ObjectName = Object ? Object->GetFName() : NAME_None;
-	const FName InstanceName = Instance->GetFName();
+	const FName InstanceName = Instance ? Instance->GetFName() : NAME_None;
 	UE_LOG(LogMutable, Verbose, TEXT("Enqueue Update Skeletal Mesh Async. CustomizableObject=%s Instance=%s, Frame=%d"), *ObjectName.ToString(), *InstanceName.ToString(), GFrameNumber);
 	
+	if (!Context->IsContextValid()) 
+	{
+		Context->UpdateResult = EUpdateResult::Error;
+		FinishUpdateGlobal(Context);
+		return;
+	}
+
+	check(Object);
+	check(Instance);
+		
 	if (!bIsMutableEnabled)
 	{
 		// Mutable is disabled. Set the reference SkeletalMesh and finish the update with success to avoid breaking too many things.
 		Context->UpdateResult = EUpdateResult::Success;
-		InstancePrivate->SetDefaultSkeletalMesh();
+		Instance->GetPrivate()->SetDefaultSkeletalMesh();
 		FinishUpdateGlobal(Context);
 		return;
 	}
-
+	
 	if (!Instance->CanUpdateInstance())
 	{
 		Context->UpdateResult = EUpdateResult::Error;
 		FinishUpdateGlobal(Context);
 		return;
 	}
-
+	
 	const EUpdateRequired UpdateRequired = IsUpdateRequired(*Instance, Context->bOnlyUpdateIfNotGenerated, false, Context->bIgnoreCloseDist);
 	switch (UpdateRequired)
 	{
@@ -1268,7 +1290,7 @@ void UCustomizableObjectSystemPrivate::EnqueueUpdateSkeletalMesh(const TSharedRe
 			return; // The requested update is equal to the running update.
 		}
 	
-		if (Context->GetCapturedDescriptorHash().IsSubset(InstancePrivate->CommittedDescriptorHash) &&
+		if (Context->GetCapturedDescriptorHash().IsSubset(Instance->GetPrivate()->CommittedDescriptorHash) &&
 			!(CurrentMutableOperation &&
 			Instance == CurrentMutableOperation->Instance)) // This condition is necessary because even if the descriptor is a subset, it will be replaced by the CurrentMutableOperation
 		{
@@ -3882,7 +3904,7 @@ int32 UCustomizableObjectSystem::TickInternal(const bool bBlocking)
 				const FCustomizableObjectInstanceDescriptor& Descriptor = bGenerated ? Instance->GetPrivate()->CommittedDescriptor : Instance->GetPrivate()->GetDescriptor();
 
 				const TSharedRef<FUpdateContextPrivate> Context = MakeShared<FUpdateContextPrivate>(*Instance, Descriptor);
-
+				
 				// Commit the LOD changes
 				LODUpdateCandidateFound->ApplyLODUpdateParamsToInstance(*Context);
 
@@ -4402,10 +4424,11 @@ void UCustomizableObjectSystemPrivate::StartUpdateSkeletalMesh(const TSharedRef<
 	TRACE_BEGIN_REGION(UE_MUTABLE_UPDATE_REGION);
 
 	check(!CurrentMutableOperation); // Can not start an update if there is already another in progress
+	check(Context->bValid);
 	check(Context->Instance.IsValid()) // The COI has to be alive to start the update
+	check(Context->Object.IsValid()) // The CO has to be alive to start the update
 
-	const UCustomizableObject* Object = Context->Instance->GetCustomizableObject();
-	check(Object) // The CO has to be alive to start the update
+	const UCustomizableObject* Object = Context->Object.Get();
 
 	if (FPlatformTime::Seconds() > LogStartedUpdateUnmute)
 	{
