[Mutable] Fixed texture analyzer Texture2D cast fail crash.

#rnx

diff --git a/Engine/Plugins/Mutable/Source/CustomizableObjectEditor/Private/MuCOE/SCustomizableObjectEditorTextureAnalyzer.cpp b/Engine/Plugins/Mutable/Source/CustomizableObjectEditor/Private/MuCOE/SCustomizableObjectEditorTextureAnalyzer.cpp
--- a/Engine/Plugins/Mutable/Source/CustomizableObjectEditor/Private/MuCOE/SCustomizableObjectEditorTextureAnalyzer.cpp
+++ b/Engine/Plugins/Mutable/Source/CustomizableObjectEditor/Private/MuCOE/SCustomizableObjectEditorTextureAnalyzer.cpp
@@ -28,7 +28,6 @@ struct FGeometry;
 #define LOCTEXT_NAMESPACE "SCustomizableObjecEditorTextureAnalyzer"
 
 
-
 void SCustomizableObjecEditorTextureAnalyzer::Construct(const FArguments& InArgs)
 {
 	CustomizableObjectEditor = InArgs._CustomizableObjectEditor;
@@ -148,8 +147,8 @@ void SCustomizableObjecEditorTextureAnalyzer::FillTextureAnalyzerTable(UCustomiz
 
 			for (const FGeneratedTexture& GeneratedTexture : GeneratedMaterial.Textures)
 			{
-				UTexture* Texture = GeneratedTexture.Texture;
-				if (!Texture)
+				UTexture2D* Texture = Cast<UTexture2D>(GeneratedTexture.Texture);
+				if (!Texture || !Texture->HasAllFlags(EObjectFlags::RF_Transient))
 				{
 					continue;
 				}
@@ -162,12 +161,12 @@ void SCustomizableObjecEditorTextureAnalyzer::FillTextureAnalyzerTable(UCustomiz
 				Entry->Texture = Texture;
 				Entry->TextureName = Texture->GetName();
 				Entry->TextureParameterName = GeneratedTexture.Name;
-				Entry->ResolutionX = Cast<UTexture2D>(Texture)->GetSizeX();
-				Entry->ResolutionY = Cast<UTexture2D>(Texture)->GetSizeY();
+				Entry->ResolutionX = Texture->GetSizeX();
+				Entry->ResolutionY = Texture->GetSizeY();
 				Entry->LODBias = Texture->GetCachedLODBias();
 				uint32 TextureSize = Texture->CalcTextureMemorySizeEnum(TMC_AllMips);
 				Entry->Size = TextureSize / 1024.0f;
-				Entry->Format = Cast<UTexture2D>(Texture)->GetPixelFormat();
+				Entry->Format = Texture->GetPixelFormat();
 				Entry->IsStreamed = Texture->IsCurrentlyVirtualTextured() ? "Virtual Streamed" : (!Texture->IsStreamable() ? "Not Streamed" : "Streamed");
 				Entry->LODGroup = Texture->LODGroup;
 
