[mutable] Fixed unrequired instance updates.
- Some COIs were triggering constant updates as the update candidate structure was not being fully updated.

diff --git a/Engine/Plugins/Mutable/Source/CustomizableObject/Private/MuCO/CustomizableObjectInstance.cpp b/Engine/Plugins/Mutable/Source/CustomizableObject/Private/MuCO/CustomizableObjectInstance.cpp
--- a/Engine/Plugins/Mutable/Source/CustomizableObject/Private/MuCO/CustomizableObjectInstance.cpp
+++ b/Engine/Plugins/Mutable/Source/CustomizableObject/Private/MuCO/CustomizableObjectInstance.cpp
@@ -7137,6 +7137,11 @@ void UCustomizableObjectInstance::SetRequestedLODs(const TMap<FName, uint8>& InM
 {
 	check(PrivateData);
 
+	// DOCUMENTATION 
+	// MinLod =					Hard limit set by the budgeter or the user (artifical limit).			Per Skeletal Mesh Component.
+	// First LOD Available =	First available lod per platform and quality (scalability) settings.	Per Skeletal Mesh.
+	// Requested LODs =			What the renderer wants to render.										Per Skeletal Mesh Component.
+
 	if (!CanUpdateInstance())
 	{
 		return;
@@ -7176,7 +7181,7 @@ void UCustomizableObjectInstance::SetRequestedLODs(const TMap<FName, uint8>& InM
 	
 	for (const FName& ComponentName : CustomizableObject->GetPrivate()->GetModelResources().ComponentNamesPerObjectComponent)
 	{
-		uint8 InMinLOD = 0;
+		uint8& InMinLOD = MutableUpdateCandidate.MinLOD.FindOrAdd(ComponentName);
 		if (const uint8* Result = InMinLODs.Find(ComponentName))
 		{
 			InMinLOD = *Result;
