[Mutable] Improved UCustomizableObject public API.

- Deprecated all functions that use ParameterIndex.
- Added missing ParameterName signatures.
- Removed old unused deprecated functions.
- Started with the Int -> Enum / Option -> Value renaming. Only renamed functions which we had to change the signature.

diff --git a/Engine/Plugins/Mutable/Config/BaseMutable.ini b/Engine/Plugins/Mutable/Config/BaseMutable.ini
--- a/Engine/Plugins/Mutable/Config/BaseMutable.ini
+++ b/Engine/Plugins/Mutable/Config/BaseMutable.ini
@@ -18,3 +18,8 @@
 +ClassRedirects=(OldName="/Script/CustomizableObjectEditor.CustomizableObjectNodeRemoveMeshBlocks", NewName="/Script/CustomizableObjectEditor.CustomizableObjectNodeModifierRemoveMeshBlocks")
 +ClassRedirects=(OldName="/Script/CustomizableObjectEditor.CustomizableObjectNodeEditMaterialBase", NewName="/Script/CustomizableObjectEditor.CustomizableObjectNodeModifierEditMeshSectionBase")
 +ClassRedirects=(OldName="/Script/CustomizableObjectEditor.CustomizableObjectNodeEditLayoutBlocks", NewName="/Script/CustomizableObjectEditor.CustomizableObjectNodeModifierEditLayoutBlocks")
+
++FunctionRedirects=(OldName="/Script/CustomizableObject.CustomizableObject.GetIntParameterOptionUIMetadata",NewName="/Script/CustomizableObject.CustomizableObject.GetEnumParameterValueUIMetadata")
++FunctionRedirects=(OldName="/Script/CustomizableObject.CustomizableObject.GetIntParameterGroupType",NewName="/Script/CustomizableObject.CustomizableObject.GetEnumParameterGroupType")
++FunctionRedirects=(OldName="/Script/CustomizableObject.CustomizableObject.GetIntParameterOptionDataTable",NewName="/Script/CustomizableObject.CustomizableObject.GetEnumParameterValueDataTable")
++FunctionRedirects=(OldName="/Script/CustomizableObject.CustomizableObject.GetIntParameterDefaultValue",NewName="/Script/CustomizableObject.CustomizableObject.GetEnumParameterDefaultValue")
\ No newline at end of file
diff --git a/Engine/Plugins/Mutable/Source/CustomizableObject/Internal/MuCO/CustomizableObjectPrivate.h b/Engine/Plugins/Mutable/Source/CustomizableObject/Internal/MuCO/CustomizableObjectPrivate.h
--- a/Engine/Plugins/Mutable/Source/CustomizableObject/Internal/MuCO/CustomizableObjectPrivate.h
+++ b/Engine/Plugins/Mutable/Source/CustomizableObject/Internal/MuCO/CustomizableObjectPrivate.h
@@ -58,6 +58,110 @@ struct FCustomizableObjectInstanceDescriptor;
 
 FGuid CUSTOMIZABLEOBJECT_API GenerateIdentifier(const UCustomizableObject& CustomizableObject);
 
+
+// A USTRUCT version of FMeshToMeshVertData in SkeletalMeshTypes.h
+// We are taking advantage of the padding data to store from which asset this data comes from
+// maintaining the same memory footprint than the original.
+USTRUCT()
+struct FCustomizableObjectMeshToMeshVertData
+{
+	GENERATED_USTRUCT_BODY()
+
+	FCustomizableObjectMeshToMeshVertData() = default;
+
+	explicit FCustomizableObjectMeshToMeshVertData(const FMeshToMeshVertData& Original)
+		: Weight(Original.Weight)
+	{
+		for (int32 i = 0; i < 4; ++i)
+		{
+			PositionBaryCoordsAndDist[i] = Original.PositionBaryCoordsAndDist[i];
+			NormalBaryCoordsAndDist[i] = Original.NormalBaryCoordsAndDist[i];
+			TangentBaryCoordsAndDist[i] = Original.TangentBaryCoordsAndDist[i];
+			SourceMeshVertIndices[i] = Original.SourceMeshVertIndices[i];
+		}
+	}
+		
+
+	explicit operator FMeshToMeshVertData() const
+	{
+		FMeshToMeshVertData ReturnValue;
+
+		for (int32 i = 0; i < 4; ++i)
+		{
+			ReturnValue.PositionBaryCoordsAndDist[i] = PositionBaryCoordsAndDist[i];
+			ReturnValue.NormalBaryCoordsAndDist[i] = NormalBaryCoordsAndDist[i];
+			ReturnValue.TangentBaryCoordsAndDist[i] = TangentBaryCoordsAndDist[i];
+			ReturnValue.SourceMeshVertIndices[i] = SourceMeshVertIndices[i];
+		}
+		ReturnValue.Weight = Weight;
+		ReturnValue.Padding = 0;
+
+		return ReturnValue;
+	}
+
+	
+	// Barycentric coords and distance along normal for the position of the final vert
+	UPROPERTY()
+	float PositionBaryCoordsAndDist[4] = {0.f};
+
+	// Barycentric coords and distance along normal for the location of the unit normal endpoint
+	// Actual normal = ResolvedNormalPosition - ResolvedPosition
+	UPROPERTY()
+	float NormalBaryCoordsAndDist[4] = {0.f};
+
+	// Barycentric coords and distance along normal for the location of the unit Tangent endpoint
+	// Actual normal = ResolvedNormalPosition - ResolvedPosition
+	UPROPERTY()
+	float TangentBaryCoordsAndDist[4] = {0.f};
+
+	// Contains the 3 indices for verts in the source mesh forming a triangle, the last element
+	// is a flag to decide how the skinning works, 0xffff uses no simulation, and just normal
+	// skinning, anything else uses the source mesh and the above skin data to get the final position
+	UPROPERTY()
+	uint16	 SourceMeshVertIndices[4] = {0u, 0u, 0u, 0u};
+
+	UPROPERTY()
+	float Weight = 0.0f;
+
+	// Non serialized, unused padding. This is present in the FMeshToMeshVertData struct as Padding for alignment.
+	uint32 UnusedPadding = 0;
+
+	/**
+	 * Serializer
+	 *
+	 * @param Ar - archive to serialize with
+	 * @param V - vertex to serialize
+	 * @return archive that was used
+	 */
+	friend FArchive& operator<<(FArchive& Ar, FCustomizableObjectMeshToMeshVertData& V)
+	{
+		Ar << V.PositionBaryCoordsAndDist[0];
+		Ar << V.PositionBaryCoordsAndDist[1];
+		Ar << V.PositionBaryCoordsAndDist[2];
+		Ar << V.PositionBaryCoordsAndDist[3];
+		Ar << V.NormalBaryCoordsAndDist[0];
+		Ar << V.NormalBaryCoordsAndDist[1];
+		Ar << V.NormalBaryCoordsAndDist[2];
+		Ar << V.NormalBaryCoordsAndDist[3];
+		Ar << V.TangentBaryCoordsAndDist[0];
+		Ar << V.TangentBaryCoordsAndDist[1];
+		Ar << V.TangentBaryCoordsAndDist[2];
+		Ar << V.TangentBaryCoordsAndDist[3];
+
+		Ar << V.SourceMeshVertIndices[0];
+		Ar << V.SourceMeshVertIndices[1];
+		Ar << V.SourceMeshVertIndices[2];
+		Ar << V.SourceMeshVertIndices[3];
+
+		Ar << V.Weight;
+
+		return Ar;
+	}
+};
+static_assert(sizeof(FCustomizableObjectMeshToMeshVertData) == sizeof(float)*4*3 + sizeof(uint16)*4 + sizeof(float) + sizeof(uint32));
+template<> struct TCanBulkSerialize<FCustomizableObjectMeshToMeshVertData> { enum { Value = true }; };
+
+
 // Warning! MutableCompiledDataHeader must be the first data serialized in a stream
 struct MutableCompiledDataStreamHeader
 {
@@ -1067,7 +1171,7 @@ struct FModelResources
 	 * Index with CustomizableObject Component index
 	 */
 	UPROPERTY()
-	TArray<FMutableRefSkeletalMeshData> ReferenceSkeletalMeshesData;
+	TArray<FMutableRefSkeletalMeshData> ReferenceSkeletalMeshesData;	
 
 	/** Skeletons used by the compiled mu::FModel. */
 	UPROPERTY()
@@ -1484,6 +1588,40 @@ public:
 	void BackwardsCompatibleFixup(int32 CustomizableObjectCustomVersion);
 
 	FName GetComponentName(FCustomizableObjectComponentIndex ObjectComponentIndex) const;
+
+#if WITH_EDITORONLY_DATA
+	EMutableCompileMeshType GetMeshCompileType() const;
+	
+	const TArray<TSoftObjectPtr<UCustomizableObject>>& GetWorkingSet() const;
+
+	bool IsAssetUserDataMergeEnabled() const;
+	
+	bool IsTableMaterialsParentCheckDisabled() const;
+
+	bool IsRealTimeMorphTargetsEnabled() const;
+
+	bool IsClothingEnabled() const;
+	
+	bool Is16BitBoneWeightsEnabled() const;
+	
+	bool IsAltSkinWeightProfilesEnabled() const;
+	
+	bool IsPhysicsAssetMergeEnabled() const;
+	
+	bool IsEnabledAnimBpPhysicsAssetsManipulation() const;
+#endif
+	
+	const FString& GetIntParameterAvailableOption(int32 ParamIndex, int32 K) const;
+
+	int32 GetEnumParameterNumValues(int32 ParamIndex) const;
+
+	FString FindIntParameterValueName(int32 ParamIndex, int32 ParamValue) const;
+
+	int32 FindState(const FString& Name) const;
+
+	int32 GetStateParameterIndex(int32 StateIndex, int32 ParameterIndex) const;
+
+	bool IsParameterMultidimensional(const int32& InParamIndex) const;
 	
 	/** Cache of generated SkeletalMeshes */
 	FMeshCache MeshCache;
diff --git a/Engine/Plugins/Mutable/Source/CustomizableObject/Private/MuCO/CustomizableObject.cpp b/Engine/Plugins/Mutable/Source/CustomizableObject/Private/MuCO/CustomizableObject.cpp
--- a/Engine/Plugins/Mutable/Source/CustomizableObject/Private/MuCO/CustomizableObject.cpp
+++ b/Engine/Plugins/Mutable/Source/CustomizableObject/Private/MuCO/CustomizableObject.cpp
@@ -98,6 +98,8 @@ TAutoConsoleVariable<bool> CVarMutableAsyncCook(
 
 #endif
 
+static const FString s_EmptyString;
+
 #if WITH_EDITORONLY_DATA
 
 namespace UE::Mutable::Private
@@ -155,14 +157,6 @@ FGuid UCustomizableObjectPrivate::GetVersionId() const
 }
 
 
-void UCustomizableObject::GetAssetRegistryTags(TArray<FAssetRegistryTag>& OutTags) const
-{
-	PRAGMA_DISABLE_DEPRECATION_WARNINGS;
-	Super::GetAssetRegistryTags(OutTags);
-	PRAGMA_ENABLE_DEPRECATION_WARNINGS;
-}
-
-
 void UCustomizableObject::GetAssetRegistryTags(FAssetRegistryTagsContext Context) const
 {
 	int32 isRoot = 0;
@@ -1322,13 +1316,7 @@ USkeletalMesh* UCustomizableObject::GetComponentMeshReferenceSkeletalMesh(const
 
 int32 UCustomizableObject::FindState( const FString& Name ) const
 {
-	int32 Result = -1;
-	if (Private->GetModel())
-	{
-		Result = Private->GetModel()->FindState(Name);
-	}
-
-	return Result;
+	return GetPrivate()->FindState(Name);
 }
 
 
@@ -1378,35 +1366,32 @@ int32 UCustomizableObject::GetStateParameterCount( int32 StateIndex ) const
 
 int32 UCustomizableObject::GetStateParameterIndex(int32 StateIndex, int32 ParameterIndex) const
 {
-	int32 Result = 0;
-
-	if (Private->GetModel())
-	{
-		Result = Private->GetModel()->GetStateParameterIndex(StateIndex, ParameterIndex);
-	}
-
-	return Result;
+	return GetPrivate()->GetStateParameterIndex(StateIndex, ParameterIndex);
 }
 
 
 int32 UCustomizableObject::GetStateParameterCount(const FString& StateName) const
 {
-	int32 StateIndex = FindState(StateName);
-	
+	int32 StateIndex = GetPrivate()->FindState(StateName);
+
+	PRAGMA_DISABLE_DEPRECATION_WARNINGS
 	return GetStateParameterCount(StateIndex);
+	PRAGMA_ENABLE_DEPRECATION_WARNINGS
 }
 
 
 FString UCustomizableObject::GetStateParameterName(const FString& StateName, int32 ParameterIndex) const
 {
-	int32 StateIndex = FindState(StateName);
-	
+	int32 StateIndex = GetPrivate()->FindState(StateName);
+
+	PRAGMA_DISABLE_DEPRECATION_WARNINGS
 	return GetStateParameterName(StateIndex, ParameterIndex);
+	PRAGMA_ENABLE_DEPRECATION_WARNINGS
 }
 
 FString UCustomizableObject::GetStateParameterName(int32 StateIndex, int32 ParameterIndex) const
 {
-	return GetParameterName(GetStateParameterIndex(StateIndex, ParameterIndex));
+	return GetParameterName(GetPrivate()->GetStateParameterIndex(StateIndex, ParameterIndex));
 }
 
 
@@ -1490,6 +1475,173 @@ FName UCustomizableObjectPrivate::GetComponentName(FCustomizableObjectComponentI
 }
 
 
+#if WITH_EDITORONLY_DATA
+EMutableCompileMeshType UCustomizableObjectPrivate::GetMeshCompileType() const
+{
+	PRAGMA_DISABLE_DEPRECATION_WARNINGS // TODO Remove >5.6
+	return GetPublic()->MeshCompileType;
+	PRAGMA_ENABLE_DEPRECATION_WARNINGS
+}
+
+
+const TArray<TSoftObjectPtr<UCustomizableObject>>& UCustomizableObjectPrivate::GetWorkingSet() const
+{
+	PRAGMA_DISABLE_DEPRECATION_WARNINGS // TODO Remove >5.6
+	return GetPublic()->WorkingSet;
+	PRAGMA_ENABLE_DEPRECATION_WARNINGS
+}
+
+
+bool UCustomizableObjectPrivate::IsAssetUserDataMergeEnabled() const
+{
+	PRAGMA_DISABLE_DEPRECATION_WARNINGS // TODO Remove >5.6
+	return GetPublic()->bEnableAssetUserDataMerge;
+	PRAGMA_ENABLE_DEPRECATION_WARNINGS
+}
+
+
+bool UCustomizableObjectPrivate::IsTableMaterialsParentCheckDisabled() const
+{
+	PRAGMA_DISABLE_DEPRECATION_WARNINGS // TODO Remove >5.6
+	return GetPublic()->bDisableTableMaterialsParentCheck;
+	PRAGMA_ENABLE_DEPRECATION_WARNINGS
+}
+
+
+bool UCustomizableObjectPrivate::IsRealTimeMorphTargetsEnabled() const
+{
+	PRAGMA_DISABLE_DEPRECATION_WARNINGS // TODO Remove >5.6
+	return GetPublic()->bEnableRealTimeMorphTargets;
+	PRAGMA_ENABLE_DEPRECATION_WARNINGS
+}
+
+
+bool UCustomizableObjectPrivate::IsClothingEnabled() const
+{
+	PRAGMA_DISABLE_DEPRECATION_WARNINGS // TODO Remove >5.6
+	return GetPublic()->bEnableClothing;
+	PRAGMA_ENABLE_DEPRECATION_WARNINGS
+}
+
+
+bool UCustomizableObjectPrivate::Is16BitBoneWeightsEnabled() const
+{
+	PRAGMA_DISABLE_DEPRECATION_WARNINGS // TODO Remove >5.6
+	return GetPublic()->bEnable16BitBoneWeights;
+	PRAGMA_ENABLE_DEPRECATION_WARNINGS
+}
+
+
+bool UCustomizableObjectPrivate::IsAltSkinWeightProfilesEnabled() const
+{
+	PRAGMA_DISABLE_DEPRECATION_WARNINGS // TODO Remove >5.6
+	return GetPublic()->bEnableAltSkinWeightProfiles;
+	PRAGMA_ENABLE_DEPRECATION_WARNINGS
+}
+
+
+bool UCustomizableObjectPrivate::IsPhysicsAssetMergeEnabled() const
+{
+	PRAGMA_DISABLE_DEPRECATION_WARNINGS // TODO Remove >5.6
+	return GetPublic()->bEnablePhysicsAssetMerge;
+	PRAGMA_ENABLE_DEPRECATION_WARNINGS
+}
+
+
+bool UCustomizableObjectPrivate::IsEnabledAnimBpPhysicsAssetsManipulation() const
+{
+	PRAGMA_DISABLE_DEPRECATION_WARNINGS // TODO Remove >5.6
+	return GetPublic()->bEnableAnimBpPhysicsAssetsManipualtion;
+	PRAGMA_ENABLE_DEPRECATION_WARNINGS
+}
+#endif
+
+
+const FString& UCustomizableObjectPrivate::GetIntParameterAvailableOption(int32 ParamIndex, int32 K) const
+{
+	if (ParamIndex >= 0 && ParamIndex < ParameterProperties.Num())
+	{
+		if (K >= 0 && K < GetEnumParameterNumValues(ParamIndex))
+		{
+			return ParameterProperties[ParamIndex].PossibleValues[K].Name;
+		}
+		else
+		{
+			UE_LOG(LogMutable, Warning, TEXT("Index [%d] out of IntParameterNumOptions bounds at GetIntParameterAvailableOption at CO %s."), K, *GetName());
+		}
+	}
+	else
+	{
+		UE_LOG(LogMutable, Warning, TEXT("Index [%d] out of ParameterProperties bounds at GetIntParameterAvailableOption at CO %s."), ParamIndex, *GetName());
+	}
+
+	return s_EmptyString;
+}
+
+
+int32 UCustomizableObjectPrivate::GetEnumParameterNumValues(int32 ParamIndex) const
+{
+	if (ParamIndex >= 0 && ParamIndex < ParameterProperties.Num())
+	{
+		return ParameterProperties[ParamIndex].PossibleValues.Num();
+	}
+	else
+	{
+		UE_LOG(LogMutable, Warning, TEXT("Index [%d] out of ParameterProperties bounds at GetIntParameterNumOptions at CO %s."), ParamIndex, *GetName());
+	}
+
+	return 0;
+}
+
+
+FString UCustomizableObjectPrivate::FindIntParameterValueName(int32 ParamIndex, int32 ParamValue) const
+{
+	if (ParamIndex >= 0 && ParamIndex < ParameterProperties.Num())
+	{
+		const TArray<FMutableModelParameterValue> & PossibleValues = ParameterProperties[ParamIndex].PossibleValues;
+
+		const int32 MinValueIndex = !PossibleValues.IsEmpty() ? PossibleValues[0].Value : 0;
+		ParamValue = ParamValue - MinValueIndex;
+
+		if (PossibleValues.IsValidIndex(ParamValue))
+		{
+			return PossibleValues[ParamValue].Name;
+		}
+	}
+	else
+	{
+		UE_LOG(LogMutable, Warning, TEXT("Index [%d] out of ParameterProperties bounds at FindIntParameterValueName at CO %s."), ParamIndex, *GetName());
+	}
+
+	return FString();
+}
+
+
+int32 UCustomizableObjectPrivate::FindState(const FString& Name) const
+{
+	int32 Result = -1;
+	if (GetModel())
+	{
+		Result = GetModel()->FindState(Name);
+	}
+
+	return Result;
+}
+
+
+int32 UCustomizableObjectPrivate::GetStateParameterIndex(int32 StateIndex, int32 ParameterIndex) const
+{
+	int32 Result = 0;
+
+	if (GetModel())
+	{
+		Result = GetModel()->GetStateParameterIndex(StateIndex, ParameterIndex);
+	}
+
+	return Result;
+}
+
+
 FName UCustomizableObject::GetComponentName(int32 ComponentIndex) const
 {
 	return GetPrivate()->GetComponentName(FCustomizableObjectComponentIndex(ComponentIndex));
@@ -1525,7 +1677,7 @@ EMutableParameterType UCustomizableObjectPrivate::GetParameterType(int32 ParamIn
 
 EMutableParameterType UCustomizableObject::GetParameterTypeByName(const FString& Name) const
 {
-	const int32 Index = FindParameter(Name); 
+	const int32 Index = GetPrivate()->FindParameter(Name); 
 	if (GetPrivate()->ParameterProperties.IsValidIndex(Index))
 	{
 		return GetPrivate()->ParameterProperties[Index].Type;
@@ -1547,9 +1699,7 @@ EMutableParameterType UCustomizableObject::GetParameterTypeByName(const FString&
 }
 
 
-static const FString s_EmptyString;
-
-const FString & UCustomizableObject::GetParameterName(int32 ParamIndex) const
+const FString& UCustomizableObject::GetParameterName(int32 ParamIndex) const
 {
 	if (ParamIndex >= 0 && ParamIndex < GetPrivate()->ParameterProperties.Num())
 	{
@@ -1573,7 +1723,7 @@ void UCustomizableObjectPrivate::UpdateParameterPropertiesFromModel(const TShare
 
 		TArray<int32> TypedParametersCount;
 		TypedParametersCount.SetNum(static_cast<int32>(mu::EParameterType::Count));
-
+		
 		ParameterProperties.Reset(NumParameters);
 		ParameterPropertiesLookupTable.Empty(NumParameters);
 		for (int32 Index = 0; Index < NumParameters; ++Index)
@@ -1655,46 +1805,53 @@ void UCustomizableObjectPrivate::UpdateParameterPropertiesFromModel(const TShare
 }
 
 
-int32 UCustomizableObject::GetParameterDescriptionCount(const FString& ParamName) const
+int32 UCustomizableObject::GetIntParameterNumOptions(int32 ParamIndex) const
 {
-	return 0;
+	return GetPrivate()->GetEnumParameterNumValues(ParamIndex);
 }
 
 
-int32 UCustomizableObject::GetIntParameterNumOptions(int32 ParamIndex) const
+int32 UCustomizableObject::GetEnumParameterNumValues(const FString& ParamName) const
 {
-	if (ParamIndex >= 0 && ParamIndex < GetPrivate()->ParameterProperties.Num())
+	if (int32 ParamIndex = GetPrivate()->FindParameter(ParamName))
 	{
-		return GetPrivate()->ParameterProperties[ParamIndex].PossibleValues.Num();
+		return GetPrivate()->GetEnumParameterNumValues(ParamIndex);
 	}
 	else
 	{
-		UE_LOG(LogMutable, Warning, TEXT("Index [%d] out of ParameterProperties bounds at GetIntParameterNumOptions at CO %s."), ParamIndex, *GetName());
+		return 0;
 	}
-
-	return 0;
 }
 
 
-const FString& UCustomizableObject::GetIntParameterAvailableOption(int32 ParamIndex, int32 K) const
+const FString& UCustomizableObject::GetEnumParameterValue(const FString& ParamName, int32 ValueIndex) const
 {
-	if (ParamIndex >= 0 && ParamIndex < GetPrivate()->ParameterProperties.Num())
+	if (int32 ParamIndex = GetPrivate()->FindParameter(ParamName))
 	{
-		if (K >= 0 && K < GetIntParameterNumOptions(ParamIndex))
-		{
-			return GetPrivate()->ParameterProperties[ParamIndex].PossibleValues[K].Name;
-		}
-		else
-		{
-			UE_LOG(LogMutable, Warning, TEXT("Index [%d] out of IntParameterNumOptions bounds at GetIntParameterAvailableOption at CO %s."), K, *GetName());
-		}
+		return GetPrivate()->GetIntParameterAvailableOption(ParamIndex, ValueIndex);
 	}
 	else
 	{
-		UE_LOG(LogMutable, Warning, TEXT("Index [%d] out of ParameterProperties bounds at GetIntParameterAvailableOption at CO %s."), ParamIndex, *GetName());
+		return s_EmptyString;
 	}
+}
 
-	return s_EmptyString;
+
+bool UCustomizableObject::ContainsEnumParameterValue(const FString& ParameterName, const FString Value) const
+{
+	int32 ParamIndex = GetPrivate()->FindParameter(ParameterName);
+	if (ParamIndex == INDEX_NONE)
+	{
+		return false;
+	}
+
+	return GetPrivate()->FindIntParameterValue(ParamIndex, Value) != INDEX_NONE;
+}
+
+
+const FString& UCustomizableObject::GetIntParameterAvailableOption(int32 ParamIndex, int32 K) const
+{
+	return GetPrivate()->GetIntParameterAvailableOption(ParamIndex, K);
 }
 
 
@@ -1704,6 +1861,12 @@ int32 UCustomizableObject::FindParameter(const FString& Name) const
 }
 
 
+bool UCustomizableObject::ContainsParameter(const FString& ParameterName) const
+{
+	return GetPrivate()->FindParameter(ParameterName) != INDEX_NONE;
+}
+
+
 int32 UCustomizableObjectPrivate::FindParameter(const FString& Name) const
 {
 	if (const FMutableParameterIndex* Found = ParameterPropertiesLookupTable.Find(Name))
@@ -1766,24 +1929,7 @@ int32 UCustomizableObjectPrivate::FindIntParameterValue(int32 ParamIndex, const
 
 FString UCustomizableObject::FindIntParameterValueName(int32 ParamIndex, int32 ParamValue) const
 {
-	if (ParamIndex >= 0 && ParamIndex < GetPrivate()->ParameterProperties.Num())
-	{
-		const TArray<FMutableModelParameterValue> & PossibleValues = GetPrivate()->ParameterProperties[ParamIndex].PossibleValues;
-
-		const int32 MinValueIndex = !PossibleValues.IsEmpty() ? PossibleValues[0].Value : 0;
-		ParamValue = ParamValue - MinValueIndex;
-
-		if (PossibleValues.IsValidIndex(ParamValue))
-		{
-			return PossibleValues[ParamValue].Name;
-		}
-	}
-	else
-	{
-		UE_LOG(LogMutable, Warning, TEXT("Index [%d] out of ParameterProperties bounds at FindIntParameterValueName at CO %s."), ParamIndex, *GetName());
-	}
-
-	return FString();
+	return GetPrivate()->FindIntParameterValueName(ParamIndex, ParamValue);
 }
 
 
@@ -1800,9 +1946,9 @@ FMutableParamUIMetadata UCustomizableObject::GetParameterUIMetadata(const FStrin
 }
 
 
-FMutableParamUIMetadata UCustomizableObject::GetIntParameterOptionUIMetadata(const FString& ParamName, const FString& OptionName) const
+FMutableParamUIMetadata UCustomizableObject::GetEnumParameterValueUIMetadata(const FString& ParamName, const FString& OptionName) const
 {
-	const int32 ParameterIndex = FindParameter(ParamName);
+	const int32 ParameterIndex = GetPrivate()->FindParameter(ParamName);
 	if (ParameterIndex == INDEX_NONE)
 	{
 		return {};
@@ -1818,9 +1964,16 @@ FMutableParamUIMetadata UCustomizableObject::GetIntParameterOptionUIMetadata(con
 	return IntegerParameterUIData ? IntegerParameterUIData->ParamUIMetadata : FMutableParamUIMetadata();
 }
 
-ECustomizableObjectGroupType UCustomizableObject::GetIntParameterGroupType(const FString& ParamName) const
+
+FMutableParamUIMetadata UCustomizableObject::GetIntParameterOptionUIMetadata(const FString& ParamName, const FString& OptionName) const
 {
-	const int32 ParameterIndex = FindParameter(ParamName);
+	return GetEnumParameterValueUIMetadata(ParamName, OptionName);
+}
+
+
+ECustomizableObjectGroupType UCustomizableObject::GetEnumParameterGroupType(const FString& ParamName) const
+{
+	const int32 ParameterIndex = GetPrivate()->FindParameter(ParamName);
 	if (ParameterIndex == INDEX_NONE)
 	{
 		return ECustomizableObjectGroupType::COGT_TOGGLE;
@@ -1835,6 +1988,13 @@ ECustomizableObjectGroupType UCustomizableObject::GetIntParameterGroupType(const
 	return ParameterData->IntegerParameterGroupType;
 }
 
+
+ECustomizableObjectGroupType UCustomizableObject::GetIntParameterGroupType(const FString& ParamName) const
+{
+	return GetEnumParameterGroupType(ParamName);
+}
+
+
 FMutableStateUIMetadata UCustomizableObject::GetStateUIMetadata(const FString& StateName) const
 {
 	const FMutableStateData* StateData = Private->GetModelResources().StateUIDataMap.Find(StateName);
@@ -1843,22 +2003,28 @@ FMutableStateUIMetadata UCustomizableObject::GetStateUIMetadata(const FString& S
 
 
 #if WITH_EDITOR
-TArray<TSoftObjectPtr<UDataTable>> UCustomizableObject::GetIntParameterOptionDataTable(const FString& ParamName, const FString& OptionName)
+TArray<TSoftObjectPtr<UDataTable>> UCustomizableObject::GetEnumParameterValueDataTable(const FString& ParamName, const FString& ValueName)
 {
 	const FModelResources& ModelResources = GetPrivate()->GetModelResources();
-	if (const TSet<TSoftObjectPtr<UDataTable>>* Result = ModelResources.IntParameterOptionDataTable.Find(MakeTuple(ParamName, OptionName)))
+	if (const TSet<TSoftObjectPtr<UDataTable>>* Result = ModelResources.IntParameterOptionDataTable.Find(MakeTuple(ParamName, ValueName)))
 	{
 		return Result->Array();
 	}
 	
 	return {};
 }
+
+
+TArray<TSoftObjectPtr<UDataTable>> UCustomizableObject::GetIntParameterOptionDataTable(const FString& ParamName, const FString& OptionName)
+{
+	return GetEnumParameterValueDataTable(ParamName, OptionName);
+}
 #endif
 
 
 float UCustomizableObject::GetFloatParameterDefaultValue(const FString& InParameterName) const
 {
-	const int32 ParameterIndex = FindParameter(InParameterName);
+	const int32 ParameterIndex = GetPrivate()->FindParameter(InParameterName);
 	if (ParameterIndex == INDEX_NONE)
 	{
 		UE_LOG(LogMutable, Error, TEXT("Tried to access the default value of the nonexistent float parameter [%s] in the CustomizableObject [%s]."), *InParameterName, *GetName());
@@ -1876,9 +2042,9 @@ float UCustomizableObject::GetFloatParameterDefaultValue(const FString& InParame
 }
 
 
-int32 UCustomizableObject::GetIntParameterDefaultValue(const FString& InParameterName) const
+int32 UCustomizableObject::GetEnumParameterDefaultValue(const FString& InParameterName) const
 {
-	const int32 ParameterIndex = FindParameter(InParameterName);
+	const int32 ParameterIndex = GetPrivate()->FindParameter(InParameterName);
 	if (ParameterIndex == INDEX_NONE)
 	{
 		UE_LOG(LogMutable, Error, TEXT("Tried to access the default value of the nonexistent integer parameter [%s] in the CustomizableObject [%s]."), *InParameterName, *GetName());
@@ -1896,9 +2062,15 @@ int32 UCustomizableObject::GetIntParameterDefaultValue(const FString& InParamete
 }
 
 
+int32 UCustomizableObject::GetIntParameterDefaultValue(const FString& InParameterName) const
+{
+	return GetEnumParameterDefaultValue(InParameterName);
+}
+
+
 bool UCustomizableObject::GetBoolParameterDefaultValue(const FString& InParameterName) const
 {
-	const int32 ParameterIndex = FindParameter(InParameterName);
+	const int32 ParameterIndex = GetPrivate()->FindParameter(InParameterName);
 	if (ParameterIndex == INDEX_NONE)
 	{
 		UE_LOG(LogMutable, Error, TEXT("Tried to access the default value of the nonexistent boolean parameter [%s] in the CustomizableObject [%s]."), *InParameterName, *GetName());
@@ -1918,7 +2090,7 @@ bool UCustomizableObject::GetBoolParameterDefaultValue(const FString& InParamete
 
 FLinearColor UCustomizableObject::GetColorParameterDefaultValue(const FString& InParameterName) const
 {
-	const int32 ParameterIndex = FindParameter(InParameterName);
+	const int32 ParameterIndex = GetPrivate()->FindParameter(InParameterName);
 	if (ParameterIndex == INDEX_NONE)
 	{
 		UE_LOG(LogMutable, Error, TEXT("Tried to access the default value of the nonexistent color parameter [%s] in the CustomizableObject [%s]."), *InParameterName, *GetName());
@@ -1941,7 +2113,7 @@ FLinearColor UCustomizableObject::GetColorParameterDefaultValue(const FString& I
 
 FTransform UCustomizableObject::GetTransformParameterDefaultValue(const FString& InParameterName) const
 {
-	const int32 ParameterIndex = FindParameter(InParameterName);
+	const int32 ParameterIndex = GetPrivate()->FindParameter(InParameterName);
 	if (ParameterIndex == INDEX_NONE)
 	{
 		UE_LOG(LogMutable, Error, TEXT("Tried to access the default value of the nonexistent color parameter [%s] in the CustomizableObject [%s]."), *InParameterName, *GetName());
@@ -1978,7 +2150,7 @@ void UCustomizableObject::GetProjectorParameterDefaultValue(const FString& InPar
 
 FCustomizableObjectProjector UCustomizableObject::GetProjectorParameterDefaultValue(const FString& InParameterName) const 
 {
-	const int32 ParameterIndex = FindParameter(InParameterName);
+	const int32 ParameterIndex = GetPrivate()->FindParameter(InParameterName);
 	if (ParameterIndex == INDEX_NONE)
 	{
 		UE_LOG(LogMutable, Error, TEXT("Tried to access the default value of the nonexistent projector [%s] in the CustomizableObject [%s]."), *InParameterName, *GetName());
@@ -2003,7 +2175,7 @@ FCustomizableObjectProjector UCustomizableObject::GetProjectorParameterDefaultVa
 
 FName UCustomizableObject::GetTextureParameterDefaultValue(const FString& InParameterName) const
 {
-	const int32 ParameterIndex = FindParameter(InParameterName);
+	const int32 ParameterIndex = GetPrivate()->FindParameter(InParameterName);
 	if (ParameterIndex == INDEX_NONE)
 	{
 		UE_LOG(LogMutable, Error, TEXT("Tried to access the default value of the nonexistent texture parameter [%s] in the CustomizableObject [%s]."), *InParameterName, *GetName());
@@ -2023,23 +2195,30 @@ FName UCustomizableObject::GetTextureParameterDefaultValue(const FString& InPara
 
 bool UCustomizableObject::IsParameterMultidimensional(const FString& InParameterName) const
 {
-	const int32 ParameterIndex = FindParameter(InParameterName);
+	const int32 ParameterIndex = GetPrivate()->FindParameter(InParameterName);
 	if (ParameterIndex == INDEX_NONE)
 	{
 		UE_LOG(LogMutable, Error, TEXT("Tried to access the default value of the nonexistent parameter [%s] in the CustomizableObject [%s]."), *InParameterName, *GetName());
 		return false;
 	}
 
-	return IsParameterMultidimensional(ParameterIndex);
+	return GetPrivate()->IsParameterMultidimensional(ParameterIndex);
 }
 
 
 bool UCustomizableObject::IsParameterMultidimensional(const int32& InParamIndex) const
+{
+	return GetPrivate()->IsParameterMultidimensional(InParamIndex);
+}
+
+
+bool UCustomizableObjectPrivate::IsParameterMultidimensional(const int32& InParamIndex) const
 {
 	check(InParamIndex != INDEX_NONE);
-	if (Private->GetModel())
+	
+	if (const TSharedPtr<const mu::FModel>& Model = GetModel())
 	{
-		return Private->GetModel()->IsParameterMultidimensional(InParamIndex);
+		return Model->IsParameterMultidimensional(InParamIndex);
 	}
 
 	return false;
@@ -2487,12 +2666,12 @@ FCompilationOptions UCustomizableObjectPrivate::GetCompileOptions() const
 
 	Options.EmbeddedDataBytesLimit = EmbeddedDataBytesLimit;
 	Options.CustomizableObjectNumBoneInfluences = ICustomizableObjectModule::Get().GetNumBoneInfluences();
-	Options.bRealTimeMorphTargetsEnabled = GetPublic()->bEnableRealTimeMorphTargets;
-	Options.bClothingEnabled = GetPublic()->bEnableClothing;
-	Options.b16BitBoneWeightsEnabled = GetPublic()->bEnable16BitBoneWeights;
-	Options.bSkinWeightProfilesEnabled = GetPublic()->bEnableAltSkinWeightProfiles;
-	Options.bPhysicsAssetMergeEnabled = GetPublic()->bEnablePhysicsAssetMerge;
-	Options.bAnimBpPhysicsManipulationEnabled = GetPublic()->bEnableAnimBpPhysicsAssetsManipualtion;
+	Options.bRealTimeMorphTargetsEnabled = IsRealTimeMorphTargetsEnabled();
+	Options.bClothingEnabled = IsClothingEnabled();
+	Options.b16BitBoneWeightsEnabled = Is16BitBoneWeightsEnabled();
+	Options.bSkinWeightProfilesEnabled = IsAltSkinWeightProfilesEnabled();
+	Options.bPhysicsAssetMergeEnabled = IsPhysicsAssetMergeEnabled();
+	Options.bAnimBpPhysicsManipulationEnabled = IsEnabledAnimBpPhysicsAssetsManipulation();
 	Options.ImageTiling = ImageTiling;
 
 	return Options;
diff --git a/Engine/Plugins/Mutable/Source/CustomizableObject/Private/MuCO/CustomizableObjectInstance.cpp b/Engine/Plugins/Mutable/Source/CustomizableObject/Private/MuCO/CustomizableObjectInstance.cpp
--- a/Engine/Plugins/Mutable/Source/CustomizableObject/Private/MuCO/CustomizableObjectInstance.cpp
+++ b/Engine/Plugins/Mutable/Source/CustomizableObject/Private/MuCO/CustomizableObjectInstance.cpp
@@ -749,7 +749,7 @@ bool UCustomizableObjectInstance::IsParameterRelevant(const FString& ParamName)
 	}
 
 	// This should have been precalculated in the last update if the appropriate flag in the instance was set.
-	int32 ParameterIndexInObject = CustomizableObject->FindParameter(ParamName);
+	int32 ParameterIndexInObject = CustomizableObject->GetPrivate()->FindParameter(ParamName);
 	return GetPrivate()->RelevantParameters.Contains(ParameterIndexInObject);
 }
 
@@ -2815,7 +2815,7 @@ void UCustomizableObjectInstance::SetDefaultValue(const FString& ParamName)
 		return;
 	}
 
-	Descriptor.SetDefaultValue(CustomizableObject->FindParameter(ParamName));
+	Descriptor.SetDefaultValue(CustomizableObject->GetPrivate()->FindParameter(ParamName));
 }
 
 void UCustomizableObjectInstance::SetDefaultValues()
diff --git a/Engine/Plugins/Mutable/Source/CustomizableObject/Private/MuCO/CustomizableObjectInstanceDescriptor.cpp b/Engine/Plugins/Mutable/Source/CustomizableObject/Private/MuCO/CustomizableObjectInstanceDescriptor.cpp
--- a/Engine/Plugins/Mutable/Source/CustomizableObject/Private/MuCO/CustomizableObjectInstanceDescriptor.cpp
+++ b/Engine/Plugins/Mutable/Source/CustomizableObject/Private/MuCO/CustomizableObjectInstanceDescriptor.cpp
@@ -29,11 +29,11 @@ void CustomizableObjectNullErrorMessage()
 FString GetAvailableOptionsString(const UCustomizableObject& CustomizableObject, const int32 ParameterIndexInObject)
 {
 	FString OptionsString;
-	const int32 NumOptions = CustomizableObject.GetIntParameterNumOptions(ParameterIndexInObject);
+	const int32 NumOptions = CustomizableObject.GetPrivate()->GetEnumParameterNumValues(ParameterIndexInObject);
 
 	for (int32 k = 0; k < NumOptions; ++k)
 	{
-		OptionsString += CustomizableObject.GetIntParameterAvailableOption(ParameterIndexInObject, k);
+		OptionsString += CustomizableObject.GetPrivate()->GetIntParameterAvailableOption(ParameterIndexInObject, k);
 
 		if (k < NumOptions - 1)
 		{
@@ -74,7 +74,7 @@ void FCustomizableObjectInstanceDescriptor::SaveDescriptor(FArchive& Ar, bool bU
 	for (int32 ModelParameterIndex = 0; ModelParameterIndex < ModelParameterCount; ++ModelParameterIndex)
 	{
 		const FString & Name = CustomizableObject->GetParameterName(ModelParameterIndex);
-		EMutableParameterType Type = CustomizableObject->GetParameterType(ModelParameterIndex);
+		EMutableParameterType Type = CustomizableObject->GetParameterTypeByName(Name);
 
 		if (!bUseCompactDescriptor)
 		{
@@ -131,17 +131,17 @@ void FCustomizableObjectInstanceDescriptor::SaveDescriptor(FArchive& Ar, bool bU
 			if (IntParameterIndex >= 0)
 			{					
 				const FCustomizableObjectIntParameterValue & P = IntParameters[IntParameterIndex];
-					Value = CustomizableObject->FindIntParameterValue(ModelParameterIndex,P.ParameterValueName);
+				Value = CustomizableObject->GetPrivate()->FindIntParameterValue(ModelParameterIndex, P.ParameterValueName);
 
-				int32 ParameterIndexInObject = CustomizableObject->FindParameter(IntParameters[IntParameterIndex].ParameterName);
-				bIsParamMultidimensional = CustomizableObject->IsParameterMultidimensional(ParameterIndexInObject);
+				int32 ParameterIndexInObject = CustomizableObject->GetPrivate()->FindParameter(IntParameters[IntParameterIndex].ParameterName);
+				bIsParamMultidimensional = CustomizableObject->GetPrivate()->IsParameterMultidimensional(ParameterIndexInObject);
 
 				if (bIsParamMultidimensional)
 				{
 					for (int32 i = 0; i < P.ParameterRangeValueNames.Num(); ++i)
 					{
 						ValueNames.Add(P.ParameterRangeValueNames[i]);
-						Values.Add(CustomizableObject->FindIntParameterValue(ModelParameterIndex, P.ParameterRangeValueNames[i]));
+						Values.Add(CustomizableObject->GetPrivate()->FindIntParameterValue(ModelParameterIndex, P.ParameterRangeValueNames[i]));
 					}
 				}
 
@@ -284,7 +284,7 @@ void FCustomizableObjectInstanceDescriptor::LoadDescriptor(FArchive& Ar)
 		{
 			ModelParameterIndex = ParameterIndex;
 			Name = CustomizableObject->GetParameterName(ModelParameterIndex);
-			Type = CustomizableObject->GetParameterType(ModelParameterIndex);
+			Type = CustomizableObject->GetPrivate()->GetParameterType(ModelParameterIndex);
 		}
 		else
 		{
@@ -336,8 +336,8 @@ void FCustomizableObjectInstanceDescriptor::LoadDescriptor(FArchive& Ar)
 			TArray<FString> ValueNames;
 
 			const int32 IntParameterIndex = FindTypedParameterIndex(Name, EMutableParameterType::Int);
-			const int32 ParameterIndexInObject = CustomizableObject->FindParameter(IntParameters[IntParameterIndex].ParameterName);
-			const bool bIsParamMultidimensional = CustomizableObject->IsParameterMultidimensional(ParameterIndexInObject);
+			const int32 ParameterIndexInObject = CustomizableObject->GetPrivate()->FindParameter(IntParameters[IntParameterIndex].ParameterName);
+			const bool bIsParamMultidimensional = CustomizableObject->GetPrivate()->IsParameterMultidimensional(ParameterIndexInObject);
 
 			if (bUseCompactDescriptor)
 			{
@@ -364,13 +364,13 @@ void FCustomizableObjectInstanceDescriptor::LoadDescriptor(FArchive& Ar)
 				{
 					if (bUseCompactDescriptor)
 					{
-						P.ParameterValueName = CustomizableObject->FindIntParameterValueName(ModelParameterIndex, Value);
+						P.ParameterValueName = CustomizableObject->GetPrivate()->FindIntParameterValueName(ModelParameterIndex, Value);
 						//check((P.ParameterValueName.IsEmpty() && ValueName.Equals(FString("None"))) || P.ParameterValueName.Equals(ValueName));
 						P.ParameterRangeValueNames.SetNum(Values.Num());
 
 						for (int32 ParamIndex = 0; ParamIndex < Values.Num(); ++ParamIndex)
 						{
-							P.ParameterRangeValueNames[ParamIndex] = CustomizableObject->FindIntParameterValueName(ModelParameterIndex, Values[ParamIndex]);
+							P.ParameterRangeValueNames[ParamIndex] = CustomizableObject->GetPrivate()->FindIntParameterValueName(ModelParameterIndex, Values[ParamIndex]);
 						}
 					}
 					else
@@ -562,13 +562,13 @@ TSharedPtr<mu::FParameters> FCustomizableObjectInstanceDescriptor::GetParameters
 							RangeIdxPtr->SetPosition(0, RangeIndex);
 
 							const FString& RangeValueName = IntParameter.ParameterRangeValueNames[RangeIndex];
-							const int32 Value = CustomizableObject->FindIntParameterValue(ParamIndex, RangeValueName);
+							const int32 Value = CustomizableObject->GetPrivate()->FindIntParameterValue(ParamIndex, RangeValueName);
 							MutableParameters->SetIntValue(ParamIndex, Value, RangeIdxPtr.Get());
 						}
 					}
 					else
 					{
-						const int32 Value = CustomizableObject->FindIntParameterValue(ParamIndex, IntParameter.ParameterValueName);
+						const int32 Value = CustomizableObject->GetPrivate()->FindIntParameterValue(ParamIndex, IntParameter.ParameterValueName);
 						MutableParameters->SetIntValue(ParamIndex, Value);
 					}
 
@@ -864,7 +864,7 @@ void FCustomizableObjectInstanceDescriptor::ReloadParameters()
 						else
 						{
 							const int32 Value = MutableParameters->GetIntValue(ParamIndex, RangeIdxPtr.Get());
-							const FString AuxParameterValueName = CustomizableObject->FindIntParameterValueName(ParamIndex, Value);
+							const FString AuxParameterValueName = CustomizableObject->GetPrivate()->FindIntParameterValueName(ParamIndex, Value);
 							Param.ParameterRangeValueNames.Add(AuxParameterValueName);
 						}
 					}
@@ -878,7 +878,7 @@ void FCustomizableObjectInstanceDescriptor::ReloadParameters()
 					else
 					{
 						const int32 ParamValue = MutableParameters->GetIntValue(ParamIndex);
-						Param.ParameterValueName = CustomizableObject->FindIntParameterValueName(ParamIndex, ParamValue);
+						Param.ParameterValueName = CustomizableObject->GetPrivate()->FindIntParameterValueName(ParamIndex, ParamValue);
 					}
 
 					// Multilayer ints with one option are not multidimensional parameters. However, we need to preserve the layer 
@@ -886,7 +886,7 @@ void FCustomizableObjectInstanceDescriptor::ReloadParameters()
 					for (int32 RangeIndex = 0; RangeIndex < Result->ParameterRangeValueNames.Num(); ++RangeIndex)
 					{
 						const int32 Value = MutableParameters->GetIntValue(ParamIndex);
-						const FString AuxParameterValueName = CustomizableObject->FindIntParameterValueName(ParamIndex, Value);
+						const FString AuxParameterValueName = CustomizableObject->GetPrivate()->FindIntParameterValueName(ParamIndex, Value);
 						Param.ParameterRangeValueNames.Add(AuxParameterValueName);
 					}
 				}
@@ -908,14 +908,14 @@ void FCustomizableObjectInstanceDescriptor::ReloadParameters()
 						}
 
 						const int32 Value = MutableParameters->GetIntValue(ParamIndex, RangeValueIdxPtr.Get());
-						const FString AuxParameterValueName = CustomizableObject->FindIntParameterValueName(ParamIndex, Value);
+						const FString AuxParameterValueName = CustomizableObject->GetPrivate()->FindIntParameterValueName(ParamIndex, Value);
 						Param.ParameterRangeValueNames[RangeIndex] = AuxParameterValueName;
 					}
 				}
 				else
 				{
 					const int32 ParamValue = MutableParameters->GetIntValue(ParamIndex);
-					Param.ParameterValueName = CustomizableObject->FindIntParameterValueName(ParamIndex, ParamValue);
+					Param.ParameterValueName = CustomizableObject->GetPrivate()->FindIntParameterValueName(ParamIndex, ParamValue);
 				}
 			}
 
@@ -1289,14 +1289,14 @@ const FString& FCustomizableObjectInstanceDescriptor::GetIntParameterSelectedOpt
 		return FCustomizableObjectBoolParameterValue::DEFAULT_PARAMETER_VALUE_NAME;
 	}
 	
-	const int32 ParameterIndexInObject = CustomizableObject->FindParameter(ParamName);
+	const int32 ParameterIndexInObject = CustomizableObject->GetPrivate()->FindParameter(ParamName);
 	const int32 ParameterIndexInInstance = FindTypedParameterIndex(ParamName, EMutableParameterType::Int);
 
 	if (ParameterIndexInObject >= 0 && IntParameters.IsValidIndex(ParameterIndexInInstance))
 	{
 		// Due to optimizations a parameter may lose its multidimensionality (if it becomes constant). 
 		// In that case it means there is only one possible option so it is ok to set it as if wasn't multidimensional.
-		bool bIsMultidimensional = CustomizableObject->IsParameterMultidimensional(ParameterIndexInObject);
+		bool bIsMultidimensional = CustomizableObject->GetPrivate()->IsParameterMultidimensional(ParameterIndexInObject);
 		if (!bIsMultidimensional || RangeIndex == INDEX_NONE)
 		{
 			return IntParameters[ParameterIndexInInstance].ParameterValueName;
@@ -1326,7 +1326,7 @@ void FCustomizableObjectInstanceDescriptor::SetIntParameterSelectedOption(const
 
 	RETURN_ON_UNCOMPILED_CO(CustomizableObject, TEXT("Error: Cannot set Int parameter "));
 
-	const int32 ParameterIndexInObject = IntParameters.IsValidIndex(ParameterIndexInInstance) ? CustomizableObject->FindParameter(IntParameters[ParameterIndexInInstance].ParameterName) : INDEX_NONE; //-V781
+	const int32 ParameterIndexInObject = IntParameters.IsValidIndex(ParameterIndexInInstance) ? CustomizableObject->GetPrivate()->FindParameter(IntParameters[ParameterIndexInInstance].ParameterName) : INDEX_NONE; //-V781
 
 	if (ParameterIndexInObject < 0 || ParameterIndexInInstance < 0)
 	{
@@ -1335,7 +1335,7 @@ void FCustomizableObjectInstanceDescriptor::SetIntParameterSelectedOption(const
 		return;
 	}
 
-	const bool bValid = SelectedOption == TEXT("None") || CustomizableObject->FindIntParameterValue(ParameterIndexInObject, SelectedOption) >= 0;
+	const bool bValid = SelectedOption == TEXT("None") || CustomizableObject->GetPrivate()->FindIntParameterValue(ParameterIndexInObject, SelectedOption) >= 0;
 	if (!bValid)
 	{
 #if !UE_BUILD_SHIPPING
@@ -1343,7 +1343,7 @@ void FCustomizableObjectInstanceDescriptor::SetIntParameterSelectedOption(const
 			TEXT("Tried to set the invalid value [%s] to parameter [%d, %s]! Value index=[%d]. Correct values=[%s]."),
 			*SelectedOption, ParameterIndexInObject,
 			*IntParameters[ParameterIndexInInstance].ParameterName,
-			CustomizableObject->FindIntParameterValue(ParameterIndexInObject, SelectedOption),
+			CustomizableObject->GetPrivate()->FindIntParameterValue(ParameterIndexInObject, SelectedOption),
 			*GetAvailableOptionsString(*CustomizableObject, ParameterIndexInObject)
 		);
 		UE_LOG(LogMutable, Error, TEXT("%s"), *Message);
@@ -1353,7 +1353,7 @@ void FCustomizableObjectInstanceDescriptor::SetIntParameterSelectedOption(const
 
 	// Due to optimizations a parameter may lose its multidimensionality (if it becomes constant). 
 	// In that case it means there is only one possible option so it is ok to set it as if wasn't multidimensional.
-	bool bIsMultidimensional = CustomizableObject->IsParameterMultidimensional(ParameterIndexInObject);
+	bool bIsMultidimensional = CustomizableObject->GetPrivate()->IsParameterMultidimensional(ParameterIndexInObject);
 	if (!bIsMultidimensional || RangeIndex == INDEX_NONE)
 	{
 		IntParameters[ParameterIndexInInstance].ParameterValueName = SelectedOption;
@@ -1403,14 +1403,14 @@ float FCustomizableObjectInstanceDescriptor::GetFloatParameterSelectedOption(con
 		return FCustomizableObjectFloatParameterValue::DEFAULT_PARAMETER_VALUE;
 	}
 
-	const int32 ParameterIndexInObject = CustomizableObject->FindParameter(FloatParamName);
+	const int32 ParameterIndexInObject = CustomizableObject->GetPrivate()->FindParameter(FloatParamName);
 	const int32 FloatParamIndex = FindTypedParameterIndex(FloatParamName, EMutableParameterType::Float);
 
 	if (ParameterIndexInObject >= 0 && FloatParameters.IsValidIndex(FloatParamIndex))
 	{
 		// Due to optimizations a parameter may lose its multidimensionality (if it becomes constant). 
 		// In that case it means there is only one possible option so it is ok to set it as if wasn't multidimensional.
-		bool bIsMultidimensional = CustomizableObject->IsParameterMultidimensional(ParameterIndexInObject);
+		bool bIsMultidimensional = CustomizableObject->GetPrivate()->IsParameterMultidimensional(ParameterIndexInObject);
 		if (!bIsMultidimensional || RangeIndex == INDEX_NONE)
 		{
 			return FloatParameters[FloatParamIndex].ParameterValue;
@@ -1440,7 +1440,7 @@ void FCustomizableObjectInstanceDescriptor::SetFloatParameterSelectedOption(cons
 
 	RETURN_ON_UNCOMPILED_CO(CustomizableObject, TEXT("Error: Cannot set Int parameter "));
 
-	const int32 ParameterIndexInObject = CustomizableObject->FindParameter(FloatParamName);
+	const int32 ParameterIndexInObject = CustomizableObject->GetPrivate()->FindParameter(FloatParamName);
 	const int32 ParameterIndexInInstance = FindTypedParameterIndex(FloatParamName, EMutableParameterType::Float);
 
 	if (ParameterIndexInObject < 0 || ParameterIndexInInstance < 0)
@@ -1452,7 +1452,7 @@ void FCustomizableObjectInstanceDescriptor::SetFloatParameterSelectedOption(cons
 
 	// Due to optimizations a parameter may lose its multidimensionality (if it becomes constant). 
 	// In that case it means there is only one possible option so it is ok to set it as if wasn't multidimensional.
-	bool bIsMultidimensional = CustomizableObject->IsParameterMultidimensional(ParameterIndexInObject);
+	bool bIsMultidimensional = CustomizableObject->GetPrivate()->IsParameterMultidimensional(ParameterIndexInObject);
 	if (!bIsMultidimensional || RangeIndex == INDEX_NONE)
 	{
 		FloatParameters[ParameterIndexInInstance].ParameterValue = FloatValue;
@@ -1481,14 +1481,14 @@ FName FCustomizableObjectInstanceDescriptor::GetTextureParameterSelectedOption(c
 		return FName();
 	}
 
-	const int32 ParameterIndexInObject = CustomizableObject->FindParameter(TextureParamName);
+	const int32 ParameterIndexInObject = CustomizableObject->GetPrivate()->FindParameter(TextureParamName);
 	const int32 ParameterIndexInInstance = FindTypedParameterIndex(TextureParamName, EMutableParameterType::Texture);
 
 	if (ParameterIndexInObject >= 0 && ParameterIndexInInstance >= 0)
 	{
 		// Due to optimizations a parameter may lose its multidimensionality (if it becomes constant). 
 		// In that case it means there is only one possible option so it is ok to set it as if wasn't multidimensional.
-		bool bIsMultidimensional = CustomizableObject->IsParameterMultidimensional(ParameterIndexInObject);
+		bool bIsMultidimensional = CustomizableObject->GetPrivate()->IsParameterMultidimensional(ParameterIndexInObject);
 		if (!bIsMultidimensional || RangeIndex == INDEX_NONE)
 		{
 			return TextureParameters[ParameterIndexInInstance].ParameterValue;
@@ -1518,7 +1518,7 @@ void FCustomizableObjectInstanceDescriptor::SetTextureParameterSelectedOption(co
 
 	RETURN_ON_UNCOMPILED_CO(CustomizableObject, TEXT("Error: Cannot set Int parameter "));
 
-	const int32 ParameterIndexInObject = CustomizableObject->FindParameter(TextureParamName);
+	const int32 ParameterIndexInObject = CustomizableObject->GetPrivate()->FindParameter(TextureParamName);
 	const int32 ParameterIndexInInstance = FindTypedParameterIndex(TextureParamName, EMutableParameterType::Texture);
 
 	if (ParameterIndexInObject < 0 || ParameterIndexInInstance < 0)
@@ -1530,7 +1530,7 @@ void FCustomizableObjectInstanceDescriptor::SetTextureParameterSelectedOption(co
 
 	// Due to optimizations a parameter may lose its multidimensionality (if it becomes constant). 
 	// In that case it means there is only one possible option so it is ok to set it as if wasn't multidimensional.
-	bool bIsMultidimensional = CustomizableObject->IsParameterMultidimensional(ParameterIndexInObject);
+	bool bIsMultidimensional = CustomizableObject->GetPrivate()->IsParameterMultidimensional(ParameterIndexInObject);
 	if (!bIsMultidimensional || RangeIndex == INDEX_NONE)
 	{
 		TextureParameters[ParameterIndexInInstance].ParameterValue = FName(TextureValue);
@@ -1558,7 +1558,7 @@ FLinearColor FCustomizableObjectInstanceDescriptor::GetColorParameterSelectedOpt
 		return FCustomizableObjectVectorParameterValue::DEFAULT_PARAMETER_VALUE;
 	}
 
-	const int32 ParameterIndexInObject = CustomizableObject->FindParameter(ColorParamName);
+	const int32 ParameterIndexInObject = CustomizableObject->GetPrivate()->FindParameter(ColorParamName);
 	const int32 ColorParamIndex = FindTypedParameterIndex(ColorParamName, EMutableParameterType::Color);
 
 	if (ColorParamIndex == INDEX_NONE)
@@ -1594,7 +1594,7 @@ FTransform FCustomizableObjectInstanceDescriptor::GetTransformParameterSelectedO
 		return FCustomizableObjectTransformParameterValue::DEFAULT_PARAMETER_VALUE;
 	}
 
-	const int32 ParameterIndexInObject = CustomizableObject->FindParameter(TransformParamName);
+	const int32 ParameterIndexInObject = CustomizableObject->GetPrivate()->FindParameter(TransformParamName);
 	const int32 TransformParamIndex = FindTypedParameterIndex(TransformParamName, EMutableParameterType::Transform);
 
 	if (TransformParamIndex == INDEX_NONE)
@@ -1617,7 +1617,7 @@ void FCustomizableObjectInstanceDescriptor::SetTransformParameterSelectedOption(
 
 	RETURN_ON_UNCOMPILED_CO(CustomizableObject, TEXT("Error: Cannot set Transform parameter "));
 
-	const int32 ParameterIndexInObject = CustomizableObject->FindParameter(TransformParamName);
+	const int32 ParameterIndexInObject = CustomizableObject->GetPrivate()->FindParameter(TransformParamName);
 	const int32 ParameterIndexInInstance = FindTypedParameterIndex(TransformParamName, EMutableParameterType::Transform);
 
 	if (ParameterIndexInObject < 0 || ParameterIndexInInstance < 0)
@@ -1640,7 +1640,7 @@ bool FCustomizableObjectInstanceDescriptor::GetBoolParameterSelectedOption(const
 		return FCustomizableObjectBoolParameterValue::DEFAULT_PARAMETER_VALUE;
 	}
 
-	const int32 ParameterIndexInObject = CustomizableObject->FindParameter(BoolParamName);
+	const int32 ParameterIndexInObject = CustomizableObject->GetPrivate()->FindParameter(BoolParamName);
 	const int32 BoolParamIndex = FindTypedParameterIndex(BoolParamName, EMutableParameterType::Bool);
 
 	if (BoolParamIndex == INDEX_NONE)
@@ -1663,7 +1663,7 @@ void FCustomizableObjectInstanceDescriptor::SetBoolParameterSelectedOption(const
 
 	RETURN_ON_UNCOMPILED_CO(CustomizableObject, TEXT("Error: Cannot set Int parameter "));
 
-	const int32 ParameterIndexInObject = CustomizableObject->FindParameter(BoolParamName);
+	const int32 ParameterIndexInObject = CustomizableObject->GetPrivate()->FindParameter(BoolParamName);
 	const int32 ParameterIndexInInstance = FindTypedParameterIndex(BoolParamName, EMutableParameterType::Bool);
 
 	if (ParameterIndexInObject < 0 || ParameterIndexInInstance < 0)
@@ -1687,7 +1687,7 @@ void FCustomizableObjectInstanceDescriptor::SetVectorParameterSelectedOption(con
 
 	RETURN_ON_UNCOMPILED_CO(CustomizableObject, TEXT("Error: Cannot set Int parameter "));
 
-	const int32 ParameterIndexInObject = CustomizableObject->FindParameter(VectorParamName);
+	const int32 ParameterIndexInObject = CustomizableObject->GetPrivate()->FindParameter(VectorParamName);
 	const int32 ParameterIndexInInstance = FindTypedParameterIndex(VectorParamName, EMutableParameterType::Color);
 
 	if (ParameterIndexInObject < 0 || ParameterIndexInInstance < 0)
@@ -1714,7 +1714,7 @@ void FCustomizableObjectInstanceDescriptor::SetProjectorValue(const FString& Pro
 
 	RETURN_ON_UNCOMPILED_CO(CustomizableObject, TEXT("Error: Cannot set Projector parameter "))
 
-	const int32 ParameterIndexInObject = CustomizableObject->FindParameter(ProjectorParamName);
+	const int32 ParameterIndexInObject = CustomizableObject->GetPrivate()->FindParameter(ProjectorParamName);
 	const int32 ParameterIndexInInstance = FindTypedParameterIndex(ProjectorParamName, EMutableParameterType::Projector);
 
 	if (ParameterIndexInObject < 0 || ParameterIndexInInstance < 0)
@@ -1738,7 +1738,7 @@ void FCustomizableObjectInstanceDescriptor::SetProjectorValue(const FString& Pro
 
 	// Due to optimizations a parameter may lose its multidimensionality (if it becomes constant). 
 	// In that case it means there is only one possible option so it is ok to set it as if wasn't multidimensional.
-	bool bIsMultidimensional = CustomizableObject->IsParameterMultidimensional(ParameterIndexInObject);
+	bool bIsMultidimensional = CustomizableObject->GetPrivate()->IsParameterMultidimensional(ParameterIndexInObject);
 	if (!bIsMultidimensional || RangeIndex == INDEX_NONE)
 	{
 		ProjectorParameter.Value = ProjectorData;
@@ -1839,7 +1839,7 @@ void FCustomizableObjectInstanceDescriptor::GetProjectorValueF(const FString& Pr
 		return;
 	}
 
-	const int32 ParameterIndexInObject = CustomizableObject->FindParameter(ProjectorParamName);
+	const int32 ParameterIndexInObject = CustomizableObject->GetPrivate()->FindParameter(ProjectorParamName);
 	const int32 ParameterIndexInInstance = FindTypedParameterIndex(ProjectorParamName, EMutableParameterType::Projector);
 
 	if (ParameterIndexInObject >= 0 && ParameterIndexInInstance >= 0)
@@ -1849,7 +1849,7 @@ void FCustomizableObjectInstanceDescriptor::GetProjectorValueF(const FString& Pr
 
 		// Due to optimizations a parameter may lose its multidimensionality (if it becomes constant). 
 		// In that case it means there is only one possible option so it is ok to set it as if wasn't multidimensional.
-		bool bIsMultidimensional = CustomizableObject->IsParameterMultidimensional(ParameterIndexInObject);
+		bool bIsMultidimensional = CustomizableObject->GetPrivate()->IsParameterMultidimensional(ParameterIndexInObject);
 		if (!bIsMultidimensional || RangeIndex == INDEX_NONE)
 		{
 			ProjectorData = &ProjectorParameter.Value;
@@ -1886,14 +1886,14 @@ FVector FCustomizableObjectInstanceDescriptor::GetProjectorPosition(const FStrin
 		return FVector(-0.0, -0.0, -0.0);
 	}
 
-	const int32 ParameterIndexInObject = CustomizableObject->FindParameter(ParamName);
+	const int32 ParameterIndexInObject = CustomizableObject->GetPrivate()->FindParameter(ParamName);
 	const int32 ParameterIndexInInstance = FindTypedParameterIndex(ParamName, EMutableParameterType::Projector);
 
 	if (ParameterIndexInObject >= 0 && ParameterIndexInInstance >= 0)
 	{
 		// Due to optimizations a parameter may lose its multidimensionality (if it becomes constant). 
 		// In that case it means there is only one possible option so it is ok to set it as if wasn't multidimensional.
-		bool bIsMultidimensional = CustomizableObject->IsParameterMultidimensional(ParameterIndexInObject);
+		bool bIsMultidimensional = CustomizableObject->GetPrivate()->IsParameterMultidimensional(ParameterIndexInObject);
 		if (!bIsMultidimensional || RangeIndex == INDEX_NONE)
 		{
 			return static_cast<FVector>(ProjectorParameters[ParameterIndexInInstance].Value.Position);
@@ -1918,14 +1918,14 @@ FVector FCustomizableObjectInstanceDescriptor::GetProjectorDirection(const FStri
 		return FVector(-0.0, -0.0, -0.0);
 	}
 
-	const int32 ParameterIndexInObject = CustomizableObject->FindParameter(ParamName);
+	const int32 ParameterIndexInObject = CustomizableObject->GetPrivate()->FindParameter(ParamName);
 	const int32 ParameterIndexInInstance = FindTypedParameterIndex(ParamName, EMutableParameterType::Projector);
 
 	if (ParameterIndexInObject >= 0 && ParameterIndexInInstance >= 0)
 	{
 		// Due to optimizations a parameter may lose its multidimensionality (if it becomes constant). 
 		// In that case it means there is only one possible option so it is ok to set it as if wasn't multidimensional.
-		bool bIsMultidimensional = CustomizableObject->IsParameterMultidimensional(ParameterIndexInObject);
+		bool bIsMultidimensional = CustomizableObject->GetPrivate()->IsParameterMultidimensional(ParameterIndexInObject);
 		if (!bIsMultidimensional || RangeIndex == INDEX_NONE)
 		{
 			return static_cast<FVector>(ProjectorParameters[ParameterIndexInInstance].Value.Direction);
@@ -1950,14 +1950,14 @@ FVector FCustomizableObjectInstanceDescriptor::GetProjectorUp(const FString& Par
 		return FVector(-0.0, -0.0, -0.0);
 	}
 
-	const int32 ParameterIndexInObject = CustomizableObject->FindParameter(ParamName);
+	const int32 ParameterIndexInObject = CustomizableObject->GetPrivate()->FindParameter(ParamName);
 	const int32 ParameterIndexInInstance = FindTypedParameterIndex(ParamName, EMutableParameterType::Projector);
 
 	if (ParameterIndexInObject >= 0 && ParameterIndexInInstance >= 0)
 	{
 		// Due to optimizations a parameter may lose its multidimensionality (if it becomes constant). 
 		// In that case it means there is only one possible option so it is ok to set it as if wasn't multidimensional.
-		bool bIsMultidimensional = CustomizableObject->IsParameterMultidimensional(ParameterIndexInObject);
+		bool bIsMultidimensional = CustomizableObject->GetPrivate()->IsParameterMultidimensional(ParameterIndexInObject);
 		if (!bIsMultidimensional || RangeIndex == INDEX_NONE)
 		{
 			return static_cast<FVector>(ProjectorParameters[ParameterIndexInInstance].Value.Up);
@@ -1982,14 +1982,14 @@ FVector FCustomizableObjectInstanceDescriptor::GetProjectorScale(const FString&
 		return FVector(-0.0, -0.0, -0.0);
 	}
 
-	const int32 ParameterIndexInObject = CustomizableObject->FindParameter(ParamName);
+	const int32 ParameterIndexInObject = CustomizableObject->GetPrivate()->FindParameter(ParamName);
 	const int32 ParameterIndexInInstance = FindTypedParameterIndex(ParamName, EMutableParameterType::Projector);
 
 	if (ParameterIndexInObject >= 0 && ParameterIndexInInstance >= 0)
 	{
 		// Due to optimizations a parameter may lose its multidimensionality (if it becomes constant). 
 		// In that case it means there is only one possible option so it is ok to set it as if wasn't multidimensional.
-		bool bIsMultidimensional = CustomizableObject->IsParameterMultidimensional(ParameterIndexInObject);
+		bool bIsMultidimensional = CustomizableObject->GetPrivate()->IsParameterMultidimensional(ParameterIndexInObject);
 		if (!bIsMultidimensional || RangeIndex == INDEX_NONE)
 		{
 			return static_cast<FVector>(ProjectorParameters[ParameterIndexInInstance].Value.Scale);
@@ -2014,14 +2014,14 @@ float FCustomizableObjectInstanceDescriptor::GetProjectorAngle(const FString& Pa
 		return 0.0;
 	}
 
-	const int32 ParameterIndexInObject = CustomizableObject->FindParameter(ParamName);
+	const int32 ParameterIndexInObject = CustomizableObject->GetPrivate()->FindParameter(ParamName);
 	const int32 ParameterIndexInInstance = FindTypedParameterIndex(ParamName, EMutableParameterType::Projector);
 
 	if (ParameterIndexInObject >= 0 && ParameterIndexInInstance >= 0)
 	{
 		// Due to optimizations a parameter may lose its multidimensionality (if it becomes constant). 
 		// In that case it means there is only one possible option so it is ok to set it as if wasn't multidimensional.
-		bool bIsMultidimensional = CustomizableObject->IsParameterMultidimensional(ParameterIndexInObject);
+		bool bIsMultidimensional = CustomizableObject->GetPrivate()->IsParameterMultidimensional(ParameterIndexInObject);
 		if (!bIsMultidimensional || RangeIndex == INDEX_NONE)
 		{
 			return ProjectorParameters[ParameterIndexInInstance].Value.Angle;
@@ -2046,14 +2046,14 @@ ECustomizableObjectProjectorType FCustomizableObjectInstanceDescriptor::GetProje
 		return ECustomizableObjectProjectorType::Planar;
 	}
 
-	const int32 ParameterIndexInObject = CustomizableObject->FindParameter(ParamName);
+	const int32 ParameterIndexInObject = CustomizableObject->GetPrivate()->FindParameter(ParamName);
 	const int32 ParameterIndexInInstance = FindTypedParameterIndex(ParamName, EMutableParameterType::Projector);
 
 	if (ParameterIndexInObject >= 0 && ParameterIndexInInstance >= 0)
 	{
 		// Due to optimizations a parameter may lose its multidimensionality (if it becomes constant). 
 		// In that case it means there is only one possible option so it is ok to set it as if wasn't multidimensional.
-		bool bIsMultidimensional = CustomizableObject->IsParameterMultidimensional(ParameterIndexInObject);
+		bool bIsMultidimensional = CustomizableObject->GetPrivate()->IsParameterMultidimensional(ParameterIndexInObject);
 		if (!bIsMultidimensional || RangeIndex == INDEX_NONE)
 		{
 			return ProjectorParameters[ParameterIndexInInstance].Value.ProjectionType;
@@ -2072,14 +2072,14 @@ ECustomizableObjectProjectorType FCustomizableObjectInstanceDescriptor::GetProje
 
 FCustomizableObjectProjector FCustomizableObjectInstanceDescriptor::GetProjector(const FString& ParamName, const int32 RangeIndex) const
 {
-	const int32 ParameterIndexInObject = CustomizableObject->FindParameter(ParamName);
+	const int32 ParameterIndexInObject = CustomizableObject->GetPrivate()->FindParameter(ParamName);
 	const int32 ParameterIndexInInstance = FindTypedParameterIndex(ParamName, EMutableParameterType::Projector);
 
 	if (ParameterIndexInObject >= 0 && ParameterIndexInInstance >= 0)
 	{
 		// Due to optimizations a parameter may lose its multidimensionality (if it becomes constant). 
 		// In that case it means there is only one possible option so it is ok to set it as if wasn't multidimensional.
-		bool bIsMultidimensional = CustomizableObject->IsParameterMultidimensional(ParameterIndexInObject);
+		bool bIsMultidimensional = CustomizableObject->GetPrivate()->IsParameterMultidimensional(ParameterIndexInObject);
 		if (!bIsMultidimensional || RangeIndex == INDEX_NONE)
 		{
 			return ProjectorParameters[ParameterIndexInInstance].Value;
@@ -2185,9 +2185,9 @@ int32 FCustomizableObjectInstanceDescriptor::AddValueToIntRange(const FString& P
 	if (IntParameterIndex != INDEX_NONE)
 	{
 		FCustomizableObjectIntParameterValue& IntParameter = IntParameters[IntParameterIndex];
-		const int32 ParamIndexInObject = CustomizableObject->FindParameter(IntParameter.ParameterName);
+		const int32 ParamIndexInObject = CustomizableObject->GetPrivate()->FindParameter(IntParameter.ParameterName);
 		// TODO: Define the default option in the editor instead of taking the first available, like it's currently defined for GetProjectorDefaultValue()
-		const FString DefaultValue = CustomizableObject->GetIntParameterAvailableOption(ParamIndexInObject, 0);
+		const FString DefaultValue = CustomizableObject->GetPrivate()->GetIntParameterAvailableOption(ParamIndexInObject, 0);
 		return IntParameter.ParameterRangeValueNames.Add(DefaultValue);
 	}
 	return -1;
@@ -2339,7 +2339,7 @@ FString FCustomizableObjectInstanceDescriptor::GetCurrentState() const
 		return FString();
 	}
 
-	return CustomizableObject->GetStateName(GetState());
+	return CustomizableObject->GetPrivate()->GetStateName(GetState());
 }
 
 
@@ -2359,7 +2359,7 @@ void FCustomizableObjectInstanceDescriptor::SetCurrentState(const FString& State
 
 	RETURN_ON_UNCOMPILED_CO(CustomizableObject, TEXT("Error: Cannot set state"))
 
-	const int32 Result = CustomizableObject->FindState(StateName);
+	const int32 Result = CustomizableObject->GetPrivate()->FindState(StateName);
 #if WITH_EDITOR
 	if (Result != INDEX_NONE)
 #else
@@ -2410,16 +2410,15 @@ void FCustomizableObjectInstanceDescriptor::SetRandomValuesFromStream(const FRan
 	
 	for (FCustomizableObjectIntParameterValue& IntParameter : IntParameters)
 	{
-		const int32 ParamIndex = CustomizableObject->FindParameter(IntParameter.ParameterName);
-		const int32 NumValues = CustomizableObject->GetIntParameterNumOptions(ParamIndex);
+		const int32 NumValues = CustomizableObject->GetEnumParameterNumValues(IntParameter.ParameterName);
 
 		if (NumValues)
 		{
-			IntParameter.ParameterValueName = CustomizableObject->GetIntParameterAvailableOption(ParamIndex, NumValues * InStream.GetFraction());
+			IntParameter.ParameterValueName = CustomizableObject->GetEnumParameterValue(IntParameter.ParameterName, NumValues * InStream.GetFraction());
 
 			for (FString& RangeValue : IntParameter.ParameterRangeValueNames)
 			{
-				RangeValue = CustomizableObject->GetIntParameterAvailableOption(ParamIndex, NumValues * InStream.GetFraction());
+				RangeValue = CustomizableObject->GetEnumParameterValue(IntParameter.ParameterName, NumValues * InStream.GetFraction());
 			}
 		}		
 	}
@@ -2480,7 +2479,7 @@ void FCustomizableObjectInstanceDescriptor::SetDefaultValue(int32 ParamIndex)
 	}
 
 	const FString ParamName = CustomizableObject->GetParameterName(ParamIndex);
-	const EMutableParameterType ParamType = CustomizableObject->GetParameterType(ParamIndex);
+	const EMutableParameterType ParamType = CustomizableObject->GetPrivate()->GetParameterType(ParamIndex);
 	
 	const int32 TypedIndex = FindTypedParameterIndex(ParamName, ParamType);
 
@@ -2495,7 +2494,7 @@ void FCustomizableObjectInstanceDescriptor::SetDefaultValue(int32 ParamIndex)
 
 	case EMutableParameterType::Int:
 	{
-		const FString DefaultValue = CustomizableObject->FindIntParameterValueName(ParamIndex, CustomizableObject->GetIntParameterDefaultValue(ParamName));
+		const FString DefaultValue = CustomizableObject->GetPrivate()->FindIntParameterValueName(ParamIndex, CustomizableObject->GetEnumParameterDefaultValue(ParamName));
 
 		FCustomizableObjectIntParameterValue& IntParameter = IntParameters[TypedIndex];
 		IntParameter.ParameterValueName = DefaultValue;
@@ -2649,9 +2648,8 @@ void FCustomizableObjectInstanceDescriptor::CreateLayer(const FName& ProjectorPa
 		const FString ImageParamName = ParamName + IMAGE_PARAMETER_POSTFIX;
 		AddValueToTextureRange(ImageParamName);
 
-		const int32 ParamIndex = CustomizableObject->FindParameter(ImageParamName);
-		const int32 DefaultValueIndex = CustomizableObject->GetIntParameterDefaultValue(ImageParamName);
-		const FString DefaultValueName = CustomizableObject->GetIntParameterAvailableOption(ParamIndex, DefaultValueIndex);
+		const int32 DefaultValueIndex = CustomizableObject->GetEnumParameterDefaultValue(ImageParamName);
+		const FString DefaultValueName = CustomizableObject->GetEnumParameterValue(ImageParamName, DefaultValueIndex);
 		SetIntParameterSelectedOption(ImageParamName, DefaultValueName, RangeIndex);
 	}
 
@@ -2660,7 +2658,6 @@ void FCustomizableObjectInstanceDescriptor::CreateLayer(const FName& ProjectorPa
 		const FString OpacityParamName = ParamName + OPACITY_PARAMETER_POSTFIX;
 		AddValueToFloatRange(OpacityParamName);
 
-		const int32 ParamIndex = CustomizableObject->FindParameter(OpacityParamName);
 		const float DefaultValue = CustomizableObject->GetFloatParameterDefaultValue(OpacityParamName);
 		SetFloatParameterSelectedOption(OpacityParamName, DefaultValue, RangeIndex);
 	}
diff --git a/Engine/Plugins/Mutable/Source/CustomizableObject/Private/MuCO/CustomizableObjectSystem.cpp b/Engine/Plugins/Mutable/Source/CustomizableObject/Private/MuCO/CustomizableObjectSystem.cpp
--- a/Engine/Plugins/Mutable/Source/CustomizableObject/Private/MuCO/CustomizableObjectSystem.cpp
+++ b/Engine/Plugins/Mutable/Source/CustomizableObject/Private/MuCO/CustomizableObjectSystem.cpp
@@ -3409,7 +3409,7 @@ namespace impl
 		check(SystemPrivateData->ExtensionDataStreamer != nullptr);
 		SystemPrivateData->ExtensionDataStreamer->SetActiveObject(CustomizableObject);
 
-		FString StateName = CandidateInstance->GetCustomizableObject()->GetStateName(CandidateInstance->GetPrivate()->GetState());
+		FString StateName = CandidateInstance->GetCustomizableObject()->GetPrivate()->GetStateName(CandidateInstance->GetPrivate()->GetState());
 		const FMutableStateData* StateData = CandidateInstance->GetCustomizableObject()->GetPrivate()->GetModelResources().StateUIDataMap.Find(StateName);
 
 		Operation->bLiveUpdateMode = false;
diff --git a/Engine/Plugins/Mutable/Source/CustomizableObject/Public/MuCO/CustomizableObject.h b/Engine/Plugins/Mutable/Source/CustomizableObject/Public/MuCO/CustomizableObject.h
--- a/Engine/Plugins/Mutable/Source/CustomizableObject/Public/MuCO/CustomizableObject.h
+++ b/Engine/Plugins/Mutable/Source/CustomizableObject/Public/MuCO/CustomizableObject.h
@@ -24,13 +24,6 @@ class UCustomizableObjectBulk;
 struct FFrame;
 struct FStreamableHandle;
 template <typename FuncType> class TFunctionRef;
- 
-// Forward declaration of Mutable classes necessary for the interface
-namespace mu
-{
-	class FModel;
-	class FParameters;
-}
 
 
 CUSTOMIZABLEOBJECT_API DECLARE_LOG_CATEGORY_EXTERN(LogMutable, Log, All);
@@ -101,107 +94,6 @@ struct FProfileParameterDat
 	TArray<FCustomizableObjectTransformParameterValue> TransformParameters;
 };
 
-// A USTRUCT version of FMeshToMeshVertData in SkeletalMeshTypes.h
-// We are taking advantage of the padding data to store from which asset this data comes from
-// maintaining the same memory footprint than the original.
-USTRUCT()
-struct FCustomizableObjectMeshToMeshVertData
-{
-	GENERATED_USTRUCT_BODY()
-
-	FCustomizableObjectMeshToMeshVertData() = default;
-
-	explicit FCustomizableObjectMeshToMeshVertData(const FMeshToMeshVertData& Original)
-		: Weight(Original.Weight)
-	{
-		for (int32 i = 0; i < 4; ++i)
-		{
-			PositionBaryCoordsAndDist[i] = Original.PositionBaryCoordsAndDist[i];
-			NormalBaryCoordsAndDist[i] = Original.NormalBaryCoordsAndDist[i];
-			TangentBaryCoordsAndDist[i] = Original.TangentBaryCoordsAndDist[i];
-			SourceMeshVertIndices[i] = Original.SourceMeshVertIndices[i];
-		}
-	}
-		
-
-	explicit operator FMeshToMeshVertData() const
-	{
-		FMeshToMeshVertData ReturnValue;
-
-		for (int32 i = 0; i < 4; ++i)
-		{
-			ReturnValue.PositionBaryCoordsAndDist[i] = PositionBaryCoordsAndDist[i];
-			ReturnValue.NormalBaryCoordsAndDist[i] = NormalBaryCoordsAndDist[i];
-			ReturnValue.TangentBaryCoordsAndDist[i] = TangentBaryCoordsAndDist[i];
-			ReturnValue.SourceMeshVertIndices[i] = SourceMeshVertIndices[i];
-		}
-		ReturnValue.Weight = Weight;
-		ReturnValue.Padding = 0;
-
-		return ReturnValue;
-	}
-
-	
-	// Barycentric coords and distance along normal for the position of the final vert
-	UPROPERTY()
-	float PositionBaryCoordsAndDist[4] = {0.f};
-
-	// Barycentric coords and distance along normal for the location of the unit normal endpoint
-	// Actual normal = ResolvedNormalPosition - ResolvedPosition
-	UPROPERTY()
-	float NormalBaryCoordsAndDist[4] = {0.f};
-
-	// Barycentric coords and distance along normal for the location of the unit Tangent endpoint
-	// Actual normal = ResolvedNormalPosition - ResolvedPosition
-	UPROPERTY()
-	float TangentBaryCoordsAndDist[4] = {0.f};
-
-	// Contains the 3 indices for verts in the source mesh forming a triangle, the last element
-	// is a flag to decide how the skinning works, 0xffff uses no simulation, and just normal
-	// skinning, anything else uses the source mesh and the above skin data to get the final position
-	UPROPERTY()
-	uint16	 SourceMeshVertIndices[4] = {0u, 0u, 0u, 0u};
-
-	UPROPERTY()
-	float Weight = 0.0f;
-
-	// Non serialized, unused padding. This is present in the FMeshToMeshVertData struct as Padding for alignment.
-	uint32 UnusedPadding = 0;
-
-	/**
-	 * Serializer
-	 *
-	 * @param Ar - archive to serialize with
-	 * @param V - vertex to serialize
-	 * @return archive that was used
-	 */
-	friend FArchive& operator<<(FArchive& Ar, FCustomizableObjectMeshToMeshVertData& V)
-	{
-		Ar << V.PositionBaryCoordsAndDist[0];
-		Ar << V.PositionBaryCoordsAndDist[1];
-		Ar << V.PositionBaryCoordsAndDist[2];
-		Ar << V.PositionBaryCoordsAndDist[3];
-		Ar << V.NormalBaryCoordsAndDist[0];
-		Ar << V.NormalBaryCoordsAndDist[1];
-		Ar << V.NormalBaryCoordsAndDist[2];
-		Ar << V.NormalBaryCoordsAndDist[3];
-		Ar << V.TangentBaryCoordsAndDist[0];
-		Ar << V.TangentBaryCoordsAndDist[1];
-		Ar << V.TangentBaryCoordsAndDist[2];
-		Ar << V.TangentBaryCoordsAndDist[3];
-
-		Ar << V.SourceMeshVertIndices[0];
-		Ar << V.SourceMeshVertIndices[1];
-		Ar << V.SourceMeshVertIndices[2];
-		Ar << V.SourceMeshVertIndices[3];
-
-		Ar << V.Weight;
-
-		return Ar;
-	}
-};
-static_assert(sizeof(FCustomizableObjectMeshToMeshVertData) == sizeof(float)*4*3 + sizeof(uint16)*4 + sizeof(float) + sizeof(uint32));
-template<> struct TCanBulkSerialize<FCustomizableObjectMeshToMeshVertData> { enum { Value = true }; };
 
 USTRUCT()
 struct FMutableLODSettings
@@ -258,28 +150,31 @@ private:
 
 public:
 
-	/**
-	  * The optional VersionBridge asset, which must implement the ICustomizableObjectVersionBridgeInterface, will be used to
-	  * decide which Mutable child CustomizableObjects and table rows must be included in a compilation/cook depending on its 
-	  * version struct/column by comparing it to the game-specific version system.
-	*/
+	/** Optional Version Bridge asset.
+	  *
+	  * Used to decide which Mutable child Customizable Objects and Data Table rows must be included in a compilation/cook.
+	  * An asset will be included if the struct/column version matches the game-specific system version.
+	  *
+	  * The provided asset must implement ICustomizableObjectVersionBridgeInterface. */
 	UPROPERTY(EditAnywhere, Category = Versioning)
 	TObjectPtr<UObject> VersionBridge;
 
-	/**
-	  * This optional struct is used to define which version this child CustomizableObject belongs to. It will be used during
-	  * cook/compilation to decide whether this CO should be included or not in the final compiled CO. To be used, the root
-	  * CO must have defined the VersionBridge property, which must implement the ICustomizableObjectVersionBridgeInterface
-	*/
+	/** Optional struct.
+	  *
+	  * Used to define which version this child Customizable Object belongs to. It will be used during
+	  * cook/compilation to decide whether this Customizable Object should be included in the final compiled Customizable Object.
+	  *
+	  * To work, the root Customizable Object must have a Version Bridge. */
 	UPROPERTY(EditAnywhere, Category = Versioning)
 	FInstancedStruct VersionStruct;
 #endif
-	
+
+	// DEPRECATED, Will be per component
 	UPROPERTY(EditAnywhere, Category = CustomizableObject, meta = (DisplayName = "LOD Settings"))
 	FMutableLODSettings LODSettings;
 	
-	/** Use the SkeletalMesh of reference as a placeholder until the custom mesh is ready to use.
-	  * Note: If disabled, a null mesh will be used to replace the discarded mesh due to 'ReplaceDiscardedWithReferenceMesh' being enabled. */
+	/** true to use the Reference Skeletal Mesh as a placeholder while the Generated Skeletal Mesh is not ready.
+	  * If false, a null mesh will be used to replace the discarded mesh due to 'ReplaceDiscardedWithReferenceMesh' being enabled. */
 	UPROPERTY(EditAnywhere, Category = CustomizableObject)
 	bool bEnableUseRefSkeletalMeshAsPlaceholder = true;
 
@@ -303,42 +198,52 @@ private:
 	FCompilationOptions_DEPRECATED CompileOptions_DEPRECATED;
 
 public:
+	UE_DEPRECATED(5.6, "Only modify this value through the Customizable ObjectEditor")
 	UPROPERTY(EditAnywhere, Category = CompileOptions)
-	bool bEnableRealTimeMorphTargets = false;
+	bool bEnableRealTimeMorphTargets = false; // TODO >5.6 make private
 
+	UE_DEPRECATED(5.6, "Only modify this value through the Customizable ObjectEditor")
 	UPROPERTY(EditAnywhere, Category = CompileOptions)
-	bool bEnableClothing = false;
+	bool bEnableClothing = false; // TODO >5.6 make private
 
+	UE_DEPRECATED(5.6, "Only modify this value through the Customizable ObjectEditor")
 	UPROPERTY(EditAnywhere, Category = CompileOptions)
-	bool bEnable16BitBoneWeights = false;
+	bool bEnable16BitBoneWeights = false; // TODO >5.6 make private
 
+	UE_DEPRECATED(5.6, "Only modify this value through the Customizable ObjectEditor")
 	UPROPERTY(EditAnywhere, Category = CompileOptions)
-	bool bEnableAltSkinWeightProfiles = false;
+	bool bEnableAltSkinWeightProfiles = false; // TODO >5.6 make private
 
+	UE_DEPRECATED(5.6, "Only modify this value through the Customizable ObjectEditor")
 	UPROPERTY(EditAnywhere, Category = CompileOptions)
-	bool bEnablePhysicsAssetMerge = false;
+	bool bEnablePhysicsAssetMerge = false; // TODO >5.6 make private
 
 	/** Experimental */
+	UE_DEPRECATED(5.6, "Only modify this value through the Customizable ObjectEditor")
 	UPROPERTY(EditAnywhere, Category = CompileOptions)
-	bool bEnableAnimBpPhysicsAssetsManipualtion = false;
+	bool bEnableAnimBpPhysicsAssetsManipualtion = false; // TODO >5.6 make private
 
 	// When this is enabled generated meshes will merge the AssetUserData from all of its constituent mesh parts
+	UE_DEPRECATED(5.6, "Only modify this value through the Customizable ObjectEditor")
 	UPROPERTY(EditAnywhere, Category = CompileOptions)
-	bool bEnableAssetUserDataMerge = false;
+	bool bEnableAssetUserDataMerge = false; // TODO >5.6 make private
 
 	// Disabling the Table Materials parent material check will let the user use any material regardless of its parent when connecting a material from a table column to a material node.
 	// Warning! But it will not check if the table material channels connected to the Material node exist in the actual material used in the Instance, and will fail silently at runtime 
 	// when setting the value of those channels if they don't exist.
+	UE_DEPRECATED(5.6, "Only modify this value through the Customizable ObjectEditor")
 	UPROPERTY(EditAnywhere, Category = CompileOptions)
-	bool bDisableTableMaterialsParentCheck = false;
+	bool bDisableTableMaterialsParentCheck = false; // TODO >5.6 make private
 
 	// Options when compiling this customizable object (see EMutableCompileMeshType declaration for info)
+	UE_DEPRECATED(5.6, "Only modify this value through the Customizable ObjectEditor")
 	UPROPERTY(EditAnywhere, Category = CompileOptions)
-	EMutableCompileMeshType MeshCompileType = EMutableCompileMeshType::LocalAndChildren;
+	EMutableCompileMeshType MeshCompileType = EMutableCompileMeshType::LocalAndChildren; // TODO >5.6 make private
 
 	// Array of elements to use with compile option CompileType = WorkingSet
+	UE_DEPRECATED(5.6, "Only modify this value through the Customizable ObjectEditor")
 	UPROPERTY(EditAnywhere, Category = CompileOptions)
-	TArray<TSoftObjectPtr<UCustomizableObject>> WorkingSet;
+	TArray<TSoftObjectPtr<UCustomizableObject>> WorkingSet; // TODO >5.6 make private
 
 private:
 	// Editor graph
@@ -368,39 +273,51 @@ public:
 	UFUNCTION(BlueprintCallable, Category = CustomizableObject)
 	int32 GetParameterCount() const;
 
-	/** Get the index of a parameter from its name. Return -1 if the parameter is not found. */
+	UE_DEPRECATED(5.6, "Use ContainsParameter instead")
 	UFUNCTION(BlueprintCallable, Category = CustomizableObject)
 	int32 FindParameter(const FString& Name) const;
 
-	/** Get the type of a parameter from its index. */
+	UFUNCTION(BlueprintCallable, Category = CustomizableObject)
+	bool ContainsParameter(const FString& ParameterName) const;
+	
+	UE_DEPRECATED(5.6, "Use ParameterName signature instead")
 	UFUNCTION(BlueprintCallable, Category = CustomizableObject)
 	EMutableParameterType GetParameterType(int32 ParamIndex) const;
 
 	/** Get the type of a parameter from its name. */
 	UFUNCTION(BlueprintCallable, Category = CustomizableObject)
-	EMutableParameterType GetParameterTypeByName(const FString& Name) const;
-
-	/** Get the name of a parameter from its index. */
+	EMutableParameterType GetParameterTypeByName(const FString& Name) const; // TODO >5.6 rename to GetParameterType (with redirector). Name -> ParameterName (create redirector)
+	
 	UFUNCTION(BlueprintCallable, Category = CustomizableObject)
 	const FString& GetParameterName(int32 ParamIndex) const;
 
-	/** Deprecated. It will always return 0. */
-	UFUNCTION(BlueprintCallable, Category = CustomizableObject, meta = (DeprecatedFunction, DeprecationMessage = "Parameter decorations have been removed. This method will be removed in future versions."))
-	int32 GetParameterDescriptionCount(const FString& ParamName) const;
-
-	/** Returns how many possible options an int parameter has, if the parameter is an enumeration. Otherwise return 0. */
+	UE_DEPRECATED(5.6, "Use ParameterName signature instead")
 	UFUNCTION(BlueprintCallable, Category = CustomizableObject)
 	int32 GetIntParameterNumOptions(int32 ParamIndex) const;
 
-	/** Gets the Name of the option at position K in the list of available options for the int parameter.
-	  * Useful to enumerate the int parameter's possible options (Ex: "Hat1", "Hat2", "Cap", "Nothing") */
+	UE_DEPRECATED(5.6, "Use ParameterName signature instead")
 	UFUNCTION(BlueprintCallable, Category = CustomizableObject)
 	const FString& GetIntParameterAvailableOption(int32 ParamIndex, int32 K) const;
+	
+	/** If the given parameter is enum int parameter, return how many possible values an enum parameter has. Otherwise, return 0. */
+	UFUNCTION(BlueprintCallable, Category = CustomizableObject)
+	int32 GetEnumParameterNumValues(const FString& ParameterName) const;
+
+	/** Gets the Name of the value at position ValueIndex in the list of available values for the int parameter. */
+	UFUNCTION(BlueprintCallable, Category = CustomizableObject)
+	const FString& GetEnumParameterValue(const FString& ParameterName, int32 ValueIndex) const;
 
+	/** Return true if the enum parameter contains this value a possible option. */
+	UFUNCTION(BlueprintCallable, Category = CustomizableObject)
+	bool ContainsEnumParameterValue(const FString& ParameterName, const FString Value) const;
+	
+	UE_DEPRECATED(5.6, "Use ParameterName signature instead")
 	int32 FindIntParameterValue( int32 ParamIndex, const FString& Value ) const;
+
+	UE_DEPRECATED(5.6, "Use ParameterName signature instead")
 	FString FindIntParameterValueName(int32 ParamIndex, int32 ParamValue) const;
 
-	// DEPRECATED
+	UE_DEPRECATED(5.6, "Use GetComponentMeshReferenceSkeletalMesh instead")
 	USkeletalMesh* GetRefSkeletalMesh(int32 ObjectComponentIndex = 0) const;
 
 	/** Given a Mesh Component name, return its reference Skeletal Mesh. */
@@ -408,45 +325,39 @@ public:
 	USkeletalMesh* GetComponentMeshReferenceSkeletalMesh(const FName& Name) const;
 	
 	/** Get the default value of a parameter of type Float.
-	  * @param InParameterName The name of the Float parameter to get the default value of.
+	  * @param ParameterName The name of the Float parameter to get the default value of.
 	  * @return The default value of the provided parameter name. */
 	UFUNCTION(BlueprintCallable, Category = CustomizableObject)
-	float GetFloatParameterDefaultValue(UPARAM(DisplayName = "Parameter Name") const FString& InParameterName) const;
+	float GetFloatParameterDefaultValue(const FString& ParameterName) const;
 	
 	/** Get the default value of a parameter of type Int. 
-	  * @param InParameterName The name of the Int parameter to get the default value of.
+	  * @param ParameterName The name of the Int parameter to get the default value of.
 	  * @return The default value of the provided parameter name. */
 	UFUNCTION(BlueprintCallable, Category = CustomizableObject)
+	int32 GetEnumParameterDefaultValue(const FString& ParameterName) const;
+
+	UE_DEPRECATED(5.6, "Use GetEnumParameterDefaultValue instead")
 	int32 GetIntParameterDefaultValue(UPARAM(DisplayName = "Parameter Name") const FString& InParameterName) const;
  
 	/** Get the default value of a parameter of type Bool.
-	  * @param InParameterName The name of the Bool parameter to get the default value of.
+	  * @param ParameterName The name of the Bool parameter to get the default value of.
 	  * @return The default value of the provided parameter name. */
 	UFUNCTION(BlueprintCallable, Category = CustomizableObject)
-	bool GetBoolParameterDefaultValue(UPARAM(DisplayName = "Parameter Name") const FString& InParameterName) const;
+	bool GetBoolParameterDefaultValue(const FString& ParameterName) const;
 
 	/** Get the default value of a parameter of type Color.
-	  * @param InParameterName The name of the Color parameter to get the default value of.
+	  * @param ParameterName The name of the Color parameter to get the default value of.
 	  * @return The default value of the provided parameter name. */
 	UFUNCTION(BlueprintCallable, Category = CustomizableObject)
-	FLinearColor GetColorParameterDefaultValue(UPARAM(DisplayName = "Parameter Name") const FString& InParameterName) const;
+	FLinearColor GetColorParameterDefaultValue(const FString& ParameterName) const;
 	
 	/** Get the default value of a parameter of type Transform.
-	  * @param InParameterName The name of the Transform parameter to get the default value of.
+	  * @param ParameterName The name of the Transform parameter to get the default value of.
 	  * @return The default value of the provided parameter name. */
 	UFUNCTION(BlueprintCallable, Category = CustomizableObject)
-	FTransform GetTransformParameterDefaultValue(UPARAM(DisplayName = "Parameter Name") const FString& InParameterName) const;
-
+	FTransform GetTransformParameterDefaultValue(const FString& ParameterName) const;
 	
-	/** Get the default value of a parameter of type Projector.
-	  * @param InParameterName The name of the Projector parameter to get the default value of.
-	  * @param OutPos The default position of the Projector.
-	  * @param OutDirection The default projection direction of the Projector.
-	  * @param OutUp The default up vector of the Projector.
-	  * @param OutScale The default scale of the Projector.
-	  * @param OutAngle The default angle of the Projector.
-	  * @param OutType The default type of the Projector. */
-	UFUNCTION(BlueprintCallable, Category = CustomizableObject)
+	UE_DEPRECATED(5.6, "Use the return FCustomizableObjectProjector signature instead")
 	void GetProjectorParameterDefaultValue(
 		UPARAM(DisplayName = "Parameter Name") const FString& InParameterName,
 		UPARAM(DisplayName = "Position") FVector3f& OutPos,
@@ -457,18 +368,17 @@ public:
 		UPARAM(DisplayName = "Type") ECustomizableObjectProjectorType& OutType) const;
 
 	/** Get the default value of a projector with the provided name
-	  * @param InParameterName The name of the parameter to get the default value of.
+	  * @param ParameterName The name of the parameter to get the default value of.
 	  * @return A data structure containing all the default data for the targeted projector parameter. */
-	FCustomizableObjectProjector GetProjectorParameterDefaultValue(const FString& InParameterName) const;
+	UFUNCTION(BlueprintCallable, Category = CustomizableObject)
+	FCustomizableObjectProjector GetProjectorParameterDefaultValue(const FString& ParameterName) const;
 	
 	/** Get the default value of a parameter of type Texture.
 	  * @param InParameterName The name of the Projector parameter to get the default value of.
 	  * @return An id representing the default parameter's texture. */
 	FName GetTextureParameterDefaultValue(const FString& InParameterName) const;
 
-	/** Return true if the parameter at the index provided is multidimensional.
-	  * @param InParamIndex The index of the parameter to check.
-	  * @return True if the parameter is multidimensional and false if it is not. */
+	UE_DEPRECATED(5.6, "Use the ParameterName signature instead")
 	bool IsParameterMultidimensional(const int32& InParamIndex) const;
 	
 	/** Return true if the parameter at the index provided is multidimensional.
@@ -477,31 +387,29 @@ public:
 	UFUNCTION(BlueprintCallable, Category = CustomizableObject)
 	bool IsParameterMultidimensional(UPARAM(DisplayName = "Parameter Name") const FString& InParameterName) const;
 
-	// Begin UObject interface.
-	void PostLoad() override;
-	void Serialize(FArchive& Ar) override;
 #if WITH_EDITOR
-	
 	/** Compile the object if Automatic Compilation is enabled and has not been already compiled.
 	  * Automatic compilation can be enabled/disabled in the Mutable's Plugin Settings.
 	  * @return true if compiled */
 	bool ConditionalAutoCompile();
-	
+#endif
+
+	// UObject interface
+	virtual void PostLoad() override;
+	virtual void Serialize(FArchive& Ar) override;
+
+#if WITH_EDITOR
 	virtual void GetAssetRegistryTags(FAssetRegistryTagsContext Context) const override;
-	UE_DEPRECATED(5.4, "Implement the version that takes FAssetRegistryTagsContext instead.")
-	virtual void GetAssetRegistryTags(TArray<FAssetRegistryTag>& OutTags) const override;
-
-	/** Returns a one line description of an object for viewing in the thumbnail view of the generic browser */
-	FString GetDesc() override;
-	bool IsEditorOnly() const override;
-	void PreSave(FObjectPreSaveContext ObjectSaveContext) override;
-	void PostSaveRoot(FObjectPostSaveRootContext ObjectSaveContext) override;
-	void PostRename(UObject* OldOuter, const FName OldName) override;	
-	void BeginCacheForCookedPlatformData(const ITargetPlatform* TargetPlatform) override;
-	bool IsCachedCookedPlatformDataLoaded(const ITargetPlatform* TargetPlatform) override;
+	virtual FString GetDesc() override;
+	virtual bool IsEditorOnly() const override;
+	virtual void PreSave(FObjectPreSaveContext ObjectSaveContext) override;
+	virtual void PostSaveRoot(FObjectPostSaveRootContext ObjectSaveContext) override;
+	virtual void PostRename(UObject* OldOuter, const FName OldName) override;	
+	virtual void BeginCacheForCookedPlatformData(const ITargetPlatform* TargetPlatform) override;
+	virtual bool IsCachedCookedPlatformDataLoaded(const ITargetPlatform* TargetPlatform) override;
 #endif
-	// End UObject interface.
 
+	UE_DEPRECATED(5.6, "Use StateName functions instead")
 	int32 FindState(const FString& Name) const;
 
 	/** Return the number of object states that are defined in the CustomizableObject. */
@@ -509,31 +417,43 @@ public:
 	int32 GetStateCount() const;
 
 	/** Return the name of an object state from its index. */
+	UE_DEPRECATED(5.6, "Use StateName functions instead")
 	UFUNCTION(BlueprintCallable, Category = CustomizableObject)
 	FString GetStateName(int32 StateIndex) const;
 
+	UE_DEPRECATED(5.6, "Use StateName signature instead")
 	int32 GetStateParameterCount(int32 StateIndex) const;
+
+	UE_DEPRECATED(5.6, "Use StateName signature instead")
 	int32 GetStateParameterIndex(int32 StateIndex, int32 ParameterIndex) const;
 
 	/** Return the number of parameters that are editable at runtime for a specific state. */
 	UFUNCTION(BlueprintCallable, Category = CustomizableObject)
 	int32 GetStateParameterCount(const FString& StateName) const;
 
-	/** Return the name of one of the state's runtime parameters, by its index (from 0 to GetStateParameterCount-1). */
+	/** Return the name of one of the state's runtime parameters, by its index (from 0 to GetStateParameterCount - 1). */
 	UFUNCTION(BlueprintCallable, Category = CustomizableObject)
 	FString GetStateParameterName(const FString& StateName, int32 ParameterIndex) const;
+
+	UE_DEPRECATED(5.6, "Use StateName signature instead")
 	FString GetStateParameterName(int32 StateIndex, int32 ParameterIndex) const;
 	
 	/** Return the metadata associated to a parameter. */
 	UFUNCTION(BlueprintCallable, Category = CustomizableObject)
 	FMutableParamUIMetadata GetParameterUIMetadata(const FString& ParamName) const;
 
-	/** Return the metadata associated to an int parameter option. */
+	/** Return the metadata associated to the given enum parameter value. */
 	UFUNCTION(BlueprintCallable, Category = CustomizableObject)
+	FMutableParamUIMetadata GetEnumParameterValueUIMetadata(const FString& ParamName, const FString& OptionName) const; // TODO GMT Rename "OptionName" -> "Value" using core redirects
+
+	UE_DEPRECATED(5.6, "Use GetEnumParameterValueUIMetadata instead")
 	FMutableParamUIMetadata GetIntParameterOptionUIMetadata(const FString& ParamName, const FString& OptionName) const;
 	
 	/** Returns the group type of the given integer parameter */
 	UFUNCTION(BlueprintCallable, Category = CustomizableObject)
+	ECustomizableObjectGroupType GetEnumParameterGroupType(const FString& ParamName) const;
+
+	UE_DEPRECATED(5.6, "Use GetEnumParameterGroupType instead")
 	ECustomizableObjectGroupType GetIntParameterGroupType(const FString& ParamName) const;
 
 	/** Return the metadata associated to a state. */
@@ -543,6 +463,9 @@ public:
 #if WITH_EDITOR
 	/** Return the DataTables used by the given parameter and its value (if any). */
 	UFUNCTION(BlueprintCallable, Category = CustomizableObject)
+	TArray<TSoftObjectPtr<UDataTable>> GetEnumParameterValueDataTable(const FString& ParamName, const FString& Value);
+
+	UE_DEPRECATED(5.6, "Use GetEnumParameterValueDataTable instead")
 	TArray<TSoftObjectPtr<UDataTable>> GetIntParameterOptionDataTable(const FString& ParamName, const FString& OptionName);
 #endif
 	
@@ -584,17 +507,21 @@ public:
 	/** Create a new instance of this object. The instance parameters will be initialized with the object default values. */
 	UFUNCTION(BlueprintCallable, Category = CustomizableObject)
 	UCustomizableObjectInstance* CreateInstance();
-
-	const UCustomizableObjectPrivate* GetPrivate() const;
-
-	UCustomizableObjectPrivate* GetPrivate();
-
+	
 	/** Check if the CustomizableObject asset has been compiled. This will always be true in a packaged game, but it could be false in the editor. */
 	UFUNCTION(BlueprintCallable, Category = CustomizableObject)
 	bool IsCompiled() const;
 
+#if WITH_EDITOR
+	/** Return true if this Customizable Object has references to a parent Customizable Object. Only root Customizable Objects will return false. */
+	UFUNCTION(BlueprintCallable, Category = CustomizableObject)
 	bool IsChildObject() const;
+#endif
+	
+	const UCustomizableObjectPrivate* GetPrivate() const;
 
+	UCustomizableObjectPrivate* GetPrivate();
+	
 private:
 	/** BulkData that stores all in-game resources used by Mutable when generating instances.
 	  * Only valid in packaged builds */
diff --git a/Engine/Plugins/Mutable/Source/CustomizableObject/Public/MuCO/CustomizableObjectInstance.h b/Engine/Plugins/Mutable/Source/CustomizableObject/Public/MuCO/CustomizableObjectInstance.h
--- a/Engine/Plugins/Mutable/Source/CustomizableObject/Public/MuCO/CustomizableObjectInstance.h
+++ b/Engine/Plugins/Mutable/Source/CustomizableObject/Public/MuCO/CustomizableObjectInstance.h
@@ -381,7 +381,7 @@ public:
 	UFUNCTION(BlueprintCallable, Category = CustomizableObjectInstance)
 	const FString& GetIntParameterSelectedOption(const FString& ParamName, int32 RangeIndex = -1) const;
 
-	/** Set the currently selected option value for the given parameter, by parameter index and option name. */
+	UE_DEPRECATED(5.6, "Use ParameterName signature instead")
 	void SetIntParameterSelectedOption(int32 IntParamIndex, const FString& SelectedOption, int32 RangeIndex = -1);
 
 	/** Set the currently selected option value for the given parameter, by parameter name and option name. */
@@ -432,7 +432,6 @@ public:
 	UFUNCTION(BlueprintCallable, Category = CustomizableObjectInstance)
 	void SetTransformParameterSelectedOption(const FString& TransformParamName, const FTransform& TransformValue);
 	
-	
 	/** Sets the projector values of a projector parameter with index "ProjectorParamIndex". */
 	UFUNCTION(BlueprintCallable, Category = CustomizableObjectInstance)
 	void SetProjectorValue(const FString& ProjectorParamName,
diff --git a/Engine/Plugins/Mutable/Source/CustomizableObject/Public/MuCO/CustomizableObjectParameterTypeDefinitions.h b/Engine/Plugins/Mutable/Source/CustomizableObject/Public/MuCO/CustomizableObjectParameterTypeDefinitions.h
--- a/Engine/Plugins/Mutable/Source/CustomizableObject/Public/MuCO/CustomizableObjectParameterTypeDefinitions.h
+++ b/Engine/Plugins/Mutable/Source/CustomizableObject/Public/MuCO/CustomizableObjectParameterTypeDefinitions.h
@@ -254,28 +254,28 @@ enum class ECustomizableObjectProjectorType : uint8
 };
 
 
-USTRUCT()
+USTRUCT(BlueprintType)
 struct CUSTOMIZABLEOBJECT_API FCustomizableObjectProjector
 {
 	GENERATED_USTRUCT_BODY()
 
-	UPROPERTY()
+	UPROPERTY(EditAnywhere, Category = CustomizableObjectProjector)
 	FVector3f Position = FVector3f(0, 0, 0);
 
-	UPROPERTY()
+	UPROPERTY(EditAnywhere, Category = CustomizableObjectProjector)
 	FVector3f Direction = FVector3f(1, 0, 0);
 
-	UPROPERTY()
+	UPROPERTY(EditAnywhere, Category = CustomizableObjectProjector)
 	FVector3f Up = FVector3f(0, 1, 0);
 
-	UPROPERTY()
+	UPROPERTY(EditAnywhere, Category = CustomizableObjectProjector)
 	FVector3f Scale = FVector3f(10, 10, 100);
 
-	UPROPERTY(EditAnywhere, Category = CustomizableObject)
+	UPROPERTY(EditAnywhere, Category = CustomizableObjectProjector)
 	ECustomizableObjectProjectorType ProjectionType = ECustomizableObjectProjectorType::Planar;
 
 	// Just for cylindrical projectors, in radians
-	UPROPERTY()
+	UPROPERTY(EditAnywhere, Category = CustomizableObjectProjector)
 	float Angle = 2.0f * PI;
 
 	bool operator==(const FCustomizableObjectProjector& Other) const = default;
diff --git a/Engine/Plugins/Mutable/Source/CustomizableObjectEditor/Private/MuCOE/CustomizableInstanceDetails.cpp b/Engine/Plugins/Mutable/Source/CustomizableObjectEditor/Private/MuCOE/CustomizableInstanceDetails.cpp
--- a/Engine/Plugins/Mutable/Source/CustomizableObjectEditor/Private/MuCOE/CustomizableInstanceDetails.cpp
+++ b/Engine/Plugins/Mutable/Source/CustomizableObjectEditor/Private/MuCOE/CustomizableInstanceDetails.cpp
@@ -341,12 +341,12 @@ TSharedRef<SWidget> FCustomizableInstanceDetails::GenerateStateSelector()
 	{
 		if (StateIndex == CurrentState)
 		{
-			CurrentStateName = MakeShareable(new FString(CustomizableObject->GetStateName(StateIndex)));
+			CurrentStateName = MakeShareable(new FString(CustomizableObject->GetPrivate()->GetStateName(StateIndex)));
 			StateNames.Add(CurrentStateName);
 		}
 		else
 		{
-			StateNames.Add(MakeShareable(new FString(CustomizableObject->GetStateName(StateIndex))));
+			StateNames.Add(MakeShareable(new FString(CustomizableObject->GetPrivate()->GetStateName(StateIndex))));
 		}
 	}
 
@@ -644,7 +644,7 @@ bool FCustomizableInstanceDetails::GenerateParametersView(IDetailCategoryBuilder
 
 	if (CustomInstance->GetPrivate()->bShowOnlyRuntimeParameters)
 	{
-		const int32 NumStateParameters = CustomizableObject->GetStateParameterCount(CustomInstance->GetPrivate()->GetState());
+		const int32 NumStateParameters = CustomizableObject->GetStateParameterCount(CustomInstance->GetCurrentState());
 
 		if (NumStateParameters < CustomizableObject->GetParameterCount())
 		{
@@ -654,7 +654,7 @@ bool FCustomizableInstanceDetails::GenerateParametersView(IDetailCategoryBuilder
 		for (int32 ParamIndexInState = 0; ParamIndexInState < NumStateParameters; ++ParamIndexInState)
 		{
 			FParameterInfo ParameterSortInfo;
-			ParameterSortInfo.ParamIndexInObject = CustomizableObject->GetStateParameterIndex(CustomInstance->GetPrivate()->GetState(), ParamIndexInState);
+			ParameterSortInfo.ParamIndexInObject = CustomizableObject->GetPrivate()->GetStateParameterIndex(CustomInstance->GetPrivate()->GetState(), ParamIndexInState);
 
 			if (CustomInstance->IsParameterRelevant(ParameterSortInfo.ParamIndexInObject) && IsVisible(ParameterSortInfo.ParamIndexInObject))
 			{
@@ -832,15 +832,15 @@ bool FCustomizableInstanceDetails::IsVisible(int32 ParamIndexInObject)
 	const FString* ParentName = UIMetadata.ExtraInformation.Find(FString("__ParentParamName"));
 
 	const bool IsAProjectorParam = ParamName.EndsWith(NUM_LAYERS_PARAMETER_POSTFIX)
-		|| (ParamName.EndsWith(IMAGE_PARAMETER_POSTFIX) && CustomizableObject->IsParameterMultidimensional(ParamIndexInObject))
-		|| (ParamName.EndsWith(OPACITY_PARAMETER_POSTFIX) && CustomizableObject->IsParameterMultidimensional(ParamIndexInObject))
+		|| (ParamName.EndsWith(IMAGE_PARAMETER_POSTFIX) && CustomizableObject->GetPrivate()->IsParameterMultidimensional(ParamIndexInObject))
+		|| (ParamName.EndsWith(OPACITY_PARAMETER_POSTFIX) && CustomizableObject->GetPrivate()->IsParameterMultidimensional(ParamIndexInObject))
 		|| (ParamName.EndsWith(POSE_PARAMETER_POSTFIX));
 
 	if (!IsAProjectorParam && ParentName && CustomInstance->GetPrivate()->bShowOnlyRelevantParameters)
 	{
 		const FString* Value = UIMetadata.ExtraInformation.Find(FString("__DisplayWhenParentValueEquals"));
 
-		if (Value && CustomizableObject->FindParameter(*ParentName) != INDEX_NONE && CustomInstance->GetIntParameterSelectedOption(*ParentName) != *Value)
+		if (Value && CustomizableObject->GetPrivate()->FindParameter(*ParentName) != INDEX_NONE && CustomInstance->GetIntParameterSelectedOption(*ParentName) != *Value)
 		{
 			return false;
 		}
@@ -857,7 +857,7 @@ bool FCustomizableInstanceDetails::IsIntParameterFilteredOut(const UCustomizable
 		return false;
 	}
 
-	FMutableParamUIMetadata Metadata = CustomizableObject.GetIntParameterOptionUIMetadata(ParamName, ParamOption);
+	FMutableParamUIMetadata Metadata = CustomizableObject.GetEnumParameterValueUIMetadata(ParamName, ParamOption);
 
 	switch (FilterType)
 	{
@@ -881,7 +881,7 @@ bool FCustomizableInstanceDetails::IsMultidimensionalProjector(int32 ParamIndexI
 {
 	const UCustomizableObject* CustomizableObject = CustomInstance->GetCustomizableObject();
 
-	return CustomizableObject->GetParameterType(ParamIndexInObject) == EMutableParameterType::Projector && CustomInstance->GetCustomizableObject()->IsParameterMultidimensional(ParamIndexInObject);
+	return CustomizableObject->GetPrivate()->GetParameterType(ParamIndexInObject) == EMutableParameterType::Projector && CustomInstance->GetCustomizableObject()->GetPrivate()->IsParameterMultidimensional(ParamIndexInObject);
 }
 
 
@@ -933,7 +933,7 @@ void FCustomizableInstanceDetails::GenerateWidgetRow(FDetailWidgetRow& WidgetRow
 
 TSharedRef<SWidget> FCustomizableInstanceDetails::GenerateParameterWidget(const UCustomizableObject& CustomizableObject, const FString& ParamName, const int32 ParamIndexInObject)
 {
-	switch (CustomizableObject.GetParameterType(ParamIndexInObject))
+	switch (CustomizableObject.GetPrivate()->GetParameterType(ParamIndexInObject))
 	{
 		case EMutableParameterType::Bool:
 		{
@@ -957,7 +957,7 @@ TSharedRef<SWidget> FCustomizableInstanceDetails::GenerateParameterWidget(const
 		}
 		case EMutableParameterType::Projector:
 		{
-			bool bIsParamMultidimensional = CustomInstance->GetCustomizableObject()->IsParameterMultidimensional(ParamIndexInObject);
+			bool bIsParamMultidimensional = CustomInstance->GetCustomizableObject()->GetPrivate()->IsParameterMultidimensional(ParamIndexInObject);
 
 			if (!bIsParamMultidimensional)
 			{
@@ -986,13 +986,13 @@ TSharedRef<SWidget> FCustomizableInstanceDetails::GenerateParameterWidget(const
 
 TSharedRef<SWidget> FCustomizableInstanceDetails::GenerateIntWidget(const UCustomizableObject& CustomizableObject, const FString& ParamName, const int32 ParamIndexInObject)
 {
-	const bool bMultidimensional = CustomInstance->GetCustomizableObject()->IsParameterMultidimensional(ParamIndexInObject);
+	const bool bMultidimensional = CustomInstance->GetCustomizableObject()->GetPrivate()->IsParameterMultidimensional(ParamIndexInObject);
 	if (bMultidimensional)
 	{
 		return SNew(STextBlock).Text(LOCTEXT("MultidimensionalINTParameter_Text", "Multidimensional INT Parameter not supported"));
 	}
 
-	const int32 NumValues = CustomizableObject.GetIntParameterNumOptions(ParamIndexInObject);
+	const int32 NumValues = CustomizableObject.GetPrivate()->GetEnumParameterNumValues(ParamIndexInObject);
 	if (!NumValues)
 	{
 		return SNew(STextBlock).Text(LOCTEXT("NoAvailableOptions", "No Available Options"));
@@ -1020,7 +1020,7 @@ TSharedRef<SWidget> FCustomizableInstanceDetails::GenerateIntWidget(const UCusto
 	TSharedPtr<FString> SelectedOptionString;
 	for (int32 i = 0; i < NumValues; ++i)
 	{
-		const FString PossibleValue = CustomizableObject.GetIntParameterAvailableOption(ParamIndexInObject, i);
+		const FString PossibleValue = CustomizableObject.GetPrivate()->GetIntParameterAvailableOption(ParamIndexInObject, i);
 
 		if (PossibleValue == SelectedOption) // Always add the selected option, even if it should be filtered.
 		{
@@ -1063,7 +1063,7 @@ TSharedRef<SWidget> FCustomizableInstanceDetails::OnGenerateWidgetIntParameter(T
 		FAssetRegistryModule& AssetRegistryModule = FModuleManager::LoadModuleChecked<FAssetRegistryModule>("AssetRegistry");
 
 		// Metadata of the represented int option
-		const FMutableParamUIMetadata ParameterMetadata = CustomInstance->GetCustomizableObject()->GetIntParameterOptionUIMetadata(ParameterName, *OptionName);
+		const FMutableParamUIMetadata ParameterMetadata = CustomInstance->GetCustomizableObject()->GetEnumParameterValueUIMetadata(ParameterName, *OptionName);
 
 		// Custom thumbnail has preference
 		if (!ParameterMetadata.UIThumbnail.IsNull())
@@ -1850,7 +1850,7 @@ FReply FCustomizableInstanceDetails::OnProjectorSelectChanged(const FString Para
 
 FReply FCustomizableInstanceDetails::OnProjectorCopyTransform(const FString ParamName, const int32 RangeIndex) const
 {
-	const int32 ParameterIndexInObject = CustomInstance->GetCustomizableObject()->FindParameter(ParamName);
+	const int32 ParameterIndexInObject = CustomInstance->GetCustomizableObject()->GetPrivate()->FindParameter(ParamName);
 	const int32 ProjectorParamIndex = CustomInstance->FindProjectorParameterNameIndex(ParamName);
 
 	if ((ParameterIndexInObject >= 0) && (ProjectorParamIndex >= 0))
@@ -1942,12 +1942,12 @@ TSharedRef<SWidget> FCustomizableInstanceDetails::GenerateMultidimensionalProjec
 
 	// Selected Pose UI
 	const FString PoseSwitchEnumParamName = ParamName + POSE_PARAMETER_POSTFIX;
-	const int32 PoseSwitchEnumParamIndexInObject = CustomizableObject.FindParameter(PoseSwitchEnumParamName);
+	const int32 PoseSwitchEnumParamIndexInObject = CustomizableObject.GetPrivate()->FindParameter(PoseSwitchEnumParamName);
 	TSharedPtr<SVerticalBox> ProjectorBox = SNew(SVerticalBox);
 
 	if (PoseSwitchEnumParamIndexInObject != INDEX_NONE)
 	{
-		const int32 NumPoseValues = CustomizableObject.GetIntParameterNumOptions(PoseSwitchEnumParamIndexInObject);
+		const int32 NumPoseValues = CustomizableObject.GetPrivate()->GetEnumParameterNumValues(PoseSwitchEnumParamIndexInObject);
 
 		TSharedPtr<TArray<TSharedPtr<FString>>>* FoundOptions = ProjectorParameterPoseOptions.Find(PoseSwitchEnumParamIndexInObject);
 		TArray<TSharedPtr<FString>>& PoseOptionNamesAttribute = FoundOptions && FoundOptions->IsValid() ?
@@ -1961,7 +1961,7 @@ TSharedRef<SWidget> FCustomizableInstanceDetails::GenerateMultidimensionalProjec
 
 		for (int32 j = 0; j < NumPoseValues; ++j)
 		{
-			const FString PossibleValue = CustomizableObject.GetIntParameterAvailableOption(PoseSwitchEnumParamIndexInObject, j);
+			const FString PossibleValue = CustomizableObject.GetPrivate()->GetIntParameterAvailableOption(PoseSwitchEnumParamIndexInObject, j);
 			if (PossibleValue == PoseValue)
 			{
 				PoseValueIndex = j;
@@ -2031,7 +2031,7 @@ TSharedRef<SWidget> FCustomizableInstanceDetails::GenerateMultidimensionalProjec
 
 	for (int32 RangeIndex = 0; RangeIndex < ProjectorParameters[ProjectorParamIndex].RangeValues.Num(); ++RangeIndex)
 	{
-		const int32 TextureSwitchEnumParamIndexInObject = CustomizableObject.FindParameter(TextureSwitchEnumParamName);
+		const int32 TextureSwitchEnumParamIndexInObject = CustomizableObject.GetPrivate()->FindParameter(TextureSwitchEnumParamName);
 		check(TextureSwitchEnumParamIndexInObject >= 0); TSharedPtr<FString> CurrentStateName = nullptr;
 
 		const UProjectorParameter* ProjectorParameter = Editor->GetProjectorParameter();
@@ -2126,7 +2126,7 @@ TSharedRef<SWidget> FCustomizableInstanceDetails::GenerateMultidimensionalProjec
 		Button->SetBorderBackgroundColor(bSelectedProjector ? FLinearColor::Green : FLinearColor::White);
 
 		// If number of options is equal to 1, Mutable does not consider it multidimensional parameters
-		int32 NumValues = CustomizableObject.GetIntParameterNumOptions(TextureSwitchEnumParamIndexInObject);
+		int32 NumValues = CustomizableObject.GetPrivate()->GetEnumParameterNumValues(TextureSwitchEnumParamIndexInObject);
 		FString Value = CustomizableObject.IsParameterMultidimensional(TextureSwitchEnumParamName) ? 
 			CustomInstance->GetIntParameterSelectedOption(TextureSwitchEnumParamName, RangeIndex) : CustomInstance->GetIntParameterSelectedOption(TextureSwitchEnumParamName);
 
@@ -2135,13 +2135,13 @@ TSharedRef<SWidget> FCustomizableInstanceDetails::GenerateMultidimensionalProjec
 
 		for (int32 CandidateIndex = 0; CandidateIndex < NumValues; ++CandidateIndex)
 		{
-			FString PossibleValue = CustomizableObject.GetIntParameterAvailableOption(TextureSwitchEnumParamIndexInObject, CandidateIndex);
+			FString PossibleValue = CustomizableObject.GetPrivate()->GetIntParameterAvailableOption(TextureSwitchEnumParamIndexInObject, CandidateIndex);
 			if (PossibleValue == Value)
 			{
 				ValueIndex = CandidateIndex;
 			}
 
-			OptionNamesAttribute.Add(MakeShared<FString>(CustomizableObject.GetIntParameterAvailableOption(TextureSwitchEnumParamIndexInObject, CandidateIndex)));
+			OptionNamesAttribute.Add(MakeShared<FString>(CustomizableObject.GetPrivate()->GetIntParameterAvailableOption(TextureSwitchEnumParamIndexInObject, CandidateIndex)));
 		}
 		
 		// Avoid filling this arraw with repeated array  options
diff --git a/Engine/Plugins/Mutable/Source/CustomizableObjectEditor/Private/MuCOE/CustomizableObjectCompiler.cpp b/Engine/Plugins/Mutable/Source/CustomizableObjectEditor/Private/MuCOE/CustomizableObjectCompiler.cpp
--- a/Engine/Plugins/Mutable/Source/CustomizableObjectEditor/Private/MuCOE/CustomizableObjectCompiler.cpp
+++ b/Engine/Plugins/Mutable/Source/CustomizableObjectEditor/Private/MuCOE/CustomizableObjectCompiler.cpp
@@ -522,7 +522,7 @@ mu::Ptr<mu::NodeObject> GenerateMutableRoot(
 		GenerationContext.NoNameNodeObjectArray.AddUnique(LocalRootNodeObject);
 	}
 
-	if ((Object->MeshCompileType == EMutableCompileMeshType::Full) || GenerationContext.Options.bIsCooking)
+	if ((Object->GetPrivate()->GetMeshCompileType() == EMutableCompileMeshType::Full) || GenerationContext.Options.bIsCooking)
 	{
 		if (LocalRootNodeObject->ParentObject!=nullptr && GenerationContext.Options.bIsCooking)
 		{
@@ -553,13 +553,13 @@ mu::Ptr<mu::NodeObject> GenerateMutableRoot(
 			return nullptr;
 		}
 
-		if ((Object->MeshCompileType == EMutableCompileMeshType::AddWorkingSetNoChildren) ||
-			(Object->MeshCompileType == EMutableCompileMeshType::AddWorkingSetAndChildren))
+		if ((Object->GetPrivate()->GetMeshCompileType() == EMutableCompileMeshType::AddWorkingSetNoChildren) ||
+			(Object->GetPrivate()->GetMeshCompileType() == EMutableCompileMeshType::AddWorkingSetAndChildren))
 		{
-			const int32 MaxIndex = Object->WorkingSet.Num();
+			const int32 MaxIndex = Object->GetPrivate()->GetWorkingSet().Num();
 			for (int32 i = 0; i < MaxIndex; ++i)
 			{
-				if (UCustomizableObject* WorkingSetObject = GenerationContext.LoadObject(Object->WorkingSet[i], true))
+				if (UCustomizableObject* WorkingSetObject = GenerationContext.LoadObject(Object->GetPrivate()->GetWorkingSet()[i], true))
 				{
 					ArrayCustomizableObject.Reset();
 
@@ -572,8 +572,8 @@ mu::Ptr<mu::NodeObject> GenerateMutableRoot(
 			}
 		}
 
-		if ((Object->MeshCompileType == EMutableCompileMeshType::LocalAndChildren) ||
-			(Object->MeshCompileType == EMutableCompileMeshType::AddWorkingSetAndChildren))
+		if ((Object->GetPrivate()->GetMeshCompileType() == EMutableCompileMeshType::LocalAndChildren) ||
+			(Object->GetPrivate()->GetMeshCompileType() == EMutableCompileMeshType::AddWorkingSetAndChildren))
 		{
 			TArray<UCustomizableObjectNodeObjectGroup*> GroupNodes;
 			Object->GetPrivate()->GetSource()->GetNodesOfClass<UCustomizableObjectNodeObjectGroup>(GroupNodes);
@@ -1253,13 +1253,13 @@ mu::Ptr<mu::Node> FCustomizableObjectCompiler::Export(UCustomizableObject* Objec
 	FSlateNotificationManager::Get().AddNotification(Info);
 
 	FCompilationOptions CompilerOptions = InCompilerOptions;
-	CompilerOptions.bRealTimeMorphTargetsEnabled = Object->bEnableRealTimeMorphTargets;
-	CompilerOptions.bClothingEnabled = Object->bEnableClothing;
-	CompilerOptions.b16BitBoneWeightsEnabled = Object->bEnable16BitBoneWeights;
-	CompilerOptions.bSkinWeightProfilesEnabled = Object->bEnableAltSkinWeightProfiles;
-	CompilerOptions.bPhysicsAssetMergeEnabled = Object->bEnablePhysicsAssetMerge;
-	CompilerOptions.bAnimBpPhysicsManipulationEnabled = Object->bEnableAnimBpPhysicsAssetsManipualtion;
-
+	CompilerOptions.bRealTimeMorphTargetsEnabled = Object->GetPrivate()->IsRealTimeMorphTargetsEnabled();
+	CompilerOptions.bClothingEnabled = Object->GetPrivate()->IsClothingEnabled();
+	CompilerOptions.b16BitBoneWeightsEnabled = Object->GetPrivate()->Is16BitBoneWeightsEnabled();
+	CompilerOptions.bSkinWeightProfilesEnabled = Object->GetPrivate()->IsAltSkinWeightProfilesEnabled();
+	CompilerOptions.bPhysicsAssetMergeEnabled = Object->GetPrivate()->IsPhysicsAssetMergeEnabled();
+	CompilerOptions.bAnimBpPhysicsManipulationEnabled = Object->GetPrivate()->IsEnabledAnimBpPhysicsAssetsManipulation();
+	
 	FMutableGraphGenerationContext GenerationContext(Object, this, CompilerOptions);
 
 	GenerationContext.bSkipParticipatingObjectsPass = true;
diff --git a/Engine/Plugins/Mutable/Source/CustomizableObjectEditor/Private/MuCOE/GenerateMutableSource/GenerateMutableSourceMesh.cpp b/Engine/Plugins/Mutable/Source/CustomizableObjectEditor/Private/MuCOE/GenerateMutableSource/GenerateMutableSourceMesh.cpp
--- a/Engine/Plugins/Mutable/Source/CustomizableObjectEditor/Private/MuCOE/GenerateMutableSource/GenerateMutableSourceMesh.cpp
+++ b/Engine/Plugins/Mutable/Source/CustomizableObjectEditor/Private/MuCOE/GenerateMutableSource/GenerateMutableSourceMesh.cpp
@@ -3346,7 +3346,7 @@ mu::Ptr<mu::NodeMesh> GenerateMutableSourceMesh(const UEdGraphPin* Pin,
 
 			TArray<FCustomizableObjectStreameableResourceId> StreamedResources;
 
-			if (GenerationContext.Object->bEnableAssetUserDataMerge)
+			if (GenerationContext.Object->GetPrivate()->IsAssetUserDataMergeEnabled())
 			{
 				const TArray<UAssetUserData*>* AssetUserDataArray = TypedNodeSkel->SkeletalMesh->GetAssetUserDataArray();
 
diff --git a/Engine/Plugins/Mutable/Source/CustomizableObjectEditor/Private/MuCOE/GenerateMutableSource/GenerateMutableSourceTable.cpp b/Engine/Plugins/Mutable/Source/CustomizableObjectEditor/Private/MuCOE/GenerateMutableSource/GenerateMutableSourceTable.cpp
--- a/Engine/Plugins/Mutable/Source/CustomizableObjectEditor/Private/MuCOE/GenerateMutableSource/GenerateMutableSourceTable.cpp
+++ b/Engine/Plugins/Mutable/Source/CustomizableObjectEditor/Private/MuCOE/GenerateMutableSource/GenerateMutableSourceTable.cpp
@@ -248,7 +248,7 @@ bool FillTableColumn(const UCustomizableObjectNodeTable* TableNode, mu::Ptr<mu::
 			
 			TArray<FCustomizableObjectStreameableResourceId> StreamedResources;
 
-			if (GenerationContext.Object->bEnableAssetUserDataMerge)
+			if (GenerationContext.Object->GetPrivate()->IsAssetUserDataMergeEnabled())
 			{
 				const TArray<UAssetUserData*>* AssetUserDataArray = SkeletalMesh->GetAssetUserDataArray();
 
@@ -460,7 +460,7 @@ bool FillTableColumn(const UCustomizableObjectNodeTable* TableNode, mu::Ptr<mu::
 
 			GenerationContext.AddParticipatingObject(*ReferenceMaterial);
 
-			const bool bTableMaterialCheckDisabled = GenerationContext.Object->bDisableTableMaterialsParentCheck;
+			const bool bTableMaterialCheckDisabled = GenerationContext.Object->GetPrivate()->IsTableMaterialsParentCheckDisabled();
 			const bool bMaterialParentMismatch = !bTableMaterialCheckDisabled && MaterialInstance
 												 && ReferenceMaterial->GetMaterial() != MaterialInstance->GetMaterial();
 
