[Mutable] Multiple Quality+LOD management fixes
- Fixed overriding of generated SKM quality settings (using the ones in the CO) with the ones in the reference SKM.
- Fixed incorrect usage of "FirstLODAvailable".

diff --git a/Engine/Plugins/Mutable/Source/CustomizableObject/Private/MuCO/CustomizableObjectInstance.cpp b/Engine/Plugins/Mutable/Source/CustomizableObject/Private/MuCO/CustomizableObjectInstance.cpp
--- a/Engine/Plugins/Mutable/Source/CustomizableObject/Private/MuCO/CustomizableObjectInstance.cpp
+++ b/Engine/Plugins/Mutable/Source/CustomizableObject/Private/MuCO/CustomizableObjectInstance.cpp
@@ -2036,7 +2036,7 @@ bool UCustomizableInstancePrivate::DoComponentsNeedUpdate(UCustomizableObjectIns
 		// than the new FirstLODAvailable. This is an artifact of our LOD management system setting MinLOD.
 		if (!OperationData->MeshChangedPerInstanceComponent[InstanceComponentIndex] && bHadSkeletalMesh) 
 		{
-			int32 MinLOD = OperationData->GetMinLOD()[ComponentName];
+			const int32 MinLOD = OperationData->GetMinLOD()[ComponentName];
 			for (int32 LODIndex = OperationData->FirstLODAvailable[ComponentName]; LODIndex < MinLOD; ++LODIndex)
 			{
 				TIndirectArray<FSkeletalMeshLODRenderData>& LODRenderData = (*SkeletalMeshPtr)->GetResourceForRendering()->LODRenderData;
@@ -3585,10 +3585,6 @@ void UCustomizableInstancePrivate::InitSkeletalMeshData(const TSharedRef<FUpdate
 		SkeletalMesh->SetSupportRayTracing(false);
 	}
 	
-	// Set Min LOD
-	SkeletalMesh->SetMinLod(FMath::Max(CustomizableObject.LODSettings.MinLOD.GetDefault(), static_cast<int32>(OperationData->FirstLODAvailable[ComponentName])));
-	SkeletalMesh->SetQualityLevelMinLod(CustomizableObject.LODSettings.MinQualityLevelLOD);
-
 	SkeletalMesh->SetHasVertexColors(false);
 
 	// Set the default Physics Assets
@@ -3718,6 +3714,10 @@ void UCustomizableInstancePrivate::InitSkeletalMeshData(const TSharedRef<FUpdate
 #endif
 	}
 
+	// Set Min LOD (Override the Reference Skeletal Mesh LOD Settings)
+	SkeletalMesh->SetMinLod(FMath::Max(CustomizableObject.LODSettings.MinLOD.GetDefault(), static_cast<int32>(OperationData->FirstLODAvailable[ComponentName])));
+	SkeletalMesh->SetQualityLevelMinLod(CustomizableObject.LODSettings.MinQualityLevelLOD);
+
 	// Set up unreal's default material, will be replaced when building materials
 	{
 		MUTABLE_CPUPROFILER_SCOPE(InitSkeletalMesh_AddDefaultMaterial);
diff --git a/Engine/Plugins/Mutable/Source/CustomizableObject/Private/MuCO/CustomizableObjectSystem.cpp b/Engine/Plugins/Mutable/Source/CustomizableObject/Private/MuCO/CustomizableObjectSystem.cpp
--- a/Engine/Plugins/Mutable/Source/CustomizableObject/Private/MuCO/CustomizableObjectSystem.cpp
+++ b/Engine/Plugins/Mutable/Source/CustomizableObject/Private/MuCO/CustomizableObjectSystem.cpp
@@ -157,6 +157,8 @@ static void CVarMutableSinkFunction()
 	{
 		UCustomizableObjectSystemPrivate* PrivateSystem = UCustomizableObjectSystem::GetInstance()->GetPrivate();
 
+		// Store the quality level set in the scalability settings so we can later determine what the MinLOD should be used by mutable.
+		// Does not seem to be triggered when changing the visibility quality directly.
 		static const IConsoleVariable* CVarSkeletalMeshMinLodQualityLevelCVarName = IConsoleManager::Get().FindConsoleVariable(TEXT("r.SkeletalMesh.MinLodQualityLevel"));
 		PrivateSystem->SkeletalMeshMinLodQualityLevel = CVarSkeletalMeshMinLodQualityLevelCVarName ? CVarSkeletalMeshMinLodQualityLevelCVarName->GetInt() : INDEX_NONE;
 	}
@@ -1771,7 +1773,7 @@ namespace impl
 			else
 			{
 				FirstResidentLOD = FirstLODAvailable;
-				RequestedLOD = FMath::Min(RequestedLOD, static_cast<uint8>(FirstLODAvailable - 1));
+				RequestedLOD = FMath::Min(RequestedLOD, static_cast<uint8>(NumLODsAvailable - 1));
 			}
 		}
 		
