Fix Bounds.ExpandBy(1.2f); in MutableMeshPreviewUtils
Found with [[nodiscard]]

#rnx

diff --git a/Engine/Plugins/Mutable/Source/CustomizableObjectEditor/Private/MuCOE/MutableMeshPreviewUtils.cpp b/Engine/Plugins/Mutable/Source/CustomizableObjectEditor/Private/MuCOE/MutableMeshPreviewUtils.cpp
--- a/Engine/Plugins/Mutable/Source/CustomizableObjectEditor/Private/MuCOE/MutableMeshPreviewUtils.cpp
+++ b/Engine/Plugins/Mutable/Source/CustomizableObjectEditor/Private/MuCOE/MutableMeshPreviewUtils.cpp
@@ -322,7 +322,7 @@ namespace MutableMeshPreviewUtils
 				}
 
 				Bounds = FBoxSphereBounds(Points.GetData(), Points.Num());
-				Bounds.ExpandBy(1.2f);
+				Bounds = Bounds.ExpandBy(1.2f);
 			}
 			else
 			{
