[mutable] Made scalability setting changes made in editor and package trigger instance update to discard no longer required meshes.

diff --git a/Engine/Plugins/Mutable/Source/CustomizableObject/Internal/MuCO/CustomizableObjectSystemPrivate.h b/Engine/Plugins/Mutable/Source/CustomizableObject/Internal/MuCO/CustomizableObjectSystemPrivate.h
--- a/Engine/Plugins/Mutable/Source/CustomizableObject/Internal/MuCO/CustomizableObjectSystemPrivate.h
+++ b/Engine/Plugins/Mutable/Source/CustomizableObject/Internal/MuCO/CustomizableObjectSystemPrivate.h
@@ -92,6 +92,8 @@ public:
 
 	TMap<FName, uint8> RequestedLODLevels;
 
+	TMap<FName, uint8> FirstLODAvailable;
+
 	FMutableUpdateCandidate(UCustomizableObjectInstance* InCustomizableObjectInstance);
 	
 	FMutableUpdateCandidate(const UCustomizableObjectInstance* InCustomizableObjectInstance, const TMap<FName, uint8>& InMinLOD,
@@ -538,6 +540,10 @@ public:
 	
 	void SetRequestedLODs(const TMap<FName, uint8>& RequestedLODs);
 
+	void SetFirstLODAvailable(const TMap<FName, uint8>& FirstLODs);
+
+	const TMap<FName, uint8>& GetFirstLODAvailable() const;
+
 	const FCustomizableObjectInstanceDescriptor& GetCapturedDescriptor() const;
 
 	const FDescriptorHash& GetCapturedDescriptorHash() const;
diff --git a/Engine/Plugins/Mutable/Source/CustomizableObject/Internal/MuCO/DescriptorHash.h b/Engine/Plugins/Mutable/Source/CustomizableObject/Internal/MuCO/DescriptorHash.h
--- a/Engine/Plugins/Mutable/Source/CustomizableObject/Internal/MuCO/DescriptorHash.h
+++ b/Engine/Plugins/Mutable/Source/CustomizableObject/Internal/MuCO/DescriptorHash.h
@@ -13,10 +13,6 @@ public:
 	FDescriptorHash() = default;
 
 	explicit FDescriptorHash(const FCustomizableObjectInstanceDescriptor& Descriptor);
-
-	bool operator==(const FDescriptorHash& Other) const;
-
-	bool operator!=(const FDescriptorHash& Other) const;
 	
 	/** Return true if this Hash is a subset of the other Hash (i.e., this Descriptor is a subset of the other Descriptor). */
 	bool IsSubset(const FDescriptorHash& Other) const;
@@ -29,6 +25,8 @@ private:
 public:
 	TMap<FName, uint8> MinLODs;
 
+	TMap<FName, uint8> FirstLODAvailable;
+	
 	// Array of bitmasks that indicate which LODs of each component have been requested
 	TMap<FName, uint8> RequestedLODsPerComponent;
 };
diff --git a/Engine/Plugins/Mutable/Source/CustomizableObject/Private/MuCO/CustomizableObject.cpp b/Engine/Plugins/Mutable/Source/CustomizableObject/Private/MuCO/CustomizableObject.cpp
--- a/Engine/Plugins/Mutable/Source/CustomizableObject/Private/MuCO/CustomizableObject.cpp
+++ b/Engine/Plugins/Mutable/Source/CustomizableObject/Private/MuCO/CustomizableObject.cpp
@@ -2280,11 +2280,17 @@ uint8 UCustomizableObjectPrivate::GetMinLODIndex(const FName& ComponentName) con
 {
 	uint8 MinLODIdx = 0;
 
+	// Use the Scalability quality settings to determine what the MinLOD will be used
 	if (GEngine && GEngine->UseSkeletalMeshMinLODPerQualityLevels)
 	{
 		if (UCustomizableObjectSystem::GetInstance() != nullptr)
 		{
-			MinLODIdx = GetPublic()->LODSettings.MinQualityLevelLOD.GetValue(UCustomizableObjectSystem::GetInstance()->GetSkeletalMeshMinLODQualityLevel());
+			// Get the quality level for the skeletal meshes for the current scalability setting. The bigger the better.
+			// Value extracted from the Scalability.ini of the current platform (if found).
+			const int32 QualityLevel = UCustomizableObjectSystem::GetInstance()->GetSkeletalMeshMinLODQualityLevel();
+			
+			// Get the MinLOD to use for the target quality level. The lower the better.
+			MinLODIdx = GetPublic()->LODSettings.MinQualityLevelLOD.GetValue(QualityLevel);
 		}
 	}
 	else
@@ -2292,6 +2298,7 @@ uint8 UCustomizableObjectPrivate::GetMinLODIndex(const FName& ComponentName) con
 		MinLODIdx = GetPublic()->LODSettings.MinLOD.GetValue();
 	}
 
+	// Get the first lod the current platform can get generated. It represents the Min LOD mutable can generate for the current platform.
 	uint8 FirstLODAvailable = 0;
 	if (const uint8* Result = GetModelResources().ComponentFirstLODAvailable.Find(ComponentName))
 	{
diff --git a/Engine/Plugins/Mutable/Source/CustomizableObject/Private/MuCO/CustomizableObjectInstance.cpp b/Engine/Plugins/Mutable/Source/CustomizableObject/Private/MuCO/CustomizableObjectInstance.cpp
--- a/Engine/Plugins/Mutable/Source/CustomizableObject/Private/MuCO/CustomizableObjectInstance.cpp
+++ b/Engine/Plugins/Mutable/Source/CustomizableObject/Private/MuCO/CustomizableObjectInstance.cpp
@@ -3178,6 +3178,7 @@ FMutableUpdateCandidate::FMutableUpdateCandidate(UCustomizableObjectInstance* In
 {
 	const FCustomizableObjectInstanceDescriptor& Descriptor = InCustomizableObjectInstance->GetPrivate()->GetDescriptor();
 	MinLOD = Descriptor.MinLOD;
+	FirstLODAvailable = Descriptor.FirstLODAvailable;
 	RequestedLODLevels = Descriptor.GetRequestedLODLevels();
 }
 
@@ -3197,9 +3198,11 @@ void FMutableUpdateCandidate::Issue()
 void FMutableUpdateCandidate::ApplyLODUpdateParamsToInstance(FUpdateContextPrivate& Context)
 {
 	CustomizableObjectInstance->Descriptor.MinLOD = MinLOD;
+	CustomizableObjectInstance->Descriptor.FirstLODAvailable = FirstLODAvailable;
 	CustomizableObjectInstance->Descriptor.RequestedLODLevels = RequestedLODLevels;
 
 	Context.SetMinLOD(MinLOD);
+	Context.SetFirstLODAvailable(FirstLODAvailable);
 	Context.SetRequestedLODs(RequestedLODLevels);
 }
 
@@ -7177,6 +7180,8 @@ void UCustomizableObjectInstance::SetRequestedLODs(const TMap<FName, uint8>& InM
 	MutableUpdateCandidate.MinLOD = InMinLODs;
 	MutableUpdateCandidate.RequestedLODLevels = Descriptor.GetRequestedLODLevels();
 
+	MutableUpdateCandidate.FirstLODAvailable.Reset();
+	
 	const TMap<FName, uint8>& GeneratedLODsPerComponent = GetPrivate()->CommittedDescriptorHash.RequestedLODsPerComponent;
 	
 	for (const FName& ComponentName : CustomizableObject->GetPrivate()->GetModelResources().ComponentNamesPerObjectComponent)
@@ -7188,6 +7193,8 @@ void UCustomizableObjectInstance::SetRequestedLODs(const TMap<FName, uint8>& InM
 		}
 
 		const uint8 MinLODIdx = CustomizableObject->GetPrivate()->GetMinLODIndex(ComponentName);
+		MutableUpdateCandidate.FirstLODAvailable.Add(ComponentName, MinLODIdx);
+
 		const int32 MaxLODIdx = ModelResources.ComponentLODs[ComponentName] - 1;
 		InMinLOD = FMath::Clamp(InMinLOD, MinLODIdx, MaxLODIdx);
 
diff --git a/Engine/Plugins/Mutable/Source/CustomizableObject/Private/MuCO/CustomizableObjectSystem.cpp b/Engine/Plugins/Mutable/Source/CustomizableObject/Private/MuCO/CustomizableObjectSystem.cpp
--- a/Engine/Plugins/Mutable/Source/CustomizableObject/Private/MuCO/CustomizableObjectSystem.cpp
+++ b/Engine/Plugins/Mutable/Source/CustomizableObject/Private/MuCO/CustomizableObjectSystem.cpp
@@ -244,6 +244,19 @@ void FUpdateContextPrivate::SetMinLOD(const TMap<FName, uint8>& MinLOD)
 }
 
 
+void FUpdateContextPrivate::SetFirstLODAvailable(const TMap<FName, uint8>& FirstLODs)
+{
+	CapturedDescriptor.FirstLODAvailable = FirstLODs;
+	CapturedDescriptorHash.FirstLODAvailable = FirstLODs;
+}
+
+
+const TMap<FName, uint8>& FUpdateContextPrivate::GetFirstLODAvailable() const
+{
+	return CapturedDescriptor.FirstLODAvailable;
+}
+
+
 const TMap<FName, uint8>& FUpdateContextPrivate::GetRequestedLODs() const
 {
 	return CapturedDescriptor.GetRequestedLODLevels();
@@ -1750,6 +1763,7 @@ namespace impl
 
 		TMap<FName, uint8> MinLODs = Operation->bStreamMeshLODs ? TMap<FName, uint8>() : Operation->GetMinLOD();
 		TMap<FName, uint8> RequestedLODs = Operation->GetRequestedLODs();
+		TMap<FName, uint8> FirstLODsAvailable = Operation->GetFirstLODAvailable();
 		
 		for (int32 InstanceComponentIndex = 0; InstanceComponentIndex < Operation->NumInstanceComponents; ++InstanceComponentIndex)
 		{
@@ -1760,8 +1774,8 @@ namespace impl
 			uint8& RequestedLOD = RequestedLODs.FindOrAdd(ComponentName, 0);
 			uint8& NumLODsAvailable = Operation->NumLODsAvailable.FindOrAdd(ComponentName, 0);
 			uint8& FirstResidentLOD = Operation->FirstResidentLOD.FindOrAdd(ComponentName, 0);
-			uint8& FirstLODAvailable = Operation->FirstLODAvailable.FindOrAdd(ComponentName, 0);
-			
+			uint8& FirstLODAvailable = FirstLODsAvailable.FindOrAdd(ComponentName, 0);
+
 			NumLODsAvailable = Operation->MutableInstance->GetLODCount(InstanceComponentIndex);
 			MinLOD = FMath::Clamp(MinLOD, FirstLODAvailable, NumLODsAvailable - 1);
 			
@@ -1779,6 +1793,7 @@ namespace impl
 		
 		Operation->SetMinLOD(MinLODs);
 		Operation->SetRequestedLODs(RequestedLODs);
+		Operation->SetFirstLODAvailable(FirstLODsAvailable);
 	}
 	
 
@@ -3853,6 +3868,7 @@ int32 UCustomizableObjectSystem::TickInternal(const bool bBlocking)
 						FDescriptorHash LODUpdateDescriptorHash = CustomizableInstancePrivate->CommittedDescriptorHash;
 						LODUpdateDescriptorHash.MinLODs = LODUpdateCandidate.MinLOD;
 						LODUpdateDescriptorHash.RequestedLODsPerComponent = LODUpdateCandidate.RequestedLODLevels;
+						LODUpdateDescriptorHash.FirstLODAvailable = LODUpdateCandidate.FirstLODAvailable;
 
 						if (CustomizableInstancePrivate->MinSquareDistFromComponentToPlayer < MaxSquareDistanceFound &&
 							!LODUpdateDescriptorHash.IsSubset(CustomizableInstancePrivate->CommittedDescriptorHash))
diff --git a/Engine/Plugins/Mutable/Source/CustomizableObject/Private/MuCO/DescriptorHash.cpp b/Engine/Plugins/Mutable/Source/CustomizableObject/Private/MuCO/DescriptorHash.cpp
--- a/Engine/Plugins/Mutable/Source/CustomizableObject/Private/MuCO/DescriptorHash.cpp
+++ b/Engine/Plugins/Mutable/Source/CustomizableObject/Private/MuCO/DescriptorHash.cpp
@@ -53,25 +53,11 @@ FDescriptorHash::FDescriptorHash(const FCustomizableObjectInstanceDescriptor& De
 	Hash = HashCombine(Hash, GetTypeHash(Descriptor.GetBuildParameterRelevancy()));
 	
 	MinLODs = Descriptor.MinLOD;
-
+	FirstLODAvailable = Descriptor.FirstLODAvailable;
 	RequestedLODsPerComponent = Descriptor.RequestedLODLevels;
 }
 
 
-bool FDescriptorHash::operator==(const FDescriptorHash& Other) const
-{
-	return Hash == Other.Hash &&
-		MinLODs == Other.MinLODs &&
-		RequestedLODsPerComponent == Other.RequestedLODsPerComponent;
-}
-
-
-bool FDescriptorHash::operator!=(const FDescriptorHash& Other) const
-{
-	return !(*this == Other);
-}
-
-
 bool FDescriptorHash::IsSubset(const FDescriptorHash& Other) const
 {
 	if (Hash != Other.Hash)
@@ -92,10 +78,11 @@ bool FDescriptorHash::IsSubset(const FDescriptorHash& Other) const
 			return false;
 		}
 	}
-
-	if (RequestedLODsPerComponent == Other.RequestedLODsPerComponent)
+	
+	// Scalability quality change : Do not care if it is a subset, always update
+	if (FirstLODAvailable != Other.FirstLODAvailable)
 	{
-		return true;
+		return false;
 	}
 
 	if (RequestedLODsPerComponent.Num() != Other.RequestedLODsPerComponent.Num())
@@ -134,6 +121,13 @@ FString FDescriptorHash::ToString() const
 	}
 	Builder.Appendf(TEXT("])"));
 
+	Builder.Appendf(TEXT("FirstLODAvailable=["));
+	for (const TTuple<FName, uint8>& FirstLOD : FirstLODAvailable)
+	{
+		Builder.Appendf(TEXT("(%s, %i),"), *FirstLOD.Key.ToString(), FirstLOD.Value);
+	}
+	Builder.Appendf(TEXT("])"));
+
 	Builder.Appendf(TEXT("RequestedLODLevels=["));
 	for (const TTuple<FName, uint8>& RequestedLODs : RequestedLODsPerComponent)
 	{
diff --git a/Engine/Plugins/Mutable/Source/CustomizableObject/Public/MuCO/CustomizableObjectInstanceDescriptor.h b/Engine/Plugins/Mutable/Source/CustomizableObject/Public/MuCO/CustomizableObjectInstanceDescriptor.h
--- a/Engine/Plugins/Mutable/Source/CustomizableObject/Public/MuCO/CustomizableObjectInstanceDescriptor.h
+++ b/Engine/Plugins/Mutable/Source/CustomizableObject/Public/MuCO/CustomizableObjectInstanceDescriptor.h
@@ -323,7 +323,7 @@ struct CUSTOMIZABLEOBJECT_API FCustomizableObjectInstanceDescriptor
 	UPROPERTY()
 	TArray<FCustomizableObjectTransformParameterValue> TransformParameters;
 
-	/** Mutable parameters optimization state. Transient UProperty to make it transactable. */
+	/** Mutable parameters optimization state. Transient UProperty to make it translatable. */
 	UPROPERTY(Transient)
 	int32 State = 0;
 	
@@ -337,6 +337,8 @@ struct CUSTOMIZABLEOBJECT_API FCustomizableObjectInstanceDescriptor
 	 * Key is the Component name. Value represents the first LOD to generate. Zero means all LODs. */
 	TMap<FName, uint8> RequestedLODLevels;
 	
+	TMap<FName, uint8> FirstLODAvailable;
+	
 	// Friends
 	friend FDescriptorHash;
 	friend UCustomizableObjectInstance;
diff --git a/Engine/Plugins/Mutable/Source/CustomizableObject/Public/MuCO/MultilayerProjector.h b/Engine/Plugins/Mutable/Source/CustomizableObject/Public/MuCO/MultilayerProjector.h
--- a/Engine/Plugins/Mutable/Source/CustomizableObject/Public/MuCO/MultilayerProjector.h
+++ b/Engine/Plugins/Mutable/Source/CustomizableObject/Public/MuCO/MultilayerProjector.h
@@ -21,7 +21,6 @@ struct FMultilayerProjectorLayer
 	/** Read the Layer from the Instance Parameters.
 	 *
 	 * @param Descriptor Instance Descriptor.
-	 * @param MultilayerProjector Multilayer Projector helper.
 	 * @param Index Index where to read the Layer in the Instance Parameters.
 	 */
 	void Read(const FCustomizableObjectInstanceDescriptor& Descriptor, const FString& ParamName, int32 Index);
@@ -29,7 +28,6 @@ struct FMultilayerProjectorLayer
 	/** Write to Layer to the Instance Parameters.
 	 *
  	 * @param Descriptor Instance Descriptor.
-	 * @param MultilayerProjector Multilayer Projector helper.
 	 * @param Index Index where to write the Layer in the Instance Parameters.
 	 */
 	void Write(FCustomizableObjectInstanceDescriptor& Descriptor, const FString& ParamName, int32 Index) const;
diff --git a/Engine/Plugins/Mutable/Source/CustomizableObjectEditor/Private/MuCOE/GenerateMutableSource/GenerateMutableSourceComponent.cpp b/Engine/Plugins/Mutable/Source/CustomizableObjectEditor/Private/MuCOE/GenerateMutableSource/GenerateMutableSourceComponent.cpp
--- a/Engine/Plugins/Mutable/Source/CustomizableObjectEditor/Private/MuCOE/GenerateMutableSource/GenerateMutableSourceComponent.cpp
+++ b/Engine/Plugins/Mutable/Source/CustomizableObjectEditor/Private/MuCOE/GenerateMutableSource/GenerateMutableSourceComponent.cpp
@@ -472,15 +472,19 @@ void FirstPass(UCustomizableObjectNodeComponentMesh& Node, FMutableGraphGenerati
 	// Find the MinLOD available for the target platform
 	if (RefSkeletalMesh->IsMinLodQualityLevelEnable()) // Engine global setting. Current implementation has nothing to do with the USkeletalMesh.
 	{
+		// Array of Quality values for each of the scalability settings. The valid values go from [0,4] each one representing one Scalability setting:
+		//		0 = Low
+		//		4 = Cine = Cinematic
 		FSupportedQualityLevelArray SupportedQualityLevels = LODSettings.MinQualityLevelLOD.GetSupportedQualityLevels(*GenerationContext.Options.TargetPlatform->GetPlatformInfo().IniPlatformName.ToString());
 			
+		// If no scalability settings are found, use the MaxLOD as the MinLOD to be used.
 		int32 MinValue = NumLODs - 1;
-		for (int32& QL : SupportedQualityLevels)
+		for (const int32& QualityLevel : SupportedQualityLevels)
 		{
-			// check if have data for the supported quality level or set to default.
-			if (LODSettings.MinQualityLevelLOD.IsQualityLevelValid(QL))
+			// check if we have data for the supported quality level or set to default.
+			if (LODSettings.MinQualityLevelLOD.IsQualityLevelValid(QualityLevel))
 			{
-				MinValue = FMath::Min(LODSettings.MinQualityLevelLOD.GetValueForQualityLevel(QL), MinValue);
+				MinValue = FMath::Min(LODSettings.MinQualityLevelLOD.GetValueForQualityLevel(QualityLevel), MinValue);
 			}
 			else 
 			{
