[Mutable] RegenerateImportedModels optimization.

- The optimization tries to parallelize all mesh index and vertex conversions using a parallelfor. It batches the work for all meshes and lods trying to assign a similar workload to all batches.

#rnx

diff --git a/Engine/Plugins/Mutable/Source/CustomizableObject/Private/MuCO/CustomizableObjectInstance.cpp b/Engine/Plugins/Mutable/Source/CustomizableObject/Private/MuCO/CustomizableObjectInstance.cpp
--- a/Engine/Plugins/Mutable/Source/CustomizableObject/Private/MuCO/CustomizableObjectInstance.cpp
+++ b/Engine/Plugins/Mutable/Source/CustomizableObject/Private/MuCO/CustomizableObjectInstance.cpp
@@ -86,8 +86,8 @@ static FAutoConsoleVariableRef CVarDisableNotifyComponentsOfTextureUpdates(
 	bDisableNotifyComponentsOfTextureUpdates,
 	TEXT("If set to true, disables Mutable notifying the streaming system that a component has had a change in at least one texture of its components."),
 	ECVF_Default);
-}
 
+}
 
 const FString MULTILAYER_PROJECTOR_PARAMETERS_INVALID = TEXT("Invalid Multilayer Projector Parameters.");
 
@@ -7598,6 +7598,29 @@ void UCustomizableInstancePrivate::RegenerateImportedModels()
 {
 	MUTABLE_CPUPROFILER_SCOPE(RegenerateImportedModels);
 
+	using UVsVectorType = typename TDecay<decltype(DeclVal<FSoftSkinVertex>().UVs[0])>::Type;
+	struct FMeshDataConvertJob
+	{
+		int32 NumIndices = 0;
+		int32 IndicesOffset = 0;
+		const FRawStaticIndexBuffer16or32Interface* IndexBuffer = nullptr;
+		uint32* DestIndexBuffer = nullptr;
+
+		int32 NumVertices = 0;
+		int32 VerticesOffset = 0;
+		const FStaticMeshVertexBuffers*	StaticVertexBuffers = nullptr;
+		const FSkinWeightVertexBuffer* SkinWeightVertexBuffer = nullptr;
+		FSoftSkinVertex* DestVertexBuffer = nullptr;
+	};
+	
+	constexpr int32 MaxJobCost = 1 << 18;
+	constexpr int32 MaxVerticesPerJob = FMath::Max<int32>(1, MaxJobCost / sizeof(FSoftSkinVertex));
+	constexpr int32 MaxIndicesPerJob =  FMath::Max<int32>(1, MaxJobCost / sizeof(int32));
+
+	TArray<FMeshDataConvertJob, TInlineAllocator<64>> Jobs;   
+	TArray<int32, TInlineAllocator<64>> JobRanges;
+	JobRanges.Add(0);
+	
 	for (const TTuple<FName, TObjectPtr<USkeletalMesh>>& Tuple : SkeletalMeshes)
 	{
 		USkeletalMesh* SkeletalMesh = Tuple.Get<1>();
@@ -7659,7 +7682,7 @@ void UCustomizableInstancePrivate::RegenerateImportedModels()
 		ImportedModel->SkeletalMeshModelGUID = FGuid::NewGuid();
 
 		ImportedModel->LODModels.Empty();
-
+	
 		int32 OriginalIndex = 0;
 		for (int32 LODIndex = 0; LODIndex < RenderData->LODRenderData.Num(); ++LODIndex)
 		{
@@ -7667,7 +7690,6 @@ void UCustomizableInstancePrivate::RegenerateImportedModels()
 			FSkeletalMeshLODModel& LODModel = ImportedModel->LODModels[LODIndex];
 
 			FSkeletalMeshLODRenderData& LODRenderData = RenderData->LODRenderData[LODIndex];
-			int32 CurrentSectionInitialVertex = 0;
 
 			LODModel.ActiveBoneIndices = LODRenderData.ActiveBoneIndices;
 			LODModel.NumTexCoords = LODRenderData.GetNumTexCoords();
@@ -7677,17 +7699,37 @@ void UCustomizableInstancePrivate::RegenerateImportedModels()
 			// Indices
 			if (LODRenderData.MultiSizeIndexContainer.IsIndexBufferValid())
 			{
-				const int32 NumIndices = LODRenderData.MultiSizeIndexContainer.GetIndexBuffer()->Num();
+				const FRawStaticIndexBuffer16or32Interface* IndexBuffer =
+						LODRenderData.MultiSizeIndexContainer.GetIndexBuffer();
+
+				const int32 NumIndices = IndexBuffer->Num();
 				LODModel.IndexBuffer.SetNum(NumIndices);
-				for (int32 Index = 0; Index < NumIndices; ++Index)
+
+				uint32* BaseDestIndexBuffer = LODModel.IndexBuffer.GetData();
+				
+				const int32 NumIndicesJobs = FMath::DivideAndRoundUp(NumIndices, MaxIndicesPerJob);
+
+				int32 CurrentJobIndex = Jobs.Num();
+				Jobs.SetNum(Jobs.Num() + NumIndicesJobs);
+		
+				for (int32 I = 0; I < NumIndicesJobs; ++I)
 				{
-					LODModel.IndexBuffer[Index] = LODRenderData.MultiSizeIndexContainer.GetIndexBuffer()->Get(Index);
+					FMeshDataConvertJob Job;
+					Job.NumIndices = FMath::Min(MaxIndicesPerJob, NumIndices - I*MaxIndicesPerJob); 
+					Job.IndexBuffer = IndexBuffer;
+					Job.IndicesOffset = I*MaxIndicesPerJob; 
+					Job.DestIndexBuffer = BaseDestIndexBuffer + I*MaxIndicesPerJob;
+
+					Jobs[I + CurrentJobIndex] = Job;  
 				}
 			}
 
 			LODModel.Sections.SetNum(LODRenderData.RenderSections.Num());
 
-			for (int SectionIndex = 0; SectionIndex < LODRenderData.RenderSections.Num(); ++SectionIndex)
+			check(!LODRenderData.StaticVertexBuffers.StaticMeshVertexBuffer.GetUseHighPrecisionTangentBasis());
+			check(LODRenderData.StaticVertexBuffers.StaticMeshVertexBuffer.GetUseFullPrecisionUVs());
+
+			for (int32 SectionIndex = 0; SectionIndex < LODRenderData.RenderSections.Num(); ++SectionIndex)
 			{
 				const FSkelMeshRenderSection& RenderSection = LODRenderData.RenderSections[SectionIndex];
 				FSkelMeshSection& ImportedSection = ImportedModel->LODModels[LODIndex].Sections[SectionIndex];
@@ -7707,58 +7749,30 @@ void UCustomizableInstancePrivate::RegenerateImportedModels()
 				ImportedSection.SoftVertices.AddUninitialized(RenderSection.NumVertices);
 				ImportedSection.bUse16BitBoneIndex = LODRenderData.DoesVertexBufferUse16BitBoneIndex();
 
-				TArray<FColor> VertexColors;
-				LODRenderData.StaticVertexBuffers.ColorVertexBuffer.GetVertexColors(VertexColors);
-
-				for (uint32 i = 0; i < RenderSection.NumVertices; ++i)
-				{
-					const FPositionVertex* PosPtr = static_cast<const FPositionVertex*>(LODRenderData.StaticVertexBuffers.PositionVertexBuffer.GetVertexData());
-					PosPtr += (CurrentSectionInitialVertex + i);
+				const int32 SectionNumVertices = RenderSection.NumVertices;
+				const int32 SectionBaseVertexIndex = RenderSection.BaseVertexIndex;
+				const FStaticMeshVertexBuffers*	StaticVertexBuffers = &LODRenderData.StaticVertexBuffers;
+				const FSkinWeightVertexBuffer* SkinWeightVertexBuffer = &LODRenderData.SkinWeightVertexBuffer;
 
-					check(!LODRenderData.StaticVertexBuffers.StaticMeshVertexBuffer.GetUseHighPrecisionTangentBasis());
-					const FPackedNormal* TangentPtr = static_cast<const FPackedNormal*>(LODRenderData.StaticVertexBuffers.StaticMeshVertexBuffer.GetTangentData());
-					TangentPtr += ((CurrentSectionInitialVertex + i) * 2);
+				FSoftSkinVertex* BaseSectionSoftVertex = ImportedSection.SoftVertices.GetData();
 
-					check(LODRenderData.StaticVertexBuffers.StaticMeshVertexBuffer.GetUseFullPrecisionUVs());
+				int32 NumSectionJobs = FMath::DivideAndRoundUp<int32>(RenderSection.NumVertices, MaxVerticesPerJob);
 
-					using UVsVectorType = typename TDecay<decltype(DeclVal<FSoftSkinVertex>().UVs[0])>::Type;
+				int32 FirstSectionJobIndex = Jobs.Num();
+				Jobs.SetNum(Jobs.Num() + NumSectionJobs);
 
-					const UVsVectorType* TexCoordPosPtr = static_cast<const UVsVectorType*>(LODRenderData.StaticVertexBuffers.StaticMeshVertexBuffer.GetTexCoordData());
-					const uint32 NumTexCoords = LODRenderData.StaticVertexBuffers.StaticMeshVertexBuffer.GetNumTexCoords();
-					TexCoordPosPtr += ((CurrentSectionInitialVertex + i) * NumTexCoords);
-
-					FSoftSkinVertex& Vertex = ImportedSection.SoftVertices[i];
-					for (int32 j = 0; j < RenderSection.MaxBoneInfluences; ++j)
-					{
-						Vertex.InfluenceBones[j] = LODRenderData.SkinWeightVertexBuffer.GetBoneIndex(CurrentSectionInitialVertex + i, j);
-						Vertex.InfluenceWeights[j] = LODRenderData.SkinWeightVertexBuffer.GetBoneWeight(CurrentSectionInitialVertex + i, j);
-					}
-
-					for (int32 j = RenderSection.MaxBoneInfluences; j < MAX_TOTAL_INFLUENCES; ++j)
-					{
-						Vertex.InfluenceBones[j] = 0;
-						Vertex.InfluenceWeights[j] = 0;
-					}
-
-					Vertex.Color = VertexColors.IsValidIndex(i) ?
-						VertexColors[i] :
-						FColor::White;
-
-					Vertex.Position = PosPtr->Position;
-
-					Vertex.TangentX = TangentPtr[0].ToFVector3f();
-					Vertex.TangentZ = TangentPtr[1].ToFVector3f();
-					float TangentSign = TangentPtr[1].Vector.W < 0 ? -1.f : 1.f;
-					Vertex.TangentY = FVector3f::CrossProduct(Vertex.TangentZ, Vertex.TangentX) * TangentSign;
-
-					Vertex.UVs[0] = TexCoordPosPtr[0];
-					Vertex.UVs[1] = NumTexCoords > 1 ? TexCoordPosPtr[1] : UVsVectorType::ZeroVector;
-					Vertex.UVs[2] = NumTexCoords > 2 ? TexCoordPosPtr[2] : UVsVectorType::ZeroVector;
-					Vertex.UVs[3] = NumTexCoords > 3 ? TexCoordPosPtr[3] : UVsVectorType::ZeroVector;
+				for (int32 I = 0; I < NumSectionJobs; ++I)
+				{
+					FMeshDataConvertJob Job;
+					Job.NumVertices = FMath::Min(MaxVerticesPerJob, SectionNumVertices - I*MaxVerticesPerJob); 
+					Job.StaticVertexBuffers = StaticVertexBuffers;
+					Job.SkinWeightVertexBuffer = SkinWeightVertexBuffer;
+					Job.VerticesOffset = SectionBaseVertexIndex + I*MaxVerticesPerJob; 
+					Job.DestVertexBuffer = BaseSectionSoftVertex + I*MaxVerticesPerJob;
+
+					Jobs[I + FirstSectionJobIndex] = Job;  
 				}
 
-				CurrentSectionInitialVertex += RenderSection.NumVertices;
-
 				// Triangles
 				ImportedSection.NumTriangles = RenderSection.NumTriangles;
 				ImportedSection.BaseIndex = RenderSection.BaseIndex;
@@ -7814,6 +7828,121 @@ void UCustomizableInstancePrivate::RegenerateImportedModels()
 				LODModel.BuildStringID = LODModel.GetLODModelDeriveDataKey();
 			}
 		}
+
+		// Try to bundle Jobs so all cost roughly the same. Large Jobs are already split so they cost about MaxJobCost.
+		// It uses a gready approach and assumes in general Jobs are sorted by cost.
+		const int32 NumJobs = Jobs.Num();
+		for (int32 JobIndex = 0; JobIndex < NumJobs;)
+		{
+			int32 RangeJobCost = 0;
+			for (; JobIndex < NumJobs;)
+			{
+				int32 CurrentJobCost = Jobs[JobIndex].NumVertices*sizeof(FSoftSkinVertex) + Jobs[JobIndex].NumIndices*sizeof(int32);
+
+				RangeJobCost += CurrentJobCost;
+				if (RangeJobCost >= MaxJobCost)
+				{
+					// Go to the next Job if the current job alone cost is larger than MaxJobCost
+					// and no other job has been processed for the range.
+					JobIndex += static_cast<int32>(CurrentJobCost == RangeJobCost);
+					break;
+				}
+
+				++JobIndex;
+			}
+
+			JobRanges.Add(JobIndex);
+		}
+	}
+
+	{
+		MUTABLE_CPUPROFILER_SCOPE(DoImportedModelMeshDataConversion)
+
+		ParallelFor(JobRanges.Num() - 1, [&JobRanges, &Jobs](int32 JobId)
+		{
+			int32 JobRangeBegin = JobRanges[JobId];
+			int32 JobRangeEnd = JobRanges[JobId + 1];
+			for (int32 J = JobRangeBegin; J < JobRangeEnd; ++J)
+			{
+				FMeshDataConvertJob Job = Jobs[J];
+		
+				if (Job.NumIndices > 0)
+				{
+					MUTABLE_CPUPROFILER_SCOPE(DoImportedModelMeshDataConversion_Indices)
+
+					for (int32 Index = 0; Index < Job.NumIndices; ++Index)
+					{
+						Job.DestIndexBuffer[Index] = Job.IndexBuffer->Get(Job.IndicesOffset + Index);
+					}
+				}
+
+				if (Job.NumVertices > 0)
+				{
+					MUTABLE_CPUPROFILER_SCOPE(DoImportedModelMeshDataConversion_Vertices)
+
+					check(Job.StaticVertexBuffers);
+					check(Job.SkinWeightVertexBuffer);
+					check(Job.DestVertexBuffer);
+
+					const FPositionVertex* PositionBuffer = 
+							static_cast<const FPositionVertex*>(Job.StaticVertexBuffers->PositionVertexBuffer.GetVertexData()) + Job.VerticesOffset;
+
+					const FPackedNormal* TangentBuffer = 
+							static_cast<const FPackedNormal*>(Job.StaticVertexBuffers->StaticMeshVertexBuffer.GetTangentData()) + Job.VerticesOffset*2;
+
+					const int32 NumTexCoords = Job.StaticVertexBuffers->StaticMeshVertexBuffer.GetNumTexCoords();
+					const UVsVectorType* TexCoordBuffer = 
+							static_cast<const UVsVectorType*>(Job.StaticVertexBuffers->StaticMeshVertexBuffer.GetTexCoordData()) + Job.VerticesOffset*NumTexCoords;
+					
+					const FColor* ColorBuffer = 
+							static_cast<const FColor*>(Job.StaticVertexBuffers->ColorVertexBuffer.GetVertexData());
+			
+					const bool bHasColor = !!ColorBuffer;
+					ColorBuffer += Job.VerticesOffset;
+
+					const FSkinWeightVertexBuffer* SkinWeightBuffer = Job.SkinWeightVertexBuffer;
+
+					const int32 MaxBoneInfluences = Job.SkinWeightVertexBuffer->GetMaxBoneInfluences();
+
+					for (int32 JobVertexIndex = 0; JobVertexIndex < Job.NumVertices; ++JobVertexIndex)
+					{
+						FSoftSkinVertex* Vertex = Job.DestVertexBuffer + JobVertexIndex;
+						FMemory::Memzero(Vertex, sizeof(FSoftSkinVertex));
+
+						Vertex->Position = PositionBuffer[JobVertexIndex].Position;
+
+						const FPackedNormal* Tangent = TangentBuffer + JobVertexIndex*2;
+					
+						Vertex->TangentX = Tangent[0].ToFVector3f();
+						Vertex->TangentZ = Tangent[1].ToFVector3f();
+						float TangentSign = Tangent[1].Vector.W < 0 ? -1.0f : 1.0f;
+						Vertex->TangentY = FVector3f::CrossProduct(Vertex->TangentZ, Vertex->TangentX) * TangentSign;
+
+						const UVsVectorType* TexCoord = TexCoordBuffer + JobVertexIndex*NumTexCoords;
+
+						// Switch based jumptable.
+						switch (NumTexCoords)
+						{
+							case 4: Vertex->UVs[3] = TexCoord[3]; // Fall through
+							case 3: Vertex->UVs[2] = TexCoord[2]; // Fall through
+							case 2: Vertex->UVs[1] = TexCoord[1]; // Fall through
+							case 1: Vertex->UVs[0] = TexCoord[0]; // Fall through
+							default: break;
+						}
+
+						Vertex->Color = bHasColor ? ColorBuffer[JobVertexIndex] : FColor::White;
+
+						const int32 SourceVertexIndex = (JobVertexIndex + Job.VerticesOffset);
+
+						for (int32 InfluenceIndex = 0; InfluenceIndex < MaxBoneInfluences; ++InfluenceIndex)
+						{
+							Vertex->InfluenceBones[InfluenceIndex] = SkinWeightBuffer->GetBoneIndex(SourceVertexIndex, InfluenceIndex);
+							Vertex->InfluenceWeights[InfluenceIndex] = SkinWeightBuffer->GetBoneWeight(SourceVertexIndex, InfluenceIndex);
+						}
+					}
+				}
+			}
+		});
 	}
 }
 
