[mutable] Refactor part of the CO compilation process to separate the mutable graph generation context from the "whole compilation" context. This is to prepare the on-demand mesh conversion during compilation, but it is not enabled yet.
- add compile-time referenced meshes support.
- moves some asset references to soft object pointers to avoid loading them unless strictly necessary.
- use asset registry to detect participating objects without loading assets in some cases.
- Cleanup some unnecessary functions and parameters.
#rnx

diff --git a/Engine/Plugins/Mutable/Source/CustomizableObject/Internal/MuCO/CustomizableObjectPrivate.h b/Engine/Plugins/Mutable/Source/CustomizableObject/Internal/MuCO/CustomizableObjectPrivate.h
--- a/Engine/Plugins/Mutable/Source/CustomizableObject/Internal/MuCO/CustomizableObjectPrivate.h
+++ b/Engine/Plugins/Mutable/Source/CustomizableObject/Internal/MuCO/CustomizableObjectPrivate.h
@@ -1187,12 +1187,16 @@ struct FModelResources
 
 	/** PassThrough meshes used by the mu::FModel. */
 	UPROPERTY()
-	TArray<TSoftObjectPtr<USkeletalMesh>> PassThroughMeshes;
+	TArray<TSoftObjectPtr<UStreamableRenderAsset>> PassThroughMeshes;
 
 #if WITH_EDITORONLY_DATA
 	/** Runtime referenced textures used by the mu::FModel. */
 	UPROPERTY()
 	TArray<TSoftObjectPtr<const UTexture>> RuntimeReferencedTextures;
+	
+	/** Runtime referenced meshes used by the mu::FModel. */
+	UPROPERTY()
+	TArray<TSoftObjectPtr<const UStreamableRenderAsset>> RuntimeReferencedMeshes;
 #endif
 	
 	/** Physics assets gathered from the SkeletalMeshes, to be used in mesh generation in-game */
diff --git a/Engine/Plugins/Mutable/Source/CustomizableObject/Internal/MuCO/UnrealToMutableTextureConversionUtils.h b/Engine/Plugins/Mutable/Source/CustomizableObject/Internal/MuCO/UnrealToMutableTextureConversionUtils.h
--- a/Engine/Plugins/Mutable/Source/CustomizableObject/Internal/MuCO/UnrealToMutableTextureConversionUtils.h
+++ b/Engine/Plugins/Mutable/Source/CustomizableObject/Internal/MuCO/UnrealToMutableTextureConversionUtils.h
@@ -6,6 +6,16 @@
 #include "Engine/Texture.h"
 
 #if WITH_EDITOR
+
+// Forward declarations
+class UTexture2D;
+class UAnimInstance;
+class USkeletalMesh;
+namespace mu
+{
+	class FImage;
+}
+
 struct CUSTOMIZABLEOBJECT_API FMutableSourceTextureData
 {
 	FMutableSourceTextureData() = default;
@@ -31,13 +41,50 @@ private:
 };
 
 
-namespace mu
+
+// Flags that can influence the mesh conversion
+enum class EMutableMeshConversionFlags : uint32
 {
-	class FImage;
-}
+	// 
+	None = 0,
+	// Ignore the skeleton and skinning
+	IgnoreSkinning = 1 << 0,
+
+	// Ignore Physics assets
+	IgnorePhysics = 1 << 1,
+
+	// Prevent this mesh generation from adding per mesh metadata. 
+	DoNotCreateMeshMetadata = 1 << 2
+};
+
+ENUM_CLASS_FLAGS(EMutableMeshConversionFlags)
+
+
+struct CUSTOMIZABLEOBJECT_API FMutableSourceMeshData
+{
+	/** Assets involved in the conversion. */
+	TSoftObjectPtr<const UStreamableRenderAsset> Mesh;
+	TSoftClassPtr<UAnimInstance> AnimInstance;
+	TSoftObjectPtr<const USkeletalMesh> TableReferenceSkeletalMesh;
+	FName Component = NAME_None;
+	uint32 SurfaceMetadataId = 0;
+
+	/** Selection of the mesh section*/
+	uint8 LODIndex=0;
+	uint8 SectionIndex=0;
+	bool bMeshMustExist=false;
+
+	/** Required mesh properties. */
+	EMutableMeshConversionFlags Flags = EMutableMeshConversionFlags::None;
+
+	/** Required realtime mesh morphs. */
+	bool bUseAllRealTimeMorphs = false;
+	TArray<FString> UsedRealTimeMorphTargetNames;
+
+	/** Context for log messages. */
+	const UObject* MessageContext = nullptr;
+};
 
-// Forward declarations
-class UTexture2D;
 
 enum class EUnrealToMutableConversionError 
 {
diff --git a/Engine/Plugins/Mutable/Source/CustomizableObject/Private/MuCO/CustomizableObject.cpp b/Engine/Plugins/Mutable/Source/CustomizableObject/Private/MuCO/CustomizableObject.cpp
--- a/Engine/Plugins/Mutable/Source/CustomizableObject/Private/MuCO/CustomizableObject.cpp
+++ b/Engine/Plugins/Mutable/Source/CustomizableObject/Private/MuCO/CustomizableObject.cpp
@@ -714,7 +714,7 @@ void FModelResources::Serialize(FObjectAndNameAsStringProxyArchive& MemoryWriter
 	int32 NumPassthroughMeshes = PassThroughMeshes.Num();
 	MemoryWriter << NumPassthroughMeshes;
 
-	for (const TSoftObjectPtr<USkeletalMesh>& PassthroughMesh : PassThroughMeshes)
+	for (const TSoftObjectPtr<UStreamableRenderAsset>& PassthroughMesh : PassThroughMeshes)
 	{
 		FString StringRef = PassthroughMesh.ToString();
 		MemoryWriter << StringRef;
@@ -869,7 +869,7 @@ bool FModelResources::Unserialize(FObjectAndNameAsStringProxyArchive& MemoryRead
 		FString StringRef;
 		MemoryReader << StringRef;
 
-		PassThroughMeshes.Add(TSoftObjectPtr<USkeletalMesh>(FSoftObjectPath(StringRef)));
+		PassThroughMeshes.Add(TSoftObjectPtr<UStreamableRenderAsset>(FSoftObjectPath(StringRef)));
 	}
 
 #if WITH_EDITORONLY_DATA
diff --git a/Engine/Plugins/Mutable/Source/CustomizableObject/Private/MuCO/CustomizableObjectInstance.cpp b/Engine/Plugins/Mutable/Source/CustomizableObject/Private/MuCO/CustomizableObjectInstance.cpp
--- a/Engine/Plugins/Mutable/Source/CustomizableObject/Private/MuCO/CustomizableObjectInstance.cpp
+++ b/Engine/Plugins/Mutable/Source/CustomizableObject/Private/MuCO/CustomizableObjectInstance.cpp
@@ -2320,15 +2320,25 @@ bool UCustomizableInstancePrivate::UpdateSkeletalMesh_PostBeginUpdate0(UCustomiz
 		if (ComponentMesh->IsReference())
         {
         	int32 ReferenceID = ComponentMesh->GetReferencedMesh();
-        	TSoftObjectPtr<USkeletalMesh> Ref = ModelResources.PassThroughMeshes[ReferenceID];
+        	TSoftObjectPtr<UStreamableRenderAsset> Ref = ModelResources.PassThroughMeshes[ReferenceID];
 
         	if (!Ref.IsValid())
         	{
         		// This shouldn't happen here synchronosuly. It should have been requested as an async load.
-        		UE_LOG(LogMutable, Error, TEXT("Referenced skeletal mesh [%s] was not pre-loaded. It will be sync-loaded probably causing a hitch. CO [%s]"), *Ref.ToString(), *GetNameSafe(CustomizableObject));
+        		UE_LOG(LogMutable, Error, TEXT("Referenced mesh [%s] was not pre-loaded. It will be sync-loaded probably causing a hitch. CO [%s]"), *Ref.ToString(), *GetNameSafe(CustomizableObject));
         	}
 
-        	SkeletalMeshes.Add(ComponentName, Ref.LoadSynchronous());
+			UStreamableRenderAsset* Asset = Ref.LoadSynchronous();
+			USkeletalMesh* SkeletalMesh = Cast<USkeletalMesh>(Asset);
+			if (SkeletalMesh)
+			{
+				SkeletalMeshes.Add(ComponentName, SkeletalMesh);
+			}
+			else
+			{
+				// Pass-through static meshes not implemented yet.
+				UE_LOG(LogMutable, Error, TEXT("Referenced static meshes [%s] are not supported yet. CO [%s]"), *Ref.ToString(), *GetNameSafe(CustomizableObject));
+			}
         	continue;
         }
 
diff --git a/Engine/Plugins/Mutable/Source/CustomizableObject/Private/MuCO/CustomizableObjectSystem.cpp b/Engine/Plugins/Mutable/Source/CustomizableObject/Private/MuCO/CustomizableObjectSystem.cpp
--- a/Engine/Plugins/Mutable/Source/CustomizableObject/Private/MuCO/CustomizableObjectSystem.cpp
+++ b/Engine/Plugins/Mutable/Source/CustomizableObject/Private/MuCO/CustomizableObjectSystem.cpp
@@ -2466,7 +2466,7 @@ namespace impl
 
 			if (ModelResources.PassThroughMeshes.IsValidIndex(ReferenceID))
 			{
-				TSoftObjectPtr<const USkeletalMesh> Ref = ModelResources.PassThroughMeshes[ReferenceID];
+				TSoftObjectPtr<const UStreamableRenderAsset> Ref = ModelResources.PassThroughMeshes[ReferenceID];
 				Instance->GetPrivate()->PassThroughMeshesToLoad.Add(Ref);
 			}
 			else
diff --git a/Engine/Plugins/Mutable/Source/CustomizableObjectEditor/Private/MuCOE/CustomizableObjectCompileRunnable.cpp b/Engine/Plugins/Mutable/Source/CustomizableObjectEditor/Private/MuCOE/CustomizableObjectCompileRunnable.cpp
--- a/Engine/Plugins/Mutable/Source/CustomizableObjectEditor/Private/MuCOE/CustomizableObjectCompileRunnable.cpp
+++ b/Engine/Plugins/Mutable/Source/CustomizableObjectEditor/Private/MuCOE/CustomizableObjectCompileRunnable.cpp
@@ -56,7 +56,7 @@ FCustomizableObjectCompileRunnable::FCustomizableObjectCompileRunnable(mu::Ptr<m
 }
 
 
-TSharedPtr<mu::FImage> FCustomizableObjectCompileRunnable::LoadResourceReferenced(int32 ID)
+TSharedPtr<mu::FImage> FCustomizableObjectCompileRunnable::LoadImageResourceReferenced(int32 ID)
 {
 	MUTABLE_CPUPROFILER_SCOPE(LoadResourceReferenced);
 
@@ -87,6 +87,44 @@ TSharedPtr<mu::FImage> FCustomizableObjectCompileRunnable::LoadResourceReference
 }
 
 
+TSharedPtr<mu::FMesh> FCustomizableObjectCompileRunnable::LoadMeshResourceReferenced(int32 ID)
+{
+	MUTABLE_CPUPROFILER_SCOPE(LoadResourceReferenced);
+
+	TSharedPtr<mu::FMesh> Mesh;
+
+	// TODO
+	check(false);
+
+	//if (!ReferencedTextures.IsValidIndex(ID))
+	//{
+	//	// The id is not valid for this CO
+	//	check(false);
+	//	return Mesh;
+	//}
+
+	//// Find the texture id
+	//FMutableSourceMeshData& MeshData = ReferencedMeshes[ID];
+	//check(IsInGameThread());
+
+	//// In the editor the src data can be directly accessed
+	////Mesh = MakeShared<mu::FMesh>();
+	////EUnrealToMutableConversionError Error = ConvertTextureUnrealSourceToMutable(Mesh.Get(), MeshData);
+
+	//Mesh = ConvertSkeletalMeshToMutable(MeshData, Context);
+
+
+	//if (Error != EUnrealToMutableConversionError::Success)
+	//{
+	//	// This could happen in the editor, because some source textures may have changed while there was a background compilation.
+	//	// We just show a warning and move on. This cannot happen during cooks, so it is fine.
+	//	UE_LOG(LogMutable, Warning, TEXT("Failed to load some source mesh data for mesh ID [%d]. Some meshes may be corrupted."), ID);
+	//}
+
+	return Mesh;
+}
+
+
 uint32 FCustomizableObjectCompileRunnable::Run()
 {
 	TRACE_BEGIN_REGION(UE_MUTABLE_CORE_REGION);
@@ -160,19 +198,35 @@ uint32 FCustomizableObjectCompileRunnable::Run()
 		CompilerOptions->SetImagePixelFormatOverride( UnrealPixelFormatFunc );
 	}
 
-	CompilerOptions->SetReferencedResourceCallback([this](int32 ID, TSharedPtr<TSharedPtr<mu::FImage>> ResolvedImage, bool bRunImmediatlyIfPossible)
+	CompilerOptions->SetReferencedResourceCallback(
+		[this](int32 ID, TSharedPtr<TSharedPtr<mu::FImage>> ResolvedImage, bool bRunImmediatlyIfPossible)
 		{
 			UE::Tasks::FTask LaunchTask = UE::Tasks::Launch(TEXT("ConstantGeneratorLaunchTasks"),
 				[ID,ResolvedImage,this]()
 				{
-					TSharedPtr<mu::FImage> Result = LoadResourceReferenced(ID);
+					TSharedPtr<mu::FImage> Result = LoadImageResourceReferenced(ID);
 					*ResolvedImage = Result;
 				},
 				LowLevelTasks::ETaskPriority::BackgroundLow
 			);
 
+			return LaunchTask;
+		},
+
+		[this](int32 ID, TSharedPtr<TSharedPtr<mu::FMesh>> ResolvedMesh, bool bRunImmediatlyIfPossible)
+		{
+			UE::Tasks::FTask LaunchTask = UE::Tasks::Launch(TEXT("ConstantGeneratorLaunchTasks"),
+			[ID, ResolvedMesh, this]()
+				{
+					TSharedPtr<mu::FMesh> Result = LoadMeshResourceReferenced(ID);
+					*ResolvedMesh = Result;
+				},
+				LowLevelTasks::ETaskPriority::BackgroundLow
+			);
+
 			return LaunchTask;
 		}
+
 	);
 
 	const int32 MinResidentMips = UTexture::GetStaticMinTextureResidentMipCount();
@@ -272,7 +326,7 @@ void FCustomizableObjectCompileRunnable::Tick()
 	FReferenceResourceRequest Request;
 	while (PendingResourceReferenceRequests.Dequeue(Request))
 	{
-		*Request.ResolvedImage = LoadResourceReferenced(Request.ID);
+		*Request.ResolvedImage = LoadImageResourceReferenced(Request.ID);
 		Request.CompletionEvent->Trigger();
 
 		// Simple time limit enforcement to avoid blocking the game thread if there are many requests.
diff --git a/Engine/Plugins/Mutable/Source/CustomizableObjectEditor/Private/MuCOE/CustomizableObjectCompileRunnable.h b/Engine/Plugins/Mutable/Source/CustomizableObjectEditor/Private/MuCOE/CustomizableObjectCompileRunnable.h
--- a/Engine/Plugins/Mutable/Source/CustomizableObjectEditor/Private/MuCOE/CustomizableObjectCompileRunnable.h
+++ b/Engine/Plugins/Mutable/Source/CustomizableObjectEditor/Private/MuCOE/CustomizableObjectCompileRunnable.h
@@ -6,6 +6,8 @@
 #include "MuCO/CustomizableObjectPrivate.h"
 #include "MuCO/CustomizableObjectCompilerTypes.h"
 #include "MuCO/UnrealToMutableTextureConversionUtils.h"
+#include "MuR/Image.h"
+#include "MuR/Mesh.h"
 #include "MuR/Ptr.h"
 #include "MuT/Node.h"
 #include "MuCOE/CustomizableObjectEditorLogger.h"
@@ -55,8 +57,8 @@ private:
 	};
 	TQueue<FReferenceResourceRequest, EQueueMode::Mpsc> PendingResourceReferenceRequests;
 
-	TSharedPtr<mu::FImage> LoadResourceReferenced(int32 ID);
-
+	TSharedPtr<mu::FImage> LoadImageResourceReferenced(int32 ID);
+	TSharedPtr<mu::FMesh> LoadMeshResourceReferenced(int32 ID);
 
 public:
 
@@ -80,6 +82,7 @@ public:
 	FCompilationOptions Options;
 	
 	TArray<FMutableSourceTextureData> ReferencedTextures;
+	TArray<FMutableSourceMeshData> ReferencedMeshes;
 
 	FString ErrorMsg;
 
diff --git a/Engine/Plugins/Mutable/Source/CustomizableObjectEditor/Private/MuCOE/CustomizableObjectCompiler.cpp b/Engine/Plugins/Mutable/Source/CustomizableObjectEditor/Private/MuCOE/CustomizableObjectCompiler.cpp
--- a/Engine/Plugins/Mutable/Source/CustomizableObjectEditor/Private/MuCOE/CustomizableObjectCompiler.cpp
+++ b/Engine/Plugins/Mutable/Source/CustomizableObjectEditor/Private/MuCOE/CustomizableObjectCompiler.cpp
@@ -522,9 +522,9 @@ mu::Ptr<mu::NodeObject> GenerateMutableRoot(
 		GenerationContext.NoNameNodeObjectArray.AddUnique(LocalRootNodeObject);
 	}
 
-	if ((Object->GetPrivate()->GetMeshCompileType() == EMutableCompileMeshType::Full) || GenerationContext.Options.bIsCooking)
+	if ((Object->GetPrivate()->GetMeshCompileType() == EMutableCompileMeshType::Full) || GenerationContext.CompilationContext.Options.bIsCooking)
 	{
-		if (LocalRootNodeObject->ParentObject!=nullptr && GenerationContext.Options.bIsCooking)
+		if (LocalRootNodeObject->ParentObject!=nullptr && GenerationContext.CompilationContext.Options.bIsCooking)
 		{
 			// This happens while packaging.
 			return nullptr;
@@ -604,9 +604,9 @@ mu::Ptr<mu::NodeObject> GenerateMutableRoot(
 		}
 	});
 	
-    GenerationContext.RealTimeMorphTargetsOverrides = RootNodeObject->RealTimeMorphSelectionOverrides;
+    GenerationContext.CompilationContext.RealTimeMorphTargetsOverrides = RootNodeObject->RealTimeMorphSelectionOverrides;
 
-	if (!GenerationContext.Options.ParamNamesToSelectedOptions.IsEmpty())
+	if (!GenerationContext.CompilationContext.Options.ParamNamesToSelectedOptions.IsEmpty())
 	{
 		GenerationContext.TableToParamNames = Object->GetPrivate()->GetModelResources().TableToParamNames;
 	}
@@ -670,18 +670,6 @@ mu::Ptr<mu::NodeObject> GenerateMutableRoot(
 }
 
 
-void FCustomizableObjectCompiler::LaunchMutableCompile()
-{
-	AddCompileNotification(LOCTEXT("CustomizableObjectCompileInProgress", "Compiling"));
-
-	// Even for async build, we spawn a thread, so that we can set a large stack. 
-	// Thread names need to be unique, apparently.
-	static int32 ThreadCount = 0;
-	FString ThreadName = FString::Printf(TEXT("MutableCompile-%03d"), ++ThreadCount);
-	CompileThread = MakeShareable(FRunnableThread::Create(CompileTask.Get(), *ThreadName, 16 * 1024 * 1024, TPri_Normal));
-}
-
-
 void FCustomizableObjectCompiler::SaveCODerivedData()
 {
 	if (!SaveDDTask.IsValid())
@@ -741,7 +729,8 @@ void FCustomizableObjectCompiler::CompileInternal(bool bAsync)
 		return;
 	}
 
-	FMutableGraphGenerationContext GenerationContext(CurrentObject, this, CurrentOptions);
+	FMutableCompilationContext CompilationContext(CurrentObject, this, CurrentOptions);
+	FMutableGraphGenerationContext GenerationContext(CompilationContext);
 
 	// Perform a first participating objects pass
 	GenerationContext.ParticipatingObjects = ICustomizableObjectEditorModule::GetChecked().GetParticipatingObjects(CurrentObject, true, &CurrentOptions);
@@ -765,8 +754,8 @@ void FCustomizableObjectCompiler::CompileInternal(bool bAsync)
 		
 		ModelResources.ReferenceSkeletalMeshesData = MoveTemp(GenerationContext.ReferenceSkeletalMeshesData);
 
-		ModelResources.Skeletons.Reserve(GenerationContext.ReferencedSkeletons.Num());
-		for (const USkeleton* Skeleton : GenerationContext.ReferencedSkeletons)
+		ModelResources.Skeletons.Reserve(GenerationContext.CompilationContext.ReferencedSkeletons.Num());
+		for (const USkeleton* Skeleton : GenerationContext.CompilationContext.ReferencedSkeletons)
 		{
 			ModelResources.Skeletons.Emplace(const_cast<USkeleton*>(Skeleton));
 		}
@@ -777,12 +766,18 @@ void FCustomizableObjectCompiler::CompileInternal(bool bAsync)
 			ModelResources.Materials.Emplace(const_cast<UMaterialInterface*>(Material));
 		}
 
-		for (const TPair<TSoftObjectPtr<USkeletalMesh>, FMutableGraphGenerationContext::FGeneratedReferencedMesh>& Pair : GenerationContext.PassthroughMeshMap)
+		for (const TPair<TSoftObjectPtr<UStreamableRenderAsset>, FMutableGraphGenerationContext::FGeneratedReferencedMesh>& Pair : GenerationContext.PassthroughMeshMap)
 		{
 			check(Pair.Value.ID == ModelResources.PassThroughMeshes.Num());
 			ModelResources.PassThroughMeshes.Add(Pair.Key);
 		}
 
+		for (const TPair<TSoftObjectPtr<const UStreamableRenderAsset>, FMutableGraphGenerationContext::FGeneratedReferencedMesh>& Pair : GenerationContext.RuntimeReferencedMeshMap)
+		{
+			check(Pair.Value.ID == ModelResources.RuntimeReferencedMeshes.Num());
+			ModelResources.RuntimeReferencedMeshes.Add(Pair.Key);
+		}
+
 		for (const TPair<TSoftObjectPtr<UTexture>, FMutableGraphGenerationContext::FGeneratedReferencedTexture>& Pair : GenerationContext.PassthroughTextureMap)
 		{
 			check(Pair.Value.ID == ModelResources.PassThroughTextures.Num());
@@ -795,28 +790,28 @@ void FCustomizableObjectCompiler::CompileInternal(bool bAsync)
 			ModelResources.RuntimeReferencedTextures.Add(Pair.Key);
 		}
 		
-		ModelResources.PhysicsAssets = MoveTemp(GenerationContext.PhysicsAssets);
+		ModelResources.PhysicsAssets = MoveTemp(GenerationContext.CompilationContext.PhysicsAssets);
 
 		ModelResources.AnimBPs = MoveTemp(GenerationContext.AnimBPAssets);
-		ModelResources.AnimBpOverridePhysiscAssetsInfo = MoveTemp(GenerationContext.AnimBpOverridePhysicsAssetsInfo);
+		ModelResources.AnimBpOverridePhysiscAssetsInfo = MoveTemp(GenerationContext.CompilationContext.AnimBpOverridePhysicsAssetsInfo);
 
 		ModelResources.MaterialSlotNames = MoveTemp(GenerationContext.ReferencedMaterialSlotNames);
 		ModelResources.SocketArray = MoveTemp(GenerationContext.SocketArray);
 
-		const int32 NumBones = GenerationContext.UniqueBoneNames.Num() + GenerationContext.RemappedBoneNames.Num();
+		const int32 NumBones = GenerationContext.CompilationContext.UniqueBoneNames.Num() + GenerationContext.CompilationContext.RemappedBoneNames.Num();
 		ModelResources.BoneNamesMap.Reserve(NumBones);
 
-		for (auto& It : GenerationContext.UniqueBoneNames)
+		for (auto& It : GenerationContext.CompilationContext.UniqueBoneNames)
 		{
 			ModelResources.BoneNamesMap.Add(It.Value, It.Key.Id);
 		}
 
-		for (auto& It : GenerationContext.RemappedBoneNames)
+		for (auto& It : GenerationContext.CompilationContext.RemappedBoneNames)
 		{
 			ModelResources.BoneNamesMap.Add(It.Key, It.Value.Id);
 		}
 
-		ModelResources.SkinWeightProfilesInfo = MoveTemp(GenerationContext.SkinWeightProfilesInfo);
+		ModelResources.SkinWeightProfilesInfo = MoveTemp(GenerationContext.CompilationContext.SkinWeightProfilesInfo);
 
 		TArray<FGeneratedImageProperties> ImageProperties;
 		GenerationContext.ImageProperties.GenerateValueArray(ImageProperties);
@@ -846,7 +841,7 @@ void FCustomizableObjectCompiler::CompileInternal(bool bAsync)
 
 		// Create the RealTimeMorphsTargets Blocks from the per mesh Morph data.
 		uint64 RealTimeMorphDataSize = 0;
-		for (const TPair<uint32, FRealTimeMorphMeshData>& MeshData : GenerationContext.RealTimeMorphTargetPerMeshData)
+		for (const TPair<uint32, FRealTimeMorphMeshData>& MeshData : GenerationContext.CompilationContext.RealTimeMorphTargetPerMeshData)
 		{
 			RealTimeMorphDataSize += MeshData.Value.Data.Num();
 		}
@@ -855,7 +850,7 @@ void FCustomizableObjectCompiler::CompileInternal(bool bAsync)
 		ModelResources.EditorOnlyMorphTargetReconstructionData.Empty(RealTimeMorphDataSize);
 
 		uint64 RealTimeMorphDataOffsetInBytes = 0;
-		for (const TPair<uint32, FRealTimeMorphMeshData>& MeshData : GenerationContext.RealTimeMorphTargetPerMeshData)
+		for (const TPair<uint32, FRealTimeMorphMeshData>& MeshData : GenerationContext.CompilationContext.RealTimeMorphTargetPerMeshData)
 		{
 			const uint32 DataSizeInBytes = (uint32)MeshData.Value.Data.Num()*sizeof(FMorphTargetVertexData); 
 			FRealTimeMorphStreamable& ResourceMeshData = ModelStreamableBulkData->RealTimeMorphStreamables.FindOrAdd(MeshData.Key);
@@ -875,7 +870,7 @@ void FCustomizableObjectCompiler::CompileInternal(bool bAsync)
 	
 		// Create the Clothing Blocks from the per mesh Morph data.
 		uint64 ClothingDataNum = 0;
-		for (const TPair<uint32, FClothingMeshData>& MeshData : GenerationContext.ClothingPerMeshData)
+		for (const TPair<uint32, FClothingMeshData>& MeshData : GenerationContext.CompilationContext.ClothingPerMeshData)
 		{
 			ClothingDataNum += MeshData.Value.Data.Num();
 		}
@@ -884,7 +879,7 @@ void FCustomizableObjectCompiler::CompileInternal(bool bAsync)
 		ModelResources.EditorOnlyClothingMeshToMeshVertData.Empty(ClothingDataNum);
 
 		uint64 ClothingDataOffsetInBytes = 0;
-		for (const TPair<uint32, FClothingMeshData>& MeshData : GenerationContext.ClothingPerMeshData)
+		for (const TPair<uint32, FClothingMeshData>& MeshData : GenerationContext.CompilationContext.ClothingPerMeshData)
 		{
 			const uint32 DataSizeInBytes = (uint32)MeshData.Value.Data.Num()*sizeof(FCustomizableObjectMeshToMeshVertData); 
 			FClothingStreamable& ResourceMeshData = ModelStreamableBulkData->ClothingStreamables.FindOrAdd(MeshData.Key);
@@ -905,7 +900,7 @@ void FCustomizableObjectCompiler::CompileInternal(bool bAsync)
 			ModelResources.EditorOnlyClothingMeshToMeshVertData.Append(MeshData.Value.Data);
 		}
 
-		ModelResources.ClothingAssetsData = MoveTemp(GenerationContext.ClothingAssetsData);
+		ModelResources.ClothingAssetsData = MoveTemp(GenerationContext.CompilationContext.ClothingAssetsData);
 
 		// A clothing backend, e.g. Chaos cloth, can use 2 config files, one owned by the asset, and another that is shared 
 		// among all assets in a SkeletalMesh. When merging different assets in a skeletalmesh we need to make sure only one of 
@@ -945,18 +940,18 @@ void FCustomizableObjectCompiler::CompileInternal(bool bAsync)
 			 ClothingAssetData.ConfigsData.RemoveAllSwap(IsSharedConfigData);
 		}
 
-		ModelResources.MeshMetadata = MoveTemp(GenerationContext.MeshMetadata);
+		ModelResources.MeshMetadata = MoveTemp(GenerationContext.CompilationContext.MeshMetadata);
 		ModelResources.SurfaceMetadata = MoveTemp(GenerationContext.SurfaceMetadata);
 
 		ModelResources.GroupNodeMap = GenerationContext.GroupNodeMap;
 
 		// If the optimization level is "none" disable texture streaming, because textures are all referenced
 		// unreal assets and progressive generation is not supported.
-		ModelResources.bIsTextureStreamingDisabled = GenerationContext.Options.OptimizationLevel == 0;
+		ModelResources.bIsTextureStreamingDisabled = GenerationContext.CompilationContext.Options.OptimizationLevel == 0;
 		
-		ModelResources.bIsCompiledWithOptimization = GenerationContext.Options.OptimizationLevel == UE_MUTABLE_MAX_OPTIMIZATION;
+		ModelResources.bIsCompiledWithOptimization = GenerationContext.CompilationContext.Options.OptimizationLevel == UE_MUTABLE_MAX_OPTIMIZATION;
 
-		ModelResources.bCompiledWithHDTextureCompression = GenerationContext.Options.TextureCompression == ECustomizableObjectTextureCompression::HighQuality;
+		ModelResources.bCompiledWithHDTextureCompression = GenerationContext.CompilationContext.Options.TextureCompression == ECustomizableObjectTextureCompression::HighQuality;
 
 		ModelResources.AlwaysLoadedExtensionData = MoveTemp(GenerationContext.AlwaysLoadedExtensionData);
 
@@ -1048,9 +1043,20 @@ void FCustomizableObjectCompiler::CompileInternal(bool bAsync)
 			NewCompileTimeReferencedTextures.Add(Tex);
 		}
 
+		TArray<FMutableSourceMeshData> NewCompileTimeReferencedMeshes;
+		for (const TPair<TSoftObjectPtr<const UStreamableRenderAsset>, FMutableGraphGenerationContext::FGeneratedReferencedMesh>& Pair : GenerationContext.CompileTimeMeshMap)
+		{
+			check(Pair.Value.ID == NewCompileTimeReferencedMeshes.Num());
+
+			FMutableSourceMeshData Data;
+			Data.Mesh = Pair.Key;
+			NewCompileTimeReferencedMeshes.Add(Data);
+		}
+
 		CompileTask = MakeShareable(new FCustomizableObjectCompileRunnable(MutableRoot));
 		CompileTask->Options = CurrentOptions;
 		CompileTask->ReferencedTextures = NewCompileTimeReferencedTextures;
+		CompileTask->ReferencedMeshes = NewCompileTimeReferencedMeshes;
 
 		if (!bAsync)
 		{
@@ -1069,7 +1075,13 @@ void FCustomizableObjectCompiler::CompileInternal(bool bAsync)
 		}
 		else
 		{
-			LaunchMutableCompile();
+			AddCompileNotification(LOCTEXT("CustomizableObjectCompileInProgress", "Compiling"));
+
+			// Even for async build, we spawn a thread, so that we can set a large stack. 
+			// Thread names need to be unique, apparently.
+			static int32 ThreadCount = 0;
+			FString ThreadName = FString::Printf(TEXT("MutableCompile-%03d"), ++ThreadCount);
+			CompileThread = MakeShareable(FRunnableThread::Create(CompileTask.Get(), *ThreadName, 16 * 1024 * 1024, TPri_Normal));
 		}
 	}
 
@@ -1241,7 +1253,9 @@ bool FCustomizableObjectCompiler::TryLoadCompiledDataFromDDC(UCustomizableObject
 
 mu::Ptr<mu::Node> FCustomizableObjectCompiler::Export(UCustomizableObject* Object, const FCompilationOptions& InCompilerOptions, 
 	TArray<TSoftObjectPtr<const UTexture>>& OutRuntimeReferencedTextures,
-	TArray<FMutableSourceTextureData>& OutCompilerReferencedTextures )
+	TArray<FMutableSourceTextureData>& OutCompilerReferencedTextures, 
+	TArray<TSoftObjectPtr<const UObject>>& OutRuntimeReferencedMeshes,
+	TArray<FMutableSourceMeshData>& OutCompilerReferencedMeshes)
 {
 	UE_LOG(LogMutable, Log, TEXT("Started Customizable Object Export %s."), *Object->GetName());
 
@@ -1260,7 +1274,8 @@ mu::Ptr<mu::Node> FCustomizableObjectCompiler::Export(UCustomizableObject* Objec
 	CompilerOptions.bPhysicsAssetMergeEnabled = Object->GetPrivate()->IsPhysicsAssetMergeEnabled();
 	CompilerOptions.bAnimBpPhysicsManipulationEnabled = Object->GetPrivate()->IsEnabledAnimBpPhysicsAssetsManipulation();
 	
-	FMutableGraphGenerationContext GenerationContext(Object, this, CompilerOptions);
+	FMutableCompilationContext CompilationContext(Object, this, CompilerOptions);
+	FMutableGraphGenerationContext GenerationContext(CompilationContext);
 
 	GenerationContext.bSkipParticipatingObjectsPass = true;
 	
@@ -1272,7 +1287,7 @@ mu::Ptr<mu::Node> FCustomizableObjectCompiler::Export(UCustomizableObject* Objec
 		return nullptr;
 	}
 
-	// Pass out the references textures
+	// Pass out the referenced textures
 	OutRuntimeReferencedTextures.Empty();
 	for (const TPair<TSoftObjectPtr<const UTexture>, FMutableGraphGenerationContext::FGeneratedReferencedTexture>& Pair : GenerationContext.RuntimeReferencedTextureMap)
 	{
@@ -1289,6 +1304,24 @@ mu::Ptr<mu::Node> FCustomizableObjectCompiler::Export(UCustomizableObject* Objec
 		OutCompilerReferencedTextures.Add(Tex);
 	}
 
+	// Pass out the referenced meshes
+	OutRuntimeReferencedMeshes.Empty();
+	for (const TPair<TSoftObjectPtr<const UStreamableRenderAsset>, FMutableGraphGenerationContext::FGeneratedReferencedMesh>& Pair : GenerationContext.RuntimeReferencedMeshMap)
+	{
+		check(Pair.Value.ID == OutRuntimeReferencedMeshes.Num());
+		OutRuntimeReferencedMeshes.Add(Pair.Key);
+	}
+
+	OutCompilerReferencedMeshes.Empty();
+	for (const TPair<TSoftObjectPtr<const UStreamableRenderAsset>, FMutableGraphGenerationContext::FGeneratedReferencedMesh>& Pair : GenerationContext.CompileTimeMeshMap)
+	{
+		check(Pair.Value.ID == OutCompilerReferencedMeshes.Num());
+
+		FMutableSourceMeshData Data;
+		Data.Mesh = Pair.Key;
+		OutCompilerReferencedMeshes.Add(Data);
+	}
+
 	return MutableRoot;
 }
 
diff --git a/Engine/Plugins/Mutable/Source/CustomizableObjectEditor/Private/MuCOE/CustomizableObjectCompiler.h b/Engine/Plugins/Mutable/Source/CustomizableObjectEditor/Private/MuCOE/CustomizableObjectCompiler.h
--- a/Engine/Plugins/Mutable/Source/CustomizableObjectEditor/Private/MuCOE/CustomizableObjectCompiler.h
+++ b/Engine/Plugins/Mutable/Source/CustomizableObjectEditor/Private/MuCOE/CustomizableObjectCompiler.h
@@ -53,7 +53,9 @@ public:
 	virtual void TickCook(float DeltaTime, bool bCookCompete) override;
 	
 	/** Generate the Mutable Graph from the Unreal Graph. */
-	mu::Ptr<mu::Node> Export(UCustomizableObject* Object, const FCompilationOptions& Options, TArray<TSoftObjectPtr<const UTexture>>& OutRuntimeReferencedTextures, TArray<FMutableSourceTextureData>& OutCompilerReferencedTextures);
+	mu::Ptr<mu::Node> Export(UCustomizableObject* Object, const FCompilationOptions& Options, 
+		TArray<TSoftObjectPtr<const UTexture>>& OutRuntimeReferencedTextures, TArray<FMutableSourceTextureData>& OutCompilerReferencedTextures,
+		TArray<TSoftObjectPtr<const UObject>>& OutRuntimeReferencedMeshes, TArray<FMutableSourceMeshData>& OutCompilerReferencedMeshes);
 
 	void CompilerLog(const FText& Message, const TArray<const UObject*>& UObject, const EMessageSeverity::Type MessageSeverity = EMessageSeverity::Warning, const bool bAddBaseObjectInfo = true, const ELoggerSpamBin SpamBin = ELoggerSpamBin::ShowAll);
 	void CompilerLog(const FText& Message, const UObject* Context = nullptr, const EMessageSeverity::Type MessageSeverity = EMessageSeverity::Warning, const bool bAddBaseObjectInfo = true, const ELoggerSpamBin SpamBin = ELoggerSpamBin::ShowAll);
@@ -102,11 +104,6 @@ private:
 	// Will output to Mutable Log the warning and error messages generated during the CO compilation
 	// and update the values of NumWarnings and NumErrors
 	void UpdateCompilerLogData();
-	
-	/** Launches the compile task in another thread when compiling a CO in the editor
-	* @param bShowNotification [in] whether to show the compiling CO notification or not
-	* @return nothing */
-	void LaunchMutableCompile();
 
 	/** Launches the save derived data task in another thread after compiling a CO in the
 	* editor
diff --git a/Engine/Plugins/Mutable/Source/CustomizableObjectEditor/Private/MuCOE/CustomizableObjectEditor.cpp b/Engine/Plugins/Mutable/Source/CustomizableObjectEditor/Private/MuCOE/CustomizableObjectEditor.cpp
--- a/Engine/Plugins/Mutable/Source/CustomizableObjectEditor/Private/MuCOE/CustomizableObjectEditor.cpp
+++ b/Engine/Plugins/Mutable/Source/CustomizableObjectEditor/Private/MuCOE/CustomizableObjectEditor.cpp
@@ -1043,13 +1043,13 @@ void FCustomizableObjectEditor::ShowGizmoClipMesh(UCustomizableObjectNode& Node,
 
 			if (const UCustomizableObjectNodeStaticMesh* StaticMeshNode = Cast<UCustomizableObjectNodeStaticMesh>(ConnectedNode))
 			{
-				ClipMesh = StaticMeshNode->GetMesh();
+				ClipMesh = StaticMeshNode->GetMesh().LoadSynchronous();
 				StaticMeshNode->GetPinSection(*ConnectedPin, LODIndex, SectionIndex, DummyIndex);
 				MaterialSlotIndex = SectionIndex;
 			}
 			else if (const UCustomizableObjectNodeSkeletalMesh* SkeletalMeshNode = Cast<UCustomizableObjectNodeSkeletalMesh>(ConnectedNode))
 			{
-				ClipMesh = SkeletalMeshNode->GetMesh();
+				ClipMesh = SkeletalMeshNode->GetMesh().LoadSynchronous();
 				SkeletalMeshNode->GetPinSection(*ConnectedPin, LODIndex, SectionIndex, DummyIndex);
 				MaterialSlotIndex = SkeletalMeshNode->GetSkeletalMaterialIndexFor(*ConnectedPin);
 			}
diff --git a/Engine/Plugins/Mutable/Source/CustomizableObjectEditor/Private/MuCOE/CustomizableObjectEditorModule.cpp b/Engine/Plugins/Mutable/Source/CustomizableObjectEditor/Private/MuCOE/CustomizableObjectEditorModule.cpp
--- a/Engine/Plugins/Mutable/Source/CustomizableObjectEditor/Private/MuCOE/CustomizableObjectEditorModule.cpp
+++ b/Engine/Plugins/Mutable/Source/CustomizableObjectEditor/Private/MuCOE/CustomizableObjectEditorModule.cpp
@@ -748,7 +748,8 @@ TMap<FName, FGuid> FCustomizableObjectEditorModule::GetParticipatingObjects(cons
 
 	FCompilationOptions Options = InOptions ? *InOptions : Object->GetPrivate()->GetCompileOptions();
 		
-	FMutableGraphGenerationContext Context(Object, nullptr, Options);
+	FMutableCompilationContext CompilationContext(Object, nullptr, Options);
+	FMutableGraphGenerationContext Context(CompilationContext);
 	Context.bParticipatingObjectsPass = true;
 	Context.bLoadObjects = bLoadObjects;
 
diff --git a/Engine/Plugins/Mutable/Source/CustomizableObjectEditor/Private/MuCOE/CustomizableObjectLayout.cpp b/Engine/Plugins/Mutable/Source/CustomizableObjectEditor/Private/MuCOE/CustomizableObjectLayout.cpp
--- a/Engine/Plugins/Mutable/Source/CustomizableObjectEditor/Private/MuCOE/CustomizableObjectLayout.cpp
+++ b/Engine/Plugins/Mutable/Source/CustomizableObjectEditor/Private/MuCOE/CustomizableObjectLayout.cpp
@@ -84,7 +84,7 @@ void UCustomizableObjectLayout::SetLayoutName(FString Name)
 void UCustomizableObjectLayout::GenerateAutomaticBlocksFromUVs()
 {
 	UCustomizableObjectNode* Node = Cast<UCustomizableObjectNode>(GetOuter());
-	const UObject* Mesh = GetMesh();
+	TSoftObjectPtr<const UObject> Mesh = GetMesh();
 
 	if (!Node || !Mesh)
 	{
@@ -100,8 +100,9 @@ void UCustomizableObjectLayout::GenerateAutomaticBlocksFromUVs()
 	FCustomizableObjectCompiler Compiler;
 	UCustomizableObject* Object = Node->GetGraphEditor()->GetCustomizableObject();
 	FCompilationOptions Options = Object->GetPrivate()->GetCompileOptions();
-
-	FMutableGraphGenerationContext GenerationContext(Object, &Compiler, Options);
+	
+	FMutableCompilationContext CompilationContext(Object, &Compiler, Options);
+	FMutableGraphGenerationContext GenerationContext(CompilationContext);
 	GenerationContext.bSkipParticipatingObjectsPass = true;
 
 	bool bOutWasEmpty = false;
@@ -175,7 +176,7 @@ void UCustomizableObjectLayout::ConsolidateAutomaticBlocks()
 
 void UCustomizableObjectLayout::GetUVs(TArray<FVector2f>& UVs) const
 {
-	if(const UObject* Mesh = GetMesh())
+	if(const UObject* Mesh = GetMesh().LoadSynchronous())
 	{
 		if (const USkeletalMesh* SkeletalMesh = Cast<const USkeletalMesh>(Mesh))
 		{
@@ -215,7 +216,7 @@ void UCustomizableObjectLayout::SetIgnoreWarningsLOD(int32 LODValue)
 }
 
 
-UObject* UCustomizableObjectLayout::GetMesh() const
+TSoftObjectPtr<UObject> UCustomizableObjectLayout::GetMesh() const
 {
 	if (UObject* Node = GetOuter())
 	{
diff --git a/Engine/Plugins/Mutable/Source/CustomizableObjectEditor/Private/MuCOE/ExtensionDataCompilerInterface.cpp b/Engine/Plugins/Mutable/Source/CustomizableObjectEditor/Private/MuCOE/ExtensionDataCompilerInterface.cpp
--- a/Engine/Plugins/Mutable/Source/CustomizableObjectEditor/Private/MuCOE/ExtensionDataCompilerInterface.cpp
+++ b/Engine/Plugins/Mutable/Source/CustomizableObjectEditor/Private/MuCOE/ExtensionDataCompilerInterface.cpp
@@ -49,8 +49,8 @@ TSharedPtr<const mu::FExtensionData> FExtensionDataCompilerInterface::MakeAlways
 
 const UObject* FExtensionDataCompilerInterface::GetOuterForAlwaysLoadedObjects()
 {
-	check(GenerationContext.Object);
-	return GenerationContext.Object;
+	check(GenerationContext.CompilationContext.Object);
+	return GenerationContext.CompilationContext.Object.Get();
 }
 
 void FExtensionDataCompilerInterface::AddGeneratedNode(const UCustomizableObjectNode* InNode)
diff --git a/Engine/Plugins/Mutable/Source/CustomizableObjectEditor/Private/MuCOE/GenerateMutableSource/GenerateMutableSource.cpp b/Engine/Plugins/Mutable/Source/CustomizableObjectEditor/Private/MuCOE/GenerateMutableSource/GenerateMutableSource.cpp
--- a/Engine/Plugins/Mutable/Source/CustomizableObjectEditor/Private/MuCOE/GenerateMutableSource/GenerateMutableSource.cpp
+++ b/Engine/Plugins/Mutable/Source/CustomizableObjectEditor/Private/MuCOE/GenerateMutableSource/GenerateMutableSource.cpp
@@ -93,7 +93,7 @@ bool FGraphCycle::FoundCycle() const
 
 	if (const UCustomizableObject** Result = Context.VisitedPins.Find(Key))
 	{
-		Context.Log(LOCTEXT("CycleFoundNode", "Cycle detected."), &Node, EMessageSeverity::Error, true);
+		Context.CompilationContext.Log(LOCTEXT("CycleFoundNode", "Cycle detected."), &Node, EMessageSeverity::Error, true);
 		Context.CustomizableObjectWithCycle = *Result;
 		return true;
 	}
@@ -123,7 +123,7 @@ void CheckNumOutputs(const UEdGraphPin& Pin, const FMutableGraphGenerationContex
 
 			if (numOutLinks > 1)
 			{
-				GenerationContext.Log(LOCTEXT("MultipleOutgoing", "The node has several outgoing connections, but it should be limited to 1."), CastChecked<UCustomizableObjectNode>(Pin.GetOwningNode()));
+				GenerationContext.CompilationContext.Log(LOCTEXT("MultipleOutgoing", "The node has several outgoing connections, but it should be limited to 1."), CastChecked<UCustomizableObjectNode>(Pin.GetOwningNode()));
 			}
 		}
 	}
@@ -140,8 +140,132 @@ uint32 GetTypeHash(const FGeneratedGroupProjectorsKey& Key)
 }
 
 
-FMutableGraphGenerationContext::FMutableGraphGenerationContext(const UCustomizableObject* InObject, FCustomizableObjectCompiler* InCompiler, const FCompilationOptions& InOptions)
-	: Object(InObject), Compiler(InCompiler), Options(InOptions)
+void FMutableCompilationContext::Log(const FText& Message, const TArray<const UObject*>& Context, const EMessageSeverity::Type MessageSeverity, const bool bAddBaseObjectInfo, const ELoggerSpamBin SpamBin) const
+{
+	if (Compiler)
+	{
+		Compiler->CompilerLog(Message, Context, MessageSeverity, bAddBaseObjectInfo, SpamBin);
+	}
+}
+
+
+void FMutableCompilationContext::Log(const FText& Message, const UObject* Context, const EMessageSeverity::Type MessageSeverity, const bool bAddBaseObjectInfo, const ELoggerSpamBin SpamBin) const
+{
+	Compiler->CompilerLog(Message, Context, MessageSeverity, bAddBaseObjectInfo, SpamBin);
+}
+
+
+mu::FBoneName FMutableCompilationContext::GetBoneUnique(const FName& InBoneName)
+{
+	const FString BoneNameString = InBoneName.ToString().ToLower();
+	mu::FBoneName Bone(CityHash32(reinterpret_cast<const char*>(*BoneNameString), BoneNameString.Len() * sizeof(FString::ElementType)));
+
+	bool bNewRemappedBoneName = false;
+
+	FString& BoneName = UniqueBoneNames.FindOrAdd(Bone, BoneNameString);
+	while (BoneName != InBoneName)
+	{
+		if (mu::FBoneName* RemappedBoneName = RemappedBoneNames.Find(BoneNameString))
+		{
+			Bone.Id = RemappedBoneName->Id;
+			break;
+		}
+
+		// Id collision detected
+		bNewRemappedBoneName = true;
+
+		// Increase Id in an attempt to make it unique again.
+		++Bone.Id;
+		BoneName = UniqueBoneNames.FindOrAdd(Bone, BoneNameString);
+	}
+
+	if (bNewRemappedBoneName)
+	{
+		RemappedBoneNames.Add(BoneNameString, Bone);
+	}
+
+	return Bone;
+}
+
+
+uint32 FMutableCompilationContext::GetSkinWeightProfileIdUnique(const FName InProfileName)
+{
+	const FString ProfileNameString = InProfileName.ToString();
+
+	uint32 UniqueProfileId = CityHash32(reinterpret_cast<const char*>(*ProfileNameString), ProfileNameString.Len() * sizeof(FString::ElementType));
+
+	bool bRemappedProfile = false;
+
+	FName& ProfileName = UniqueSkinWeightProfileIds.FindOrAdd(UniqueProfileId, InProfileName);
+	while (ProfileName != InProfileName)
+	{
+		if (uint32* RemappedProfileId = RemappedSkinWeightProfileIds.Find(InProfileName))
+		{
+			UniqueProfileId = *RemappedProfileId;
+			break;
+		}
+
+		// Id collision detected
+		bRemappedProfile = true;
+
+		// Increase Id in an attempt to make it unique again.
+		++UniqueProfileId;
+		ProfileName = UniqueSkinWeightProfileIds.FindOrAdd(UniqueProfileId, InProfileName);
+	}
+
+	if (bRemappedProfile)
+	{
+		RemappedSkinWeightProfileIds.Add(InProfileName, UniqueProfileId);
+	}
+
+	return UniqueProfileId;
+}
+
+
+FMutableCompilationContext::FMutableCompilationContext(const UCustomizableObject* InObject, class FCustomizableObjectCompiler* InCompiler, const FCompilationOptions& InOptions)
+	: Object(InObject)
+	, Options(InOptions)
+	, Compiler(InCompiler)
+{
+}
+
+
+FMutableComponentInfo* FMutableCompilationContext::GetComponentInfo(FName ComponentName)
+{
+	return ComponentInfos.FindByPredicate([ComponentName](const FMutableComponentInfo& Component)
+		{
+			return Component.ComponentName == ComponentName;
+		});
+}
+
+
+FString FMutableCompilationContext::GetObjectName() const
+{
+	return GetNameSafe(Object.Get());
+}
+
+
+bool FMutableCompilationContext::FindBone(const FName& InBoneName, mu::FBoneName& OutBoneName) const
+{
+	const FString BoneNameString = InBoneName.ToString().ToLower();
+	OutBoneName.Id = CityHash32(reinterpret_cast<const char*>(*BoneNameString), BoneNameString.Len() * sizeof(FString::ElementType));
+	if (UniqueBoneNames.Contains(OutBoneName))
+	{
+		return true;
+	}
+
+	if (const mu::FBoneName* BoneName = RemappedBoneNames.Find(BoneNameString))
+	{
+		OutBoneName = *BoneName;
+		return true;
+	}
+
+	return false;
+}
+
+
+FMutableGraphGenerationContext::FMutableGraphGenerationContext(FMutableCompilationContext& InCompilationContext)
+	: CompilationContext(InCompilationContext)
 	, ExtensionDataCompilerInterface(*this)
 {
 	// Default flags for mesh generation nodes.
@@ -151,7 +275,26 @@ FMutableGraphGenerationContext::FMutableGraphGenerationContext(const UCustomizab
 	LayoutGenerationFlags.Push(FLayoutGenerationFlags());
 }
 
-FMutableGraphGenerationContext::~FMutableGraphGenerationContext() = default;
+
+FString FMutableGraphGenerationContext::GetObjectName()
+{
+	return CompilationContext.GetObjectName();
+}
+
+
+void FMutableGraphGenerationContext::Log(const FText& Message, const TArray<const UObject*>& Context, const EMessageSeverity::Type MessageSeverity, const bool bAddBaseObjectInfo, const ELoggerSpamBin SpamBin) const
+{
+	CompilationContext.Log(Message, Context, MessageSeverity, bAddBaseObjectInfo, SpamBin);
+}
+
+
+void FMutableGraphGenerationContext::Log(const FText& Message, const UObject* Context, const EMessageSeverity::Type MessageSeverity, const bool bAddBaseObjectInfo, const ELoggerSpamBin SpamBin) const
+{
+	if (!bParticipatingObjectsPass)
+	{
+		CompilationContext.Log(Message, Context, MessageSeverity, bAddBaseObjectInfo, SpamBin);
+	}
+}
 
 
 void FMutableGraphGenerationContext::AddParticipatingObject(const FSoftObjectPath& SoftPath)
@@ -168,7 +311,7 @@ void FMutableGraphGenerationContext::AddParticipatingObject(const FSoftObjectPat
 	PRAGMA_DISABLE_DEPRECATION_WARNINGS
 	const FGuid PackageGuid = AssetPackageData.PackageGuid;
 	PRAGMA_ENABLE_DEPRECATION_WARNINGS
-		
+
 	AddParticipatingObjectChecked(SoftPath.GetLongPackageFName(), PackageGuid);
 }
 
@@ -189,20 +332,19 @@ void FMutableGraphGenerationContext::AddParticipatingObject(const UObject& InObj
 }
 
 
-void FMutableGraphGenerationContext::Log(const FText& Message, const TArray<const UObject*>& Context, const EMessageSeverity::Type MessageSeverity, const bool bAddBaseObjectInfo, const ELoggerSpamBin SpamBin) const
+void FMutableGraphGenerationContext::AddParticipatingObjectChecked(const FName& PackageName, const FGuid& PackageGuid)
 {
-	if (Compiler)
+	if (bParticipatingObjectsPass)
 	{
-		Compiler->CompilerLog(Message, Context, MessageSeverity, bAddBaseObjectInfo, SpamBin);
+		ParticipatingObjects.Emplace(PackageName, PackageGuid);
 	}
-}
-
-
-void FMutableGraphGenerationContext::Log(const FText& Message, const UObject* Context, const EMessageSeverity::Type MessageSeverity, const bool bAddBaseObjectInfo, const ELoggerSpamBin SpamBin) const
-{
-	if (!bParticipatingObjectsPass)
+	else
 	{
-		Compiler->CompilerLog(Message, Context, MessageSeverity, bAddBaseObjectInfo, SpamBin);
+		if (!bSkipParticipatingObjectsPass)
+		{
+			FGuid* Result = ParticipatingObjects.Find(PackageName);
+			ensureAlwaysMsgf(Result, TEXT("Participating Object and Generate Source pass missatch. Asset mismatch: %s"), *PackageName.ToString());
+		}
 	}
 }
 
@@ -237,7 +379,7 @@ int32 FMutableGraphGenerationContext::AddStreamedResource(uint32 InResourceHash,
 	}
 	
 	int32 NewResourceIndex = StreamedResourceData.Num();
-	const FString ContainerName = GetNameSafe(Object) + FString::Printf(TEXT("_SR_%d"), NewResourceIndex);
+	const FString ContainerName = GetObjectName() + FString::Printf(TEXT("_SR_%d"), NewResourceIndex);
 
 	OutNewResource = NewObject<UCustomizableObjectResourceDataContainer>(
 		GetTransientPackage(),
@@ -272,40 +414,6 @@ int32 FMutableGraphGenerationContext::AddAssetUserDataToStreamedResources(UAsset
 }
 
 
-uint32 FMutableGraphGenerationContext::GetSkinWeightProfileIdUnique(const FName InProfileName)
-{
-	const FString ProfileNameString = InProfileName.ToString();
-	
-	uint32 UniqueProfileId = CityHash32(reinterpret_cast<const char*>(*ProfileNameString), ProfileNameString.Len() * sizeof(FString::ElementType));
-
-	bool bRemappedProfile = false;
-
-	FName& ProfileName = UniqueSkinWeightProfileIds.FindOrAdd(UniqueProfileId, InProfileName);
-	while (ProfileName != InProfileName)
-	{
-		if (uint32* RemappedProfileId = RemappedSkinWeightProfileIds.Find(InProfileName))
-		{
-			UniqueProfileId = *RemappedProfileId;
-			break;
-		}
-
-		// Id collision detected
-		bRemappedProfile = true;
-
-		// Increase Id in an attempt to make it unique again.
-		++UniqueProfileId;
-		ProfileName = UniqueSkinWeightProfileIds.FindOrAdd(UniqueProfileId, InProfileName);
-	}
-
-	if (bRemappedProfile)
-	{
-		RemappedSkinWeightProfileIds.Add(InProfileName, UniqueProfileId);
-	}
-
-	return UniqueProfileId;
-}
-
-
 const FGuid FMutableGraphGenerationContext::GetNodeIdUnique(const UCustomizableObjectNode* Node)
 {
 	TArray<const UObject*>* ArrayResult = NodeIdsMap.Find(Node->NodeGuid);
@@ -369,57 +477,6 @@ void FMutableGraphGenerationContext::GenerateSharedSurfacesUniqueIds()
 	}
 }
 
-bool FMutableGraphGenerationContext::FindBone(const FName& InBoneName, mu::FBoneName& OutBoneName) const
-{
-	const FString BoneNameString = InBoneName.ToString().ToLower();
-	OutBoneName.Id = CityHash32(reinterpret_cast<const char*>(*BoneNameString), BoneNameString.Len() * sizeof(FString::ElementType));
-	if (UniqueBoneNames.Contains(OutBoneName))
-	{
-		return true;
-	}
-
-	if (const mu::FBoneName* BoneName = RemappedBoneNames.Find(BoneNameString))
-	{
-		OutBoneName = *BoneName;
-		return true;
-	}
-
-	return false;
-}
-
-
-mu::FBoneName FMutableGraphGenerationContext::GetBoneUnique(const FName& InBoneName)
-{
-	const FString BoneNameString = InBoneName.ToString().ToLower();
-	mu::FBoneName Bone(CityHash32(reinterpret_cast<const char*>(*BoneNameString), BoneNameString.Len() * sizeof(FString::ElementType)));
-
-	bool bNewRemappedBoneName = false;
-
-	FString& BoneName = UniqueBoneNames.FindOrAdd(Bone, BoneNameString);
-	while (BoneName != InBoneName)
-	{
-		if (mu::FBoneName* RemappedBoneName = RemappedBoneNames.Find(BoneNameString))
-		{
-			Bone.Id = RemappedBoneName->Id;
-			break;
-		}
-
-		// Id collision detected
-		bNewRemappedBoneName = true;
-
-		// Increase Id in an attempt to make it unique again.
-		++Bone.Id;
-		BoneName = UniqueBoneNames.FindOrAdd(Bone, BoneNameString);
-	}
-
-	if (bNewRemappedBoneName)
-	{
-		RemappedBoneNames.Add(BoneNameString, Bone);
-	}
-
-	return Bone;
-}
-
 
 //void FMutableGraphGenerationContext::CheckPhysicsAssetInSkeletalMesh(const USkeletalMesh* SkeletalMesh)
 //{
@@ -439,10 +496,7 @@ mu::FBoneName FMutableGraphGenerationContext::GetBoneUnique(const FName& InBoneN
 
 FMutableComponentInfo* FMutableGraphGenerationContext::GetCurrentComponentInfo()
 {
-	return ComponentInfos.FindByPredicate([this](const FMutableComponentInfo& Component)
-	{
-		return Component.ComponentName == CurrentMeshComponent;
-	});
+	return CompilationContext.GetComponentInfo(CurrentMeshComponent);
 }
 
 
@@ -464,23 +518,6 @@ UObject* FMutableGraphGenerationContext::LoadObject(const FSoftObjectPtr& SoftOb
 }
 
 
-void FMutableGraphGenerationContext::AddParticipatingObjectChecked(const FName& PackageName, const FGuid& PackageGuid)
-{
-	if (bParticipatingObjectsPass)
-	{
-		ParticipatingObjects.Emplace(PackageName, PackageGuid);
-	}
-	else
-	{
-		if (!bSkipParticipatingObjectsPass)
-		{
-			FGuid* Result = ParticipatingObjects.Find(PackageName);
-			ensureAlwaysMsgf(Result, TEXT("Participating Object and Generate Source pass missatch. Asset mismatch: %s"), *PackageName.ToString());
-		}
-	}
-}
-
-
 FGeneratedKey::FGeneratedKey(void* InFunctionAddress, const UEdGraphPin& InPin, const UCustomizableObjectNode& Node, FMutableGraphGenerationContext& GenerationContext, const bool UseMesh, const bool InbOnlyConnectedLOD)
 {
 	FunctionAddress = InFunctionAddress;
@@ -763,7 +800,7 @@ mu::NodeMeshApplyPosePtr CreateNodeMeshApplyPose(FMutableGraphGenerationContext&
 	for (int32 i = 0; i < ArrayBoneName.Num(); ++i)
 	{
 		const FName BoneName = ArrayBoneName[i];
-		const mu::FBoneName& BoneId = GenerationContext.GetBoneUnique(BoneName);
+		const mu::FBoneName& BoneId = GenerationContext.CompilationContext.GetBoneUnique(BoneName);
 
 		MutableSkeleton->SetDebugName(i, BoneName);
 		MutableSkeleton->SetBoneName(i, { BoneId });
@@ -887,7 +924,7 @@ mu::Ptr<mu::NodeObject> GenerateMutableSource(const UEdGraphPin * Pin, FMutableG
 				ObjectNode->AddStateParam(StateIndex, State.RuntimeParameters[ParamIndex]);
 			}
 
-			const ITargetPlatform* TargetPlatform = GenerationContext.Options.TargetPlatform;
+			const ITargetPlatform* TargetPlatform = GenerationContext.CompilationContext.Options.TargetPlatform;
 
 			int32 NumExtraLODsToBuildAfterFirstLOD = 0;
 			if (const int32* AuxNumExtraLODs = State.NumExtraLODsToBuildPerPlatform.Find(TargetPlatform->PlatformName()))
@@ -933,7 +970,7 @@ mu::Ptr<mu::NodeObject> GenerateMutableSource(const UEdGraphPin * Pin, FMutableG
 					continue;
 				}
 				
-				for (const FMutableComponentInfo& Pair : GenerationContext.ComponentInfos)
+				for (const FMutableComponentInfo& Pair : GenerationContext.CompilationContext.ComponentInfos)
 				{
 					GenerationContext.CurrentMeshComponent = Pair.ComponentName;
 					GenerationContext.CurrentAutoLODStrategy = Pair.Node->AutoLODStrategy;
@@ -1067,7 +1104,7 @@ mu::Ptr<mu::NodeObject> GenerateMutableSource(const UEdGraphPin * Pin, FMutableG
 
 			UCustomizableObjectNodeObject* CustomizableObjectNodeObject = Cast<UCustomizableObjectNodeObject>(ConnectedChildrenPins[ChildIndex]->GetOwningNode());
 
-			const FString* SelectedOptionName = GenerationContext.Options.ParamNamesToSelectedOptions.Find(TypedNodeGroup->GroupName); // If the param is in the map restrict to only the selected option
+			const FString* SelectedOptionName = GenerationContext.CompilationContext.Options.ParamNamesToSelectedOptions.Find(TypedNodeGroup->GroupName); // If the param is in the map restrict to only the selected option
 			mu::NodeObjectPtr ChildNode = nullptr;
 
 			if (bConnectAtLeastTheLastChild || !SelectedOptionName || (CustomizableObjectNodeObject && *SelectedOptionName == CustomizableObjectNodeObject->ObjectName) )
@@ -1173,7 +1210,7 @@ mu::Ptr<mu::NodeObject> GenerateMutableSource(const UEdGraphPin * Pin, FMutableG
 
 			UCustomizableObjectNodeObject* CustomizableObjectNodeObject = Cast<UCustomizableObjectNodeObject>(ExternalChildNode->OutputPin()->GetOwningNode());
 
-			const FString* SelectedOptionName = GenerationContext.Options.ParamNamesToSelectedOptions.Find(TypedNodeGroup->GroupName); // If the param is in the map restrict to only the selected option
+			const FString* SelectedOptionName = GenerationContext.CompilationContext.Options.ParamNamesToSelectedOptions.Find(TypedNodeGroup->GroupName); // If the param is in the map restrict to only the selected option
 			mu::NodeObjectPtr ChildNode = nullptr;
 
 			if (bConnectAtLeastTheLastChild || !SelectedOptionName || (CustomizableObjectNodeObject  && *SelectedOptionName == CustomizableObjectNodeObject->ObjectName) )
@@ -1297,14 +1334,14 @@ FString GenerateGameplayTag(const FString& GameplayTag)
 
 void PopulateReferenceSkeletalMeshesData(FMutableGraphGenerationContext& GenerationContext)
 {
-	const FString PlatformName = GenerationContext.Options.TargetPlatform->IniPlatformName();
+	const FString PlatformName = GenerationContext.CompilationContext.Options.TargetPlatform->IniPlatformName();
 	
-	const uint32 ComponentCount = GenerationContext.ComponentInfos.Num();
+	const uint32 ComponentCount = GenerationContext.CompilationContext.ComponentInfos.Num();
 
 	GenerationContext.ReferenceSkeletalMeshesData.AddDefaulted(ComponentCount);
 	for (uint32 ComponentIndex = 0; ComponentIndex < ComponentCount; ++ComponentIndex)
 	{
-		const FMutableComponentInfo& ComponentInfo = GenerationContext.ComponentInfos[ComponentIndex];
+		const FMutableComponentInfo& ComponentInfo = GenerationContext.CompilationContext.ComponentInfos[ComponentIndex];
 		const FName& ComponentName = ComponentInfo.ComponentName;
 		
 		USkeletalMesh* RefSkeletalMesh = ComponentInfo.RefSkeletalMesh;
@@ -1460,9 +1497,9 @@ uint32 ComputeLODBiasForTexture(const FMutableGraphGenerationContext& Generation
 	constexpr int32 MaxAllowedLODBias = 6;
 
 	// Force a large LODBias for debug
-	if (GenerationContext.Options.bForceLargeLODBias)
+	if (GenerationContext.CompilationContext.Options.bForceLargeLODBias)
 	{
-		return FMath::Min(GenerationContext.Options.DebugBias, MaxAllowedLODBias);
+		return FMath::Min(GenerationContext.CompilationContext.Options.DebugBias, MaxAllowedLODBias);
 	}
 
 	// Max size and number of mips from Texture. 
@@ -1481,7 +1518,7 @@ uint32 ComputeLODBiasForTexture(const FMutableGraphGenerationContext& Generation
 		return 0;
 	}
 
-	const UTextureLODSettings& LODSettings = GenerationContext.Options.TargetPlatform->GetTextureLODSettings();
+	const UTextureLODSettings& LODSettings = GenerationContext.CompilationContext.Options.TargetPlatform->GetTextureLODSettings();
 
 	// Get the MaxTextureSize for the TargetPlatform.
 	const int32 MaxTextureSize = GetMaxTextureSize(ReferenceTexture ? *ReferenceTexture : Texture, LODSettings);
@@ -1597,7 +1634,7 @@ TSharedPtr<mu::FImage> GenerateImageConstant(UTexture* Texture, FMutableGraphGen
 	if (!bIsReference)
 	{
 		bForceLoad = true;
-		if (GenerationContext.Options.OptimizationLevel == 0)
+		if (GenerationContext.CompilationContext.Options.OptimizationLevel == 0)
 		{
 			bIsCompileTime = false;
 		}
@@ -1645,7 +1682,7 @@ TSharedPtr<mu::FImage> GenerateImageConstant(UTexture* Texture, FMutableGraphGen
 }
 
 
-TSharedPtr<mu::FMesh> GenerateMeshConstant(USkeletalMesh* Mesh, FMutableGraphGenerationContext& GenerationContext, bool bIsReference)
+TSharedPtr<mu::FMesh> GenerateMeshConstant(TSoftObjectPtr<UStreamableRenderAsset> Mesh, FMutableGraphGenerationContext& GenerationContext, bool bIsReference)
 {
 	MUTABLE_CPUPROFILER_SCOPE(GenerateMeshConstant)
 
@@ -1659,7 +1696,7 @@ TSharedPtr<mu::FMesh> GenerateMeshConstant(USkeletalMesh* Mesh, FMutableGraphGen
 	if (!bIsReference)
 	{
 		bForceLoad = true;
-		if (GenerationContext.Options.OptimizationLevel == 0)
+		if (GenerationContext.CompilationContext.Options.OptimizationLevel == 0)
 		{
 			bIsCompileTime = false;
 		}
@@ -1681,22 +1718,15 @@ TSharedPtr<mu::FMesh> GenerateMeshConstant(USkeletalMesh* Mesh, FMutableGraphGen
 		Entry = &GenerationContext.PassthroughMeshMap.FindOrAdd(Mesh, InvalidEntry);
 		Num = GenerationContext.PassthroughMeshMap.Num();
 	}
-	// Not supported yet: this will be implemented when mesh conversion happens on-demand in the compilation process
-	// like it happens with images. For now we only implement support for pure pass-through.
-	//else if (bIsCompileTime)
-	//{
-	//	Entry = &GenerationContext.CompileTimeTextureMap.FindOrAdd(Texture, InvalidEntry);
-	//	Num = GenerationContext.CompileTimeTextureMap.Num();
-	//}
-	//else
-	//{
-	//	Entry = &GenerationContext.RuntimeReferencedTextureMap.FindOrAdd(Texture, InvalidEntry);
-	//	Num = GenerationContext.RuntimeReferencedTextureMap.Num();
-	//}
+	else if (bIsCompileTime)
+	{
+		Entry = &GenerationContext.CompileTimeMeshMap.FindOrAdd(Mesh, InvalidEntry);
+		Num = GenerationContext.CompileTimeMeshMap.Num();
+	}
 	else
 	{
-		check(false);
-		return nullptr;
+		Entry = &GenerationContext.RuntimeReferencedMeshMap.FindOrAdd(Mesh, InvalidEntry);
+		Num = GenerationContext.RuntimeReferencedMeshMap.Num();
 	}
 
 	if (Entry->ID == TNumericLimits<uint32>::Max())
diff --git a/Engine/Plugins/Mutable/Source/CustomizableObjectEditor/Private/MuCOE/GenerateMutableSource/GenerateMutableSource.h b/Engine/Plugins/Mutable/Source/CustomizableObjectEditor/Private/MuCOE/GenerateMutableSource/GenerateMutableSource.h
--- a/Engine/Plugins/Mutable/Source/CustomizableObjectEditor/Private/MuCOE/GenerateMutableSource/GenerateMutableSource.h
+++ b/Engine/Plugins/Mutable/Source/CustomizableObjectEditor/Private/MuCOE/GenerateMutableSource/GenerateMutableSource.h
@@ -9,6 +9,7 @@
 #include "MuCO/CustomizableObject.h"
 #include "MuCO/CustomizableObjectIdentifier.h"
 #include "MuCO/CustomizableObjectCompilerTypes.h"
+#include "MuCO/UnrealToMutableTextureConversionUtils.h"
 #include "MuCOE/CustomizableObjectEditorLogger.h"
 #include "MuCOE/ExtensionDataCompilerInterface.h"
 #include "MuCOE/Nodes/CustomizableObjectNodeMaterialBase.h"
@@ -89,24 +90,6 @@
 	int32 TextureSize = 0;
 };
 
-
-// Flags that can influence the mesh conversion
-enum class EMutableMeshConversionFlags : uint32
-{
-	// 
-	None = 0,
-	// Ignore the skeleton and skinning
-	IgnoreSkinning = 1 << 0,
-
-	// Ignore Physics assets
-	IgnorePhysics = 1 << 1,
-
-	// Prevent this mesh generation from adding per mesh metadata. 
-	DoNotCreateMeshMetadata = 1 << 2
-};
-
-ENUM_CLASS_FLAGS(EMutableMeshConversionFlags)
-
 struct FLayoutGenerationFlags
 {
 	bool operator==(const FLayoutGenerationFlags& Other) const = default;
@@ -476,43 +459,99 @@
 uint32 GetTypeHash(const FGeneratedGroupProjectorsKey& Key);
 
 
-struct FMutableGraphGenerationContext
+/** This structure stores the information that is used during the CustomizableObject compilation. 
+* This includes graph generation, core compilation and data storage.
+* This context should have nothing to do with CO-level nodes.
+* It should only be accessed from the game thread.
+*/
+struct FMutableCompilationContext
 {
-	FMutableGraphGenerationContext(const UCustomizableObject* CustomizableObject, class FCustomizableObjectCompiler* InCompiler, const FCompilationOptions& InOptions);
-	~FMutableGraphGenerationContext();
+	/** */
+	TStrongObjectPtr<const UCustomizableObject> Object;
+		
+	/** Compilation options, including target platform. */
+	const FCompilationOptions& Options;
+	/** */
+	TArray<const USkeleton*> ReferencedSkeletons;
 
-	/** See FCustomizableObjectPrivateData::ParticipatingObjects. */
-	void AddParticipatingObject(const FSoftObjectPath& SoftPath);
-	
-	/** See FCustomizableObjectPrivateData::ParticipatingObjects. */
-	void AddParticipatingObject(const FSoftObjectPtr& SoftObject);
+	/** Global morph selection overrides. */
+	TArray<FRealTimeMorphSelectionOverride> RealTimeMorphTargetsOverrides;
 
-	/** See FCustomizableObjectPrivateData::ParticipatingObjects. */
-	template<typename T>
-	void AddParticipatingObject(const TSoftClassPtr<T>& SoftClass)
-	{
-		AddParticipatingObject(SoftClass.ToSoftObjectPath());
-	}
-	
-	/** See FCustomizableObjectPrivateData::ParticipatingObjects. */
-	void AddParticipatingObject(const UObject& Object);
+	/** Data used for MorphTarget reconstruction. */
+	TMap<uint32, FRealTimeMorphMeshData> RealTimeMorphTargetPerMeshData;
 
+	/** Data used for Clothing reconstruction. */
+	TArray<FCustomizableObjectClothingAssetData> ClothingAssetsData;
+	TMap<uint32, FClothingMeshData> ClothingPerMeshData;
+
+	/** Stores the physics assets gathered from the SkeletalMesh nodes during compilation, to be used in mesh generation in-game */
+	TArray<TSoftObjectPtr<UPhysicsAsset>> PhysicsAssets;
+
+	TArray<FAnimBpOverridePhysicsAssetsInfo> AnimBpOverridePhysicsAssetsInfo;
+
+	/** Data used for SkinWeightProfiles reconstruction. */
+	TArray<FMutableSkinWeightProfileInfo> SkinWeightProfilesInfo;
+
+	/** */
+	TMap<uint32, FMutableMeshMetadata> MeshMetadata;
+
+	/** Only Mesh Components (no passthrough). */
+	TArray<FMutableComponentInfo> ComponentInfos;
+
+	/** Array of unique Bone identifiers. */
+	TMap<mu::FBoneName, FString> UniqueBoneNames;
+
+	/** Bone identifiers that had a collision. */
+	TMap<FString, mu::FBoneName> RemappedBoneNames;
+
+private:
+
+	/** Non-owned reference to the compiler object */
+	FCustomizableObjectCompiler* Compiler = nullptr;
+
+	/** */
+	TMap<uint32, FName> UniqueSkinWeightProfileIds;
+	TMap<FName, uint32> RemappedSkinWeightProfileIds;
+
+public:
+
+	FMutableCompilationContext(const UCustomizableObject* InObject, FCustomizableObjectCompiler* InCompiler, const FCompilationOptions& InOptions);
+
+	/** Return the name of the current customizable object, for logging purposes. */
+	FString GetObjectName() const;
+
+	/** Message logging */
 	void Log(const FText& Message, const TArray<const UObject*>& UObject, const EMessageSeverity::Type MessageSeverity = EMessageSeverity::Warning, const bool bAddBaseObjectInfo = true, const ELoggerSpamBin SpamBin = ELoggerSpamBin::ShowAll) const;
-
 	void Log(const FText& Message, const UObject* Context = nullptr, const EMessageSeverity::Type MessageSeverity = EMessageSeverity::Warning, const bool bAddBaseObjectInfo = true, const ELoggerSpamBin SpamBin = ELoggerSpamBin::ShowAll) const;
+
+	/** Component access (for Mesh components only). */
+	FMutableComponentInfo* GetComponentInfo(FName ComponentName);
+
+	/** Get unique identifier for BoneName built from its FString. */
+	mu::FBoneName GetBoneUnique(const FName& BoneName);
+
+	/** */
+	bool FindBone(const FName& BoneName, mu::FBoneName& OutBone) const;
+
+	/** */
+	uint32 GetSkinWeightProfileIdUnique(const FName ProfileName);
+
+};
+
+
+/** This structure stores the information that is used only during the graph generation stage of the CustomizableObject compilation. 
+* Data needed for node graph processing goes here.
+*/
+struct FMutableGraphGenerationContext
+{
+	FMutableGraphGenerationContext(FMutableCompilationContext& InCompilationContext);
 	
-	const UCustomizableObject* Object = nullptr;
+	FMutableCompilationContext& CompilationContext;
 
 	/** Full hierarchy root. */
 	UCustomizableObjectNodeObject* Root = nullptr;
 
-private:
-	// Non-owned reference to the compiler object
-	FCustomizableObjectCompiler* Compiler = nullptr;
-
 public:
-	// Compilation options, including target platform
-	const FCompilationOptions& Options;
 	
 	// Cache of generated pins per LOD
 	TMap<FGeneratedKey, FGeneratedData> Generated;
@@ -521,7 +560,8 @@
 	TSet<UCustomizableObjectNode*> GeneratedNodes;
 
 	/** Struct that stores the relevant information of a data table generated during the compilation. 
-	e.g. all data tables must have the same compilation restrictions */
+	* e.g. all data tables must have the same compilation restrictions 
+	*/
 	struct FGeneratedDataTablesData
 	{
 		// Pointer to the generated mutable Table
@@ -591,7 +631,15 @@
 
 	// Cache of runtime pass-through meshes and their IDs used in the core to identify them.
 	// These meshes will remain as external references even in optimized models.
-	TMap<TSoftObjectPtr<USkeletalMesh>, FGeneratedReferencedMesh> PassthroughMeshMap;
+	TMap<TSoftObjectPtr<UStreamableRenderAsset>, FGeneratedReferencedMesh> PassthroughMeshMap;
 
+	// Cache of runtime meshes and their IDs used in the core to identify them.
+	// These meshes will remain as external references even in optimized models.
+	TMap<TSoftObjectPtr<const UStreamableRenderAsset>, FGeneratedReferencedMesh> RuntimeReferencedMeshMap;
+
+	// Cache of runtime pass-through meshes and their IDs used in the core to identify them
+	// These meshes will become mutable images in the compiled model.
+	TMap<TSoftObjectPtr<const UStreamableRenderAsset>, FGeneratedReferencedMesh> CompileTimeMeshMap;
+
 	// Cache of runtime pass-through images and their IDs used in the core to identify them.
 	// These textures will remain as external references even in optimized models.
@@ -605,16 +653,13 @@
 	// These textures will become mutable images in the compiled model.
 	TMap<TSoftObjectPtr<const UTexture>, FGeneratedReferencedTexture> CompileTimeTextureMap;
 
-    // Global morph selection overrides.
-    TArray<FRealTimeMorphSelectionOverride> RealTimeMorphTargetsOverrides;
-
 	// Mutable meshes already build for source UStaticMesh or USkeletalMesh.
 	struct FGeneratedMeshData
 	{
 		struct FKey
 		{
 			/** Source mesh data. */
-			const UObject* Mesh = nullptr;
+			TSoftObjectPtr<const UStreamableRenderAsset> Mesh;
 			int32 LOD = 0; // Mesh Data LOD (i.e., LOD where we are getting the vertices from)
 			int32 CurrentLOD = 0; // Derived data LOD (i.e., LOD where we are generating the non-Core Data like morphs)
 			int32 MaterialIndex = 0;
@@ -673,6 +718,29 @@
 	 * When a Mesh Section node is found, it will compile all Group Projector nodes in the stack. */
 	TArray<TArray<UCustomizableObjectNodeGroupProjectorParameter*>> CurrentGroupProjectors;
 
+	/** See FCustomizableObjectPrivateData::ParticipatingObjects. */
+	TMap<FName, FGuid> ParticipatingObjects;
+
+	/** Participating objects.
+	* See FCustomizableObjectPrivateData::ParticipatingObjects.
+	*/
+	void AddParticipatingObject(const UObject& Object);
+
+	void AddParticipatingObject(const FSoftObjectPath& SoftPath);
+
+	void AddParticipatingObject(const FSoftObjectPtr& SoftObject);
+
+	template<typename T>
+	void AddParticipatingObject(const TSoftClassPtr<T>& SoftClass)
+	{
+		AddParticipatingObject(SoftClass.ToSoftObjectPath());
+	}
+
+	/** Message logging */
+	FString GetObjectName();
+	void Log(const FText& Message, const TArray<const UObject*>& Context, const EMessageSeverity::Type MessageSeverity = EMessageSeverity::Warning, const bool bAddBaseObjectInfo = true, const ELoggerSpamBin SpamBin = ELoggerSpamBin::ShowAll) const;
+	void Log(const FText& Message, const UObject* Context = nullptr, const EMessageSeverity::Type MessageSeverity = EMessageSeverity::Warning, const bool bAddBaseObjectInfo = true, const ELoggerSpamBin SpamBin = ELoggerSpamBin::ShowAll) const;
+
 	/** Find a mesh if already generated for a given source and flags. */
 	TSharedPtr<mu::FMesh> FindGeneratedMesh(const FGeneratedMeshData::FKey& Key);
 
@@ -685,18 +753,12 @@
 	  * Returns resource index in the array of streamed resources. */
 	int32 AddAssetUserDataToStreamedResources(UAssetUserData* AssetUserData);
 
-	uint32 GetSkinWeightProfileIdUnique(const FName ProfileName);
-
 	// Check if the Id of the node Node already exists, if it's new adds it to NodeIds array, otherwise, returns new Id
 	const FGuid GetNodeIdUnique(const UCustomizableObjectNode* Node);
 
 	/** Generates shared surface IDs for all surface nodes. If one or more nodes are equal, they will use the same SharedSurfaceId */
 	void GenerateSharedSurfacesUniqueIds();
 
-	bool FindBone(const FName& BoneName, mu::FBoneName& OutBone) const;
 
-	/** Get unique identifier for BoneName built from its FString. */
-	mu::FBoneName GetBoneUnique(const FName& BoneName);
-
 	/** Check if the PhysicsAsset of a given SkeletalMesh has any SkeletalBodySetup with BoneNames not present in the
 	* InSkeletalMesh's RefSkeleton, if so, adds the PhysicsAsset to the DiscartedPhysicsAssetMap to display a warning later on */
@@ -723,15 +785,7 @@
 			nullptr;
 	}
 
-private:
-	/** Two operation modes:
-	 * 1. bParticipatingObjectsPass = true, add the participating object.
-	 * 2. bParticipatingObjectsPass = false, check that the participating object was discovered in the participating object pass. */
-	void AddParticipatingObjectChecked(const FName& PackageName, const FGuid& PackageGuid);
-
 public:
-	/** Only Mesh Components (no passthrough). */
-	TArray<FMutableComponentInfo> ComponentInfos;
 
 	/** Only compiled components. All components types. Index is the ObjectComponentIndex. */
 	TArray<FName> ComponentNames;
@@ -749,11 +803,6 @@
 	TMap<TTuple<FString, FString>, TSet<TSoftObjectPtr<UDataTable>>> IntParameterOptionDataTable;
 	//TMap<UPhysicsAsset*, uint32> DiscartedPhysicsAssetMap;
 
-	TArray<const USkeleton*> ReferencedSkeletons;
 
-	// Array of unique Bone identifiers. 
-	TMap<mu::FBoneName, FString> UniqueBoneNames;
-	TMap<FString, mu::FBoneName> RemappedBoneNames; // Bone identifiers that had a collision.
-
 	// Used to aviod Nodes with duplicated ids
 	TMap<FGuid, TArray<const UObject*>> NodeIdsMap;
@@ -762,24 +811,9 @@
 	// For a given material node (the key is node package path + node uid + image index in node) stores images generated for the same node at a higher quality LOD to reuse that image node
 	TMap<FString, FGroupProjectorImageInfo> GroupProjectorLODCache;
 
-	// Data used for MorphTarget reconstruction.
-	TMap<uint32, FRealTimeMorphMeshData> RealTimeMorphTargetPerMeshData;
 
-	// Data used for Clothing reconstruction.
-	TArray<FCustomizableObjectClothingAssetData> ClothingAssetsData;
-	TMap<uint32, FClothingMeshData> ClothingPerMeshData;
-
-	TMap<uint32, FMutableMeshMetadata> MeshMetadata;
 	TMap<uint32, FMutableSurfaceMetadata> SurfaceMetadata;
 
-	// Data used for SkinWeightProfiles reconstruction
-	TArray<FMutableSkinWeightProfileInfo> SkinWeightProfilesInfo;
-
-	TMap<uint32, FName> UniqueSkinWeightProfileIds;
-	TMap<FName, uint32> RemappedSkinWeightProfileIds;
-
-	TArray<FAnimBpOverridePhysicsAssetsInfo> AnimBpOverridePhysicsAssetsInfo;
-
 	uint8 FromLOD = 0; // LOD to append to the CurrentLOD when using AutomaticLODs. 
 	uint8 CurrentLOD = 0;
 	FName CurrentMeshComponent;
@@ -821,9 +855,6 @@
 	TMap<FGraphCycleKey, const UCustomizableObject*> VisitedPins;
 	const UCustomizableObject* CustomizableObjectWithCycle = nullptr;
 
-	/** Stores the physics assets gathered from the SkeletalMesh nodes during compilation, to be used in mesh generation in-game */
-	TArray<TSoftObjectPtr<UPhysicsAsset>> PhysicsAssets;
-
 	/** Stores the anim BP assets gathered from the SkeletalMesh nodes during compilation, to be used in mesh generation in-game */
 	TArray<TSoftClassPtr<UAnimInstance>> AnimBPAssets;
 
@@ -875,14 +906,18 @@
 	TArray<FCustomizableObjectResourceData> AlwaysLoadedExtensionData;
 	TArray<TPair<FName, UCustomizableObjectResourceDataContainer*>> StreamedExtensionData;
 
-	/** See FCustomizableObjectPrivateData::ParticipatingObjects. */
-	TMap<FName, FGuid> ParticipatingObjects;
-
 	/** Map to relate a Composite Data Table Row and its original DataTable */
 	TMap<UDataTable*,TMap<FName, TArray<UDataTable*>>> CompositeDataTableRowToOriginalDataTableMap;
 
 	/** Version Bridge of the root object */
 	TObjectPtr<UObject> RootVersionBridge;
+
+private:
+
+	/** Two operation modes:
+	 * 1. bParticipatingObjectsPass = true, add the participating object.
+	 * 2. bParticipatingObjectsPass = false, check that the participating object was discovered in the participating object pass. */
+	void AddParticipatingObjectChecked(const FName& PackageName, const FGuid& PackageGuid);
 };
 
 
@@ -982,7 +1017,7 @@
 int32 GetTextureSizeInGame(const UTexture2D& Texture, const UTextureLODSettings& LODSettings);
 
 TSharedPtr<mu::FImage> GenerateImageConstant(UTexture*, FMutableGraphGenerationContext&, bool bIsReference);
-TSharedPtr<mu::FMesh> GenerateMeshConstant(USkeletalMesh*, FMutableGraphGenerationContext&, bool bIsReference);
+TSharedPtr<mu::FMesh> GenerateMeshConstant(TSoftObjectPtr<UStreamableRenderAsset> Mesh, FMutableGraphGenerationContext&, bool bIsReference);
 
 /** Generates a mutable image descriptor from an unreal engine texture */
 mu::FImageDesc GenerateImageDescriptor(UTexture* Texture);
diff --git a/Engine/Plugins/Mutable/Source/CustomizableObjectEditor/Private/MuCOE/GenerateMutableSource/GenerateMutableSourceComponent.cpp b/Engine/Plugins/Mutable/Source/CustomizableObjectEditor/Private/MuCOE/GenerateMutableSource/GenerateMutableSourceComponent.cpp
--- a/Engine/Plugins/Mutable/Source/CustomizableObjectEditor/Private/MuCOE/GenerateMutableSource/GenerateMutableSourceComponent.cpp
+++ b/Engine/Plugins/Mutable/Source/CustomizableObjectEditor/Private/MuCOE/GenerateMutableSource/GenerateMutableSourceComponent.cpp
@@ -130,7 +130,7 @@ mu::Ptr<mu::NodeComponent> GenerateMutableSourceComponent(const UEdGraphPin* Pin
 	
 	if (const UCustomizableObjectNodeComponentMesh* TypedComponentMesh = Cast<UCustomizableObjectNodeComponentMesh>(Node))
 	{
-		if (!GenerationContext.ComponentInfos.ContainsByPredicate([&](const FMutableComponentInfo& ComponentInfo)
+		if (!GenerationContext.CompilationContext.ComponentInfos.ContainsByPredicate([&](const FMutableComponentInfo& ComponentInfo)
 		{
 			return ComponentInfo.Node == TypedComponentMesh;
 		}))
@@ -155,7 +155,7 @@ mu::Ptr<mu::NodeComponent> GenerateMutableSourceComponent(const UEdGraphPin* Pin
 
 	else if (const UCustomizableObjectNodeComponentMeshAddTo* TypedComponentMeshExtend = Cast<UCustomizableObjectNodeComponentMeshAddTo>(Node))
 	{
-		if (FMutableComponentInfo* FindResult = GenerationContext.ComponentInfos.FindByPredicate([&](const FMutableComponentInfo& Element)
+		if (FMutableComponentInfo* FindResult = GenerationContext.CompilationContext.ComponentInfos.FindByPredicate([&](const FMutableComponentInfo& Element)
 		{
 			return Element.ComponentName == TypedComponentMeshExtend->ParentComponentName;
 		}))
@@ -419,7 +419,7 @@ void FirstPass(UCustomizableObjectNodeComponentMesh& Node, FMutableGraphGenerati
 		return;
 	}
 	
-	if (FMutableComponentInfo* Result = GenerationContext.ComponentInfos.FindByPredicate([&](const FMutableComponentInfo& Element)
+	if (FMutableComponentInfo* Result = GenerationContext.CompilationContext.ComponentInfos.FindByPredicate([&](const FMutableComponentInfo& Element)
 	{
 		return Element.ComponentName == Node.ComponentName;
 	}))
@@ -465,7 +465,7 @@ void FirstPass(UCustomizableObjectNodeComponentMesh& Node, FMutableGraphGenerati
 		NumLODs = MaxRefMeshLODs;
 	}
 	
-	const FMutableLODSettings& LODSettings = GenerationContext.Object->LODSettings;
+	const FMutableLODSettings& LODSettings = GenerationContext.CompilationContext.Object->LODSettings;
 
 	uint8 FirstLODAvailable = MAX_MESH_LOD_COUNT;
 	
@@ -475,7 +475,7 @@ void FirstPass(UCustomizableObjectNodeComponentMesh& Node, FMutableGraphGenerati
 		// Array of Quality values for each of the scalability settings. The valid values go from [0,4] each one representing one Scalability setting:
 		//		0 = Low
 		//		4 = Cine = Cinematic
-		FSupportedQualityLevelArray SupportedQualityLevels = LODSettings.MinQualityLevelLOD.GetSupportedQualityLevels(*GenerationContext.Options.TargetPlatform->GetPlatformInfo().IniPlatformName.ToString());
+		FSupportedQualityLevelArray SupportedQualityLevels = LODSettings.MinQualityLevelLOD.GetSupportedQualityLevels(*GenerationContext.CompilationContext.Options.TargetPlatform->GetPlatformInfo().IniPlatformName.ToString());
 			
 		// If no scalability settings are found, use the MaxLOD as the MinLOD to be used.
 		int32 MinValue = NumLODs - 1;
@@ -497,7 +497,7 @@ void FirstPass(UCustomizableObjectNodeComponentMesh& Node, FMutableGraphGenerati
 	}
 	else
 	{
-		FirstLODAvailable = LODSettings.MinLOD.GetValueForPlatform(*GenerationContext.Options.TargetPlatform->IniPlatformName());
+		FirstLODAvailable = LODSettings.MinLOD.GetValueForPlatform(*GenerationContext.CompilationContext.Options.TargetPlatform->IniPlatformName());
 	}
 
 	FirstLODAvailable = FMath::Clamp(FirstLODAvailable, 0, NumLODs - 1);
@@ -507,20 +507,20 @@ void FirstPass(UCustomizableObjectNodeComponentMesh& Node, FMutableGraphGenerati
 	// Find the streaming settings for the target platform
 	if (LODSettings.bOverrideLODStreamingSettings)
 	{
-		GenerationContext.bEnableLODStreaming = LODSettings.bEnableLODStreaming.GetValueForPlatform(*GenerationContext.Options.TargetPlatform->IniPlatformName());
-		NumMaxLODsToStream = LODSettings.NumMaxStreamedLODs.GetValueForPlatform(*GenerationContext.Options.TargetPlatform->IniPlatformName());
+		GenerationContext.bEnableLODStreaming = LODSettings.bEnableLODStreaming.GetValueForPlatform(*GenerationContext.CompilationContext.Options.TargetPlatform->IniPlatformName());
+		NumMaxLODsToStream = LODSettings.NumMaxStreamedLODs.GetValueForPlatform(*GenerationContext.CompilationContext.Options.TargetPlatform->IniPlatformName());
 	}
 	else
 	{
-		for (int32 MeshIndex = 0; MeshIndex < GenerationContext.ComponentInfos.Num(); ++MeshIndex)
+		for (int32 MeshIndex = 0; MeshIndex < GenerationContext.CompilationContext.ComponentInfos.Num(); ++MeshIndex)
 		{
-			RefSkeletalMesh = GenerationContext.ComponentInfos[MeshIndex].RefSkeletalMesh;
+			RefSkeletalMesh = GenerationContext.CompilationContext.ComponentInfos[MeshIndex].RefSkeletalMesh;
 			check(RefSkeletalMesh);
 
 			GenerationContext.bEnableLODStreaming = GenerationContext.bEnableLODStreaming &&
-				RefSkeletalMesh->GetEnableLODStreaming(GenerationContext.Options.TargetPlatform);
+				RefSkeletalMesh->GetEnableLODStreaming(GenerationContext.CompilationContext.Options.TargetPlatform);
 
-			NumMaxLODsToStream = FMath::Min(NumMaxLODsToStream, static_cast<uint8>(RefSkeletalMesh->GetMaxNumStreamedLODs(GenerationContext.Options.TargetPlatform)));
+			NumMaxLODsToStream = FMath::Min(NumMaxLODsToStream, static_cast<uint8>(RefSkeletalMesh->GetMaxNumStreamedLODs(GenerationContext.CompilationContext.Options.TargetPlatform)));
 		}
 	}
 
@@ -535,10 +535,10 @@ void FirstPass(UCustomizableObjectNodeComponentMesh& Node, FMutableGraphGenerati
 	ComponentInfo.Node = &Node;
 	ComponentInfo.AccumulateBonesToRemovePerLOD(Node.LODReductionSettings, Node.NumLODs);
 
-	GenerationContext.ComponentInfos.Add(ComponentInfo);
+	GenerationContext.CompilationContext.ComponentInfos.Add(ComponentInfo);
 
 	// Make sure the Skeleton from the reference mesh is added to the list of referenced Skeletons.
-	GenerationContext.ReferencedSkeletons.Add(RefSkeleton);
+	GenerationContext.CompilationContext.ReferencedSkeletons.Add(RefSkeleton);
 
 	GenerationContext.ComponentNames.Add(Node.ComponentName);
 
diff --git a/Engine/Plugins/Mutable/Source/CustomizableObjectEditor/Private/MuCOE/GenerateMutableSource/GenerateMutableSourceGroupProjector.cpp b/Engine/Plugins/Mutable/Source/CustomizableObjectEditor/Private/MuCOE/GenerateMutableSource/GenerateMutableSourceGroupProjector.cpp
--- a/Engine/Plugins/Mutable/Source/CustomizableObjectEditor/Private/MuCOE/GenerateMutableSource/GenerateMutableSourceGroupProjector.cpp
+++ b/Engine/Plugins/Mutable/Source/CustomizableObjectEditor/Private/MuCOE/GenerateMutableSource/GenerateMutableSourceGroupProjector.cpp
@@ -332,6 +332,8 @@ TOptional<FGroupProjectorTempData> GenerateMutableGroupProjector(UCustomizableOb
 
 	FGroupProjectorTempData GroupProjectorTempData;
 	
+	const FCompilationOptions& CompilationOptions = GenerationContext.CompilationContext.Options;
+
 	// The static cast works here because it's already known to be a mu::NodeProjectorParameter* because of the UE5 Cast in the previous line
 	GroupProjectorTempData.NodeProjectorParameterPtr = static_cast<mu::NodeProjectorParameter*>(GenerateMutableSourceProjector(&ProjParamNode->OutputPin(), GenerationContext).get());
 
@@ -526,8 +528,8 @@ TOptional<FGroupProjectorTempData> GenerateMutableGroupProjector(UCustomizableOb
 				if (ProjParamNode->ReferenceTexture)
 				{
 					int32 TotalMips = mu::FImage::GetMipmapCount(ProjParamNode->ReferenceTexture->GetSizeX(), ProjParamNode->ReferenceTexture->GetSizeY());
-					int32 NumMipsBeyondMin = FMath::Max(0, TotalMips - GenerationContext.Options.MinDiskMips);
-					int32 HighResMipsForThisImage = FMath::Min(NumMipsBeyondMin, GenerationContext.Options.NumHighResImageMips);
+					int32 NumMipsBeyondMin = FMath::Max(0, TotalMips - CompilationOptions.MinDiskMips);
+					int32 HighResMipsForThisImage = FMath::Min(NumMipsBeyondMin, CompilationOptions.NumHighResImageMips);
 					ImageNode->SourceDataDescriptor.SourceHighResMips = HighResMipsForThisImage;
 
 					const FString TextureName = GetNameSafe(Texture).ToLower();
diff --git a/Engine/Plugins/Mutable/Source/CustomizableObjectEditor/Private/MuCOE/GenerateMutableSource/GenerateMutableSourceImage.cpp b/Engine/Plugins/Mutable/Source/CustomizableObjectEditor/Private/MuCOE/GenerateMutableSource/GenerateMutableSourceImage.cpp
--- a/Engine/Plugins/Mutable/Source/CustomizableObjectEditor/Private/MuCOE/GenerateMutableSource/GenerateMutableSourceImage.cpp
+++ b/Engine/Plugins/Mutable/Source/CustomizableObjectEditor/Private/MuCOE/GenerateMutableSource/GenerateMutableSourceImage.cpp
@@ -98,7 +98,9 @@ mu::NodeImagePtr GenerateMutableSourceImage(const UEdGraphPin* Pin, FMutableGrap
 	bool bDoNotAddToGeneratedCache = false;
 
 	mu::NodeImagePtr Result;
-	
+
+	const FCompilationOptions& CompilationOptions = GenerationContext.CompilationContext.Options;
+
 	if (const UCustomizableObjectNodeTexture* TypedNodeTex = Cast<UCustomizableObjectNodeTexture>(Node))
 	{
 		UTexture2D* BaseTexture = TypedNodeTex->Texture;
@@ -124,8 +126,8 @@ mu::NodeImagePtr GenerateMutableSourceImage(const UEdGraphPin* Pin, FMutableGrap
 				ConstantImageNode->SetValue( ImageConstant );
 				
 				int32 TotalMips = mu::FImage::GetMipmapCount(ImageConstant->GetSizeX(), ImageConstant->GetSizeY());
-				int32 NumMipsBeyondMin = FMath::Max(0, TotalMips - int32(MipsToSkip) - GenerationContext.Options.MinDiskMips);
-				int32 HighResMipsForThisImage = FMath::Min(NumMipsBeyondMin, GenerationContext.Options.NumHighResImageMips);
+				int32 NumMipsBeyondMin = FMath::Max(0, TotalMips - int32(MipsToSkip) - CompilationOptions.MinDiskMips);
+				int32 HighResMipsForThisImage = FMath::Min(NumMipsBeyondMin, CompilationOptions.NumHighResImageMips);
 				ConstantImageNode->SourceDataDescriptor.SourceHighResMips = HighResMipsForThisImage;
 
 				const FString TextureName = GetNameSafe(BaseTexture).ToLower();
@@ -226,8 +228,8 @@ mu::NodeImagePtr GenerateMutableSourceImage(const UEdGraphPin* Pin, FMutableGrap
 			Result = ResizeTextureByNumMips(ImageNode, MipsToSkip);
 
 			int32 TotalMips = mu::FImage::GetMipmapCount(ImageConstant->GetSizeX(), ImageConstant->GetSizeY());
-			int32 NumMipsBeyondMin = FMath::Max(0, TotalMips - int32(MipsToSkip) - GenerationContext.Options.MinDiskMips);
-			int32 HighResMipsForThisImage = FMath::Min(NumMipsBeyondMin, GenerationContext.Options.NumHighResImageMips);
+			int32 NumMipsBeyondMin = FMath::Max(0, TotalMips - int32(MipsToSkip) - CompilationOptions.MinDiskMips);
+			int32 HighResMipsForThisImage = FMath::Min(NumMipsBeyondMin, CompilationOptions.NumHighResImageMips);
 			ImageNode->SourceDataDescriptor.SourceHighResMips = HighResMipsForThisImage;
 
 			const FString TextureName = GetNameSafe(Texture).ToLower();
@@ -618,7 +620,7 @@ mu::NodeImagePtr GenerateMutableSourceImage(const UEdGraphPin* Pin, FMutableGrap
 		{
 			GenerationContext.AddParticipatingObject(*TypedNodeProject->ReferenceTexture);
 
-			const UTextureLODSettings& TextureLODSettings = GenerationContext.Options.TargetPlatform->GetTextureLODSettings();
+			const UTextureLODSettings& TextureLODSettings = CompilationOptions.TargetPlatform->GetTextureLODSettings();
 			MaxReferenceTextureSizeInGame = GetTextureSizeInGame(*TypedNodeProject->ReferenceTexture, TextureLODSettings);
 		}
 
@@ -851,7 +853,7 @@ mu::NodeImagePtr GenerateMutableSourceImage(const UEdGraphPin* Pin, FMutableGrap
 		{
 			GenerationContext.AddParticipatingObject(*TypedNodeTransform->ReferenceTexture);
 
-			const UTextureLODSettings& TextureLODSettings = GenerationContext.Options.TargetPlatform->GetTextureLODSettings();
+			const UTextureLODSettings& TextureLODSettings = CompilationOptions.TargetPlatform->GetTextureLODSettings();
 			MaxReferenceTextureSizeInGame = GetTextureSizeInGame(*TypedNodeTransform->ReferenceTexture, TextureLODSettings);
 		}
 
@@ -1124,8 +1126,8 @@ mu::NodeImagePtr GenerateMutableSourceImage(const UEdGraphPin* Pin, FMutableGrap
 
 								// Calculate the number of mips to tag as high res for this image.
 								int32 TotalMips = mu::FImage::GetMipmapCount(ImageDesc.m_size[0], ImageDesc.m_size[1]);
-								int32 NumMipsBeyondMin = FMath::Max(0, TotalMips - GenerationContext.Options.MinDiskMips);
-								int32 HighResMipsForThisImage = FMath::Min(NumMipsBeyondMin, GenerationContext.Options.NumHighResImageMips);
+								int32 NumMipsBeyondMin = FMath::Max(0, TotalMips - CompilationOptions.MinDiskMips);
+								int32 HighResMipsForThisImage = FMath::Min(NumMipsBeyondMin, CompilationOptions.NumHighResImageMips);
 								ImageTableNode->SourceDataDescriptor.SourceHighResMips = HighResMipsForThisImage;
 								ImageTableNode->SourceDataDescriptor.SourceId = TableId; // It will be combined with the RowId when generating constants
 							}
diff --git a/Engine/Plugins/Mutable/Source/CustomizableObjectEditor/Private/MuCOE/GenerateMutableSource/GenerateMutableSourceLayout.cpp b/Engine/Plugins/Mutable/Source/CustomizableObjectEditor/Private/MuCOE/GenerateMutableSource/GenerateMutableSourceLayout.cpp
--- a/Engine/Plugins/Mutable/Source/CustomizableObjectEditor/Private/MuCOE/GenerateMutableSource/GenerateMutableSourceLayout.cpp
+++ b/Engine/Plugins/Mutable/Source/CustomizableObjectEditor/Private/MuCOE/GenerateMutableSource/GenerateMutableSourceLayout.cpp
@@ -83,9 +83,10 @@ mu::Ptr<mu::NodeLayout> CreateMutableLayoutNode(FMutableGraphGenerationContext&
 	else
 	{
 		// Convert the UE mesh in the layout into a Mutable mesh.
-		TSharedPtr<mu::FMesh> MutableMesh = nullptr;
+		TSharedPtr<mu::FMesh> MutableMesh;
 
-		if (USkeletalMesh* SkeletalMesh = Cast<USkeletalMesh>(UnrealLayout->GetMesh()))
+		// TODO: Don't force load the mesh
+		if (USkeletalMesh* SkeletalMesh = Cast<USkeletalMesh>(UnrealLayout->GetMesh().LoadSynchronous()))
 		{
 			// We don't need all the data to generate the blocks
 			const EMutableMeshConversionFlags ShapeFlags =
@@ -94,11 +95,12 @@ mu::Ptr<mu::NodeLayout> CreateMutableLayoutNode(FMutableGraphGenerationContext&
 
 			GenerationContext.MeshGenerationFlags.Push(ShapeFlags);
 
-			MutableMesh = ConvertSkeletalMeshToMutable(SkeletalMesh, TSoftClassPtr<UAnimInstance>(), UnrealLayout->GetLOD(), UnrealLayout->GetMaterial(), UnrealLayout->GetLOD(), UnrealLayout->GetMaterial(), 0, GenerationContext, nullptr, nullptr);
+			bool bMeshMustExist = true;
+			MutableMesh = ConvertSkeletalMeshToMutable(SkeletalMesh, TSoftClassPtr<UAnimInstance>(), bMeshMustExist, UnrealLayout->GetLOD(), UnrealLayout->GetMaterial(), 0, GenerationContext, nullptr, nullptr);
 
 			GenerationContext.MeshGenerationFlags.Pop();
 		}
-		else if (UStaticMesh* StaticMesh = Cast<UStaticMesh>(UnrealLayout->GetMesh()))
+		else if (UStaticMesh* StaticMesh = Cast<UStaticMesh>(UnrealLayout->GetMesh().LoadSynchronous()))
 		{
 			MutableMesh = ConvertStaticMeshToMutable(StaticMesh, UnrealLayout->GetLOD(), UnrealLayout->GetMaterial(), GenerationContext, nullptr);
 		}
diff --git a/Engine/Plugins/Mutable/Source/CustomizableObjectEditor/Private/MuCOE/GenerateMutableSource/GenerateMutableSourceMesh.cpp b/Engine/Plugins/Mutable/Source/CustomizableObjectEditor/Private/MuCOE/GenerateMutableSource/GenerateMutableSourceMesh.cpp
--- a/Engine/Plugins/Mutable/Source/CustomizableObjectEditor/Private/MuCOE/GenerateMutableSource/GenerateMutableSourceMesh.cpp
+++ b/Engine/Plugins/Mutable/Source/CustomizableObjectEditor/Private/MuCOE/GenerateMutableSource/GenerateMutableSourceMesh.cpp
@@ -50,6 +50,7 @@
 #include "PhysicsEngine/PhysicsAsset.h"
 #include "PhysicsEngine/SkeletalBodySetup.h"
 #include "Engine/SkinnedAssetCommon.h"
+#include "AssetRegistry/AssetRegistryModule.h"
 
 #include "Hash/CityHash.h"
 
@@ -425,7 +426,7 @@ void SetAndPropagatePoseBoneUsage(mu::FMesh& MutableMesh, int32 PoseIndex, mu::E
 
 }
 
-TArray<TTuple<UPhysicsAsset*, int32>> GetPhysicsAssetsFromAnimInstance(FMutableGraphGenerationContext& GenerationContext, const TSoftClassPtr<UAnimInstance>& AnimInstance)
+TArray<TTuple<UPhysicsAsset*, int32>> GetPhysicsAssetsFromAnimInstance(const TSoftClassPtr<UAnimInstance>& AnimInstance)
 {
 	// TODO: Consider caching the result in the GenerationContext.
 	TArray<TTuple<UPhysicsAsset*, int32>> Result;
@@ -435,7 +436,7 @@ TArray<TTuple<UPhysicsAsset*, int32>> GetPhysicsAssetsFromAnimInstance(FMutableG
 		return Result;
 	}
 
-	UClass* AnimInstanceClass = GenerationContext.LoadClass(AnimInstance);
+	UClass* AnimInstanceClass = AnimInstance.LoadSynchronous();
 	UAnimBlueprintGeneratedClass* AnimClass = Cast<UAnimBlueprintGeneratedClass>(AnimInstanceClass);
 
 	if (AnimClass)
@@ -460,7 +461,26 @@ TArray<TTuple<UPhysicsAsset*, int32>> GetPhysicsAssetsFromAnimInstance(FMutableG
 	return Result;
 }
 
-TArray<uint8> MakePhysicsAssetBodySetupRelevancyMap(const FMutableGraphGenerationContext& GenerationContext, UPhysicsAsset* Asset, const TSharedPtr<mu::FMesh>& Mesh)
+TArray<TTuple<UPhysicsAsset*, int32>> GetPhysicsAssetsFromAnimInstance(FMutableGraphGenerationContext& GenerationContext, const TSoftClassPtr<UAnimInstance>& AnimInstance)
+{
+	// TODO: Consider caching the result in the GenerationContext.
+	TArray<TTuple<UPhysicsAsset*, int32>> Result;
+
+	if (AnimInstance.IsNull())
+	{
+		return Result;
+	}
+
+	UClass* AnimInstanceClass = GenerationContext.LoadClass(AnimInstance);
+	if (AnimInstanceClass)
+	{
+		Result = GetPhysicsAssetsFromAnimInstance(AnimInstance);
+	}
+
+	return Result;
+}
+
+TArray<uint8> MakePhysicsAssetBodySetupRelevancyMap(const FMutableCompilationContext& Context, UPhysicsAsset* Asset, const TSharedPtr<mu::FMesh>& Mesh)
 {
 	const int32 BodySetupsNum = Asset->SkeletalBodySetups.Num();
 
@@ -475,13 +495,13 @@ TArray<uint8> MakePhysicsAssetBodySetupRelevancyMap(const FMutableGraphGeneratio
 	for (int32 BodyIndex = 0; BodyIndex < BodySetupsNum; ++BodyIndex)
 	{
 		mu::FBoneName Bone; 
-		RelevancyMap[BodyIndex] = GenerationContext.FindBone(Asset->SkeletalBodySetups[BodyIndex]->BoneName, Bone);
+		RelevancyMap[BodyIndex] = Context.FindBone(Asset->SkeletalBodySetups[BodyIndex]->BoneName, Bone);
 	}
 
 	return RelevancyMap;
 }
 
-TSharedPtr<mu::FPhysicsBody> MakePhysicsBodyFromAsset(FMutableGraphGenerationContext& GenerationContext, UPhysicsAsset* Asset, const TArray<uint8>& BodySetupRelevancyMap)
+TSharedPtr<mu::FPhysicsBody> MakePhysicsBodyFromAsset(FMutableCompilationContext& Context, UPhysicsAsset* Asset, const TArray<uint8>& BodySetupRelevancyMap)
 {
 	check(Asset);
 	check(Asset->SkeletalBodySetups.Num() == BodySetupRelevancyMap.Num());
@@ -514,7 +534,7 @@ TSharedPtr<mu::FPhysicsBody> MakePhysicsBodyFromAsset(FMutableGraphGenerationCon
 
 		TObjectPtr<USkeletalBodySetup>& BodySetup = SkeletalBodySetups[SourceBodyIndex++];
 
-		const mu::FBoneName& BodyBoneId = GenerationContext.GetBoneUnique(BodySetup->BoneName);
+		const mu::FBoneName& BodyBoneId = Context.GetBoneUnique(BodySetup->BoneName);
 		PhysicsBody->SetBodyBoneId(B, BodyBoneId);
 		
 		const int32 NumSpheres = BodySetup->AggGeom.SphereElems.Num();
@@ -683,37 +703,27 @@ namespace MutablePrivate
 }
 
 
-
 TSharedPtr<mu::FMesh> ConvertSkeletalMeshToMutable(
-		const USkeletalMesh* InSkeletalMesh, const TSoftClassPtr<UAnimInstance>& AnimBp, int32 LODIndexConnected, 
-		int32 SectionIndexConnected, int32 LODIndex, int32 SectionIndex, uint32 SurfaceMetadataId,
-		FMutableGraphGenerationContext& GenerationContext, const UCustomizableObjectNode* CurrentNode, USkeletalMesh* TableReferenceSkeletalMesh)
+	const FMutableSourceMeshData& Source,
+	FMutableCompilationContext& Context
+)
 {
-	MUTABLE_CPUPROFILER_SCOPE(ConvertSkeletalMeshToMutable);
-
-	if(!InSkeletalMesh)
-	{
-		return nullptr;
-	}
-
-	GenerationContext.AddParticipatingObject(*InSkeletalMesh);
-
-	const FString MeshName = *GetNameSafe(InSkeletalMesh).ToLower();
+	const FString MeshName = Source.Mesh.GetLongPackageName().ToLower();
 	const uint32 MeshId = CityHash32(reinterpret_cast<const char*>(*MeshName), MeshName.Len() * sizeof(FString::ElementType));
 
-	const FSkeletalMeshModel* ImportedModel = InSkeletalMesh->GetImportedModel();
+	const USkeletalMesh* SkeletalMesh = Cast<USkeletalMesh>( Source.Mesh.LoadSynchronous() );
+	check(SkeletalMesh);
+	const FSkeletalMeshModel* ImportedModel = SkeletalMesh->GetImportedModel();
 	if (!ImportedModel)
 	{
-		const FString Msg = FString::Printf(TEXT("The SkeletalMesh [%s] doesn't have an imported resource."), *InSkeletalMesh->GetName());
-		GenerationContext.Log(FText::FromString(Msg), CurrentNode);
-
+		const FString Msg = FString::Printf(TEXT("The SkeletalMesh [%s] doesn't have an imported resource."), *MeshName);
+		Context.Log(FText::FromString(Msg), Source.MessageContext);
 		return nullptr;
 	}
 
-	if (!ImportedModel->LODModels.IsValidIndex(LODIndex))
+	if (!ImportedModel->LODModels.IsValidIndex(Source.LODIndex))
 	{
-		if (GenerationContext.CurrentAutoLODStrategy == ECustomizableObjectAutomaticLODStrategy::AutomaticFromMesh && 
-			LODIndex != LODIndexConnected) // If we are using automatic LODs and not generating the base LOD (the connected one) is not an error.
+		if (!Source.bMeshMustExist)
 		{
 			return MakeShared<mu::FMesh>(); // Return empty mesh to preserve the layouts
 		}
@@ -721,20 +731,19 @@ TSharedPtr<mu::FMesh> ConvertSkeletalMeshToMutable(
 		{
 			const FString Msg = FString::Printf(
 				TEXT("The SkeletalMesh [%s] doesn't have the expected number of LODs [need %d, has %d]. Changed after reimporting?"),
-				InSkeletalMesh ? *InSkeletalMesh->GetName() : TEXT("none"),
-				LODIndex + 1,
+				*MeshName,
+				Source.LODIndex + 1,
 				ImportedModel->LODModels.Num());
-			GenerationContext.Log(FText::FromString(Msg), CurrentNode);
+			Context.Log(FText::FromString(Msg), Source.MessageContext);
 
 			return nullptr;
 		}
 	}
 
-	const FSkeletalMeshLODModel& LODModel = ImportedModel->LODModels[LODIndex];
-	if (!LODModel.Sections.IsValidIndex(SectionIndex))
+	const FSkeletalMeshLODModel& LODModel = ImportedModel->LODModels[Source.LODIndex];
+	if (!LODModel.Sections.IsValidIndex(Source.SectionIndex))
 	{
-		if (GenerationContext.CurrentAutoLODStrategy == ECustomizableObjectAutomaticLODStrategy::AutomaticFromMesh && 
-			SectionIndex != SectionIndexConnected) // If we are using automatic LODs and not generating the base LOD (the connected one) is not an error.
+		if (!Source.bMeshMustExist)
 		{
 			return MakeShared<mu::FMesh>(); // Return empty mesh to preserve the layouts
 		}
@@ -742,21 +751,20 @@ TSharedPtr<mu::FMesh> ConvertSkeletalMeshToMutable(
 		{
 			const FString Msg = FString::Printf(
 				TEXT("The SkeletalMesh [%s] doesn't have the expected structure. Maybe the number of LODs [need %d, has %d] or Materials [need %d, has %d] has changed after reimporting?"),
-				*InSkeletalMesh->GetName(),
-				LODIndex + 1,
+				*SkeletalMesh->GetName(),
+				Source.LODIndex + 1,
 				ImportedModel->LODModels.Num(),
-				SectionIndex + 1,
+				Source.SectionIndex + 1,
 				LODModel.Sections.Num());
-			GenerationContext.Log(FText::FromString(Msg), CurrentNode);
-
+			Context.Log(FText::FromString(Msg), Source.MessageContext);
 			return nullptr;
 		}
 	}
 
-	const FSkelMeshSection& MeshSection = LODModel.Sections[SectionIndex];
+	const FSkelMeshSection& MeshSection = LODModel.Sections[Source.SectionIndex];
 
 	// Get the mesh generation flags to use
-	const EMutableMeshConversionFlags CurrentFlags = GenerationContext.MeshGenerationFlags.Last();
+	const EMutableMeshConversionFlags CurrentFlags = Source.Flags;
 	const bool bIgnoreSkeleton = EnumHasAnyFlags(CurrentFlags, EMutableMeshConversionFlags::IgnoreSkinning);
 	const bool bIgnorePhysics = EnumHasAnyFlags(CurrentFlags, EMutableMeshConversionFlags::IgnorePhysics);
 	const bool bDoNotCreateMeshMetadata = EnumHasAnyFlags(CurrentFlags, EMutableMeshConversionFlags::DoNotCreateMeshMetadata);
@@ -770,19 +778,20 @@ TSharedPtr<mu::FMesh> ConvertSkeletalMeshToMutable(
 	// Check if the Skeleton is valid and build the mu::Skeleton
 	if (!bIgnoreSkeleton)
 	{
-		const USkeleton* InSkeleton = InSkeletalMesh->GetSkeleton();
-		if (!InSkeleton)
+		const USkeleton* Skeleton = SkeletalMesh->GetSkeleton();
+		if (!Skeleton)
 		{
-			FString Msg = FString::Printf(TEXT("No skeleton provided when converting SkeletalMesh [%s]."), *InSkeletalMesh->GetName());
-			GenerationContext.Log(FText::FromString(Msg), CurrentNode);
+			FString Msg = FString::Printf(TEXT("No skeleton provided when converting SkeletalMesh [%s]."), *MeshName);
+			Context.Log(FText::FromString(Msg), Source.MessageContext);
 			return nullptr;
 		}
 
-		GenerationContext.AddParticipatingObject(*InSkeleton);
+		// Done while in GenerateContext
+		//Context.AddParticipatingObject(*Skeleton);
 
-		FMutableComponentInfo* MutComponentInfo = GenerationContext.GetCurrentComponentInfo();
+		FMutableComponentInfo* MutComponentInfo = Context.GetComponentInfo(Source.Component);
 		check(MutComponentInfo);
-		
+
 		USkeletalMesh* ComponentRefSkeletalMesh = MutComponentInfo->RefSkeletalMesh;
 		USkeleton* ComponentRefSkeleton = MutComponentInfo->RefSkeleton;
 		check(ComponentRefSkeletalMesh);
@@ -791,32 +800,32 @@ TSharedPtr<mu::FMesh> ConvertSkeletalMeshToMutable(
 		// Compatibility check
 		{
 			FString ErrorMessage;
-			const bool bCompatible = IsSkeletalMeshCompatibleWithRefSkeleton(*MutComponentInfo, InSkeletalMesh, ErrorMessage);
-			MutComponentInfo->SkeletonCompatibility.Add(InSkeleton, bCompatible);
+			const bool bCompatible = IsSkeletalMeshCompatibleWithRefSkeleton(*MutComponentInfo, SkeletalMesh, ErrorMessage);
+			MutComponentInfo->SkeletonCompatibility.Add(Skeleton, bCompatible);
 
 			if (!bCompatible)
 			{
 				if (!ErrorMessage.IsEmpty())
 				{
-					GenerationContext.Log(FText::FromString(ErrorMessage), CurrentNode, EMessageSeverity::Warning);
+					Context.Log(FText::FromString(ErrorMessage), Source.MessageContext, EMessageSeverity::Warning);
 				}
 				return nullptr;
 			}
 
 			// Add the RefSkeleton ID to the mesh.
-			const int32 RefSkeletonID = GenerationContext.ReferencedSkeletons.AddUnique(ComponentRefSkeleton);
+			const int32 RefSkeletonID = Context.ReferencedSkeletons.AddUnique(ComponentRefSkeleton);
 			MutableMesh->AddSkeletonID(RefSkeletonID);
 
 			// Add the skeleton to the list of referenced skeletons and add its index to the mesh
-			const int32 SkeletonID = GenerationContext.ReferencedSkeletons.AddUnique(InSkeleton);
+			const int32 SkeletonID = Context.ReferencedSkeletons.AddUnique(Skeleton);
 			MutableMesh->AddSkeletonID(SkeletonID);
 		}
 
 		// RefSkeleton check
 		{
 			// Ensure the bones used by the Skeletal Mesh exits in the Mesh's Skeleton
-			const TArray<FMeshBoneInfo>& RawRefBoneInfo = InSkeletalMesh->GetRefSkeleton().GetRawRefBoneInfo();
-			const FReferenceSkeleton& InSkeletonRefSkeleton = InSkeleton->GetReferenceSkeleton();
+			const TArray<FMeshBoneInfo>& RawRefBoneInfo = SkeletalMesh->GetRefSkeleton().GetRawRefBoneInfo();
+			const FReferenceSkeleton& InSkeletonRefSkeleton = Skeleton->GetReferenceSkeleton();
 
 			bool bIsSkeletonMissingBones = false;
 
@@ -825,11 +834,11 @@ TSharedPtr<mu::FMesh> ConvertSkeletalMeshToMutable(
 				if (InSkeletonRefSkeleton.FindRawBoneIndex(BoneInfo.Name) == INDEX_NONE)
 				{
 					bIsSkeletonMissingBones = true;
-					UE_LOG(LogMutable, Warning, TEXT("In object [%s] SkeletalMesh [%s] uses bone [%s] not present in skeleton [%s]."),
-						*GenerationContext.Object->GetName(),
-						*InSkeletalMesh->GetName(),
+					FString Msg = FString::Printf(TEXT("SkeletalMesh [%s] uses bone [%s] not present in skeleton [%s]."),
+						*MeshName,
 						*BoneInfo.ExportName,
-						*InSkeleton->GetName());
+						*Skeleton->GetName());
+					Context.Log(FText::FromString(Msg), Source.MessageContext);
 				}
 			}
 
@@ -838,9 +847,8 @@ TSharedPtr<mu::FMesh> ConvertSkeletalMeshToMutable(
 			{
 				FString Msg = FString::Printf(
 					TEXT("The Skeleton [%s] is missing bones that SkeletalMesh [%s] needs. The mesh will be discarded! Information about missing bones can be found in the Output Log."),
-					*InSkeleton->GetName(), *InSkeletalMesh->GetName());
-				
-				GenerationContext.Log(FText::FromString(Msg), CurrentNode, EMessageSeverity::Warning);
+					*Skeleton->GetName(), *MeshName);
+				Context.Log(FText::FromString(Msg), Source.MessageContext);
 
 				return nullptr;
 			}
@@ -850,8 +858,8 @@ TSharedPtr<mu::FMesh> ConvertSkeletalMeshToMutable(
 
 		// Remove bones and build an array to remap indices of the BoneMap
 		TArray<FBoneIndexType> RemappedBones;
-		BuildRemappedBonesArray(*MutComponentInfo, InSkeletalMesh, LODIndex, SourceRequiredBones, RemappedBones);
-		
+		BuildRemappedBonesArray(*MutComponentInfo, SkeletalMesh, Source.LODIndex, SourceRequiredBones, RemappedBones);
+
 		// Build RequiredBones array
 		TArray<FBoneIndexType> RequiredBones;
 		RequiredBones.Reserve(SourceRequiredBones.Num());
@@ -888,11 +896,11 @@ TSharedPtr<mu::FMesh> ConvertSkeletalMeshToMutable(
 		// MutableBoneMap will not keep an index to the Skeleton, but to the BoneName
 		TArray<mu::FBoneName> MutableBoneMap;
 		MutableBoneMap.SetNum(BoneMap.Num());
-		
+
 		TArray<FMatrix> ComposedRefPoseMatrices;
 		ComposedRefPoseMatrices.SetNum(NumRequiredBones);
 
-		const TArray<FMeshBoneInfo>& RefBoneInfo = InSkeletalMesh->GetRefSkeleton().GetRefBoneInfo();
+		const TArray<FMeshBoneInfo>& RefBoneInfo = SkeletalMesh->GetRefSkeleton().GetRefBoneInfo();
 		for (int32 BoneIndex = 0; BoneIndex < NumRequiredBones; ++BoneIndex)
 		{
 			const int32 RefSkeletonBoneIndex = RequiredBones[BoneIndex];
@@ -901,7 +909,7 @@ TSharedPtr<mu::FMesh> ConvertSkeletalMeshToMutable(
 			const int32 ParentBoneIndex = RequiredBones.Find(BoneInfo.ParentIndex);
 
 			// Set bone hierarchy
-			const mu::FBoneName& BoneName = GenerationContext.GetBoneUnique(BoneInfo.Name);
+			const mu::FBoneName& BoneName = Context.GetBoneUnique(BoneInfo.Name);
 
 			MutableSkeleton->SetBoneName(BoneIndex, BoneName);
 			MutableSkeleton->SetBoneParent(BoneIndex, ParentBoneIndex);
@@ -918,11 +926,11 @@ TSharedPtr<mu::FMesh> ConvertSkeletalMeshToMutable(
 
 			if (ParentBoneIndex >= 0)
 			{
-				ComposedRefPoseMatrices[BoneIndex] = InSkeletalMesh->GetRefPoseMatrix(RefSkeletonBoneIndex) * ComposedRefPoseMatrices[ParentBoneIndex];
+				ComposedRefPoseMatrices[BoneIndex] = SkeletalMesh->GetRefPoseMatrix(RefSkeletonBoneIndex) * ComposedRefPoseMatrices[ParentBoneIndex];
 			}
 			else
 			{
-				ComposedRefPoseMatrices[BoneIndex] = InSkeletalMesh->GetRefPoseMatrix(RefSkeletonBoneIndex);
+				ComposedRefPoseMatrices[BoneIndex] = SkeletalMesh->GetRefPoseMatrix(RefSkeletonBoneIndex);
 			}
 
 			// Set bone pose
@@ -950,7 +958,7 @@ TSharedPtr<mu::FMesh> ConvertSkeletalMeshToMutable(
 	MutableMesh->GetVertexBuffers().SetBufferCount(2);
 
 	const int32 MaxSectionInfluences = MeshSection.MaxBoneInfluences;
-	const bool bUseUnlimitedInfluences = FGPUBaseSkinVertexFactory::UseUnlimitedBoneInfluences(MaxSectionInfluences, GenerationContext.Options.TargetPlatform);
+	const bool bUseUnlimitedInfluences = FGPUBaseSkinVertexFactory::UseUnlimitedBoneInfluences(MaxSectionInfluences, Context.Options.TargetPlatform);
 
 	// Create a mutable vertex buffer definition that matches the unreal soft vertex format.
 	mu::FMeshBufferSet UnrealSourceVertexFormat;
@@ -966,10 +974,10 @@ TSharedPtr<mu::FMesh> ConvertSkeletalMeshToMutable(
 		const EMeshBufferFormat Formats[ChannelCount] = { EMeshBufferFormat::Float32, EMeshBufferFormat::Float32, EMeshBufferFormat::Float32, EMeshBufferFormat::Float32, EMeshBufferFormat::Float32, EMeshBufferFormat::Float32, EMeshBufferFormat::Float32, EMeshBufferFormat::Float32, EMeshBufferFormat::NUInt8, EMeshBufferFormat::UInt16, BoneWeightFormat };
 
 		int32 Components[ChannelCount] = { 3, 3, 3, 4, 2, 2, 2, 2, 4, 4, 4 };
-		if (GenerationContext.Options.CustomizableObjectNumBoneInfluences != ECustomizableObjectNumBoneInfluences::Four &&
+		if (Context.Options.CustomizableObjectNumBoneInfluences != ECustomizableObjectNumBoneInfluences::Four &&
 			MaxSectionInfluences > 4)
 		{
-			int32 NewBoneInfluencesNum = (int32)GenerationContext.Options.CustomizableObjectNumBoneInfluences;
+			int32 NewBoneInfluencesNum = (int32)Context.Options.CustomizableObjectNumBoneInfluences;
 
 			if (bUseUnlimitedInfluences &&
 				MaxSectionInfluences < NewBoneInfluencesNum)
@@ -1065,10 +1073,10 @@ TSharedPtr<mu::FMesh> ConvertSkeletalMeshToMutable(
 		const EMeshBufferFormat Formats[ChannelCount] = { EMeshBufferFormat::UInt16, BoneWeightFormat };
 
 		int32 Components[ChannelCount] = { 4, 4 };
-		if (GenerationContext.Options.CustomizableObjectNumBoneInfluences != ECustomizableObjectNumBoneInfluences::Four && 
+		if (Context.Options.CustomizableObjectNumBoneInfluences != ECustomizableObjectNumBoneInfluences::Four &&
 			MaxSectionInfluences > 4)
 		{
-			int32 NewBoneInfluencesNum = (int32)GenerationContext.Options.CustomizableObjectNumBoneInfluences;
+			int32 NewBoneInfluencesNum = (int32)Context.Options.CustomizableObjectNumBoneInfluences;
 
 			if (bUseUnlimitedInfluences &&
 				MaxSectionInfluences < NewBoneInfluencesNum)
@@ -1130,7 +1138,7 @@ TSharedPtr<mu::FMesh> ConvertSkeletalMeshToMutable(
 				TransferRemovedBonesInfluences(&Vertex.InfluenceBones[0], &Vertex.InfluenceWeights[0], MaxSectionInfluences, RemappedBoneMapIndices);
 			}
 
-			if (GenerationContext.Options.CustomizableObjectNumBoneInfluences == ECustomizableObjectNumBoneInfluences::Four)
+			if (Context.Options.CustomizableObjectNumBoneInfluences == ECustomizableObjectNumBoneInfluences::Four)
 			{
 				// Normalize weights
 				const int32 MaxMutableWeights = 4;
@@ -1140,7 +1148,7 @@ TSharedPtr<mu::FMesh> ConvertSkeletalMeshToMutable(
 				NormalizeWeights(&Vertex.InfluenceBones[0], &Vertex.InfluenceWeights[0], MaxSectionInfluences, MaxMutableWeights,
 					&MaxOrderedWeighsIndices[0], MaxSectionBoneMapIndex, MaxBoneWeightValue);
 			}
-			else if (GenerationContext.Options.CustomizableObjectNumBoneInfluences == ECustomizableObjectNumBoneInfluences::Eight)
+			else if (Context.Options.CustomizableObjectNumBoneInfluences == ECustomizableObjectNumBoneInfluences::Eight)
 			{
 				if (!bUseUnlimitedInfluences &&
 					MaxSectionInfluences < EXTRA_BONE_INFLUENCES) // EXTRA_BONE_INFLUENCES is ECustomizableObjectNumBoneInfluences::Eight
@@ -1148,7 +1156,7 @@ TSharedPtr<mu::FMesh> ConvertSkeletalMeshToMutable(
 					FMemory::Memzero(&Vertex.InfluenceWeights[MaxSectionInfluences], EXTRA_BONE_INFLUENCES - MaxSectionInfluences);
 				}
 			}
-			else if (GenerationContext.Options.CustomizableObjectNumBoneInfluences == ECustomizableObjectNumBoneInfluences::Twelve)
+			else if (Context.Options.CustomizableObjectNumBoneInfluences == ECustomizableObjectNumBoneInfluences::Twelve)
 			{
 				if (!bUseUnlimitedInfluences &&
 					MaxSectionInfluences < MAX_TOTAL_INFLUENCES) // MAX_TOTAL_INFLUENCES is ECustomizableObjectNumBoneInfluences::Twelve
@@ -1165,91 +1173,87 @@ TSharedPtr<mu::FMesh> ConvertSkeletalMeshToMutable(
 
 	// TODO: Add Mesh generation flags to not include RT Morph and clothing if not needed.
 	int32 NextBufferIndex = MutableMesh->VertexBuffers.Buffers.Num();
-	if (GenerationContext.Options.bRealTimeMorphTargetsEnabled && CurrentNode)
+	if (Context.Options.bRealTimeMorphTargetsEnabled)
 	{
 
 		// This call involves resolving every TObjectPtr<UMorphTarget> to a UMorphTarget*, so
 		// cache the result here to avoid calling it repeatedly.
-		const TArray<UMorphTarget*>& SkeletalMeshMorphTargets = InSkeletalMesh->GetMorphTargets();
+		const TArray<UMorphTarget*>& SkeletalMeshMorphTargets = SkeletalMesh->GetMorphTargets();
 
 		// Find realtime MorphTargets to be used.
 		TArray<const UMorphTarget*> UsedMorphTargets;
 		UsedMorphTargets.Reserve(SkeletalMeshMorphTargets.Num());
 
-		const UCustomizableObjectNodeSkeletalMesh* NodeTypedSkMesh = Cast<UCustomizableObjectNodeSkeletalMesh>(CurrentNode);
-		const UCustomizableObjectNodeTable* NodeTypedTable = Cast<UCustomizableObjectNodeTable>(CurrentNode);
-		check(NodeTypedSkMesh || (NodeTypedTable));
-
-        // Add SkeletalMesh node used defined realtime morph targets to a temporal array where
-        // the actual to be used real-time morphs names will be placed.        
-        TArray<FName> UsedMorphTargetsNames = Invoke([&]()
-        {
-            TArray<FName> MorphTargetsNames;
-            MorphTargetsNames.Reserve(SkeletalMeshMorphTargets.Num());
-            
-			if (NodeTypedSkMesh)
+		// Add SkeletalMesh node used defined realtime morph targets to a temporal array where
+		// the actual to be used real-time morphs names will be placed.        
+		TArray<FName> UsedMorphTargetsNames = Invoke([&]()
 			{
-				if (NodeTypedSkMesh->bUseAllRealTimeMorphs)
+				TArray<FName> MorphTargetsNames;
+				MorphTargetsNames.Reserve(SkeletalMeshMorphTargets.Num());
+
+				if (!Source.TableReferenceSkeletalMesh)
 				{
-					for (const UMorphTarget* MorphTarget : SkeletalMeshMorphTargets)
+					if (Source.bUseAllRealTimeMorphs)
 					{
-						check(MorphTarget);
-						MorphTargetsNames.Add(MorphTarget->GetFName());
+						for (const UMorphTarget* MorphTarget : SkeletalMeshMorphTargets)
+						{
+							check(MorphTarget);
+							MorphTargetsNames.Add(MorphTarget->GetFName());
+						}
 					}
-				}
-				else
-				{
-					for (const FString& MorphName : NodeTypedSkMesh->UsedRealTimeMorphTargetNames)
+					else
 					{
-						MorphTargetsNames.Emplace(*MorphName);
+						for (const FString& MorphName : Source.UsedRealTimeMorphTargetNames)
+						{
+							MorphTargetsNames.Emplace(*MorphName);
+						}
 					}
 				}
-			}
-			else if (NodeTypedTable && TableReferenceSkeletalMesh)
-			{
-				for (const UMorphTarget* MorphTarget : SkeletalMeshMorphTargets)
+				else
 				{
-					check(MorphTarget);
-
-					if (TableReferenceSkeletalMesh->FindMorphTarget(MorphTarget->GetFName()))
+					for (const UMorphTarget* MorphTarget : SkeletalMeshMorphTargets)
 					{
-						MorphTargetsNames.Add(MorphTarget->GetFName());
+						check(MorphTarget);
+
+						if (Source.TableReferenceSkeletalMesh.LoadSynchronous()->FindMorphTarget(MorphTarget->GetFName()))
+						{
+							MorphTargetsNames.Add(MorphTarget->GetFName());
+						}
 					}
 				}
-			}
 
-            return MorphTargetsNames;
-        });
+				return MorphTargetsNames;
+			});
 
-		TArray<FRealTimeMorphSelectionOverride>& RealTimeMorphTargetOverrides = GenerationContext.RealTimeMorphTargetsOverrides;
-        for (FRealTimeMorphSelectionOverride& MorphTargetOverride : RealTimeMorphTargetOverrides)
-        {
-            const ECustomizableObjectSelectionOverride OverrideValue = 
+		TArray<FRealTimeMorphSelectionOverride>& RealTimeMorphTargetOverrides = Context.RealTimeMorphTargetsOverrides;
+		for (FRealTimeMorphSelectionOverride& MorphTargetOverride : RealTimeMorphTargetOverrides)
+		{
+			const ECustomizableObjectSelectionOverride OverrideValue =
 				Invoke([&]() -> ECustomizableObjectSelectionOverride
-				{
-					const int32 FoundMeshIndex = MorphTargetOverride.SkeletalMeshes.IndexOfByPredicate(
-							[Name = InSkeletalMesh->GetFName()](const FSkeletalMeshMorphTargetOverride& Elem)
+					{
+						const int32 FoundMeshIndex = MorphTargetOverride.SkeletalMeshes.IndexOfByPredicate(
+							[Name = SkeletalMesh->GetFName()](const FSkeletalMeshMorphTargetOverride& Elem)
 							{
 								return Name == Elem.SkeletalMeshName;
 							});
 
-					if (FoundMeshIndex != INDEX_NONE)
-					{
-						return MorphTargetOverride.SkeletalMeshes[FoundMeshIndex].SelectionOverride;
-					}
+						if (FoundMeshIndex != INDEX_NONE)
+						{
+							return MorphTargetOverride.SkeletalMeshes[FoundMeshIndex].SelectionOverride;
+						}
 
-					return MorphTargetOverride.SelectionOverride;
-				});
+						return MorphTargetOverride.SelectionOverride;
+					});
 
-            if (OverrideValue == ECustomizableObjectSelectionOverride::Enable)
-            {
-                UsedMorphTargetsNames.AddUnique(MorphTargetOverride.MorphName);
-            }
-            else if (OverrideValue == ECustomizableObjectSelectionOverride::Disable)
-            {
-                UsedMorphTargetsNames.Remove(MorphTargetOverride.MorphName);
-            }
-        }
+			if (OverrideValue == ECustomizableObjectSelectionOverride::Enable)
+			{
+				UsedMorphTargetsNames.AddUnique(MorphTargetOverride.MorphName);
+			}
+			else if (OverrideValue == ECustomizableObjectSelectionOverride::Disable)
+			{
+				UsedMorphTargetsNames.Remove(MorphTargetOverride.MorphName);
+			}
+		}
 
 
 
@@ -1270,7 +1274,7 @@ TSharedPtr<mu::FMesh> ConvertSkeletalMeshToMutable(
 		if (UsedMorphTargets.Num())
 		{
 			MUTABLE_CPUPROFILER_SCOPE(RealTimeMorphTargetProcessing);
-			
+
 			MutableMesh->GetVertexBuffers().SetBufferCount(NextBufferIndex + 2);
 			// MorphTarget vertex block offset.
 			{
@@ -1304,9 +1308,9 @@ TSharedPtr<mu::FMesh> ConvertSkeletalMeshToMutable(
 
 			// Setup MorphTarget reconstruction data.
 			TArrayView<uint32> VertexMorphsOffsetAndCountBufferView(reinterpret_cast<uint32*>(
-						MutableMesh->GetVertexBuffers().GetBufferData(NextBufferIndex)), VertexCount);
+				MutableMesh->GetVertexBuffers().GetBufferData(NextBufferIndex)), VertexCount);
 			TArrayView<uint32> VertexMorphsResourceIdBufferView(reinterpret_cast<uint32*>(
-						MutableMesh->GetVertexBuffers().GetBufferData(NextBufferIndex + 1)), VertexCount);
+				MutableMesh->GetVertexBuffers().GetBufferData(NextBufferIndex + 1)), VertexCount);
 
 			for (uint32& Elem : VertexMorphsOffsetAndCountBufferView)
 			{
@@ -1318,7 +1322,7 @@ TSharedPtr<mu::FMesh> ConvertSkeletalMeshToMutable(
 			{
 				Elem = InvalidResourceId;
 			}
-			
+
 
 			TArray<FMorphTargetVertexData> MorphsMeshData;
 			MorphsMeshData.Reserve(32);
@@ -1340,32 +1344,32 @@ TSharedPtr<mu::FMesh> ConvertSkeletalMeshToMutable(
 
 					const TArray<FMorphTargetLODModel>& MorphLODModels = MorphTarget->GetMorphLODModels();
 
-					if (LODIndex >= MorphLODModels.Num() || !MorphLODModels[LODIndex].SectionIndices.Contains(SectionIndex))
+					if (Source.LODIndex >= MorphLODModels.Num() || !MorphLODModels[Source.LODIndex].SectionIndices.Contains(Source.SectionIndex))
 					{
 						continue;
 					}
 
 					// The vertices should be sorted by SourceIdx
-					check(MorphLODModels[LODIndex].Vertices.Num() < 2 || MorphLODModels[LODIndex].Vertices[0].SourceIdx < MorphLODModels[LODIndex].Vertices.Last().SourceIdx);
+					check(MorphLODModels[Source.LODIndex].Vertices.Num() < 2 || MorphLODModels[Source.LODIndex].Vertices[0].SourceIdx < MorphLODModels[Source.LODIndex].Vertices.Last().SourceIdx);
 
-					const int32 VertexFoundIndex = Algo::BinarySearchBy(MorphLODModels[LODIndex].Vertices, (uint32)VertexIdx,
-							[](const FMorphTargetDelta& Element) { return Element.SourceIdx; });
+					const int32 VertexFoundIndex = Algo::BinarySearchBy(MorphLODModels[Source.LODIndex].Vertices, (uint32)VertexIdx,
+						[](const FMorphTargetDelta& Element) { return Element.SourceIdx; });
 
 					if (VertexFoundIndex == INDEX_NONE)
 					{
 						continue;
 					}
 
-					const FMorphTargetDelta& VertexFound = MorphLODModels[LODIndex].Vertices[VertexFoundIndex];
+					const FMorphTargetDelta& VertexFound = MorphLODModels[Source.LODIndex].Vertices[VertexFoundIndex];
 					const FName MorphTargetName = MorphTarget->GetFName();
 
 					int32 MorphTargetNameIndex = NameResolutionMap.Find(MorphTargetName);
 
 					MorphTargetNameIndex = MorphTargetNameIndex != INDEX_NONE
-							? MorphTargetNameIndex
-							: NameResolutionMap.Emplace(MorphTargetName);
+						? MorphTargetNameIndex
+						: NameResolutionMap.Emplace(MorphTargetName);
 
-					MorphsUsed.Emplace(FMorphTargetVertexData{VertexFound.PositionDelta, VertexFound.TangentZDelta, (uint32)MorphTargetNameIndex});
+					MorphsUsed.Emplace(FMorphTargetVertexData{ VertexFound.PositionDelta, VertexFound.TangentZDelta, (uint32)MorphTargetNameIndex });
 				}
 
 				if (MorphsUsed.Num())
@@ -1377,10 +1381,10 @@ TSharedPtr<mu::FMesh> ConvertSkeletalMeshToMutable(
 					checkf(MorphsMeshData.Num() < (1 << Log2MaxNumVerts) - 1, TEXT("Maximum number of vertices with morphs per mesh reached."));
 					checkf(MorphsUsed.Num() < (1 << Log2MaxNumMorphs) - 1, TEXT("Maximum number of morphs per vertex reached."));
 
-					uint32 EncodedVertexOffsetAndCount = 
-							(uint32(MorphsMeshData.Num()) & ((1 << Log2MaxNumVerts) - 1)) | (uint32(MorphsUsed.Num()) << Log2MaxNumVerts);
+					uint32 EncodedVertexOffsetAndCount =
+						(uint32(MorphsMeshData.Num()) & ((1 << Log2MaxNumVerts) - 1)) | (uint32(MorphsUsed.Num()) << Log2MaxNumVerts);
 					VertexMorphsOffsetAndCountBufferView[VertexIdx - VertexStart] = EncodedVertexOffsetAndCount;
-					
+
 					MorphsMeshData.Append(MorphsUsed);
 				}
 			}
@@ -1389,20 +1393,20 @@ TSharedPtr<mu::FMesh> ConvertSkeletalMeshToMutable(
 			if (MorphsMeshData.Num())
 			{
 				const uint32 DataHash = CityHash32(
-						reinterpret_cast<const char*>(MorphsMeshData.GetData()), MorphsMeshData.Num() * sizeof(FMorphTargetVertexData));
-				
+					reinterpret_cast<const char*>(MorphsMeshData.GetData()), MorphsMeshData.Num() * sizeof(FMorphTargetVertexData));
+
 				uint32 UniqueHash = DataHash == InvalidResourceId ? InvalidResourceId + 1 : DataHash;
-			
-				const FRealTimeMorphMeshData* FoundHash = GenerationContext.RealTimeMorphTargetPerMeshData.Find(UniqueHash);
+
+				const FRealTimeMorphMeshData* FoundHash = Context.RealTimeMorphTargetPerMeshData.Find(UniqueHash);
 
 				bool bIsDataAlreadyCollected = false;
 				if (FoundHash)
 				{
-					bIsDataAlreadyCollected = 
-							FoundHash->NameResolutionMap == NameResolutionMap &&
-							FoundHash->Data.Num() == MorphsMeshData.Num() &&
-							FMemory::Memcmp(FoundHash->Data.GetData(), MorphsMeshData.GetData(), 
-									MorphsMeshData.Num()*sizeof(FMorphTargetVertexData)) == 0;
+					bIsDataAlreadyCollected =
+						FoundHash->NameResolutionMap == NameResolutionMap &&
+						FoundHash->Data.Num() == MorphsMeshData.Num() &&
+						FMemory::Memcmp(FoundHash->Data.GetData(), MorphsMeshData.GetData(),
+							MorphsMeshData.Num() * sizeof(FMorphTargetVertexData)) == 0;
 				}
 
 				// NOTE: This way of unique hash generation guarantees all valid values can be used but given its 
@@ -1413,18 +1417,18 @@ TSharedPtr<mu::FMesh> ConvertSkeletalMeshToMutable(
 					uint32 NumTries = 0;
 					for (; NumTries < TNumericLimits<uint32>::Max(); ++NumTries)
 					{
-						FoundHash = GenerationContext.RealTimeMorphTargetPerMeshData.Find(UniqueHash);
-						
+						FoundHash = Context.RealTimeMorphTargetPerMeshData.Find(UniqueHash);
+
 						if (!FoundHash)
 						{
 							break;
 						}
 
-						bIsDataAlreadyCollected = 
-								FoundHash->NameResolutionMap == NameResolutionMap &&
-								FoundHash->Data.Num() == MorphsMeshData.Num() &&
-								FMemory::Memcmp(FoundHash->Data.GetData(), MorphsMeshData.GetData(), 
-										MorphsMeshData.Num()*sizeof(FMorphTargetVertexData)) == 0;
+						bIsDataAlreadyCollected =
+							FoundHash->NameResolutionMap == NameResolutionMap &&
+							FoundHash->Data.Num() == MorphsMeshData.Num() &&
+							FMemory::Memcmp(FoundHash->Data.GetData(), MorphsMeshData.GetData(),
+								MorphsMeshData.Num() * sizeof(FMorphTargetVertexData)) == 0;
 
 						if (bIsDataAlreadyCollected)
 						{
@@ -1455,7 +1459,7 @@ TSharedPtr<mu::FMesh> ConvertSkeletalMeshToMutable(
 							VertexMorphsResourceIdBufferView[VertexIdx - VertexStart] = UniqueHash;
 						}
 					}
-				
+
 					FCustomizableObjectStreameableResourceId StreamedMorphResource;
 					StreamedMorphResource.Id = UniqueHash;
 					StreamedMorphResource.Type = static_cast<uint8>(FCustomizableObjectStreameableResourceId::EType::RealTimeMorphTarget);
@@ -1464,8 +1468,7 @@ TSharedPtr<mu::FMesh> ConvertSkeletalMeshToMutable(
 
 					if (!bIsDataAlreadyCollected)
 					{
-						FRealTimeMorphMeshData& RealTimeMorphMeshData = 
-								GenerationContext.RealTimeMorphTargetPerMeshData.FindOrAdd(UniqueHash);
+						FRealTimeMorphMeshData& RealTimeMorphMeshData = Context.RealTimeMorphTargetPerMeshData.FindOrAdd(UniqueHash);
 
 						check(RealTimeMorphMeshData.NameResolutionMap.IsEmpty());
 						check(RealTimeMorphMeshData.Data.IsEmpty());
@@ -1482,7 +1485,7 @@ TSharedPtr<mu::FMesh> ConvertSkeletalMeshToMutable(
 	}
 
 	// Clothing vertex info.
-	if (GenerationContext.Options.bClothingEnabled)
+	if (Context.Options.bClothingEnabled)
 	{
 		// Create new asset or find an already created one if the section has clothing assets.
 		// clothing assets are shared among all LODs in a section
@@ -1490,134 +1493,128 @@ TSharedPtr<mu::FMesh> ConvertSkeletalMeshToMutable(
 		int32 PhysicsAssetIndex = INDEX_NONE;
 
 		Invoke([&](int32& OutClothingAssetIndex, int32& OutPhysicsAssetIndex) -> void
-		{
-			const UClothingAssetBase* ClothingAssetBase = InSkeletalMesh->GetSectionClothingAsset(LODIndex, SectionIndex);
-
-			if (!ClothingAssetBase)
 			{
-				return;
-			}
+				const UClothingAssetBase* ClothingAssetBase = SkeletalMesh->GetSectionClothingAsset(Source.LODIndex, Source.SectionIndex);
 
-			const UClothingAssetCommon* Asset = Cast<UClothingAssetCommon>(ClothingAssetBase);
-			if (!Asset)
-			{
-				return;
-			}
+				if (!ClothingAssetBase)
+				{
+					return;
+				}
 
-			int32 FoundIndex = GenerationContext.ClothingAssetsData.IndexOfByPredicate(
-				[SourceAsset = Asset](const FCustomizableObjectClothingAssetData& Asset)
-				{ 
-					bool bIsSameAsset = 
-							Asset.OriginalAssetGuid == SourceAsset->GetAssetGuid() && 
-						   	Asset.Name == SourceAsset->GetFName() && 
-						   	Asset.ReferenceBoneIndex == SourceAsset->ReferenceBoneIndex &&
-						   	Asset.UsedBoneIndices == SourceAsset->UsedBoneIndices &&
-						   	Asset.UsedBoneNames == SourceAsset->UsedBoneNames &&
-						   	Asset.LodMap == SourceAsset->LodMap &&
-							Asset.LodData.Num() == SourceAsset->LodData.Num();
-					
-					const int32 LodDataNum = Asset.LodData.Num(); 
-					for (int32 LodDataIndex = LodDataNum - 1; LodDataIndex >= 0 && bIsSameAsset; --LodDataIndex)
+				const UClothingAssetCommon* Asset = Cast<UClothingAssetCommon>(ClothingAssetBase);
+				if (!Asset)
+				{
+					return;
+				}
+
+				int32 FoundIndex = Context.ClothingAssetsData.IndexOfByPredicate(
+					[SourceAsset = Asset](const FCustomizableObjectClothingAssetData& Asset)
 					{
-						const FClothPhysicalMeshData& DataA = Asset.LodData[LodDataIndex].PhysicalMeshData; 
-						const FClothPhysicalMeshData& DataB = SourceAsset->LodData[LodDataIndex].PhysicalMeshData; 
-					
-						bIsSameAsset = bIsSameAsset && DataA.NumFixedVerts == DataB.NumFixedVerts;
-						bIsSameAsset = bIsSameAsset && DataA.MaxBoneWeights == DataB.MaxBoneWeights;
-						bIsSameAsset = bIsSameAsset && DataA.Vertices == DataB.Vertices;
-						bIsSameAsset = bIsSameAsset && DataA.Normals == DataB.Normals;
-						bIsSameAsset = bIsSameAsset && DataA.Indices == DataB.Indices;
-						bIsSameAsset = bIsSameAsset && DataA.InverseMasses == DataB.InverseMasses;
-						bIsSameAsset = bIsSameAsset && DataA.EuclideanTethers.Tethers == DataB.EuclideanTethers.Tethers;
-						bIsSameAsset = bIsSameAsset && DataA.GeodesicTethers.Tethers == DataB.GeodesicTethers.Tethers;
-						
-						bIsSameAsset = bIsSameAsset && DataA.WeightMaps.Num() == DataB.WeightMaps.Num();
-						bIsSameAsset = bIsSameAsset && DataA.SelfCollisionVertexSet.Num() == DataB.SelfCollisionVertexSet.Num();
-						bIsSameAsset = bIsSameAsset && DataA.BoneData.Num() == DataB.BoneData.Num();
-						
-						static_assert(std::is_trivially_copyable_v<FClothVertBoneData>);
-						// Assume the FClothVertBoneData does not have any padding. In case there was padding, same assets should
-						// have the same unset memory so false negatives can only happen with different assets that have the 
-						// same data. This reasoning relays on the fact that the data buffers have been copied byte for byte using 
-						// Memcpy or similar.
-						bIsSameAsset = bIsSameAsset && FMemory::Memcmp(
-								DataA.BoneData.GetData(), DataB.BoneData.GetData(), DataA.BoneData.Num()*sizeof(FClothVertBoneData)) == 0;
-
-						if (bIsSameAsset)
+						bool bIsSameAsset =
+							Asset.OriginalAssetGuid == SourceAsset->GetAssetGuid() &&
+							Asset.Name == SourceAsset->GetFName() &&
+							Asset.ReferenceBoneIndex == SourceAsset->ReferenceBoneIndex &&
+							Asset.UsedBoneIndices == SourceAsset->UsedBoneIndices &&
+							Asset.UsedBoneNames == SourceAsset->UsedBoneNames &&
+							Asset.LodMap == SourceAsset->LodMap &&
+							Asset.LodData.Num() == SourceAsset->LodData.Num();
+
+						const int32 LodDataNum = Asset.LodData.Num();
+						for (int32 LodDataIndex = LodDataNum - 1; LodDataIndex >= 0 && bIsSameAsset; --LodDataIndex)
 						{
-							for (const TPair<uint32, FPointWeightMap>& WeightMap : DataA.WeightMaps)
+							const FClothPhysicalMeshData& DataA = Asset.LodData[LodDataIndex].PhysicalMeshData;
+							const FClothPhysicalMeshData& DataB = SourceAsset->LodData[LodDataIndex].PhysicalMeshData;
+
+							bIsSameAsset = bIsSameAsset && DataA.NumFixedVerts == DataB.NumFixedVerts;
+							bIsSameAsset = bIsSameAsset && DataA.MaxBoneWeights == DataB.MaxBoneWeights;
+							bIsSameAsset = bIsSameAsset && DataA.Vertices == DataB.Vertices;
+							bIsSameAsset = bIsSameAsset && DataA.Normals == DataB.Normals;
+							bIsSameAsset = bIsSameAsset && DataA.Indices == DataB.Indices;
+							bIsSameAsset = bIsSameAsset && DataA.InverseMasses == DataB.InverseMasses;
+							bIsSameAsset = bIsSameAsset && DataA.EuclideanTethers.Tethers == DataB.EuclideanTethers.Tethers;
+							bIsSameAsset = bIsSameAsset && DataA.GeodesicTethers.Tethers == DataB.GeodesicTethers.Tethers;
+
+							bIsSameAsset = bIsSameAsset && DataA.WeightMaps.Num() == DataB.WeightMaps.Num();
+							bIsSameAsset = bIsSameAsset && DataA.SelfCollisionVertexSet.Num() == DataB.SelfCollisionVertexSet.Num();
+							bIsSameAsset = bIsSameAsset && DataA.BoneData.Num() == DataB.BoneData.Num();
+
+							static_assert(std::is_trivially_copyable_v<FClothVertBoneData>);
+							// Assume the FClothVertBoneData does not have any padding. In case there was padding, same assets should
+							// have the same unset memory so false negatives can only happen with different assets that have the 
+							// same data. This reasoning relays on the fact that the data buffers have been copied byte for byte using 
+							// Memcpy or similar.
+							bIsSameAsset = bIsSameAsset && FMemory::Memcmp(
+								DataA.BoneData.GetData(), DataB.BoneData.GetData(), DataA.BoneData.Num() * sizeof(FClothVertBoneData)) == 0;
+
+							if (bIsSameAsset)
 							{
-								const FPointWeightMap* FoundWeightMap = DataB.WeightMaps.Find(WeightMap.Key);
-								
-								if (!FoundWeightMap)
+								for (const TPair<uint32, FPointWeightMap>& WeightMap : DataA.WeightMaps)
 								{
-									bIsSameAsset = false;
-									break;
-								}
+									const FPointWeightMap* FoundWeightMap = DataB.WeightMaps.Find(WeightMap.Key);
 
-								if (FoundWeightMap->Values != WeightMap.Value.Values)
-								{
-									bIsSameAsset = false;
-									break;
+									if (!FoundWeightMap)
+									{
+										bIsSameAsset = false;
+										break;
+									}
+
+									if (FoundWeightMap->Values != WeightMap.Value.Values)
+									{
+										bIsSameAsset = false;
+										break;
+									}
 								}
 							}
-						}
 
-						if (bIsSameAsset)
-						{
-							TArray<int32> DataASelfCollisionVertexArray = DataA.SelfCollisionVertexSet.Array();
-							Algo::Sort(DataASelfCollisionVertexArray);
+							if (bIsSameAsset)
+							{
+								TArray<int32> DataASelfCollisionVertexArray = DataA.SelfCollisionVertexSet.Array();
+								Algo::Sort(DataASelfCollisionVertexArray);
 
-							TArray<int32> DataBSelfCollisionVertexArray = DataB.SelfCollisionVertexSet.Array();
-							Algo::Sort(DataBSelfCollisionVertexArray);
+								TArray<int32> DataBSelfCollisionVertexArray = DataB.SelfCollisionVertexSet.Array();
+								Algo::Sort(DataBSelfCollisionVertexArray);
 
-							bIsSameAsset = DataASelfCollisionVertexArray == DataBSelfCollisionVertexArray;
+								bIsSameAsset = DataASelfCollisionVertexArray == DataBSelfCollisionVertexArray;
+							}
 						}
-					}
 
-					return bIsSameAsset;
-				});
-			
-			if (FoundIndex != INDEX_NONE)
-			{
-				OutClothingAssetIndex = FoundIndex;
-				OutPhysicsAssetIndex = GenerationContext.PhysicsAssets.IndexOfByPredicate(
-					[PhysicsAsset = Asset->PhysicsAsset](const TSoftObjectPtr<UPhysicsAsset>& OtherPhysicsAsset)
-					{
-						return TSoftObjectPtr<UPhysicsAsset>(PhysicsAsset) == OtherPhysicsAsset;
+						return bIsSameAsset;
 					});
 
-				return;
-			}
+				if (FoundIndex != INDEX_NONE)
+				{
+					OutClothingAssetIndex = FoundIndex;
+					OutPhysicsAssetIndex = Context.PhysicsAssets.IndexOfByPredicate(
+						[PhysicsAsset = Asset->PhysicsAsset](const TSoftObjectPtr<UPhysicsAsset>& OtherPhysicsAsset)
+						{
+							return TSoftObjectPtr<UPhysicsAsset>(PhysicsAsset) == OtherPhysicsAsset;
+						});
 
-			OutClothingAssetIndex = GenerationContext.ClothingAssetsData.AddDefaulted();
-			FCustomizableObjectClothingAssetData& AssetData = GenerationContext.ClothingAssetsData[OutClothingAssetIndex];
-			
-			AssetData.LodData = Asset->LodData;
-			AssetData.LodMap = Asset->LodMap;
-			AssetData.ReferenceBoneIndex = Asset->ReferenceBoneIndex;
-			AssetData.UsedBoneIndices = Asset->UsedBoneIndices;
-			AssetData.UsedBoneNames = Asset->UsedBoneNames;
-			AssetData.OriginalAssetGuid = Asset->GetAssetGuid();
-			AssetData.Name = Asset->GetFName();
+					return;
+				}
 
-			if (Asset->PhysicsAsset)
-			{
-				OutPhysicsAssetIndex = GenerationContext.PhysicsAssets.AddUnique(Asset->PhysicsAsset);
-				GenerationContext.AddParticipatingObject(*Asset->PhysicsAsset);
-			}
+				OutClothingAssetIndex = Context.ClothingAssetsData.AddDefaulted();
+				FCustomizableObjectClothingAssetData& AssetData = Context.ClothingAssetsData[OutClothingAssetIndex];
 
-			// Store raw clothing config serialized raw data, and info to recreate it afterwards.
-			for (const TPair<FName, TObjectPtr<UClothConfigBase>>& ClothConfig : Asset->ClothConfigs)
-			{
-				FCustomizableObjectClothConfigData& ConfigData = AssetData.ConfigsData.AddDefaulted_GetRef();
-				ConfigData.ClassPath = ClothConfig.Value->GetClass()->GetPathName();
-				ConfigData.ConfigName = ClothConfig.Key;
-				
-				FMemoryWriter MemoryWriter(ConfigData.ConfigBytes);
-                ClothConfig.Value->Serialize(MemoryWriter);
-			}
-		}, ClothingAssetIndex, PhysicsAssetIndex);
+				AssetData.LodData = Asset->LodData;
+				AssetData.LodMap = Asset->LodMap;
+				AssetData.ReferenceBoneIndex = Asset->ReferenceBoneIndex;
+				AssetData.UsedBoneIndices = Asset->UsedBoneIndices;
+				AssetData.UsedBoneNames = Asset->UsedBoneNames;
+				AssetData.OriginalAssetGuid = Asset->GetAssetGuid();
+				AssetData.Name = Asset->GetFName();
+
+				// Store raw clothing config serialized raw data, and info to recreate it afterwards.
+				for (const TPair<FName, TObjectPtr<UClothConfigBase>>& ClothConfig : Asset->ClothConfigs)
+				{
+					FCustomizableObjectClothConfigData& ConfigData = AssetData.ConfigsData.AddDefaulted_GetRef();
+					ConfigData.ClassPath = ClothConfig.Value->GetClass()->GetPathName();
+					ConfigData.ConfigName = ClothConfig.Key;
+
+					FMemoryWriter MemoryWriter(ConfigData.ConfigBytes);
+					ClothConfig.Value->Serialize(MemoryWriter);
+				}
+			}, ClothingAssetIndex, PhysicsAssetIndex);
 
 		if (ClothingAssetIndex != INDEX_NONE)
 		{
@@ -1649,7 +1646,7 @@ TSharedPtr<mu::FMesh> ConvertSkeletalMeshToMutable(
 			}
 
 			TArrayView<int32> ClothSectionIndexView(reinterpret_cast<int32*>(
-						MutableMesh->GetVertexBuffers().GetBufferData(NextBufferIndex)), VertexCount);
+				MutableMesh->GetVertexBuffers().GetBufferData(NextBufferIndex)), VertexCount);
 			for (int32& Elem : ClothSectionIndexView)
 			{
 				Elem = -1;
@@ -1657,7 +1654,7 @@ TSharedPtr<mu::FMesh> ConvertSkeletalMeshToMutable(
 
 			constexpr uint32 InvalidResourceId = 0;
 			TArrayView<uint32> ClothSectionResourceIdView(reinterpret_cast<uint32*>(
-						MutableMesh->GetVertexBuffers().GetBufferData(NextBufferIndex + 1)), VertexCount);
+				MutableMesh->GetVertexBuffers().GetBufferData(NextBufferIndex + 1)), VertexCount);
 			for (uint32& Elem : ClothSectionResourceIdView)
 			{
 				Elem = InvalidResourceId;
@@ -1676,7 +1673,7 @@ TSharedPtr<mu::FMesh> ConvertSkeletalMeshToMutable(
 			// TODO: find a better place to keep this constant.
 			constexpr int32 NumInfluencesPerVertex = 5;
 
-			int32 MeshToMeshDataIndex = 0; 
+			int32 MeshToMeshDataIndex = 0;
 
 			constexpr int32 MaxSupportedInfluences = 1;
 			for (int32& Elem : ClothSectionIndexView)
@@ -1691,11 +1688,11 @@ TSharedPtr<mu::FMesh> ConvertSkeletalMeshToMutable(
 			const int32 NumClothMappingDataVerts = ClothMappingData.Num() / ClothDataStride;
 
 			FClothingMeshData ClothingMeshData;
-			
-			const FCustomizableObjectClothingAssetData& ClothingAssetData = GenerationContext.ClothingAssetsData[ClothingAssetIndex];
-			
+
+			const FCustomizableObjectClothingAssetData& ClothingAssetData = Context.ClothingAssetsData[ClothingAssetIndex];
+
 			ClothingMeshData.ClothingAssetIndex = ClothingAssetIndex;
-			ClothingMeshData.ClothingAssetLOD = ClothingAssetData.LodMap[LODIndex];
+			ClothingMeshData.ClothingAssetLOD = ClothingAssetData.LodMap[Source.LODIndex];
 			ClothingMeshData.PhysicsAssetIndex = PhysicsAssetIndex;
 			ClothingMeshData.Data.Reserve(NumClothMappingDataVerts * ClothDataStride);
 
@@ -1718,27 +1715,27 @@ TSharedPtr<mu::FMesh> ConvertSkeletalMeshToMutable(
 			}
 
 			const uint32 DataHash = CityHash32(
-					reinterpret_cast<const char*>(ClothingMeshData.Data.GetData()), 
-					ClothingMeshData.Data.Num()*sizeof(FCustomizableObjectMeshToMeshVertData));
+				reinterpret_cast<const char*>(ClothingMeshData.Data.GetData()),
+				ClothingMeshData.Data.Num() * sizeof(FCustomizableObjectMeshToMeshVertData));
 
 			uint32 UniqueHash = DataHash == InvalidResourceId ? DataHash + 1 : DataHash;
 
-			const FClothingMeshData* FoundHash = GenerationContext.ClothingPerMeshData.Find(UniqueHash);
+			const FClothingMeshData* FoundHash = Context.ClothingPerMeshData.Find(UniqueHash);
 
 			bool bIsDataAlreadyCollected = false;
-			
+
 			auto CompareClothingData = [](const FClothingMeshData& A, const FClothingMeshData& B) -> bool
-			{
-				return A.ClothingAssetIndex == B.ClothingAssetIndex &&
-					   A.ClothingAssetLOD == B.ClothingAssetLOD &&
-					   A.PhysicsAssetIndex == B.PhysicsAssetIndex &&
-					   A.Data.Num() == B.Data.Num() &&
-					   FMemory::Memcmp(A.Data.GetData(), B.Data.GetData(), A.Data.Num()*sizeof(FCustomizableObjectMeshToMeshVertData)) == 0;
-			};
+				{
+					return A.ClothingAssetIndex == B.ClothingAssetIndex &&
+						A.ClothingAssetLOD == B.ClothingAssetLOD &&
+						A.PhysicsAssetIndex == B.PhysicsAssetIndex &&
+						A.Data.Num() == B.Data.Num() &&
+						FMemory::Memcmp(A.Data.GetData(), B.Data.GetData(), A.Data.Num() * sizeof(FCustomizableObjectMeshToMeshVertData)) == 0;
+				};
 
 			if (FoundHash)
 			{
-				bIsDataAlreadyCollected = CompareClothingData(*FoundHash, ClothingMeshData); 
+				bIsDataAlreadyCollected = CompareClothingData(*FoundHash, ClothingMeshData);
 			}
 
 			// NOTE: This way of unique hash generation guarantees all valid values can be used but given its 
@@ -1749,8 +1746,8 @@ TSharedPtr<mu::FMesh> ConvertSkeletalMeshToMutable(
 				uint32 NumTries = 0;
 				for (; NumTries < TNumericLimits<uint32>::Max(); ++NumTries)
 				{
-					FoundHash = GenerationContext.ClothingPerMeshData.Find(UniqueHash);
-					
+					FoundHash = Context.ClothingPerMeshData.Find(UniqueHash);
+
 					if (!FoundHash)
 					{
 						break;
@@ -1770,9 +1767,9 @@ TSharedPtr<mu::FMesh> ConvertSkeletalMeshToMutable(
 				{
 					UE_LOG(LogMutable, Warning, TEXT("Maximum number of meshes with clothing reached, some cloth meshes may not work as expected."));
 					UniqueHash = InvalidResourceId;
-				}	
+				}
 			}
-			
+
 			// fill the resource buffer with the generated unique id.
 			for (uint32& Elem : ClothSectionResourceIdView)
 			{
@@ -1787,7 +1784,7 @@ TSharedPtr<mu::FMesh> ConvertSkeletalMeshToMutable(
 
 			if (!bIsDataAlreadyCollected)
 			{
-				FClothingMeshData& NewClothingMeshData = GenerationContext.ClothingPerMeshData.FindOrAdd(UniqueHash);
+				FClothingMeshData& NewClothingMeshData = Context.ClothingPerMeshData.FindOrAdd(UniqueHash);
 
 				check(NewClothingMeshData.ClothingAssetIndex == INDEX_NONE);
 				check(NewClothingMeshData.Data.IsEmpty());
@@ -1799,7 +1796,7 @@ TSharedPtr<mu::FMesh> ConvertSkeletalMeshToMutable(
 				NewClothingMeshData.Data = MoveTemp(ClothingMeshData.Data);
 				NewClothingMeshData.SourceId = MeshId;
 			}
-			
+
 			NextBufferIndex += 2;
 		}
 
@@ -1807,7 +1804,7 @@ TSharedPtr<mu::FMesh> ConvertSkeletalMeshToMutable(
 
 
 	// SkinWeightProfiles vertex info.
-	if (GenerationContext.Options.bSkinWeightProfilesEnabled)
+	if (Context.Options.bSkinWeightProfilesEnabled)
 	{
 		using namespace mu;
 
@@ -1817,15 +1814,15 @@ TSharedPtr<mu::FMesh> ConvertSkeletalMeshToMutable(
 
 		// Limit skinning weights if necessary
 		const int32 MutableBonesPerVertex = bUseUnlimitedInfluences ?
-											MaxSectionInfluences :
-											(int32)GenerationContext.Options.CustomizableObjectNumBoneInfluences;
+			MaxSectionInfluences :
+			(int32)Context.Options.CustomizableObjectNumBoneInfluences;
 		const int32 BoneIndicesSize = MutableBonesPerVertex * sizeof(FBoneIndexType);
 		const int32 BoneWeightsSize = MutableBonesPerVertex * BoneWeightTypeSizeBytes;
 		const int32 SkinWeightProfileVertexSize = sizeof(int32) + BoneIndicesSize + BoneWeightsSize;
 
 		const int32 MaxSectionBoneMapIndex = MeshSection.BoneMap.Num();
 
-		const TArray<FSkinWeightProfileInfo>& SkinWeightProfilesInfo = InSkeletalMesh->GetSkinWeightProfiles();
+		const TArray<FSkinWeightProfileInfo>& SkinWeightProfilesInfo = SkeletalMesh->GetSkinWeightProfiles();
 		for (const FSkinWeightProfileInfo& Profile : SkinWeightProfilesInfo)
 		{
 			const FImportedSkinWeightProfileData* ImportedProfileData = LODModel.SkinWeightProfiles.Find(Profile.Name);
@@ -1849,7 +1846,7 @@ TSharedPtr<mu::FMesh> ConvertSkeletalMeshToMutable(
 					TransferRemovedBonesInfluences(&SkinWeight.InfluenceBones[0], &SkinWeight.InfluenceWeights[0], MaxSectionInfluences, RemappedBoneMapIndices);
 				}
 
-				if (GenerationContext.Options.CustomizableObjectNumBoneInfluences == ECustomizableObjectNumBoneInfluences::Four)
+				if (Context.Options.CustomizableObjectNumBoneInfluences == ECustomizableObjectNumBoneInfluences::Four)
 				{
 					// Normalize weights
 					const int32 MaxMutableWeights = 4;
@@ -1887,11 +1884,11 @@ TSharedPtr<mu::FMesh> ConvertSkeletalMeshToMutable(
 				MutSkinWeightData += BoneWeightsSize;
 			}
 
-			const uint32 ProfileId = GenerationContext.GetSkinWeightProfileIdUnique(Profile.Name);
-			const int32 ProfileIndex = GenerationContext.SkinWeightProfilesInfo.AddUnique({Profile.Name, ProfileId, false, 0});
-			FMutableSkinWeightProfileInfo& MutSkinWeightProfileInfo = GenerationContext.SkinWeightProfilesInfo[ProfileIndex];
+			const uint32 ProfileId = Context.GetSkinWeightProfileIdUnique(Profile.Name);
+			const int32 ProfileIndex = Context.SkinWeightProfilesInfo.AddUnique({ Profile.Name, ProfileId, false, 0 });
+			FMutableSkinWeightProfileInfo& MutSkinWeightProfileInfo = Context.SkinWeightProfilesInfo[ProfileIndex];
 
-			const FName PlatformName = *GenerationContext.Options.TargetPlatform->PlatformName();
+			const FName PlatformName = *Context.Options.TargetPlatform->PlatformName();
 			MutSkinWeightProfileInfo.DefaultProfile = MutSkinWeightProfileInfo.DefaultProfile || Profile.DefaultProfile.GetValueForPlatform(PlatformName);
 			MutSkinWeightProfileInfo.DefaultProfileFromLODIndex = FMath::Min(MutSkinWeightProfileInfo.DefaultProfileFromLODIndex, Profile.DefaultProfileFromLODIndex.GetValueForPlatform(PlatformName));
 
@@ -1911,6 +1908,7 @@ TSharedPtr<mu::FMesh> ConvertSkeletalMeshToMutable(
 		}
 	}
 
+
 	// Indices
 	{
 		const uint32 IndexStart = MeshSection.BaseIndex;
@@ -1923,7 +1921,7 @@ TSharedPtr<mu::FMesh> ConvertSkeletalMeshToMutable(
 		check(LODModel.IndexBuffer.IsValidIndex(IndexStart) && LODModel.IndexBuffer.IsValidIndex(IndexStart + IndexCount - 1));
 		const uint32* IndexDataPtr = &LODModel.IndexBuffer[IndexStart];
 
-		const int32 FinalElementSize = sizeof(uint32_t);
+		const int32 FinalElementSize = sizeof(uint32);
 		const int32 ChannelCount = 1;
 		const EMeshBufferSemantic Semantics[ChannelCount] = { EMeshBufferSemantic::VertexIndex };
 		const int32 SemanticIndices[ChannelCount] = { 0 };
@@ -1945,7 +1943,7 @@ TSharedPtr<mu::FMesh> ConvertSkeletalMeshToMutable(
 			VertexIndex = *IndexDataPtr - VertexStart;
 			if (ensureMsgf(VertexIndex < (uint32)VertexCount, TEXT("Mutable: VertexIndex >= VertexCount. VI [%d], VC [%d], VS [%d]. SKM [%s] LOD [%d] Section [%d]."),
 				VertexIndex, VertexCount, VertexStart,
-				*GetNameSafe(InSkeletalMesh), LODIndex, SectionIndex))
+				*GetNameSafe(SkeletalMesh), Source.LODIndex, Source.SectionIndex))
 			{
 				*pDest = VertexIndex;
 			}
@@ -1958,11 +1956,11 @@ TSharedPtr<mu::FMesh> ConvertSkeletalMeshToMutable(
 		}
 	}
 
-	if (!bIgnorePhysics && InSkeletalMesh->GetPhysicsAsset() && MutableMesh->GetSkeleton() && GenerationContext.Options.bPhysicsAssetMergeEnabled)
+	if (!bIgnorePhysics && SkeletalMesh->GetPhysicsAsset() && MutableMesh->GetSkeleton() && Context.Options.bPhysicsAssetMergeEnabled)
 	{
 		// Find BodySetups with relevant bones.
-		TArray<TObjectPtr<USkeletalBodySetup>>& SkeletalBodySetups = InSkeletalMesh->GetPhysicsAsset()->SkeletalBodySetups;
-		
+		TArray<TObjectPtr<USkeletalBodySetup>>& SkeletalBodySetups = SkeletalMesh->GetPhysicsAsset()->SkeletalBodySetups;
+
 		TArray<TObjectPtr<USkeletalBodySetup>> RelevantBodySetups;
 		RelevantBodySetups.Reserve(SkeletalBodySetups.Num());
 
@@ -1977,14 +1975,14 @@ TSharedPtr<mu::FMesh> ConvertSkeletalMeshToMutable(
 				continue;
 			}
 
-			const mu::FBoneName& BoneName = GenerationContext.GetBoneUnique(BodySetup->BoneName);
+			const mu::FBoneName& BoneName = Context.GetBoneUnique(BodySetup->BoneName);
 			const int32 BonePoseIndex = MutableMesh->FindBonePose(BoneName);
 
 			if (BonePoseIndex == INDEX_NONE)
 			{
 				continue;
 			}
-		
+
 			RelevantBodySetups.Add(BodySetup);
 			DiscardedBodySetups[BodySetupIndex] = 0;
 			EnumAddFlags(MutableMesh->BonePoses[BonePoseIndex].BoneUsageFlags, mu::EBoneUsageFlags::Physics);
@@ -1995,12 +1993,11 @@ TSharedPtr<mu::FMesh> ConvertSkeletalMeshToMutable(
 		constexpr bool bOptOutOfIncompleteBodyWarnings = true;
 		if (NumDiscardedSetups > 0 && !bOptOutOfIncompleteBodyWarnings)
 		{
-			FString PhysicsSetupsRemovedMsg = 
-					FString::Printf(TEXT("PhysicsBodySetups in %s attached to bones"), 
-					*(InSkeletalMesh->GetPhysicsAsset()->GetName()));
+			FString PhysicsSetupsRemovedMsg =
+				FString::Printf(TEXT("PhysicsBodySetups in %s attached to bones"), *(SkeletalMesh->GetPhysicsAsset()->GetName()));
 
 			constexpr int32 MaxNumDiscardedShown = 3;
-			
+
 			int32 NumDiscardedShown = 0;
 			for (int32 I = 0; I < SkeletalBodySetups.Num() && NumDiscardedShown < MaxNumDiscardedShown; ++I)
 			{
@@ -2010,63 +2007,63 @@ TSharedPtr<mu::FMesh> ConvertSkeletalMeshToMutable(
 					++NumDiscardedShown;
 				}
 			}
-	
+
 			if (NumDiscardedShown < NumDiscardedSetups)
 			{
 				PhysicsSetupsRemovedMsg += FString::Printf(TEXT("... and %d more "), NumDiscardedSetups - MaxNumDiscardedShown);
 			}
-	
+
 			PhysicsSetupsRemovedMsg += FString::Printf(TEXT("have been discarded because they are not present in the SkeletalMesh [%s] Skeleton."),
-				*InSkeletalMesh->GetName());
-					
-			GenerationContext.Log(FText::FromString(PhysicsSetupsRemovedMsg), CurrentNode, EMessageSeverity::Warning);
+				*SkeletalMesh->GetName());
+
+			Context.Log(FText::FromString(PhysicsSetupsRemovedMsg), Source.MessageContext, EMessageSeverity::Warning);
 		}
 
 		TSharedPtr<mu::FPhysicsBody> PhysicsBody = MakeShared<mu::FPhysicsBody>();
-		
-		const int32 NumBodySetups = RelevantBodySetups.Num();	
+
+		const int32 NumBodySetups = RelevantBodySetups.Num();
 		PhysicsBody->SetBodyCount(NumBodySetups);
 
 		auto GetKBodyElemFlags = [](const FKShapeElem& KElem) -> uint32
-		{
-			uint8 ElemCollisionEnabled = static_cast<uint8>( KElem.GetCollisionEnabled() );
-			
-			uint32 Flags = static_cast<uint32>( ElemCollisionEnabled );
-			Flags = Flags | (static_cast<uint32>(KElem.GetContributeToMass()) << 8);
+			{
+				uint8 ElemCollisionEnabled = static_cast<uint8>(KElem.GetCollisionEnabled());
 
-			return Flags; 
-		};
+				uint32 Flags = static_cast<uint32>(ElemCollisionEnabled);
+				Flags = Flags | (static_cast<uint32>(KElem.GetContributeToMass()) << 8);
 
-		for ( int32 B = 0; B < NumBodySetups; ++B )
+				return Flags;
+			};
+
+		for (int32 B = 0; B < NumBodySetups; ++B)
 		{
 			TObjectPtr<USkeletalBodySetup>& BodySetup = RelevantBodySetups[B];
-			
-			const mu::FBoneName& BoneId = GenerationContext.GetBoneUnique(BodySetup->BoneName);
-			PhysicsBody->SetBodyBoneId( B, BoneId);
-			
+
+			const mu::FBoneName& BoneId = Context.GetBoneUnique(BodySetup->BoneName);
+			PhysicsBody->SetBodyBoneId(B, BoneId);
+
 			const int32 NumSpheres = BodySetup->AggGeom.SphereElems.Num();
-			PhysicsBody->SetSphereCount( B, NumSpheres );
+			PhysicsBody->SetSphereCount(B, NumSpheres);
 
-			for ( int32 I = 0; I < NumSpheres; ++I )
+			for (int32 I = 0; I < NumSpheres; ++I)
 			{
 				const FKSphereElem& SphereElem = BodySetup->AggGeom.SphereElems[I];
-				PhysicsBody->SetSphere( B, I, FVector3f(SphereElem.Center), SphereElem.Radius );
+				PhysicsBody->SetSphere(B, I, FVector3f(SphereElem.Center), SphereElem.Radius);
 
 				const FString ElemName = SphereElem.GetName().ToString();
-				PhysicsBody->SetSphereName(B, I, StringCast<ANSICHAR>(*ElemName).Get());	
+				PhysicsBody->SetSphereName(B, I, StringCast<ANSICHAR>(*ElemName).Get());
 				PhysicsBody->SetSphereFlags(B, I, GetKBodyElemFlags(SphereElem));
 			}
 
 			const int32 NumBoxes = BodySetup->AggGeom.BoxElems.Num();
-			PhysicsBody->SetBoxCount( B, NumBoxes );
+			PhysicsBody->SetBoxCount(B, NumBoxes);
 
-			for ( int32 I = 0; I < NumBoxes; ++I )
+			for (int32 I = 0; I < NumBoxes; ++I)
 			{
 				const FKBoxElem& BoxElem = BodySetup->AggGeom.BoxElems[I];
-				PhysicsBody->SetBox( B, I, 
-						FVector3f(BoxElem.Center), 
-						FQuat4f(BoxElem.Rotation.Quaternion()), 
-						FVector3f(BoxElem.X, BoxElem.Y, BoxElem.Z));
+				PhysicsBody->SetBox(B, I,
+					FVector3f(BoxElem.Center),
+					FQuat4f(BoxElem.Rotation.Quaternion()),
+					FVector3f(BoxElem.X, BoxElem.Y, BoxElem.Z));
 
 				const FString KElemName = BoxElem.GetName().ToString();
 				PhysicsBody->SetBoxName(B, I, StringCast<ANSICHAR>(*KElemName).Get());
@@ -2074,40 +2071,40 @@ TSharedPtr<mu::FMesh> ConvertSkeletalMeshToMutable(
 			}
 
 			const int32 NumConvex = BodySetup->AggGeom.ConvexElems.Num();
-			PhysicsBody->SetConvexCount( B, NumConvex );
-			for ( int32 I = 0; I < NumConvex; ++I )
+			PhysicsBody->SetConvexCount(B, NumConvex);
+			for (int32 I = 0; I < NumConvex; ++I)
 			{
 				const FKConvexElem& ConvexElem = BodySetup->AggGeom.ConvexElems[I];
 
 				// Convert to FVector3f
 				TArray<FVector3f> VertexData;
-				VertexData.SetNumUninitialized( ConvexElem.VertexData.Num() );
-				for ( int32 Elem = VertexData.Num() - 1; Elem >= 0; --Elem )
+				VertexData.SetNumUninitialized(ConvexElem.VertexData.Num());
+				for (int32 Elem = VertexData.Num() - 1; Elem >= 0; --Elem)
 				{
 					VertexData[Elem] = FVector3f(ConvexElem.VertexData[Elem]);
 				}
-				
+
 				PhysicsBody->SetConvexMesh(B, I,
-						TArrayView<const FVector3f>(VertexData.GetData(), ConvexElem.VertexData.Num()),
-						TArrayView<const int32>(ConvexElem.IndexData.GetData(), ConvexElem.IndexData.Num()));
+					TArrayView<const FVector3f>(VertexData.GetData(), ConvexElem.VertexData.Num()),
+					TArrayView<const int32>(ConvexElem.IndexData.GetData(), ConvexElem.IndexData.Num()));
 
 				PhysicsBody->SetConvexTransform(B, I, FTransform3f(ConvexElem.GetTransform()));
-				
+
 				const FString KElemName = ConvexElem.GetName().ToString();
 				PhysicsBody->SetConvexName(B, I, StringCast<ANSICHAR>(*KElemName).Get());
 				PhysicsBody->SetConvexFlags(B, I, GetKBodyElemFlags(ConvexElem));
 			}
 
 			const int32 NumSphyls = BodySetup->AggGeom.SphylElems.Num();
-			PhysicsBody->SetSphylCount( B, NumSphyls );
+			PhysicsBody->SetSphylCount(B, NumSphyls);
 
-			for ( int32 I = 0; I < NumSphyls; ++I )
+			for (int32 I = 0; I < NumSphyls; ++I)
 			{
 				const FKSphylElem& SphylElem = BodySetup->AggGeom.SphylElems[I];
-				PhysicsBody->SetSphyl( B, I, 
-						FVector3f(SphylElem.Center), 
-						FQuat4f(SphylElem.Rotation.Quaternion()), 
-						SphylElem.Radius, SphylElem.Length );
+				PhysicsBody->SetSphyl(B, I,
+					FVector3f(SphylElem.Center),
+					FQuat4f(SphylElem.Rotation.Quaternion()),
+					SphylElem.Radius, SphylElem.Length);
 
 				const FString KElemName = SphylElem.GetName().ToString();
 				PhysicsBody->SetSphylName(B, I, StringCast<ANSICHAR>(*KElemName).Get());
@@ -2115,32 +2112,32 @@ TSharedPtr<mu::FMesh> ConvertSkeletalMeshToMutable(
 			}
 
 			const int32 NumTaperedCapsules = BodySetup->AggGeom.TaperedCapsuleElems.Num();
-			PhysicsBody->SetTaperedCapsuleCount( B, NumTaperedCapsules );
+			PhysicsBody->SetTaperedCapsuleCount(B, NumTaperedCapsules);
 
-			for ( int32 I = 0; I < NumTaperedCapsules; ++I )
+			for (int32 I = 0; I < NumTaperedCapsules; ++I)
 			{
 				const FKTaperedCapsuleElem& TaperedCapsuleElem = BodySetup->AggGeom.TaperedCapsuleElems[I];
-				PhysicsBody->SetTaperedCapsule( B, I, 
-						FVector3f(TaperedCapsuleElem.Center), 
-						FQuat4f(TaperedCapsuleElem.Rotation.Quaternion()), 
-						TaperedCapsuleElem.Radius0, TaperedCapsuleElem.Radius1, TaperedCapsuleElem.Length );
-				
+				PhysicsBody->SetTaperedCapsule(B, I,
+					FVector3f(TaperedCapsuleElem.Center),
+					FQuat4f(TaperedCapsuleElem.Rotation.Quaternion()),
+					TaperedCapsuleElem.Radius0, TaperedCapsuleElem.Radius1, TaperedCapsuleElem.Length);
+
 				const FString KElemName = TaperedCapsuleElem.GetName().ToString();
 				PhysicsBody->SetTaperedCapsuleName(B, I, StringCast<ANSICHAR>(*KElemName).Get());
-				PhysicsBody->SetTaperedCapsuleFlags( B, I, GetKBodyElemFlags(TaperedCapsuleElem));
+				PhysicsBody->SetTaperedCapsuleFlags(B, I, GetKBodyElemFlags(TaperedCapsuleElem));
 			}
 		}
 
 		MutableMesh->SetPhysicsBody(PhysicsBody);
 	}
-	
+
 	// Set Bone Parenting usages. This has to be done after all primary usages are set.
 	for (int32 I = MutableMesh->GetBonePoseCount() - 1; I >= 0; --I)
 	{
 		mu::FMesh::FBonePose& BonePose = MutableMesh->BonePoses[I];
 
 		constexpr mu::EBoneUsageFlags FlagsToPropagate =
-				mu::EBoneUsageFlags::Skinning | mu::EBoneUsageFlags::Physics | mu::EBoneUsageFlags::Deform;
+			mu::EBoneUsageFlags::Skinning | mu::EBoneUsageFlags::Physics | mu::EBoneUsageFlags::Deform;
 		if (EnumHasAnyFlags(BonePose.BoneUsageFlags, FlagsToPropagate))
 		{
 			const int32 Index = MutableMesh->GetSkeleton()->FindBone(BonePose.BoneId);
@@ -2153,28 +2150,28 @@ TSharedPtr<mu::FMesh> ConvertSkeletalMeshToMutable(
 			const int32 ParentIndex = MutableMesh->GetSkeleton()->GetBoneParent(Index);
 
 			if (ParentIndex == INDEX_NONE)
-			{	
+			{
 				continue;
 			}
 
 			const mu::EBoneUsageFlags ParentPropagationFlags =
-				(EnumHasAnyFlags(BonePose.BoneUsageFlags, mu::EBoneUsageFlags::Skinning) 
+				(EnumHasAnyFlags(BonePose.BoneUsageFlags, mu::EBoneUsageFlags::Skinning)
 					? mu::EBoneUsageFlags::SkinningParent : mu::EBoneUsageFlags::None) |
-				(EnumHasAnyFlags(BonePose.BoneUsageFlags, mu::EBoneUsageFlags::Physics) 
+				(EnumHasAnyFlags(BonePose.BoneUsageFlags, mu::EBoneUsageFlags::Physics)
 					? mu::EBoneUsageFlags::PhysicsParent : mu::EBoneUsageFlags::None) |
-				(EnumHasAnyFlags(BonePose.BoneUsageFlags, mu::EBoneUsageFlags::Deform) 
+				(EnumHasAnyFlags(BonePose.BoneUsageFlags, mu::EBoneUsageFlags::Deform)
 					? mu::EBoneUsageFlags::DeformParent : mu::EBoneUsageFlags::None);
 
 			SetAndPropagatePoseBoneUsage(*MutableMesh, ParentIndex, ParentPropagationFlags);
 		}
 	}
 
-	const bool bAnimPhysicsManipulationEnabled = GenerationContext.Options.bAnimBpPhysicsManipulationEnabled;
+	const bool bAnimPhysicsManipulationEnabled = Context.Options.bAnimBpPhysicsManipulationEnabled;
 
-	if (!bIgnorePhysics && !AnimBp.IsNull() && MutableMesh->GetSkeleton() && bAnimPhysicsManipulationEnabled)
+	if (!bIgnorePhysics && !Source.AnimInstance.IsNull() && MutableMesh->GetSkeleton() && bAnimPhysicsManipulationEnabled)
 	{
 		using AnimPhysicsInfoType = TTuple<UPhysicsAsset*, int32>;
-		const TArray<AnimPhysicsInfoType> AnimPhysicsInfo = GetPhysicsAssetsFromAnimInstance(GenerationContext, AnimBp);
+		const TArray<AnimPhysicsInfoType> AnimPhysicsInfo = GetPhysicsAssetsFromAnimInstance(Source.AnimInstance);
 
 		for (const AnimPhysicsInfoType PropertyInfo : AnimPhysicsInfo)
 		{
@@ -2183,16 +2180,16 @@ TSharedPtr<mu::FMesh> ConvertSkeletalMeshToMutable(
 
 			FAnimBpOverridePhysicsAssetsInfo Info;
 			{
-				Info.AnimInstanceClass = AnimBp;
+				Info.AnimInstanceClass = Source.AnimInstance;
 				Info.PropertyIndex = PropertyIndex;
 				Info.SourceAsset = TSoftObjectPtr<UPhysicsAsset>(PropertyAsset);
 			}
 
-			const int32 PhysicsAssetId = GenerationContext.AnimBpOverridePhysicsAssetsInfo.AddUnique(Info);
+			const int32 PhysicsAssetId = Context.AnimBpOverridePhysicsAssetsInfo.AddUnique(Info);
 
-			TSharedPtr<mu::FPhysicsBody> MutableBody = MakePhysicsBodyFromAsset(GenerationContext,
-					PropertyAsset,
-					MakePhysicsAssetBodySetupRelevancyMap(GenerationContext, PropertyAsset, MutableMesh));
+			TSharedPtr<mu::FPhysicsBody> MutableBody = MakePhysicsBodyFromAsset(Context,
+				PropertyAsset,
+				MakePhysicsAssetBodySetupRelevancyMap(Context, PropertyAsset, MutableMesh));
 			MutableBody->CustomId = PhysicsAssetId;
 
 			MutableMesh->AddAdditionalPhysicsBody(MutableBody);
@@ -2208,24 +2205,24 @@ TSharedPtr<mu::FMesh> ConvertSkeletalMeshToMutable(
 		// this info can be stored in the CO per mesh.
 		MeshMetadata.MorphMetadataId = 0;
 		MeshMetadata.ClothingMetadataId = 0;
-		MeshMetadata.SurfaceMetadataId = SurfaceMetadataId;
+		MeshMetadata.SurfaceMetadataId = Source.SurfaceMetadataId;
 
 		auto HashMeshMetadataFunc = [](const FMutableMeshMetadata& Data) -> uint32
-		{
-			return CityHash32(reinterpret_cast<const char*>(&Data), sizeof(FMutableMeshMetadata));
-		};
+			{
+				return CityHash32(reinterpret_cast<const char*>(&Data), sizeof(FMutableMeshMetadata));
+			};
 
 		auto CompareMeshMetadataFunc = [](const FMutableMeshMetadata& A, const FMutableMeshMetadata& B) -> bool
-		{
-			return FMemory::Memcmp(&A, &B, sizeof(FMutableMeshMetadata)) == 0;
-		};
+			{
+				return FMemory::Memcmp(&A, &B, sizeof(FMutableMeshMetadata)) == 0;
+			};
 
 		const uint32 MeshMetadataUniqueHash = Private::GenerateUniquePersistentHash(
-				MeshMetadata, GenerationContext.MeshMetadata, HashMeshMetadataFunc, CompareMeshMetadataFunc);
-		
+			MeshMetadata, Context.MeshMetadata, HashMeshMetadataFunc, CompareMeshMetadataFunc);
+
 		if (MeshMetadataUniqueHash != 0)
 		{
-			GenerationContext.MeshMetadata.FindOrAdd(MeshMetadataUniqueHash, MeshMetadata);
+			Context.MeshMetadata.FindOrAdd(MeshMetadataUniqueHash, MeshMetadata);
 		}
 		else
 		{
@@ -2234,14 +2231,14 @@ TSharedPtr<mu::FMesh> ConvertSkeletalMeshToMutable(
 
 		// Ensure Surface Data
 		mu::FMeshSurface& MeshSurface = MutableMesh->Surfaces.Emplace_GetRef();
-	
+
 		MeshSurface.SubMeshes.Emplace(
-				mu::FSurfaceSubMesh 
-				{
-					0, MutableMesh->VertexBuffers.GetElementCount(), 
-					0, MutableMesh->IndexBuffers.GetElementCount(),
-					MeshMetadataUniqueHash	
-				});
+			mu::FSurfaceSubMesh
+			{
+				0, MutableMesh->VertexBuffers.GetElementCount(),
+				0, MutableMesh->IndexBuffers.GetElementCount(),
+				MeshMetadataUniqueHash
+			});
 
 		MeshSurface.BoneMapCount = MutableMesh->BoneMap.Num();
 	}
@@ -2249,6 +2246,135 @@ TSharedPtr<mu::FMesh> ConvertSkeletalMeshToMutable(
 	return MutableMesh;
 }
 
+
+TSharedPtr<mu::FMesh> ConvertSkeletalMeshToMutable(
+	TSoftObjectPtr<const USkeletalMesh> InSkeletalMesh, const TSoftClassPtr<UAnimInstance>& AnimBp,
+	bool bMeshMustExist, int32 LODIndex, int32 SectionIndex, 
+	uint32 SurfaceMetadataId,
+	FMutableGraphGenerationContext& GenerationContext, 
+	const UCustomizableObjectNode* CurrentNode, 
+	USkeletalMesh* TableReferenceSkeletalMesh)
+{
+	MUTABLE_CPUPROFILER_SCOPE(ConvertSkeletalMeshToMutable);
+
+	if(!InSkeletalMesh)
+	{
+		return nullptr;
+	}
+
+	GenerationContext.AddParticipatingObject(InSkeletalMesh.ToSoftObjectPath());
+
+	// Prepare the data that is needed for the core mesh conversion.
+	FMutableSourceMeshData Source;
+	Source.Mesh = InSkeletalMesh;
+	Source.AnimInstance = AnimBp;
+	Source.bMeshMustExist = bMeshMustExist;
+	Source.LODIndex = LODIndex;
+	Source.SectionIndex = SectionIndex;
+	Source.SurfaceMetadataId = SurfaceMetadataId;
+	Source.Flags = GenerationContext.MeshGenerationFlags.Last();
+	Source.MessageContext = CurrentNode;
+	Source.Component = GenerationContext.CurrentMeshComponent;
+	Source.TableReferenceSkeletalMesh = TableReferenceSkeletalMesh;
+
+	// Prepare source data for realtime morph targets
+	const UCustomizableObjectNodeSkeletalMesh* NodeTypedSkMesh = Cast<UCustomizableObjectNodeSkeletalMesh>(CurrentNode);
+	if (NodeTypedSkMesh)
+	{
+		Source.bUseAllRealTimeMorphs = NodeTypedSkMesh->bUseAllRealTimeMorphs;
+		if (!Source.bUseAllRealTimeMorphs)
+		{
+			Source.UsedRealTimeMorphTargetNames = NodeTypedSkMesh->UsedRealTimeMorphTargetNames;
+		}
+	}
+
+	// Add additional participating objects
+	{
+		FAssetRegistryModule& AssetRegistryModule = FModuleManager::LoadModuleChecked<FAssetRegistryModule>("AssetRegistry");
+		TArray<FAssetData> Assets;
+		{
+			TArray<FName> PackageDependencies;
+			AssetRegistryModule.Get().GetDependencies(InSkeletalMesh.GetLongPackageFName(), PackageDependencies);
+
+			FARFilter Filter;
+			for (const FName& PackageName : PackageDependencies)
+			{
+				if (!PackageName.ToString().StartsWith(TEXT("/TempAutosave")))
+				{
+					Filter.PackageNames.Add(PackageName);
+				}
+			}
+			AssetRegistryModule.Get().GetAssets(Filter, Assets);
+		}
+
+		// Add the skeleton dependency
+		const bool bIgnoreSkeleton = EnumHasAnyFlags(Source.Flags, EMutableMeshConversionFlags::IgnoreSkinning);
+		if (!bIgnoreSkeleton)
+		{
+			bool bSkeletonFound = false;
+			for (const FAssetData& Dep : Assets)
+			{
+				const UClass* AssetClass = Dep.GetClass(EResolveClass::Yes);
+				if (AssetClass == USkeleton::StaticClass())
+				{
+					GenerationContext.AddParticipatingObject(Dep.GetSoftObjectPath());
+					bSkeletonFound = true;
+				}
+			}
+
+			if (!bSkeletonFound)
+			{
+				FString Msg = FString::Printf(TEXT("No skeleton provided when converting SkeletalMesh [%s]."), *InSkeletalMesh.GetLongPackageName());
+				GenerationContext.Log(FText::FromString(Msg), Source.MessageContext);
+				return nullptr;
+			}
+		}
+
+		// Add clothing asset participating objects
+		if (GenerationContext.CompilationContext.Options.bClothingEnabled)
+		{
+			for (const FAssetData& Dep : Assets)
+			{
+				const UClass* AssetClass = Dep.GetClass(EResolveClass::Yes);
+				if (AssetClass->ImplementsInterface(UClothingAssetCommon::StaticClass()))
+				{
+					GenerationContext.AddParticipatingObject(Dep.GetSoftObjectPath());
+
+					TArray<FAssetData> ClothAssets;
+					{
+						TArray<FName> ClothPackageDependencies;
+						AssetRegistryModule.Get().GetDependencies(InSkeletalMesh.GetLongPackageFName(), ClothPackageDependencies);
+
+						FARFilter ClothFilter;
+						for (const FName& PackageName : ClothPackageDependencies)
+						{
+							if (!PackageName.ToString().StartsWith(TEXT("/TempAutosave")))
+							{
+								ClothFilter.PackageNames.Add(PackageName);
+							}
+						}
+						AssetRegistryModule.Get().GetAssets(ClothFilter, ClothAssets);
+					}
+
+					for (const FAssetData& ClothDep : ClothAssets)
+					{
+						const UClass* ClothAssetClass = ClothDep.GetClass(EResolveClass::Yes);
+						if (ClothAssetClass->ImplementsInterface(UPhysicsAsset::StaticClass()))
+						{
+							GenerationContext.AddParticipatingObject(ClothDep.GetSoftObjectPath());
+						}
+					}
+				}
+			}
+		}
+	}
+
+	// Do it. In the future this will happen in a callback from the mutable core compilation process.
+	TSharedPtr<mu::FMesh> MutableMesh = ConvertSkeletalMeshToMutable(Source, GenerationContext.CompilationContext );
+
+	return MutableMesh;
+}
+
 TSharedPtr<mu::FMesh> ConvertStaticMeshToMutable(const UStaticMesh* StaticMesh, int32 LODIndex, int32 SectionIndex, FMutableGraphGenerationContext& GenerationContext, const UCustomizableObjectNode* CurrentNode)
 {
 	if (!StaticMesh->GetRenderData() ||
@@ -2439,7 +2565,7 @@ TSharedPtr<mu::FMesh> ConvertStaticMeshToMutable(const UStaticMesh* StaticMesh,
 
 // Convert a Mesh constant to a mutable format. UniqueTags are the tags that make this Mesh unique that cannot be merged in the cache 
 // with the exact same Mesh with other tags
-TSharedPtr<mu::FMesh> GenerateMutableMesh(UObject * Mesh, const TSoftClassPtr<UAnimInstance>& AnimInstance, int32 LODIndexConnected, int32 SectionIndexConnected, 
+TSharedPtr<mu::FMesh> GenerateMutableMesh( TSoftObjectPtr<UStreamableRenderAsset> Mesh, const TSoftClassPtr<UAnimInstance>& AnimInstance, int32 LODIndexConnected, int32 SectionIndexConnected,
 									  int32 LODIndex, int32 SectionIndex, const FString& UniqueTags, uint32 SurfaceMetadataId, FMutableGraphGenerationContext & GenerationContext, 
 									  const UCustomizableObjectNode* CurrentNode, USkeletalMesh* TableReferenceSkeletalMesh, bool bIsReference)
 {
@@ -2453,8 +2579,22 @@ TSharedPtr<mu::FMesh> GenerateMutableMesh(UObject * Mesh, const TSoftClassPtr<UA
 		return MutableMesh;
 	}
 
-	if (USkeletalMesh* SkeletalMesh = Cast<USkeletalMesh>(Mesh))
+	// If we are using automatic LODs and not generating the base LOD (the connected one) is not an error.
+	bool bMeshMustExist =
+		// Are we using automatic LODs?
+		GenerationContext.CurrentAutoLODStrategy != ECustomizableObjectAutomaticLODStrategy::AutomaticFromMesh
+		||
+		// Are we generating the base LOD?
+		((LODIndex == LODIndexConnected) && (SectionIndex == SectionIndexConnected));
+
+	// When we want to defer the mesh conversion to the core compilation stage, we will do this:
+	//MutableMesh = GenerateMeshConstant(Mesh, AnimInstance, bMeshMustExist, LODIndex, SectionIndex, SurfaceMetadataId, GenerationContext, bIsReference);
+
+	// Meanwhile:
+	UObject* LoadedMesh = Mesh.LoadSynchronous();
+	if (USkeletalMesh* SkeletalMesh = Cast<USkeletalMesh>(LoadedMesh))
 	{
+
 		// At some point we will want all meshes to be references at compile-time. For now, just create the actual pass-through meshes.
 		if (bIsReference)
 		{
@@ -2462,10 +2602,10 @@ TSharedPtr<mu::FMesh> GenerateMutableMesh(UObject * Mesh, const TSoftClassPtr<UA
 		}
 		else
 		{
-			MutableMesh = ConvertSkeletalMeshToMutable(SkeletalMesh, AnimInstance, LODIndexConnected, SectionIndexConnected, LODIndex, SectionIndex, SurfaceMetadataId, GenerationContext, CurrentNode, TableReferenceSkeletalMesh);
+			MutableMesh = ConvertSkeletalMeshToMutable(SkeletalMesh, AnimInstance, bMeshMustExist, LODIndex, SectionIndex, SurfaceMetadataId, GenerationContext, CurrentNode, TableReferenceSkeletalMesh);
 		}
 	}
-	else if (const UStaticMesh* StaticMesh = Cast<UStaticMesh>(Mesh))
+	else if (const UStaticMesh* StaticMesh = Cast<UStaticMesh>(LoadedMesh))
 	{
 		MutableMesh = ConvertStaticMeshToMutable(StaticMesh, LODIndex, SectionIndex, GenerationContext, CurrentNode);
 	}
@@ -2568,7 +2708,7 @@ TSharedPtr<mu::FMesh> BuildMorphedMutableMesh(const UEdGraphPin* BaseSourcePin,
 	{
 		int32 LayoutIndex;
 		TypedNodeSkeletalMesh->GetPinSection(*BaseSourcePin, LODIndexConnected, SectionIndexConnected, LayoutIndex);
-		SkeletalMesh = TypedNodeSkeletalMesh->SkeletalMesh;
+		SkeletalMesh = TypedNodeSkeletalMesh->SkeletalMesh.LoadSynchronous();
 	}
 
 	else if (const UCustomizableObjectNodeTable* TypedNodeTable = Cast<UCustomizableObjectNodeTable>(Node))
@@ -2662,7 +2802,7 @@ TArray<TTuple<USkeletalMesh*, TSoftClassPtr<UAnimInstance>>> GetSkeletalMeshesIn
 	{
 		if (SkeletalMeshNode->SkeletalMesh)
 		{
-			SkeletalMeshesInfo.Emplace(SkeletalMeshNode->SkeletalMesh, SkeletalMeshNode->AnimInstance);
+			SkeletalMeshesInfo.Emplace(SkeletalMeshNode->SkeletalMesh.Get(), SkeletalMeshNode->AnimInstance);
 		}
 	}
 	else if (const UCustomizableObjectNodeTable* TableNode = Cast<UCustomizableObjectNodeTable>(SkeletalMeshOrTableNode))
@@ -2875,7 +3015,7 @@ bool GetAndValidateReshapePhysicsToDeform(
 	{
 		TArray<PhysicsInfoType> PhysicsAssetsInfo;
 
-		const bool bAnimBpOverridePhysicsManipulationEnabled = GenerationContext.Options.bAnimBpPhysicsManipulationEnabled;
+		const bool bAnimBpOverridePhysicsManipulationEnabled = GenerationContext.CompilationContext.Options.bAnimBpPhysicsManipulationEnabled;
 		for (const MeshInfoType& Mesh : SkeletalMeshesInfo)
 		{
 			const USkeletalMesh* SkeletalMesh = Mesh.Get<USkeletalMesh*>();
@@ -3228,7 +3368,7 @@ mu::NodeMeshPtr GenerateMorphMesh(const UEdGraphPin* Pin,
 
 						for (const FName BoneName : BonesToDeform)
 						{
-							Result->AddBoneToDeform(GenerationContext.GetBoneUnique(BoneName));
+							Result->AddBoneToDeform(GenerationContext.CompilationContext.GetBoneUnique(BoneName));
 						}
 					}
 
@@ -3244,7 +3384,7 @@ mu::NodeMeshPtr GenerateMorphMesh(const UEdGraphPin* Pin,
 	
 						for (const FName& PhysicsBoneName : PhysicsToDeform)
 						{
-							Result->AddPhysicsBodyToDeform(GenerationContext.GetBoneUnique(PhysicsBoneName));
+							Result->AddPhysicsBodyToDeform(GenerationContext.CompilationContext.GetBoneUnique(PhysicsBoneName));
 						}
 					}
 						
@@ -3346,7 +3486,7 @@ mu::Ptr<mu::NodeMesh> GenerateMutableSourceMesh(const UEdGraphPin* Pin,
 
 			TArray<FCustomizableObjectStreameableResourceId> StreamedResources;
 
-			if (GenerationContext.Object->GetPrivate()->IsAssetUserDataMergeEnabled())
+			if (GenerationContext.CompilationContext.Object->GetPrivate()->IsAssetUserDataMergeEnabled())
 			{
 				const TArray<UAssetUserData*>* AssetUserDataArray = TypedNodeSkel->SkeletalMesh->GetAssetUserDataArray();
 
@@ -3383,6 +3523,7 @@ mu::Ptr<mu::NodeMesh> GenerateMutableSourceMesh(const UEdGraphPin* Pin,
 			{
 				MeshNode->SetValue(MutableMesh);
 
+				// TODO
 				// Add the potentially required morphs. For now always add all morphs. A better implementation can narrow down with:
 				// - detecting morph names in all "ModifierMorphMeshSection" nodes
 				// - detecting morph names in "ModifierMorphMeshSection" nodes relevant for this current context.
@@ -3393,7 +3534,7 @@ mu::Ptr<mu::NodeMesh> GenerateMutableSourceMesh(const UEdGraphPin* Pin,
 					for (TObjectPtr<UMorphTarget> Morph : Morphs)
 					{
 						FString MorphTargetName = Morph->GetName();
-						TSharedPtr<mu::FMesh> MorphedMesh = BuildMorphedMutableMeshFromMesh(MutableMesh, TypedNodeSkel->SkeletalMesh, MorphTargetName, LODIndex, SectionIndex);
+						TSharedPtr<mu::FMesh> MorphedMesh = BuildMorphedMutableMeshFromMesh(MutableMesh, TypedNodeSkel->SkeletalMesh.Get(), MorphTargetName, LODIndex, SectionIndex);
 						MeshNode->AddMorph(MorphTargetName,MorphedMesh);
 					}
 				}
@@ -3406,19 +3547,19 @@ mu::Ptr<mu::NodeMesh> GenerateMutableSourceMesh(const UEdGraphPin* Pin,
 
 					GenerationContext.AddParticipatingObject(*PhysicsAsset);
 
-					const int32 AssetIndex = GenerationContext.PhysicsAssets.AddUnique(PhysicsAsset);
+					const int32 AssetIndex = GenerationContext.CompilationContext.PhysicsAssets.AddUnique(PhysicsAsset);
 					FString PhysicsAssetTag = FString("__PA:") + FString::FromInt(AssetIndex);
 
 					AddTagToMutableMeshUnique(*MutableMesh, PhysicsAssetTag);
 				}
 
-				if (GenerationContext.Options.bSkinWeightProfilesEnabled && LODIndex>=0)
+				if (GenerationContext.CompilationContext.Options.bSkinWeightProfilesEnabled && LODIndex>=0)
 				{
 					if (ImportedModel && ImportedModel->LODModels.IsValidIndex(LODIndex))
 					{
 						for (const auto& SkinWeightProfile : ImportedModel->LODModels[LODIndex].SkinWeightProfiles)
 						{
-							const int32 ProfileId = static_cast<int32>(GenerationContext.GetSkinWeightProfileIdUnique(SkinWeightProfile.Key));
+							const int32 ProfileId = static_cast<int32>(GenerationContext.CompilationContext.GetSkinWeightProfileIdUnique(SkinWeightProfile.Key));
 							MeshData.SkinWeightProfilesSemanticIndices.AddUnique(ProfileId);
 						}
 					}
@@ -3439,7 +3580,7 @@ mu::Ptr<mu::NodeMesh> GenerateMutableSourceMesh(const UEdGraphPin* Pin,
 					MutableMesh->AddStreamedResource(BitCast<uint64>(ResourceId));
 				}
 
-				AddSocketTagsToMesh(TypedNodeSkel->SkeletalMesh, MutableMesh, GenerationContext);
+				AddSocketTagsToMesh(TypedNodeSkel->SkeletalMesh.Get(), MutableMesh, GenerationContext);
 
 				if (UCustomizableObjectSystem::GetInstance()->IsMutableAnimInfoDebuggingEnabled())
 				{
@@ -3531,7 +3672,7 @@ mu::Ptr<mu::NodeMesh> GenerateMutableSourceMesh(const UEdGraphPin* Pin,
 			}
 
 
-			const FString MeshName = GetNameSafe(TypedNodeSkel->SkeletalMesh).ToLower();
+			const FString MeshName = TypedNodeSkel->SkeletalMesh.GetLongPackageName().ToLower();
 			MeshNode->SourceDataDescriptor.SourceId = CityHash32(reinterpret_cast<const char*>(*MeshName), MeshName.Len() * sizeof(FString::ElementType));
 			MeshNode->SourceDataDescriptor.SourceHighResMips = 0;
 
@@ -3543,7 +3684,7 @@ mu::Ptr<mu::NodeMesh> GenerateMutableSourceMesh(const UEdGraphPin* Pin,
 		}
 		else
 		{
-			GenerationContext.Log(LOCTEXT("MissingskeletlMesh", "No Skeletal Mesh set in the SkeletalMesh node."), Node);
+			GenerationContext.Log(LOCTEXT("MissingSkeletalMesh", "No Skeletal Mesh set in the SkeletalMesh node."), Node);
 		}
 	}
 
@@ -3626,7 +3767,7 @@ mu::Ptr<mu::NodeMesh> GenerateMutableSourceMesh(const UEdGraphPin* Pin,
 				LayoutNode->SetMessageContext(Node);  // We need it here because we create multiple nodes.
 
 
-				const FString MeshName = GetNameSafe(TypedNodeStatic->StaticMesh).ToLower();
+				const FString MeshName = TypedNodeStatic->StaticMesh.GetLongPackageName().ToLower();
 				MeshNode->SourceDataDescriptor.SourceId = CityHash32(reinterpret_cast<const char*>(*MeshName), MeshName.Len() * sizeof(FString::ElementType));
 				MeshNode->SourceDataDescriptor.SourceHighResMips = 0;
 			}
@@ -4000,7 +4141,7 @@ mu::Ptr<mu::NodeMesh> GenerateMutableSourceMesh(const UEdGraphPin* Pin,
 				
 				for (const FName& BoneName : BonesToDeform)
 				{
-					MeshNode->AddBoneToDeform(GenerationContext.GetBoneUnique(BoneName));
+					MeshNode->AddBoneToDeform(GenerationContext.CompilationContext.GetBoneUnique(BoneName));
 				}
 			}
 
@@ -4015,7 +4156,7 @@ mu::Ptr<mu::NodeMesh> GenerateMutableSourceMesh(const UEdGraphPin* Pin,
 
 				for (const FName& PhysicsBoneName : PhysicsToDeform)
 				{
-					MeshNode->AddPhysicsBodyToDeform(GenerationContext.GetBoneUnique(PhysicsBoneName));
+					MeshNode->AddPhysicsBodyToDeform(GenerationContext.CompilationContext.GetBoneUnique(PhysicsBoneName));
 				}	
 			}
 			
@@ -4263,8 +4404,8 @@ mu::Ptr<mu::NodeMesh> GenerateMutableSourceMesh(const UEdGraphPin* Pin,
 									MeshData.MinNumTriangles = ImportedModel->LODModels[LODIndex].Sections[SectionIndex].NumTriangles;
 
 									// With tables we have to use the default skeletal mesh to tell if the surface mesh needs morphs or clothing.
-									MeshData.bHasRealTimeMorphs = GenerationContext.Options.bRealTimeMorphTargetsEnabled && DefaultSkeletalMesh->GetMorphTargets().Num(); 
-									MeshData.bHasClothing = GenerationContext.Options.bClothingEnabled && ImportedModel->LODModels[LODIndex].HasClothData();
+									MeshData.bHasRealTimeMorphs = GenerationContext.CompilationContext.Options.bRealTimeMorphTargetsEnabled && DefaultSkeletalMesh->GetMorphTargets().Num();
+									MeshData.bHasClothing = GenerationContext.CompilationContext.Options.bClothingEnabled && ImportedModel->LODModels[LODIndex].HasClothData();
 								}
 							}
 
diff --git a/Engine/Plugins/Mutable/Source/CustomizableObjectEditor/Private/MuCOE/GenerateMutableSource/GenerateMutableSourceMesh.h b/Engine/Plugins/Mutable/Source/CustomizableObjectEditor/Private/MuCOE/GenerateMutableSource/GenerateMutableSourceMesh.h
--- a/Engine/Plugins/Mutable/Source/CustomizableObjectEditor/Private/MuCOE/GenerateMutableSource/GenerateMutableSourceMesh.h
+++ b/Engine/Plugins/Mutable/Source/CustomizableObjectEditor/Private/MuCOE/GenerateMutableSource/GenerateMutableSourceMesh.h
@@ -14,6 +14,7 @@ class UEdGraphPin;
 class UObject;
 class USkeletalMesh;
 class UStaticMesh;
+class UStreamableRenderAsset;
 struct FMutableGraphGenerationContext;
 struct FMutableGraphMeshGenerationData;
 
@@ -49,8 +50,8 @@ void GetLODAndSectionForAutomaticLODs(const FMutableGraphGenerationContext& Cont
  * @param CurrentNode The node being compiled, can be a skeletal mesh or a table node.
  * @param TableReferenceSkeletalMesh Only set if compiling a table node. It refers to the reference mesh of the skeletal mesh column being compiled.
  * @return Mutable Mesh. Nullptr if there has been an error. Empty mesh if the Skeletal Mesh does not contain the requested LOD + Section. */
-TSharedPtr<mu::FMesh> ConvertSkeletalMeshToMutable(const USkeletalMesh* InSkeletalMesh, const TSoftClassPtr<UAnimInstance>& AnimBp,
-                                         int32 LODIndexConnected,  int32 SectionIndexConnected,
+TSharedPtr<mu::FMesh> ConvertSkeletalMeshToMutable(TSoftObjectPtr<const USkeletalMesh> InSkeletalMesh, const TSoftClassPtr<UAnimInstance>& AnimBp,
+                                         bool bMeshMustExist,
                                          int32 LODIndex, int32 SectionIndex, uint32 SurfaceMetadataId,
                                          FMutableGraphGenerationContext& GenerationContext,
                                          const UCustomizableObjectNode* CurrentNode, 
@@ -74,7 +75,7 @@ TSharedPtr<mu::FMesh> ConvertStaticMeshToMutable(const UStaticMesh* StaticMesh,
  * @param CurrentNode The node being compiled, can be a skeletal mesh or a table node.
  * @param TableReferenceSkeletalMesh Only set if compiling a table node. It refers to the reference mesh of the skeletal mesh column being compiled.
  * @return Mutable Mesh. Nullptr if there has been an error. Empty mesh if the Skeletal Mesh does not contain the requested LOD + Section. */
-TSharedPtr<mu::FMesh> GenerateMutableMesh(UObject* Mesh, const TSoftClassPtr<UAnimInstance>& AnimBp,
+TSharedPtr<mu::FMesh> GenerateMutableMesh(TSoftObjectPtr<UStreamableRenderAsset> Mesh, const TSoftClassPtr<UAnimInstance>& AnimBp,
                                 int32 LODIndexConnected, int32 SectionIndexConnected,
                                 int32 LODIndex, int32 SectionIndex, const FString& MeshUniqueTags, uint32 SurfaceMetadataId, 
                                 FMutableGraphGenerationContext& GenerationContext, 
diff --git a/Engine/Plugins/Mutable/Source/CustomizableObjectEditor/Private/MuCOE/GenerateMutableSource/GenerateMutableSourceModifier.cpp b/Engine/Plugins/Mutable/Source/CustomizableObjectEditor/Private/MuCOE/GenerateMutableSource/GenerateMutableSourceModifier.cpp
--- a/Engine/Plugins/Mutable/Source/CustomizableObjectEditor/Private/MuCOE/GenerateMutableSource/GenerateMutableSourceModifier.cpp
+++ b/Engine/Plugins/Mutable/Source/CustomizableObjectEditor/Private/MuCOE/GenerateMutableSource/GenerateMutableSourceModifier.cpp
@@ -87,7 +87,7 @@ mu::Ptr<mu::NodeModifier> GenerateMutableSourceModifier(const UEdGraphPin * Pin,
 		ClipNode->SetParams(TypedNodeClip->B, TypedNodeClip->Exponent);
 		ClipNode->SetMorphEllipse(TypedNodeClip->Radius, TypedNodeClip->Radius2, TypedNodeClip->RotationAngle);
 
-		ClipNode->SetVertexSelectionBone(GenerationContext.GetBoneUnique(TypedNodeClip->BoneName), TypedNodeClip->MaxEffectRadius);
+		ClipNode->SetVertexSelectionBone(GenerationContext.CompilationContext.GetBoneUnique(TypedNodeClip->BoneName), TypedNodeClip->MaxEffectRadius);
 
 		ClipNode->MultipleTagsPolicy = TypedNodeClip->MultipleTagPolicy;
 		ClipNode->RequiredTags = TypedNodeClip->RequiredTags;
@@ -304,8 +304,8 @@ mu::Ptr<mu::NodeModifier> GenerateMutableSourceModifier(const UEdGraphPin * Pin,
 				mu::Ptr<mu::NodeMeshFormat> MeshFormat = new mu::NodeMeshFormat();
 				SetSurfaceFormat(GenerationContext,
 					MeshFormat->GetVertexBuffers(), MeshFormat->GetIndexBuffers(), MeshData,
-					GenerationContext.Options.CustomizableObjectNumBoneInfluences,
-					GenerationContext.Options.b16BitBoneWeightsEnabled);
+					GenerationContext.CompilationContext.Options.CustomizableObjectNumBoneInfluences,
+					GenerationContext.CompilationContext.Options.b16BitBoneWeightsEnabled);
 
 				MeshFormat->SetSource(MeshPtr.get());
 
diff --git a/Engine/Plugins/Mutable/Source/CustomizableObjectEditor/Private/MuCOE/GenerateMutableSource/GenerateMutableSourceSurface.cpp b/Engine/Plugins/Mutable/Source/CustomizableObjectEditor/Private/MuCOE/GenerateMutableSource/GenerateMutableSourceSurface.cpp
--- a/Engine/Plugins/Mutable/Source/CustomizableObjectEditor/Private/MuCOE/GenerateMutableSource/GenerateMutableSourceSurface.cpp
+++ b/Engine/Plugins/Mutable/Source/CustomizableObjectEditor/Private/MuCOE/GenerateMutableSource/GenerateMutableSourceSurface.cpp
@@ -64,7 +64,7 @@ void SetSurfaceFormat( FMutableGraphGenerationContext& GenerationContext,
 {
 	// Limit skinning weights if necessary
 	// \todo: make it more flexible to support 3 or 5 or 1 weight, since there is support for this in 4.25
-	const int32 MutableBonesPerVertex = FGPUBaseSkinVertexFactory::UseUnlimitedBoneInfluences(MeshData.MaxNumBonesPerVertex, GenerationContext.Options.TargetPlatform) &&
+	const int32 MutableBonesPerVertex = FGPUBaseSkinVertexFactory::UseUnlimitedBoneInfluences(MeshData.MaxNumBonesPerVertex, GenerationContext.CompilationContext.Options.TargetPlatform) &&
 										MeshData.MaxNumBonesPerVertex < (int32)ECustomizableObjectNumBoneInfluences ?
 											MeshData.MaxNumBonesPerVertex :
 											(int32)ECustomizableObjectNumBoneInfluences;
@@ -73,7 +73,7 @@ void SetSurfaceFormat( FMutableGraphGenerationContext& GenerationContext,
 	
 	if (MutableBonesPerVertex != MeshData.MaxNumBonesPerVertex)
 	{
-		UE_LOG(LogMutable, Verbose, TEXT("In object [%s] Mesh bone number adjusted from %d to %d."), *GenerationContext.Object->GetName(), MeshData.MaxNumBonesPerVertex, MutableBonesPerVertex);
+		UE_LOG(LogMutable, Verbose, TEXT("In object [%s] Mesh bone number adjusted from %d to %d."), *GenerationContext.GetObjectName(), MeshData.MaxNumBonesPerVertex, MutableBonesPerVertex);
 	}
 
 	int MutableBufferCount = MUTABLE_VERTEXBUFFER_TEXCOORDS + 1;
@@ -345,8 +345,8 @@ mu::Ptr<mu::NodeSurface> GenerateMutableSourceSurface(const UEdGraphPin * Pin, F
 					MeshFormatNode->SetSource(MeshNode.get());
 					SetSurfaceFormat(GenerationContext,
 						MeshFormatNode->GetVertexBuffers(), MeshFormatNode->GetIndexBuffers(), MeshData,
-						GenerationContext.Options.CustomizableObjectNumBoneInfluences,
-						GenerationContext.Options.b16BitBoneWeightsEnabled);
+						GenerationContext.CompilationContext.Options.CustomizableObjectNumBoneInfluences,
+						GenerationContext.CompilationContext.Options.b16BitBoneWeightsEnabled);
 
 					// \TODO: Make it an option?
 					MeshFormatNode->SetOptimizeBuffers(true);
@@ -522,7 +522,7 @@ mu::Ptr<mu::NodeSurface> GenerateMutableSourceSurface(const UEdGraphPin * Pin, F
 
 
 								// MaxTextureSize setting. Based on the ReferenceTexture and Platform settings.
-								const UTextureLODSettings& TextureLODSettings = GenerationContext.Options.TargetPlatform->GetTextureLODSettings();
+								const UTextureLODSettings& TextureLODSettings = GenerationContext.CompilationContext.Options.TargetPlatform->GetTextureLODSettings();
 								Props.MaxTextureSize = GetMaxTextureSize(*ReferenceTexture, TextureLODSettings);
 
 								// ReferenceTexture source size. Textures contributing to this Image should be equal to or smaller than TextureSize. 
@@ -587,8 +587,8 @@ mu::Ptr<mu::NodeSurface> GenerateMutableSourceSurface(const UEdGraphPin * Pin, F
 
 										// Calculate the number of mips to tag as high res for this image.
 										int32 TotalMips = mu::FImage::GetMipmapCount(ImageConstant->GetSizeX(), ImageConstant->GetSizeY());
-										int32 NumMipsBeyondMin = FMath::Max(0, TotalMips - int32(MipsToSkip) - GenerationContext.Options.MinDiskMips);
-										int32 HighResMipsForThisImage = FMath::Min(NumMipsBeyondMin, GenerationContext.Options.NumHighResImageMips);
+										int32 NumMipsBeyondMin = FMath::Max(0, TotalMips - int32(MipsToSkip) - GenerationContext.CompilationContext.Options.MinDiskMips);
+										int32 HighResMipsForThisImage = FMath::Min(NumMipsBeyondMin, GenerationContext.CompilationContext.Options.NumHighResImageMips);
 										ConstImageNode->SourceDataDescriptor.SourceHighResMips = HighResMipsForThisImage;
 
 										const FString TextureName = GetNameSafe(Texture2D).ToLower();
@@ -686,8 +686,8 @@ mu::Ptr<mu::NodeSurface> GenerateMutableSourceSurface(const UEdGraphPin * Pin, F
 									CompositedImage->SetNormal(NormalCompositeMipmapImage);
 
 									int32 TotalMips = mu::FImage::GetMipmapCount(ImageConstant->GetSizeX(), ImageConstant->GetSizeY());
-									int32 NumMipsBeyondMin = FMath::Max(0, TotalMips - int32(MipsToSkip) - GenerationContext.Options.MinDiskMips);
-									int32 HighResMipsForThisImage = FMath::Min(NumMipsBeyondMin, GenerationContext.Options.NumHighResImageMips);
+									int32 NumMipsBeyondMin = FMath::Max(0, TotalMips - int32(MipsToSkip) - GenerationContext.CompilationContext.Options.MinDiskMips);
+									int32 HighResMipsForThisImage = FMath::Min(NumMipsBeyondMin, GenerationContext.CompilationContext.Options.NumHighResImageMips);
 									CompositeNormalImage->SourceDataDescriptor.SourceHighResMips = HighResMipsForThisImage;
 
 									const FString TextureName = GetNameSafe(ReferenceCompositeNormalTexture).ToLower();
@@ -705,17 +705,17 @@ mu::Ptr<mu::NodeSurface> GenerateMutableSourceSurface(const UEdGraphPin * Pin, F
 							LastImage = FormatImage;
 
 							TArray<TArray<FTextureBuildSettings>> BuildSettingsPerFormatPerLayer;
-							if (GenerationContext.Options.TargetPlatform)
+							if (GenerationContext.CompilationContext.Options.TargetPlatform)
 							{
-								ReferenceTexture->GetTargetPlatformBuildSettings(GenerationContext.Options.TargetPlatform, BuildSettingsPerFormatPerLayer);
+								ReferenceTexture->GetTargetPlatformBuildSettings(GenerationContext.CompilationContext.Options.TargetPlatform, BuildSettingsPerFormatPerLayer);
 
-								bool bIsServerOnly = GenerationContext.Options.TargetPlatform->IsServerOnly();
+								bool bIsServerOnly = GenerationContext.CompilationContext.Options.TargetPlatform->IsServerOnly();
 								// Suppress the message for server only platforms. Images are discarded is ServerOnly. 
 								if (!bIsServerOnly && BuildSettingsPerFormatPerLayer.IsEmpty())
 								{
 									const FString ReplacedImageFormatMsg = FString::Printf(TEXT("In object [%s] for platform [%s] the unsupported image format of texture [%s] is used, RGBA_UByte will be used instead."),
-										*GenerationContext.Object->GetName(),
-										*GenerationContext.Options.TargetPlatform->PlatformName(),
+										*GenerationContext.GetObjectName(),
+										*GenerationContext.CompilationContext.Options.TargetPlatform->PlatformName(),
 										*ReferenceTexture->GetName());
 									const FText ReplacedImageFormatText = FText::FromString(ReplacedImageFormatMsg);
 									GenerationContext.Log(ReplacedImageFormatText, Node, EMessageSeverity::Info);
@@ -724,8 +724,8 @@ mu::Ptr<mu::NodeSurface> GenerateMutableSourceSurface(const UEdGraphPin * Pin, F
 								else if (BuildSettingsPerFormatPerLayer.Num() > 1)
 								{
 									const FString ReplacedImageFormatMsg = FString::Printf(TEXT("In object [%s] for platform [%s] the image format of texture [%s] has multiple target formats. Only one will be used.."),
-										*GenerationContext.Object->GetName(),
-										*GenerationContext.Options.TargetPlatform->PlatformName(),
+										*GenerationContext.GetObjectName(),
+										*GenerationContext.CompilationContext.Options.TargetPlatform->PlatformName(),
 										*ReferenceTexture->GetName());
 									const FText ReplacedImageFormatText = FText::FromString(ReplacedImageFormatMsg);
 									GenerationContext.Log(ReplacedImageFormatText, Node, EMessageSeverity::Info);
@@ -737,7 +737,7 @@ mu::Ptr<mu::NodeSurface> GenerateMutableSourceSurface(const UEdGraphPin * Pin, F
 							{
 								const TArray<FTextureBuildSettings>& BuildSettingsPerLayer = BuildSettingsPerFormatPerLayer[0];
 
-								if (GenerationContext.Options.TextureCompression!=ECustomizableObjectTextureCompression::None)
+								if (GenerationContext.CompilationContext.Options.TextureCompression!=ECustomizableObjectTextureCompression::None)
 								{
 									static ITextureFormatManagerModule* TextureFormatManager = nullptr;
 									if (!TextureFormatManager)
@@ -755,9 +755,9 @@ mu::Ptr<mu::NodeSurface> GenerateMutableSourceSurface(const UEdGraphPin * Pin, F
 									mu::EImageFormat MutableFormatIfAlpha = QualityAndPerformanceFix(UnrealToMutablePixelFormat(UnrealTargetPlatformFormatAlpha,true));
 
 									// Temp hack to enable RG->LA 
-									if (GenerationContext.Options.TargetPlatform)
+									if (GenerationContext.CompilationContext.Options.TargetPlatform)
 									{
-										bool bUseLA = GenerationContext.Options.TargetPlatform->SupportsFeature(ETargetPlatformFeatures::NormalmapLAEncodingMode);
+										bool bUseLA = GenerationContext.CompilationContext.Options.TargetPlatform->SupportsFeature(ETargetPlatformFeatures::NormalmapLAEncodingMode);
 										if (bUseLA)
 										{
 											// See GetQualityFormat in TextureFormatASTC.cpp to understand why
@@ -786,7 +786,7 @@ mu::Ptr<mu::NodeSurface> GenerateMutableSourceSurface(const UEdGraphPin * Pin, F
 									// Unsupported format: look for something generic
 									if (MutableFormat == mu::EImageFormat::None)
 									{
-										const FString ReplacedImageFormatMsg = FString::Printf(TEXT("In object [%s] the unsupported image format %d is used, RGBA_UByte will be used instead."), *GenerationContext.Object->GetName(), UnrealTargetPlatformFormat );
+										const FString ReplacedImageFormatMsg = FString::Printf(TEXT("In object [%s] the unsupported image format %d is used, RGBA_UByte will be used instead."), *GenerationContext.GetObjectName(), UnrealTargetPlatformFormat );
 										const FText ReplacedImageFormatText = FText::FromString(ReplacedImageFormatMsg);
 										GenerationContext.Log(ReplacedImageFormatText, Node, EMessageSeverity::Info);
 										UE_LOG(LogMutable, Log, TEXT("%s"), *ReplacedImageFormatMsg);
@@ -794,7 +794,7 @@ mu::Ptr<mu::NodeSurface> GenerateMutableSourceSurface(const UEdGraphPin * Pin, F
 									}
 									if (MutableFormatIfAlpha == mu::EImageFormat::None)
 									{
-										const FString ReplacedImageFormatMsg = FString::Printf(TEXT("In object [%s] the unsupported image format %d is used, RGBA_UByte will be used instead."), *GenerationContext.Object->GetName(), UnrealTargetPlatformFormatAlpha);
+										const FString ReplacedImageFormatMsg = FString::Printf(TEXT("In object [%s] the unsupported image format %d is used, RGBA_UByte will be used instead."), *GenerationContext.GetObjectName(), UnrealTargetPlatformFormatAlpha);
 										const FText ReplacedImageFormatText = FText::FromString(ReplacedImageFormatMsg);
 										GenerationContext.Log(ReplacedImageFormatText, Node, EMessageSeverity::Info);
 										UE_LOG(LogMutable, Log, TEXT("%s"), *ReplacedImageFormatMsg);
diff --git a/Engine/Plugins/Mutable/Source/CustomizableObjectEditor/Private/MuCOE/GenerateMutableSource/GenerateMutableSourceTable.cpp b/Engine/Plugins/Mutable/Source/CustomizableObjectEditor/Private/MuCOE/GenerateMutableSource/GenerateMutableSourceTable.cpp
--- a/Engine/Plugins/Mutable/Source/CustomizableObjectEditor/Private/MuCOE/GenerateMutableSource/GenerateMutableSourceTable.cpp
+++ b/Engine/Plugins/Mutable/Source/CustomizableObjectEditor/Private/MuCOE/GenerateMutableSource/GenerateMutableSourceTable.cpp
@@ -248,7 +248,7 @@ bool FillTableColumn(const UCustomizableObjectNodeTable* TableNode, mu::Ptr<mu::
 			
 			TArray<FCustomizableObjectStreameableResourceId> StreamedResources;
 
-			if (GenerationContext.Object->GetPrivate()->IsAssetUserDataMergeEnabled())
+			if (GenerationContext.CompilationContext.Object->GetPrivate()->IsAssetUserDataMergeEnabled())
 			{
 				const TArray<UAssetUserData*>* AssetUserDataArray = SkeletalMesh->GetAssetUserDataArray();
 
@@ -288,9 +288,9 @@ bool FillTableColumn(const UCustomizableObjectNodeTable* TableNode, mu::Ptr<mu::
 				{
  					TSoftObjectPtr<UPhysicsAsset> PhysicsAsset = SkeletalMesh->GetPhysicsAsset();
 
-					GenerationContext.AddParticipatingObject(*PhysicsAsset); 						
+					GenerationContext.AddParticipatingObject(*PhysicsAsset);
 
-					const int32 AssetIndex = GenerationContext.PhysicsAssets.AddUnique(PhysicsAsset);
+					const int32 AssetIndex = GenerationContext.CompilationContext.PhysicsAssets.AddUnique(PhysicsAsset);
 					FString PhysicsAssetTag = FString("__PA:") + FString::FromInt(AssetIndex);
 
 					AddTagToMutableMeshUnique(*MutableMesh, PhysicsAssetTag);
@@ -460,7 +460,7 @@ bool FillTableColumn(const UCustomizableObjectNodeTable* TableNode, mu::Ptr<mu::
 
 			GenerationContext.AddParticipatingObject(*ReferenceMaterial);
 
-			const bool bTableMaterialCheckDisabled = GenerationContext.Object->GetPrivate()->IsTableMaterialsParentCheckDisabled();
+			const bool bTableMaterialCheckDisabled = GenerationContext.CompilationContext.Object->GetPrivate()->IsTableMaterialsParentCheckDisabled();
 			const bool bMaterialParentMismatch = !bTableMaterialCheckDisabled && MaterialInstance
 												 && ReferenceMaterial->GetMaterial() != MaterialInstance->GetMaterial();
 
@@ -602,7 +602,7 @@ bool FillTableColumn(const UCustomizableObjectNodeTable* TableNode, mu::Ptr<mu::
 				for (int32 i = 0; i < ArrayBoneName.Num(); ++i)
 				{
 					const FName BoneName = ArrayBoneName[i];
-					const mu::FBoneName& BoneId = GenerationContext.GetBoneUnique(BoneName);
+					const mu::FBoneName& BoneId = GenerationContext.CompilationContext.GetBoneUnique(BoneName);
 
 					MutableSkeleton->SetDebugName(i, BoneName);
 					MutableSkeleton->SetBoneName(i, { BoneId });
@@ -763,7 +763,7 @@ TArray<FName> GetEnabledRows(const UDataTable& DataTable, const UCustomizableObj
 
 void RestrictRowNamesToSelectedOption(TArray<FName>& InOutRowNames, const UCustomizableObjectNodeTable& TableNode, FMutableGraphGenerationContext& GenerationContext)
 {
-	if (!GenerationContext.Options.ParamNamesToSelectedOptions.IsEmpty())
+	if (!GenerationContext.CompilationContext.Options.ParamNamesToSelectedOptions.IsEmpty())
 	{
 		FMutableParamNameSet* ParamNameSet = GenerationContext.TableToParamNames.Find(TableNode.Table->GetPathName());
 
@@ -774,7 +774,7 @@ void RestrictRowNamesToSelectedOption(TArray<FName>& InOutRowNames, const UCusto
 			for (const FString& ParamName : ParamNameSet->ParamNames)
 			{
 				// If the param is in the map restrict to only the selected option
-				const FString* SelectedOptionString = GenerationContext.Options.ParamNamesToSelectedOptions.Find(ParamName);
+				const FString* SelectedOptionString = GenerationContext.CompilationContext.Options.ParamNamesToSelectedOptions.Find(ParamName);
 
 				if (SelectedOptionString)
 				{
@@ -1089,7 +1089,7 @@ mu::Ptr<mu::FTable> GenerateMutableSourceTable(const UDataTable* DataTable, cons
 {
 	check(DataTable && TableNode);
 
-	if (GenerationContext.Options.ParamNamesToSelectedOptions.IsEmpty())
+	if (GenerationContext.CompilationContext.Options.ParamNamesToSelectedOptions.IsEmpty())
 	{
 		FMutableParamNameSet* ParamNameSet = GenerationContext.TableToParamNames.Find(DataTable->GetPathName());
 
diff --git a/Engine/Plugins/Mutable/Source/CustomizableObjectEditor/Private/MuCOE/Nodes/CustomizableObjectNodeMeshMorphDetails.cpp b/Engine/Plugins/Mutable/Source/CustomizableObjectEditor/Private/MuCOE/Nodes/CustomizableObjectNodeMeshMorphDetails.cpp
--- a/Engine/Plugins/Mutable/Source/CustomizableObjectEditor/Private/MuCOE/Nodes/CustomizableObjectNodeMeshMorphDetails.cpp
+++ b/Engine/Plugins/Mutable/Source/CustomizableObjectEditor/Private/MuCOE/Nodes/CustomizableObjectNodeMeshMorphDetails.cpp
@@ -52,7 +52,7 @@ void FCustomizableObjectNodeMeshMorphDetails::CustomizeDetails( IDetailLayoutBui
 				USkeletalMesh* SkeletalMesh = nullptr;
 				if ( const UCustomizableObjectNodeSkeletalMesh* TypedNodeSkeletalMesh = Cast<UCustomizableObjectNodeSkeletalMesh>(BaseSourcePin->GetOwningNode()) )
 				{
-					SkeletalMesh = TypedNodeSkeletalMesh->SkeletalMesh;
+					SkeletalMesh = TypedNodeSkeletalMesh->SkeletalMesh.LoadSynchronous();
 				}
 				else if (const UCustomizableObjectNodeTable* TypedNodeTable = Cast<UCustomizableObjectNodeTable>(BaseSourcePin->GetOwningNode()))
 				{
diff --git a/Engine/Plugins/Mutable/Source/CustomizableObjectEditor/Private/MuCOE/Nodes/CustomizableObjectNodeMeshMorphStackApplication.cpp b/Engine/Plugins/Mutable/Source/CustomizableObjectEditor/Private/MuCOE/Nodes/CustomizableObjectNodeMeshMorphStackApplication.cpp
--- a/Engine/Plugins/Mutable/Source/CustomizableObjectEditor/Private/MuCOE/Nodes/CustomizableObjectNodeMeshMorphStackApplication.cpp
+++ b/Engine/Plugins/Mutable/Source/CustomizableObjectEditor/Private/MuCOE/Nodes/CustomizableObjectNodeMeshMorphStackApplication.cpp
@@ -82,7 +82,7 @@ TArray<FString> UCustomizableObjectNodeMeshMorphStackApplication::GetMorphList()
 
 		if (const UCustomizableObjectNodeSkeletalMesh* SkeletalMeshNode = Cast< UCustomizableObjectNodeSkeletalMesh >(MeshNodePin->GetOwningNode()))
 		{
-			SkeletalMesh = SkeletalMeshNode->SkeletalMesh;
+			SkeletalMesh = SkeletalMeshNode->SkeletalMesh.LoadSynchronous();
 		}
 		else if (const UCustomizableObjectNodeTable* TableNode = Cast< UCustomizableObjectNodeTable >(MeshNodePin->GetOwningNode()))
 		{
diff --git a/Engine/Plugins/Mutable/Source/CustomizableObjectEditor/Private/MuCOE/Nodes/CustomizableObjectNodeMeshMorphStackDefinition.cpp b/Engine/Plugins/Mutable/Source/CustomizableObjectEditor/Private/MuCOE/Nodes/CustomizableObjectNodeMeshMorphStackDefinition.cpp
--- a/Engine/Plugins/Mutable/Source/CustomizableObjectEditor/Private/MuCOE/Nodes/CustomizableObjectNodeMeshMorphStackDefinition.cpp
+++ b/Engine/Plugins/Mutable/Source/CustomizableObjectEditor/Private/MuCOE/Nodes/CustomizableObjectNodeMeshMorphStackDefinition.cpp
@@ -101,7 +101,7 @@ void UCustomizableObjectNodeMeshMorphStackDefinition::UpdateMorphList()
 
 		if (const UCustomizableObjectNodeSkeletalMesh* SkeletalMeshNode = Cast< UCustomizableObjectNodeSkeletalMesh >(MeshNodePin->GetOwningNode()))
 		{
-			SkeletalMesh = SkeletalMeshNode->SkeletalMesh;
+			SkeletalMesh = SkeletalMeshNode->SkeletalMesh.LoadSynchronous();
 		}
 		else if (const UCustomizableObjectNodeTable* TableNode = Cast< UCustomizableObjectNodeTable >(MeshNodePin->GetOwningNode()))
 		{
@@ -190,7 +190,7 @@ bool UCustomizableObjectNodeMeshMorphStackDefinition::IsNodeOutDatedAndNeedsRefr
 
 		if (const UCustomizableObjectNodeSkeletalMesh* SkeletalMeshNode = Cast< UCustomizableObjectNodeSkeletalMesh >(MeshNodePin->GetOwningNode()))
 		{
-			SkeletalMesh = SkeletalMeshNode->SkeletalMesh;
+			SkeletalMesh = SkeletalMeshNode->SkeletalMesh.LoadSynchronous();
 		}
 		else if (const UCustomizableObjectNodeTable* TableNode = Cast< UCustomizableObjectNodeTable >(MeshNodePin->GetOwningNode()))
 		{
diff --git a/Engine/Plugins/Mutable/Source/CustomizableObjectEditor/Private/MuCOE/Nodes/CustomizableObjectNodeMeshReshapeSelectionDetails.cpp b/Engine/Plugins/Mutable/Source/CustomizableObjectEditor/Private/MuCOE/Nodes/CustomizableObjectNodeMeshReshapeSelectionDetails.cpp
--- a/Engine/Plugins/Mutable/Source/CustomizableObjectEditor/Private/MuCOE/Nodes/CustomizableObjectNodeMeshReshapeSelectionDetails.cpp
+++ b/Engine/Plugins/Mutable/Source/CustomizableObjectEditor/Private/MuCOE/Nodes/CustomizableObjectNodeMeshReshapeSelectionDetails.cpp
@@ -90,7 +90,7 @@ USkeletalMesh* FMeshReshapeBonesReferenceCustomization::GetSkeletalMesh(TSharedP
 				{
 					if (const UCustomizableObjectNodeSkeletalMesh* SkeletalMeshNode = Cast<UCustomizableObjectNodeSkeletalMesh>(SourceMeshPin->GetOwningNode()))
 					{
-						return SkeletalMeshNode->SkeletalMesh;
+						return SkeletalMeshNode->SkeletalMesh.LoadSynchronous();
 					}
 					else if (UCustomizableObjectNodeTable* TableNode = Cast<UCustomizableObjectNodeTable>(SourceMeshPin->GetOwningNode()))
 					{
diff --git a/Engine/Plugins/Mutable/Source/CustomizableObjectEditor/Private/MuCOE/Nodes/CustomizableObjectNodeModifierMorphMeshSectionDetails.cpp b/Engine/Plugins/Mutable/Source/CustomizableObjectEditor/Private/MuCOE/Nodes/CustomizableObjectNodeModifierMorphMeshSectionDetails.cpp
--- a/Engine/Plugins/Mutable/Source/CustomizableObjectEditor/Private/MuCOE/Nodes/CustomizableObjectNodeModifierMorphMeshSectionDetails.cpp
+++ b/Engine/Plugins/Mutable/Source/CustomizableObjectEditor/Private/MuCOE/Nodes/CustomizableObjectNodeModifierMorphMeshSectionDetails.cpp
@@ -277,7 +277,7 @@ void FCustomizableObjectNodeModifierMorphMeshSectionDetails::AddMorphsFromNode(
 				UCustomizableObjectNodeSkeletalMesh* SkeletalNode = Cast<UCustomizableObjectNodeSkeletalMesh>(SourceMeshPin->GetOwningNode());
 				if (SkeletalNode)
 				{
-					SkeletalMesh = SkeletalNode->SkeletalMesh;
+					SkeletalMesh = SkeletalNode->SkeletalMesh.LoadSynchronous();
 				}
 			}
 		}
@@ -293,7 +293,7 @@ void FCustomizableObjectNodeModifierMorphMeshSectionDetails::AddMorphsFromNode(
 				UCustomizableObjectNodeSkeletalMesh* SkeletalNode = Cast<UCustomizableObjectNodeSkeletalMesh>(SourceMeshPin->GetOwningNode());
 				if (SkeletalNode)
 				{
-					SkeletalMesh = SkeletalNode->SkeletalMesh;
+					SkeletalMesh = SkeletalNode->SkeletalMesh.LoadSynchronous();
 				}
 			}
 		}
diff --git a/Engine/Plugins/Mutable/Source/CustomizableObjectEditor/Private/MuCOE/Nodes/CustomizableObjectNodeSkeletalMesh.cpp b/Engine/Plugins/Mutable/Source/CustomizableObjectEditor/Private/MuCOE/Nodes/CustomizableObjectNodeSkeletalMesh.cpp
--- a/Engine/Plugins/Mutable/Source/CustomizableObjectEditor/Private/MuCOE/Nodes/CustomizableObjectNodeSkeletalMesh.cpp
+++ b/Engine/Plugins/Mutable/Source/CustomizableObjectEditor/Private/MuCOE/Nodes/CustomizableObjectNodeSkeletalMesh.cpp
@@ -75,7 +75,7 @@ void UCustomizableObjectNodeSkeletalMesh::AllocateDefaultPins(UCustomizableObjec
 	if (!SkeletalMesh)
 	{
 		UCustomizableObjectNodeSkeletalMeshPinDataMesh* PinData = NewObject<UCustomizableObjectNodeSkeletalMeshPinDataMesh>(this);
-		PinData->Init(-1, -1, -1, nullptr);
+		PinData->Init(-1, -1, -1);
 		
 		DefaultPin = CustomCreatePin(EGPD_Output, Schema->PC_Mesh, FName(SKELETAL_MESH_PIN_NAME), PinData);
 		return;
@@ -109,7 +109,7 @@ void UCustomizableObjectNodeSkeletalMesh::AllocateDefaultPins(UCustomizableObjec
 					FString MeshName = FString::Printf(TEXT("LOD %i - Section %i - Mesh"), LODIndex, SectionIndex);
 
 					UCustomizableObjectNodeSkeletalMeshPinDataMesh* PinData = NewObject<UCustomizableObjectNodeSkeletalMeshPinDataMesh>(this);
-					PinData->Init(LODIndex, SectionIndex, ImportedModel->LODModels[LODIndex].NumTexCoords, SkeletalMesh);
+					PinData->Init(LODIndex, SectionIndex, ImportedModel->LODModels[LODIndex].NumTexCoords);
 					
 					UEdGraphPin* Pin = CustomCreatePin(EGPD_Output, Schema->PC_Mesh, FName(*MeshName), PinData);
 					Pin->PinFriendlyName = FText::FromString(FString::Printf(TEXT("LOD %i - %s - Mesh"), LODIndex, *SectionFriendlyName));
@@ -229,7 +229,7 @@ TArray<UCustomizableObjectLayout*> UCustomizableObjectNodeSkeletalMesh::GetLayou
 }
 
 
-UObject* UCustomizableObjectNodeSkeletalMesh::GetMesh() const
+TSoftObjectPtr<UObject> UCustomizableObjectNodeSkeletalMesh::GetMesh() const
 {
 	return SkeletalMesh;
 }
@@ -511,7 +511,7 @@ void UCustomizableObjectNodeSkeletalMesh::BackwardsCompatibleFixup(int32 Customi
 
 				{
 					UCustomizableObjectNodeSkeletalMeshPinDataMesh* PinData = NewObject<UCustomizableObjectNodeSkeletalMeshPinDataMesh>(this);
-					PinData->Init(LODIndex, SectionIndex, -1, nullptr);
+					PinData->Init(LODIndex, SectionIndex, -1);
 				
 					AddPinData(*Section.MeshPinRef.Get(), *PinData);
 				}
@@ -584,7 +584,7 @@ void UCustomizableObjectNodeSkeletalMesh::BackwardsCompatibleFixup(int32 Customi
 		if (const UEdGraphPin* Pin = DefaultPin.Get())
 		{
 			UCustomizableObjectNodeSkeletalMeshPinDataMesh* PinData = NewObject<UCustomizableObjectNodeSkeletalMeshPinDataMesh>(this);
-			PinData->Init(-1, -1, -1, nullptr);
+			PinData->Init(-1, -1, -1);
 
 			AddPinData(*Pin, *PinData);
 		}
@@ -832,7 +832,8 @@ void SGraphNodeSkeletalMesh::Construct(const FArguments& InArgs, UEdGraphNode* I
 
 	// Thumbnail
 	AssetThumbnailPool = MakeShareable(new FAssetThumbnailPool(32));
-	AssetThumbnail = MakeShareable(new FAssetThumbnail(NodeSkeletalMesh->SkeletalMesh, ThumbnailSize, ThumbnailSize, AssetThumbnailPool));
+	USkeletalMesh* SkeletalMesh = NodeSkeletalMesh->SkeletalMesh.LoadSynchronous();
+	AssetThumbnail = MakeShareable(new FAssetThumbnail(SkeletalMesh, ThumbnailSize, ThumbnailSize, AssetThumbnailPool));
 
 	// Selector
 	FPropertyEditorModule& PropPlugin = FModuleManager::LoadModuleChecked<FPropertyEditorModule>("PropertyEditor");
@@ -1016,7 +1017,7 @@ void UCustomizableObjectNodeSkeletalMeshPinDataMesh::Copy(const UCustomizableObj
 	}
 }
 
-void UCustomizableObjectNodeSkeletalMeshPinDataMesh::Init(int32 InLODIndex, int32 InSectionIndex, int32 NumTexCoords, USkeletalMesh* Mesh)
+void UCustomizableObjectNodeSkeletalMeshPinDataMesh::Init(int32 InLODIndex, int32 InSectionIndex, int32 NumTexCoords)
 {
 	Super::Init(InLODIndex, InSectionIndex);
 
diff --git a/Engine/Plugins/Mutable/Source/CustomizableObjectEditor/Private/MuCOE/Nodes/CustomizableObjectNodeStaticMesh.cpp b/Engine/Plugins/Mutable/Source/CustomizableObjectEditor/Private/MuCOE/Nodes/CustomizableObjectNodeStaticMesh.cpp
--- a/Engine/Plugins/Mutable/Source/CustomizableObjectEditor/Private/MuCOE/Nodes/CustomizableObjectNodeStaticMesh.cpp
+++ b/Engine/Plugins/Mutable/Source/CustomizableObjectEditor/Private/MuCOE/Nodes/CustomizableObjectNodeStaticMesh.cpp
@@ -267,7 +267,7 @@ TArray<class UCustomizableObjectLayout*> UCustomizableObjectNodeStaticMesh::GetL
 }
 
 
-UObject* UCustomizableObjectNodeStaticMesh::GetMesh() const
+TSoftObjectPtr<UObject> UCustomizableObjectNodeStaticMesh::GetMesh() const
 {
 	return StaticMesh;
 }
@@ -481,7 +481,7 @@ void SGraphNodeStaticMesh::Construct(const FArguments& InArgs, UEdGraphNode* InG
 
 	// Thumbnail
 	AssetThumbnailPool = MakeShareable(new FAssetThumbnailPool(32));
-	AssetThumbnail = MakeShareable(new FAssetThumbnail(NodeStaticMesh->StaticMesh, ThumbnailSize, ThumbnailSize, AssetThumbnailPool));
+	AssetThumbnail = MakeShareable(new FAssetThumbnail(NodeStaticMesh->StaticMesh.Get(), ThumbnailSize, ThumbnailSize, AssetThumbnailPool));
 
 	FPropertyEditorModule& PropPlugin = FModuleManager::LoadModuleChecked<FPropertyEditorModule>("PropertyEditor");
 	FSinglePropertyParams SingleDetails;
diff --git a/Engine/Plugins/Mutable/Source/CustomizableObjectEditor/Private/MuCOE/SCustomizableObjectNodeSkeletalMeshRTMorphSelector.cpp b/Engine/Plugins/Mutable/Source/CustomizableObjectEditor/Private/MuCOE/SCustomizableObjectNodeSkeletalMeshRTMorphSelector.cpp
--- a/Engine/Plugins/Mutable/Source/CustomizableObjectEditor/Private/MuCOE/SCustomizableObjectNodeSkeletalMeshRTMorphSelector.cpp
+++ b/Engine/Plugins/Mutable/Source/CustomizableObjectEditor/Private/MuCOE/SCustomizableObjectNodeSkeletalMeshRTMorphSelector.cpp
@@ -188,7 +188,7 @@ void SCustomizableObjectNodeSkeletalMeshRTMorphSelector::GenerateMorphsInfoList(
 {
     Items.Reset();
 
-    const USkeletalMesh* SkeletalMesh = Node->SkeletalMesh;
+    const USkeletalMesh* SkeletalMesh = Node->SkeletalMesh.LoadSynchronous();
 
     if (!SkeletalMesh)
     {
diff --git a/Engine/Plugins/Mutable/Source/CustomizableObjectEditor/Private/MuCOE/SMutableCodeViewer.cpp b/Engine/Plugins/Mutable/Source/CustomizableObjectEditor/Private/MuCOE/SMutableCodeViewer.cpp
--- a/Engine/Plugins/Mutable/Source/CustomizableObjectEditor/Private/MuCOE/SMutableCodeViewer.cpp
+++ b/Engine/Plugins/Mutable/Source/CustomizableObjectEditor/Private/MuCOE/SMutableCodeViewer.cpp
@@ -294,10 +294,13 @@ void SMutableCodeViewer::ClearSelectedTreeRow() const
 	TreeView->ClearSelection();
 }
 
-void SMutableCodeViewer::SetCurrentModel(const TSharedPtr<mu::FModel, ESPMode::ThreadSafe>& InMutableModel, const TArray<TSoftObjectPtr<const UTexture>>& InReferencedTextures)
+void SMutableCodeViewer::SetCurrentModel(const TSharedPtr<mu::FModel, ESPMode::ThreadSafe>& InMutableModel, 
+	const TArray<TSoftObjectPtr<const UTexture>>& InReferencedTextures,
+	const TArray<TSoftObjectPtr<const UObject>>& InReferencedMeshes )
 {
 	MutableModel = InMutableModel;
 	ReferencedTextures = InReferencedTextures;
+	ReferencedMeshes = InReferencedMeshes;
 	PreviewParameters = mu::FModel::NewParameters(MutableModel);
 
 	RootNodes.Empty();
@@ -343,13 +346,14 @@ void SMutableCodeViewer::SetCurrentModel(const TSharedPtr<mu::FModel, ESPMode::T
 
 
 void SMutableCodeViewer::Construct(const FArguments& InArgs, const TSharedPtr<mu::FModel, ESPMode::ThreadSafe>& InMutableModel,
-	const TArray<TSoftObjectPtr<const UTexture>>& InReferencedTextures )
+	const TArray<TSoftObjectPtr<const UTexture>>& InReferencedTextures,
+	const TArray<TSoftObjectPtr<const UObject>>& InReferencedMeshes )
 {
 	// Min width allowed for the column. Needed to avoid having issues with the constants space being to small
 	// and then getting too tall on the y axis crashing the UI drawer.
 	constexpr float MinParametersCollWidth = 200;
 
-	SetCurrentModel(InMutableModel, InReferencedTextures);
+	SetCurrentModel(InMutableModel, InReferencedTextures, InReferencedMeshes);
 	
 	FToolBarBuilder ToolbarBuilder(TSharedPtr<const FUICommandList>(), FMultiBoxCustomization::None, TSharedPtr<FExtender>(), true);
 	ToolbarBuilder.SetLabelVisibility(EVisibility::Visible);
@@ -2478,6 +2482,7 @@ namespace
 
 		/** */
 		TArray<TSoftObjectPtr<const UTexture>> ReferencedTextures;
+		TArray<TSoftObjectPtr<const UObject>> ReferencedMeshes;
 
 	public:
 
@@ -2568,6 +2573,7 @@ namespace
 		virtual TTuple<UE::Tasks::FTask, TFunction<void()>> GetMeshAsync(FName Id, TFunction<void(TSharedPtr<mu::FMesh>)>& ResultCallback) override
 		{
 			// TODO: This is a stub function where Mesh FParameters would be supported.
+			ensure(false);
 
 			// Thread: worker
 			MUTABLE_CPUPROFILER_SCOPE(FUnrealMutableImageProvider::GetMeshAsync);
@@ -2644,6 +2650,7 @@ void SMutableCodeViewer::Tick(const FGeometry& AllottedGeometry, const double In
 
 	TSharedPtr<TestResourceProvider> ExternalResourceProvider = MakeShared<TestResourceProvider>();
 	ExternalResourceProvider->ReferencedTextures = ReferencedTextures;
+	ExternalResourceProvider->ReferencedMeshes = ReferencedMeshes;
 	System->SetExternalResourceProvider(ExternalResourceProvider);
 
 	System->GetPrivate()->BeginBuild(MutableModel);
diff --git a/Engine/Plugins/Mutable/Source/CustomizableObjectEditor/Private/MuCOE/SMutableCodeViewer.h b/Engine/Plugins/Mutable/Source/CustomizableObjectEditor/Private/MuCOE/SMutableCodeViewer.h
--- a/Engine/Plugins/Mutable/Source/CustomizableObjectEditor/Private/MuCOE/SMutableCodeViewer.h
+++ b/Engine/Plugins/Mutable/Source/CustomizableObjectEditor/Private/MuCOE/SMutableCodeViewer.h
@@ -60,7 +60,8 @@ public:
 	SLATE_END_ARGS()
 
 	void Construct(const FArguments& InArgs, const TSharedPtr<mu::FModel, ESPMode::ThreadSafe>& InMutableModel,
-		const TArray<TSoftObjectPtr<const UTexture>>& ReferencedTextures );
+		const TArray<TSoftObjectPtr<const UTexture>>& ReferencedTextures,
+		const TArray<TSoftObjectPtr<const UObject>>& ReferencedMeshes);
 
 	// SWidget interface
 	virtual void Tick(const FGeometry& AllottedGeometry, const double InCurrentTime, const float InDeltaTime) override;
@@ -79,6 +80,7 @@ private:
 	
 	/** Array of external referenced textures in MutableModel, indexed by id. */
 	TArray<TSoftObjectPtr<const UTexture>> ReferencedTextures;
+	TArray<TSoftObjectPtr<const UObject>> ReferencedMeshes;
 
 	/** Selected model operation for preview. */
 	mu::OP::ADDRESS SelectedOperationAddress = 0;
@@ -186,7 +188,7 @@ private:
 	TMap< mu::OP::ADDRESS, TSharedPtr<FMutableCodeTreeElement>> ExpandedElements;
 
 	/** Prepare the widget for the given model. */
-	void SetCurrentModel(const TSharedPtr<mu::FModel, ESPMode::ThreadSafe>&, const TArray<TSoftObjectPtr<const UTexture>>& ReferencedTextures);
+	void SetCurrentModel(const TSharedPtr<mu::FModel, ESPMode::ThreadSafe>&, const TArray<TSoftObjectPtr<const UTexture>>& ReferencedTextures, const TArray<TSoftObjectPtr<const UObject>>& ReferencedMeshes);
 
 	/** Before any UI operation generate all the elements that may be navigable over the tree. No children of duplicated 
 	 * addresses will be generated.
diff --git a/Engine/Plugins/Mutable/Source/CustomizableObjectEditor/Private/MuCOE/SMutableGraphViewer.h b/Engine/Plugins/Mutable/Source/CustomizableObjectEditor/Private/MuCOE/SMutableGraphViewer.h
--- a/Engine/Plugins/Mutable/Source/CustomizableObjectEditor/Private/MuCOE/SMutableGraphViewer.h
+++ b/Engine/Plugins/Mutable/Source/CustomizableObjectEditor/Private/MuCOE/SMutableGraphViewer.h
@@ -37,6 +37,8 @@ public:
 
 	SLATE_ARGUMENT(TArray<TSoftObjectPtr<const UTexture>>, ReferencedRuntimeTextures)
 	SLATE_ARGUMENT(TArray<FMutableSourceTextureData>, ReferencedCompileTextures)
+	SLATE_ARGUMENT(TArray<TSoftObjectPtr<const UObject>>, ReferencedRuntimeMeshes)
+	SLATE_ARGUMENT(TArray<FMutableSourceMeshData>, ReferencedCompileMeshes)
 
 	SLATE_END_ARGS()
 
diff --git a/Engine/Plugins/Mutable/Source/CustomizableObjectEditor/Private/MuCOE/SMutableObjectViewer.cpp b/Engine/Plugins/Mutable/Source/CustomizableObjectEditor/Private/MuCOE/SMutableObjectViewer.cpp
--- a/Engine/Plugins/Mutable/Source/CustomizableObjectEditor/Private/MuCOE/SMutableObjectViewer.cpp
+++ b/Engine/Plugins/Mutable/Source/CustomizableObjectEditor/Private/MuCOE/SMutableObjectViewer.cpp
@@ -110,7 +110,9 @@ void SMutableObjectViewer::GenerateMutableGraphPressed()
 	FCustomizableObjectCompiler Compiler;
 	TArray<TSoftObjectPtr<const UTexture>> RuntimeTextures;
 	TArray<FMutableSourceTextureData> CompilerTextures;
-	mu::Ptr<mu::Node> RootNode = Compiler.Export(CustomizableObject, CompileOptions, RuntimeTextures, CompilerTextures);
+	TArray<TSoftObjectPtr<const UObject>> RuntimeMeshes;
+	TArray<FMutableSourceMeshData> CompilerMeshes;
+	mu::Ptr<mu::Node> RootNode = Compiler.Export(CustomizableObject, CompileOptions, RuntimeTextures, CompilerTextures, RuntimeMeshes, CompilerMeshes);
 	if (!RootNode)
 	{
 		// TODO: Show errors
@@ -122,8 +124,10 @@ void SMutableObjectViewer::GenerateMutableGraphPressed()
 	TSharedRef<SMutableGraphViewer> GraphViewer = SNew(SMutableGraphViewer, RootNode)
 		.DataTag(DataTag)
 		.ReferencedRuntimeTextures(RuntimeTextures)
-		.ReferencedCompileTextures(CompilerTextures);
-	
+		.ReferencedCompileTextures(CompilerTextures)
+		.ReferencedRuntimeMeshes(RuntimeMeshes)
+		.ReferencedCompileMeshes(CompilerMeshes);
+
 	DebuggerContentsBox->ClearContent();
 	DebuggerContentsBox->SetContent(GraphViewer);
 }
@@ -133,6 +137,8 @@ void SMutableObjectViewer::CompileMutableCodePressed()
 {
 	TArray<TSoftObjectPtr<const UTexture>> RuntimeTextures;
 	TArray<FMutableSourceTextureData> CompilerTextures;
+	TArray<TSoftObjectPtr<const UObject>> RuntimeMeshes;
+	TArray<FMutableSourceMeshData> CompilerMeshes;
 	if (CompileOptions.bForceLargeLODBias)
 	{
 		// Debug compile with many different biases
@@ -144,7 +150,9 @@ void SMutableObjectViewer::CompileMutableCodePressed()
 			FCustomizableObjectCompiler Compiler;
 			RuntimeTextures.Empty();
 			CompilerTextures.Empty();
-			mu::NodePtr RootNode = Compiler.Export(CustomizableObject, CompileOptions, RuntimeTextures, CompilerTextures);
+			RuntimeMeshes.Empty();
+			CompilerMeshes.Empty();
+			mu::NodePtr RootNode = Compiler.Export(CustomizableObject, CompileOptions, RuntimeTextures, CompilerTextures, RuntimeMeshes, CompilerMeshes);
 			if (!RootNode)
 			{
 				// TODO: Show errors
@@ -156,6 +164,7 @@ void SMutableObjectViewer::CompileMutableCodePressed()
 			TSharedPtr<FCustomizableObjectCompileRunnable> CompileTask = MakeShareable(new FCustomizableObjectCompileRunnable(RootNode));
 			CompileTask->Options = CompileOptions;
 			CompileTask->ReferencedTextures = CompilerTextures;
+			CompileTask->ReferencedMeshes = CompilerMeshes;
 			CompileTask->Init();
 			CompileTask->Run();
 		}
@@ -165,7 +174,9 @@ void SMutableObjectViewer::CompileMutableCodePressed()
 	FCustomizableObjectCompiler Compiler;
 	RuntimeTextures.Empty();
 	CompilerTextures.Empty();
-	mu::NodePtr RootNode = Compiler.Export(CustomizableObject, CompileOptions, RuntimeTextures, CompilerTextures);
+	RuntimeMeshes.Empty();
+	CompilerMeshes.Empty();
+	mu::NodePtr RootNode = Compiler.Export(CustomizableObject, CompileOptions, RuntimeTextures, CompilerTextures, RuntimeMeshes, CompilerMeshes);
 	if (!RootNode)
 	{
 		// TODO: Show errors
@@ -177,12 +188,13 @@ void SMutableObjectViewer::CompileMutableCodePressed()
 	TSharedPtr<FCustomizableObjectCompileRunnable> CompileTask = MakeShareable(new FCustomizableObjectCompileRunnable(RootNode));
 	CompileTask->Options = CompileOptions;
 	CompileTask->ReferencedTextures = CompilerTextures;
+	CompileTask->ReferencedMeshes = CompilerMeshes;
 	CompileTask->Init();
 	CompileTask->Run();
 
 	FString DataTag = FString::Printf( TEXT("Mutable Code : %s opt %d"), *CompileOptions.TargetPlatform->PlatformName(), CompileOptions.OptimizationLevel );
 
-	TSharedRef<SMutableCodeViewer> CodeViewer = SNew(SMutableCodeViewer, CompileTask->Model, RuntimeTextures)
+	TSharedRef<SMutableCodeViewer> CodeViewer = SNew(SMutableCodeViewer, CompileTask->Model, RuntimeTextures, RuntimeMeshes)
 		.DataTag(DataTag);
 	
 	DebuggerContentsBox->ClearContent();
diff --git a/Engine/Plugins/Mutable/Source/CustomizableObjectEditor/Public/MuCOE/CustomizableObjectLayout.h b/Engine/Plugins/Mutable/Source/CustomizableObjectEditor/Public/MuCOE/CustomizableObjectLayout.h
--- a/Engine/Plugins/Mutable/Source/CustomizableObjectEditor/Public/MuCOE/CustomizableObjectLayout.h
+++ b/Engine/Plugins/Mutable/Source/CustomizableObjectEditor/Public/MuCOE/CustomizableObjectLayout.h
@@ -122,7 +122,7 @@ public:
 	int32 GetMaterial() const { return Material; }
 	int32 GetUVChannel() const { return UVChannel; }
 	FString GetLayoutName() const { return LayoutName; }
-	UObject* GetMesh() const;
+	TSoftObjectPtr<UObject> GetMesh() const;
 	FIntPoint GetGridSize() const { return GridSize; }
 	FIntPoint GetMaxGridSize() const { return MaxGridSize; }
 	bool GetIgnoreVertexLayoutWarnings() const { return bIgnoreUnassignedVertexWarning; };
diff --git a/Engine/Plugins/Mutable/Source/CustomizableObjectEditor/Public/MuCOE/Nodes/CustomizableObjectNodeMesh.h b/Engine/Plugins/Mutable/Source/CustomizableObjectEditor/Public/MuCOE/Nodes/CustomizableObjectNodeMesh.h
--- a/Engine/Plugins/Mutable/Source/CustomizableObjectEditor/Public/MuCOE/Nodes/CustomizableObjectNodeMesh.h
+++ b/Engine/Plugins/Mutable/Source/CustomizableObjectEditor/Public/MuCOE/Nodes/CustomizableObjectNodeMesh.h
@@ -25,7 +25,7 @@ public:
 	virtual UTexture2D* FindTextureForPin(const UEdGraphPin* Pin) const PURE_VIRTUAL(UCustomizableObjectNodeMesh::FindTextureForPin, return {};);
 
 	/** Returns the Unreal mesh (e.g., USkeletalMesh, UStaticMesh...). */
-	virtual UObject* GetMesh() const PURE_VIRTUAL(UCustomizableObjectNodeMesh::GetMesh, return {};);
+	virtual TSoftObjectPtr<UObject> GetMesh() const PURE_VIRTUAL(UCustomizableObjectNodeMesh::GetMesh, return {};);
 
 	/** Returns the output Mesh pin associated to the given LODIndex and SectionIndex. Override. */
 	virtual UEdGraphPin* GetMeshPin(int32 LOD, int32 SectionIndex) const PURE_VIRTUAL(UCustomizableObjectNodeMesh::GetMeshPin, return {};);
diff --git a/Engine/Plugins/Mutable/Source/CustomizableObjectEditor/Public/MuCOE/Nodes/CustomizableObjectNodeSkeletalMesh.h b/Engine/Plugins/Mutable/Source/CustomizableObjectEditor/Public/MuCOE/Nodes/CustomizableObjectNodeSkeletalMesh.h
--- a/Engine/Plugins/Mutable/Source/CustomizableObjectEditor/Public/MuCOE/Nodes/CustomizableObjectNodeSkeletalMesh.h
+++ b/Engine/Plugins/Mutable/Source/CustomizableObjectEditor/Public/MuCOE/Nodes/CustomizableObjectNodeSkeletalMesh.h
@@ -127,7 +127,7 @@ public:
 	/** Virtual function used to copy pin data when remapping pins. */
 	virtual void Copy(const UCustomizableObjectNodePinData& Other) override;
 
-	void Init(int32 InLODIndex, int32 InSectionIndex, int32 NumTexCoords, USkeletalMesh* Mesh);
+	void Init(int32 InLODIndex, int32 InSectionIndex, int32 NumTexCoords);
 
 	/** Layouts related to this Mesh pin */
 	UPROPERTY()
@@ -183,7 +183,7 @@ public:
 	GENERATED_BODY()
 	
 	UPROPERTY(EditAnywhere, Category = CustomizableObject)
-	TObjectPtr<USkeletalMesh> SkeletalMesh;
+	TSoftObjectPtr<USkeletalMesh> SkeletalMesh;
 
 	/** Default pin when there is no mesh. */
 	UPROPERTY()
@@ -231,7 +231,7 @@ public:
 	// UCustomizableObjectNodeMesh interface
 	virtual UTexture2D* FindTextureForPin(const UEdGraphPin* Pin) const override;
 	virtual TArray<UCustomizableObjectLayout*> GetLayouts(const UEdGraphPin& MeshPin) const override;
-	virtual UObject* GetMesh() const override;
+	virtual TSoftObjectPtr<UObject> GetMesh() const override;
 	virtual UEdGraphPin* GetMeshPin(int32 LOD, int32 SectionIndex) const override;
 	virtual void GetPinSection(const UEdGraphPin& Pin, int32& OutLODIndex, int32& OutSectionIndex, int32& OutLayoutIndex) const override;
 
diff --git a/Engine/Plugins/Mutable/Source/CustomizableObjectEditor/Public/MuCOE/Nodes/CustomizableObjectNodeStaticMesh.h b/Engine/Plugins/Mutable/Source/CustomizableObjectEditor/Public/MuCOE/Nodes/CustomizableObjectNodeStaticMesh.h
--- a/Engine/Plugins/Mutable/Source/CustomizableObjectEditor/Public/MuCOE/Nodes/CustomizableObjectNodeStaticMesh.h
+++ b/Engine/Plugins/Mutable/Source/CustomizableObjectEditor/Public/MuCOE/Nodes/CustomizableObjectNodeStaticMesh.h
@@ -120,7 +120,7 @@ public:
 	GENERATED_BODY()
 
 	UPROPERTY(EditAnywhere, Category = CustomizableObject)
-	TObjectPtr<UStaticMesh> StaticMesh;
+	TSoftObjectPtr<UStaticMesh> StaticMesh;
 
 	/** Images */
 	UPROPERTY()
@@ -151,7 +151,7 @@ public:
 	// UCustomizableObjectNodeMesh interface
 	virtual UTexture2D* FindTextureForPin(const UEdGraphPin* Pin) const override;
 	virtual TArray<UCustomizableObjectLayout*> GetLayouts(const UEdGraphPin& OutPin) const override;
-	virtual UObject* GetMesh() const override;
+	virtual TSoftObjectPtr<UObject> GetMesh() const override;
 	virtual UEdGraphPin* GetMeshPin(int32 LOD, int32 SectionIndex) const override;
 	virtual void GetPinSection(const UEdGraphPin& Pin, int32& OutLODIndex, int32& OutSectionIndex, int32& OutLayoutIndex) const override;
 
diff --git a/Engine/Plugins/Mutable/Source/MutableTools/Internal/MuT/AST.h b/Engine/Plugins/Mutable/Source/MutableTools/Internal/MuT/AST.h
--- a/Engine/Plugins/Mutable/Source/MutableTools/Internal/MuT/AST.h
+++ b/Engine/Plugins/Mutable/Source/MutableTools/Internal/MuT/AST.h
@@ -223,8 +223,9 @@ namespace mu
 		// External resource provision functions
 		//---------------------------------------------------------------------
 
-		/** Function used to request an engine resource (image) from the compiler. */
-		FReferencedResourceFunc ReferencedResourceProvider;
+		/** Function used to request an engine resource from the compiler. */
+		FReferencedImageResourceFunc ReferencedImageResourceProvider;
+		FReferencedMeshResourceFunc ReferencedMeshResourceProvider;
 
         bool bDisableImageGeneration = false;
         bool bDisableMeshGeneration = false;
diff --git a/Engine/Plugins/Mutable/Source/MutableTools/Private/MuT/CodeOptimiser.cpp b/Engine/Plugins/Mutable/Source/MutableTools/Private/MuT/CodeOptimiser.cpp
--- a/Engine/Plugins/Mutable/Source/MutableTools/Private/MuT/CodeOptimiser.cpp
+++ b/Engine/Plugins/Mutable/Source/MutableTools/Private/MuT/CodeOptimiser.cpp
@@ -445,7 +445,8 @@ namespace mu
 		FProxyFileContext* DiskCacheContext = nullptr;
 		int32 ImageCompressionQuality = 0;
 		int32 OptimizationPass = 0;
-		FReferencedResourceFunc ReferencedResourceProvider;
+		FReferencedMeshResourceFunc ReferencedMeshResourceProvider;
+		FReferencedImageResourceFunc ReferencedImageResourceProvider;
 
 		// Intermediate
 		Ptr<ASTOp> SourceCloned;
@@ -461,7 +462,8 @@ namespace mu
 			Source = InSource;
 			DiskCacheContext = InOptions->OptimisationOptions.DiskCacheContext;
 			ImageCompressionQuality = InOptions->ImageCompressionQuality;
-			ReferencedResourceProvider = InOptions->OptimisationOptions.ReferencedResourceProvider;
+			ReferencedMeshResourceProvider = InOptions->OptimisationOptions.ReferencedMeshResourceProvider;
+			ReferencedImageResourceProvider = InOptions->OptimisationOptions.ReferencedImageResourceProvider;
 		}
 
 		void Run(FImageOperator ImOp)
@@ -811,11 +813,13 @@ namespace mu
 						Ptr<ASTOp> SubgraphRoot = ConstantSubgraphs[Index].Root;
 						UE::Tasks::FTaskEvent& SubgraphCompletionEvent = ConstantSubgraphs[Index].CompletedEvent;
 
-						// Referenced images are resolved in its own task to prevent requesting them twice (and loading them twice)
+						// Referenced images and meshes are resolved in its own task to prevent requesting them twice (and loading them twice)
 						bool bIsCompileTimeReferenceImage = (SubgraphRoot->GetOpType() == EOpType::IM_REFERENCE);
 
 						// Launch the task with its dependencies
-						if (bIsCompileTimeReferenceImage)
+						switch (SubgraphRoot->GetOpType())
+						{
+						case EOpType::IM_REFERENCE:
 						{
 							// Instead of generating the constant we resolve the reference, which also replaces the ASTOp.
 							const ASTOpReferenceResource* Typed = static_cast<const ASTOpReferenceResource*>(SubgraphRoot.get());
@@ -824,7 +828,7 @@ namespace mu
 							TSharedPtr< TSharedPtr<FImage> > ResolveImage = MakeShared<TSharedPtr<FImage>>();
 
 							constexpr bool bRunImmediatlyIfPossible = false;
-							UE::Tasks::FTask ReferenceCompletion = InOptions->OptimisationOptions.ReferencedResourceProvider(ImageID, ResolveImage, bRunImmediatlyIfPossible);
+							UE::Tasks::FTask ReferenceCompletion = InOptions->OptimisationOptions.ReferencedImageResourceProvider(ImageID, ResolveImage, bRunImmediatlyIfPossible);
 
 							UE::Tasks::FTask CompleteTask = UE::Tasks::Launch(TEXT("MutableResolveComplete"),
 								[SubgraphRoot, InOptions, ResolveImage]()
@@ -849,13 +853,53 @@ namespace mu
 								LowLevelTasks::ETaskPriority::BackgroundNormal);
 
 							SubgraphCompletionEvent.AddPrerequisites(CompleteTask);
+
+							break;
 						}
 
-						else
+						case EOpType::ME_REFERENCE:
+						{
+							// Instead of generating the constant we resolve the reference, which also replaces the ASTOp.
+							const ASTOpReferenceResource* Typed = static_cast<const ASTOpReferenceResource*>(SubgraphRoot.get());
+							uint32 MeshID = Typed->ID;
+
+							TSharedPtr< TSharedPtr<FMesh> > ResolveMesh = MakeShared<TSharedPtr<FMesh>>();
+
+							constexpr bool bRunImmediatlyIfPossible = false;
+							UE::Tasks::FTask ReferenceCompletion = InOptions->OptimisationOptions.ReferencedMeshResourceProvider(MeshID, ResolveMesh, bRunImmediatlyIfPossible);
+
+							UE::Tasks::FTask CompleteTask = UE::Tasks::Launch(TEXT("MutableResolveComplete"),
+								[SubgraphRoot, InOptions, ResolveMesh]()
+								{
+									Ptr<ASTOpConstantResource> ConstantOp;
+									{
+										MUTABLE_CPUPROFILER_SCOPE(MutableResolveComplete_CreateConstant);
+										ConstantOp = new ASTOpConstantResource;
+										ConstantOp->Type = EOpType::IM_CONSTANT;
+										{
+											MUTABLE_CPUPROFILER_SCOPE(GetSourceDataDescriptor);
+											ConstantOp->SourceDataDescriptor = SubgraphRoot->GetSourceDataDescriptor();
+										}
+										ConstantOp->SetValue(*ResolveMesh, InOptions->OptimisationOptions.DiskCacheContext);
+									}
+									{
+										MUTABLE_CPUPROFILER_SCOPE(MutableResolveComplete_Replace);
+										ASTOp::Replace(SubgraphRoot, ConstantOp);
+									}
+								},
+								ReferenceCompletion,
+								LowLevelTasks::ETaskPriority::BackgroundNormal);
+
+							SubgraphCompletionEvent.AddPrerequisites(CompleteTask);
+
+							break;
+						}
+
+						default:
 						{
 							// Scan for requisites
 							TArray< UE::Tasks::FTask, TInlineAllocator<8> > Requisites;
-							GetRequisites(SubgraphRoot,Requisites);
+							GetRequisites(SubgraphRoot, Requisites);
 
 							TUniquePtr<FConstantTask> Task(new FConstantTask(SubgraphRoot, InOptions, Pass));
 
@@ -879,6 +923,8 @@ namespace mu
 							SubgraphCompletionEvent.AddPrerequisites(CompleteTask);
 						}
 
+						}
+
 						ConstantSubgraphs[Index].Root = nullptr;
 						SubgraphCompletionEvent.Trigger();
 						
@@ -925,7 +971,7 @@ namespace mu
 							{
 								TSharedPtr< TSharedPtr<FImage> > ResolveImage = MakeShared<TSharedPtr<FImage>>();
 								constexpr bool bRunImmediatlyIfPossible = true;
-								UE::Tasks::FTask ReferenceCompletion = InOptions->OptimisationOptions.ReferencedResourceProvider(ImageID, ResolveImage, bRunImmediatlyIfPossible);
+								UE::Tasks::FTask ReferenceCompletion = InOptions->OptimisationOptions.ReferencedImageResourceProvider(ImageID, ResolveImage, bRunImmediatlyIfPossible);
 								ReferenceCompletion.Wait();
 								Ptr<ASTOpConstantResource> ConstantOp = new ASTOpConstantResource;
 								ConstantOp->SourceDataDescriptor = SubgraphRoot->GetSourceDataDescriptor();
diff --git a/Engine/Plugins/Mutable/Source/MutableTools/Private/MuT/Compiler.cpp b/Engine/Plugins/Mutable/Source/MutableTools/Private/MuT/Compiler.cpp
--- a/Engine/Plugins/Mutable/Source/MutableTools/Private/MuT/Compiler.cpp
+++ b/Engine/Plugins/Mutable/Source/MutableTools/Private/MuT/Compiler.cpp
@@ -181,9 +181,10 @@ namespace mu
 
 
 	//---------------------------------------------------------------------------------------------
-	void CompilerOptions::SetReferencedResourceCallback(const FReferencedResourceFunc& Provider)
+	void CompilerOptions::SetReferencedResourceCallback(const FReferencedImageResourceFunc& Image, const FReferencedMeshResourceFunc& Mesh)
 	{
-		m_pD->OptimisationOptions.ReferencedResourceProvider = Provider;
+		m_pD->OptimisationOptions.ReferencedImageResourceProvider = Image;
+		m_pD->OptimisationOptions.ReferencedMeshResourceProvider = Mesh;
 	}
 
     void CompilerOptions::SetDisableImageGeneration(bool bDisabled)
diff --git a/Engine/Plugins/Mutable/Source/MutableTools/Public/MuT/Compiler.h b/Engine/Plugins/Mutable/Source/MutableTools/Public/MuT/Compiler.h
--- a/Engine/Plugins/Mutable/Source/MutableTools/Public/MuT/Compiler.h
+++ b/Engine/Plugins/Mutable/Source/MutableTools/Public/MuT/Compiler.h
@@ -27,7 +27,8 @@ namespace mu
     class Node;
     class NodeTransformedObject;
 
-	typedef TFunction<UE::Tasks::FTask(int32, TSharedPtr<TSharedPtr<FImage>>, bool)> FReferencedResourceFunc;
+	typedef TFunction<UE::Tasks::FTask(int32, TSharedPtr<TSharedPtr<FImage>>, bool)> FReferencedImageResourceFunc;
+	typedef TFunction<UE::Tasks::FTask(int32, TSharedPtr<TSharedPtr<FMesh>>, bool)> FReferencedMeshResourceFunc;
 
     //! \brief Options used to compile the models with a compiler.
     class MUTABLETOOLS_API CompilerOptions : public RefCounted
@@ -78,7 +79,7 @@ namespace mu
 		void SetImagePixelFormatOverride(const FImageOperator::FImagePixelFormatFunc&);
 
 		/** */
-		void SetReferencedResourceCallback(const FReferencedResourceFunc&);
+		void SetReferencedResourceCallback(const FReferencedImageResourceFunc&, const FReferencedMeshResourceFunc&);
 
         void SetDisableImageGeneration(bool bDisabled);
         void SetDisableMeshGeneration(bool bDisabled);
@@ -131,20 +132,6 @@ namespace mu
 		uint8 NumExtraLODsToBuildAfterFirstLOD = 0;
 		bool bOnlyFirstLOD = false;
 		ETextureCompressionStrategy TextureCompressionStrategy = ETextureCompressionStrategy::None;
-
-        void Serialise( FOutputArchive& arch ) const
-        {
-			arch << bOnlyFirstLOD;
-            arch << TextureCompressionStrategy;
-			arch << NumExtraLODsToBuildAfterFirstLOD;
-        }
-
-        void Unserialise( FInputArchive& arch )
-        {
-            arch >> bOnlyFirstLOD;
-			arch >> TextureCompressionStrategy;
-			arch >> NumExtraLODsToBuildAfterFirstLOD;
-		}
     };
 
 
