[mutable] Removed check for IsGameRunning() for case where it is no longer required
- We do check for this when getting the "MutableEditorModule" module.

#rnx

diff --git a/Engine/Plugins/Mutable/Source/CustomizableObject/Private/MuCO/CustomizableObject.cpp b/Engine/Plugins/Mutable/Source/CustomizableObject/Private/MuCO/CustomizableObject.cpp
--- a/Engine/Plugins/Mutable/Source/CustomizableObject/Private/MuCO/CustomizableObject.cpp
+++ b/Engine/Plugins/Mutable/Source/CustomizableObject/Private/MuCO/CustomizableObject.cpp
@@ -1121,15 +1121,15 @@ bool UCustomizableObject::ConditionalAutoCompile()
 		return false;
 	}
 
-	// Don't compile if we're running game or if Mutable or AutoCompile is disabled.
-	if (IsRunningGame() || !System->IsActive() || !System->IsAutoCompileEnabled())
+	// Don't compile if Mutable or AutoCompile is disabled.
+	if (!System->IsActive() || !System->IsAutoCompileEnabled())
 	{
 		System->AddUncompiledCOWarning(*this);
 		return false;
 	}
 
-	ICustomizableObjectEditorModule* EditorModule = ICustomizableObjectEditorModule::Get();
-	if (ensure(EditorModule))
+	// Get the EditorModule : Will be null if IsRunningGame() is true
+	if (ICustomizableObjectEditorModule* EditorModule = ICustomizableObjectEditorModule::Get())
 	{
 		if (System->IsAutoCompilationSync() && GetPrivate()->Status.Get() == FCustomizableObjectStatusTypes::EState::Loading)
 		{
