[Mutable] Use platform LODGroup config to set compiled textures optional mips.

- Optional mips was set from a compilation option and computed using the source image size.
- Now the source texture config values are passed as additional data and the calculation is done at the rom classification stage when the final size is known.

#rnx

diff --git a/Engine/Plugins/Mutable/Source/CustomizableObject/Internal/MuCO/CustomizableObjectPrivate.h b/Engine/Plugins/Mutable/Source/CustomizableObject/Internal/MuCO/CustomizableObjectPrivate.h
--- a/Engine/Plugins/Mutable/Source/CustomizableObject/Internal/MuCO/CustomizableObjectPrivate.h
+++ b/Engine/Plugins/Mutable/Source/CustomizableObject/Internal/MuCO/CustomizableObjectPrivate.h
@@ -1817,6 +1817,8 @@ public:
 
 		LODsPerComponentTypeMismatch,
 
+		ImageHiResLODsUseLODGroupInfo,
+
 		// -----<new versions can be added above this line>--------
 		LastCustomizableObjectVersion
 	};
diff --git a/Engine/Plugins/Mutable/Source/CustomizableObjectEditor/Private/MuCOE/GenerateMutableSource/GenerateMutableSourceGroupProjector.cpp b/Engine/Plugins/Mutable/Source/CustomizableObjectEditor/Private/MuCOE/GenerateMutableSource/GenerateMutableSourceGroupProjector.cpp
--- a/Engine/Plugins/Mutable/Source/CustomizableObjectEditor/Private/MuCOE/GenerateMutableSource/GenerateMutableSourceGroupProjector.cpp
+++ b/Engine/Plugins/Mutable/Source/CustomizableObjectEditor/Private/MuCOE/GenerateMutableSource/GenerateMutableSourceGroupProjector.cpp
@@ -2,6 +2,9 @@
 
 #include "MuCOE/GenerateMutableSource/GenerateMutableSourceGroupProjector.h"
 
+#include "Engine/TextureLODSettings.h"
+#include "Interfaces/ITargetPlatform.h"
+
 #include "GenerateMutableSourceImage.h"
 #include "Materials/MaterialInterface.h"
 #include "MuCO/CustomizableObjectInstancePrivate.h"
@@ -527,10 +530,12 @@ TOptional<FGroupProjectorTempData> GenerateMutableGroupProjector(UCustomizableOb
 				// Calculate the number of mips to tag as high res for this image.
 				if (ProjParamNode->ReferenceTexture)
 				{
-					int32 TotalMips = mu::FImage::GetMipmapCount(ProjParamNode->ReferenceTexture->GetSizeX(), ProjParamNode->ReferenceTexture->GetSizeY());
-					int32 NumMipsBeyondMin = FMath::Max(0, TotalMips - CompilationOptions.MinDiskMips);
-					int32 HighResMipsForThisImage = FMath::Min(NumMipsBeyondMin, CompilationOptions.NumHighResImageMips);
-					ImageNode->SourceDataDescriptor.SourceHighResMips = HighResMipsForThisImage;
+					const UTextureLODSettings& TextureLODSettings = GenerationContext.CompilationContext.Options.TargetPlatform->GetTextureLODSettings();
+					const FTextureLODGroup& LODGroupInfo = TextureLODSettings.GetTextureLODGroup(ProjParamNode->ReferenceTexture->LODGroup);
+
+					ImageNode->SourceDataDescriptor.OptionalMaxLODSize = LODGroupInfo.OptionalMaxLODSize;
+					ImageNode->SourceDataDescriptor.OptionalLODBias = LODGroupInfo.OptionalLODBias;
+					ImageNode->SourceDataDescriptor.NumNonOptionalLODs = CompilationOptions.MinDiskMips;
 
 					const FString TextureName = GetNameSafe(Texture).ToLower();
 					ImageNode->SourceDataDescriptor.SourceId = CityHash32(reinterpret_cast<const char*>(*TextureName), TextureName.Len() * sizeof(FString::ElementType));
diff --git a/Engine/Plugins/Mutable/Source/CustomizableObjectEditor/Private/MuCOE/GenerateMutableSource/GenerateMutableSourceImage.cpp b/Engine/Plugins/Mutable/Source/CustomizableObjectEditor/Private/MuCOE/GenerateMutableSource/GenerateMutableSourceImage.cpp
--- a/Engine/Plugins/Mutable/Source/CustomizableObjectEditor/Private/MuCOE/GenerateMutableSource/GenerateMutableSourceImage.cpp
+++ b/Engine/Plugins/Mutable/Source/CustomizableObjectEditor/Private/MuCOE/GenerateMutableSource/GenerateMutableSourceImage.cpp
@@ -2,6 +2,8 @@
 
 #include "MuCOE/GenerateMutableSource/GenerateMutableSourceImage.h"
 
+#include "Engine/TextureLODSettings.h"
+
 #include "MuCOE/CustomizableObjectCompiler.h"
 #include "MuCOE/EdGraphSchema_CustomizableObject.h"
 #include "MuCOE/GenerateMutableSource/GenerateMutableSourceColor.h"
@@ -125,10 +127,12 @@ mu::NodeImagePtr GenerateMutableSourceImage(const UEdGraphPin* Pin, FMutableGrap
 				mu::Ptr<mu::NodeImageConstant> ConstantImageNode = new mu::NodeImageConstant;
 				ConstantImageNode->SetValue( ImageConstant );
 				
-				int32 TotalMips = mu::FImage::GetMipmapCount(ImageConstant->GetSizeX(), ImageConstant->GetSizeY());
-				int32 NumMipsBeyondMin = FMath::Max(0, TotalMips - int32(MipsToSkip) - CompilationOptions.MinDiskMips);
-				int32 HighResMipsForThisImage = FMath::Min(NumMipsBeyondMin, CompilationOptions.NumHighResImageMips);
-				ConstantImageNode->SourceDataDescriptor.SourceHighResMips = HighResMipsForThisImage;
+				const UTextureLODSettings& TextureLODSettings = GenerationContext.CompilationContext.Options.TargetPlatform->GetTextureLODSettings();
+				const FTextureLODGroup& LODGroupInfo = TextureLODSettings.GetTextureLODGroup(BaseTexture->LODGroup);
+
+				ConstantImageNode->SourceDataDescriptor.OptionalMaxLODSize = LODGroupInfo.OptionalMaxLODSize;
+				ConstantImageNode->SourceDataDescriptor.OptionalLODBias = LODGroupInfo.OptionalLODBias;
+				ConstantImageNode->SourceDataDescriptor.NumNonOptionalLODs = CompilationOptions.MinDiskMips;
 
 				const FString TextureName = GetNameSafe(BaseTexture).ToLower();
 				ConstantImageNode->SourceDataDescriptor.SourceId = CityHash32(reinterpret_cast<const char*>(*TextureName), TextureName.Len() * sizeof(FString::ElementType));
@@ -227,10 +231,12 @@ mu::NodeImagePtr GenerateMutableSourceImage(const UEdGraphPin* Pin, FMutableGrap
 			const uint32 MipsToSkip = ComputeLODBiasForTexture(GenerationContext, *Texture, nullptr, ReferenceTextureSize);
 			Result = ResizeTextureByNumMips(ImageNode, MipsToSkip);
 
-			int32 TotalMips = mu::FImage::GetMipmapCount(ImageConstant->GetSizeX(), ImageConstant->GetSizeY());
-			int32 NumMipsBeyondMin = FMath::Max(0, TotalMips - int32(MipsToSkip) - CompilationOptions.MinDiskMips);
-			int32 HighResMipsForThisImage = FMath::Min(NumMipsBeyondMin, CompilationOptions.NumHighResImageMips);
-			ImageNode->SourceDataDescriptor.SourceHighResMips = HighResMipsForThisImage;
+			const UTextureLODSettings& TextureLODSettings = GenerationContext.CompilationContext.Options.TargetPlatform->GetTextureLODSettings();
+			const FTextureLODGroup& LODGroupInfo = TextureLODSettings.GetTextureLODGroup(Texture->LODGroup);
+
+			ImageNode->SourceDataDescriptor.OptionalMaxLODSize = LODGroupInfo.OptionalMaxLODSize;
+			ImageNode->SourceDataDescriptor.OptionalLODBias = LODGroupInfo.OptionalLODBias;
+			ImageNode->SourceDataDescriptor.NumNonOptionalLODs = CompilationOptions.MinDiskMips;
 
 			const FString TextureName = GetNameSafe(Texture).ToLower();
 			ImageNode->SourceDataDescriptor.SourceId = CityHash32(reinterpret_cast<const char*>(*TextureName), TextureName.Len() * sizeof(FString::ElementType));
@@ -1123,12 +1129,14 @@ mu::NodeImagePtr GenerateMutableSourceImage(const UEdGraphPin* Pin, FMutableGrap
 								const uint16 MaxTextureSize = FMath::Max3(ImageDesc.m_size[0], ImageDesc.m_size[1], (uint16)1);
 								ImageTableNode->MaxTextureSize = MaxTextureSize;
 								ImageTableNode->ReferenceImageDesc = ImageDesc;
+				
+								const UTextureLODSettings& TextureLODSettings = GenerationContext.CompilationContext.Options.TargetPlatform->GetTextureLODSettings();
+								const FTextureLODGroup& LODGroupInfo = TextureLODSettings.GetTextureLODGroup(DefaultTexture2D->LODGroup);
+
+								ImageTableNode->SourceDataDescriptor.OptionalMaxLODSize = LODGroupInfo.OptionalMaxLODSize;
+								ImageTableNode->SourceDataDescriptor.OptionalLODBias = LODGroupInfo.OptionalLODBias;
+								ImageTableNode->SourceDataDescriptor.NumNonOptionalLODs = CompilationOptions.MinDiskMips;
 
-								// Calculate the number of mips to tag as high res for this image.
-								int32 TotalMips = mu::FImage::GetMipmapCount(ImageDesc.m_size[0], ImageDesc.m_size[1]);
-								int32 NumMipsBeyondMin = FMath::Max(0, TotalMips - CompilationOptions.MinDiskMips);
-								int32 HighResMipsForThisImage = FMath::Min(NumMipsBeyondMin, CompilationOptions.NumHighResImageMips);
-								ImageTableNode->SourceDataDescriptor.SourceHighResMips = HighResMipsForThisImage;
 								ImageTableNode->SourceDataDescriptor.SourceId = TableId; // It will be combined with the RowId when generating constants
 							}
 						}
diff --git a/Engine/Plugins/Mutable/Source/CustomizableObjectEditor/Private/MuCOE/GenerateMutableSource/GenerateMutableSourceMesh.cpp b/Engine/Plugins/Mutable/Source/CustomizableObjectEditor/Private/MuCOE/GenerateMutableSource/GenerateMutableSourceMesh.cpp
--- a/Engine/Plugins/Mutable/Source/CustomizableObjectEditor/Private/MuCOE/GenerateMutableSource/GenerateMutableSourceMesh.cpp
+++ b/Engine/Plugins/Mutable/Source/CustomizableObjectEditor/Private/MuCOE/GenerateMutableSource/GenerateMutableSourceMesh.cpp
@@ -3316,7 +3316,7 @@ mu::NodeMeshPtr GenerateMorphMesh(const UEdGraphPin* Pin,
 			MorphedSourceMeshNodeTable->SetParameterName(TypedNodeTable->ParameterName);
 			MorphedSourceMeshNodeTable->SetMessageContext(MorphNode);
 			MorphedSourceMeshNodeTable->SourceDataDescriptor.SourceId = TableId;
-			MorphedSourceMeshNodeTable->SourceDataDescriptor.SourceHighResMips = 0;
+			MorphedSourceMeshNodeTable->SourceDataDescriptor.OptionalMaxLODSize = 0;
 
 			mu::NodeMeshMakeMorphPtr Morph = new mu::NodeMeshMakeMorph;
 			Morph->SetBase(BaseSourceMesh.get());
@@ -3674,7 +3674,7 @@ mu::Ptr<mu::NodeMesh> GenerateMutableSourceMesh(const UEdGraphPin* Pin,
 
 			const FString MeshName = TypedNodeSkel->SkeletalMesh.GetLongPackageName().ToLower();
 			MeshNode->SourceDataDescriptor.SourceId = CityHash32(reinterpret_cast<const char*>(*MeshName), MeshName.Len() * sizeof(FString::ElementType));
-			MeshNode->SourceDataDescriptor.SourceHighResMips = 0;
+			MeshNode->SourceDataDescriptor.OptionalMaxLODSize = 0;
 
 			// Applying Mesh Morph Nodes
 			if (GenerationContext.MeshMorphStack.Num())
@@ -3769,7 +3769,7 @@ mu::Ptr<mu::NodeMesh> GenerateMutableSourceMesh(const UEdGraphPin* Pin,
 
 				const FString MeshName = TypedNodeStatic->StaticMesh.GetLongPackageName().ToLower();
 				MeshNode->SourceDataDescriptor.SourceId = CityHash32(reinterpret_cast<const char*>(*MeshName), MeshName.Len() * sizeof(FString::ElementType));
-				MeshNode->SourceDataDescriptor.SourceHighResMips = 0;
+				MeshNode->SourceDataDescriptor.OptionalMaxLODSize = 0;
 			}
 			else
 			{
@@ -4380,7 +4380,7 @@ mu::Ptr<mu::NodeMesh> GenerateMutableSourceMesh(const UEdGraphPin* Pin,
 						MeshTableNode->SetNoneOption(TypedNodeTable->bAddNoneOption);
 						MeshTableNode->SetDefaultRowName(TypedNodeTable->DefaultRowName.ToString());
 						MeshTableNode->SourceDataDescriptor.SourceId = TableId;
-						MeshTableNode->SourceDataDescriptor.SourceHighResMips = 0;
+						MeshTableNode->SourceDataDescriptor.OptionalMaxLODSize = 0;
 
 						// Pose Assets do not need this part of the code
 						if (Pin->PinType.PinCategory == Schema->PC_Mesh)
diff --git a/Engine/Plugins/Mutable/Source/CustomizableObjectEditor/Private/MuCOE/GenerateMutableSource/GenerateMutableSourceSurface.cpp b/Engine/Plugins/Mutable/Source/CustomizableObjectEditor/Private/MuCOE/GenerateMutableSource/GenerateMutableSourceSurface.cpp
--- a/Engine/Plugins/Mutable/Source/CustomizableObjectEditor/Private/MuCOE/GenerateMutableSource/GenerateMutableSourceSurface.cpp
+++ b/Engine/Plugins/Mutable/Source/CustomizableObjectEditor/Private/MuCOE/GenerateMutableSource/GenerateMutableSourceSurface.cpp
@@ -4,6 +4,7 @@
 
 #include "Engine/SkinnedAssetCommon.h"
 #include "Engine/StaticMesh.h"
+#include "Engine/TextureLODSettings.h"
 #include "GPUSkinPublicDefs.h"
 #include "Interfaces/ITargetPlatform.h"
 #include "Interfaces/ITextureFormat.h"
@@ -584,12 +585,13 @@ mu::Ptr<mu::NodeSurface> GenerateMutableSourceSurface(const UEdGraphPin * Pin, F
 
 										const uint32 MipsToSkip = ComputeLODBiasForTexture(GenerationContext, *Texture2D, nullptr, Props.TextureSize);
 										mu::Ptr<mu::NodeImage> Result =  ResizeTextureByNumMips(ConstImageNode, MipsToSkip);
+	
+										const UTextureLODSettings& TextureLODSettings = GenerationContext.CompilationContext.Options.TargetPlatform->GetTextureLODSettings();
+										const FTextureLODGroup& LODGroupInfo = TextureLODSettings.GetTextureLODGroup(Texture2D->LODGroup);
 
-										// Calculate the number of mips to tag as high res for this image.
-										int32 TotalMips = mu::FImage::GetMipmapCount(ImageConstant->GetSizeX(), ImageConstant->GetSizeY());
-										int32 NumMipsBeyondMin = FMath::Max(0, TotalMips - int32(MipsToSkip) - GenerationContext.CompilationContext.Options.MinDiskMips);
-										int32 HighResMipsForThisImage = FMath::Min(NumMipsBeyondMin, GenerationContext.CompilationContext.Options.NumHighResImageMips);
-										ConstImageNode->SourceDataDescriptor.SourceHighResMips = HighResMipsForThisImage;
+										ConstImageNode->SourceDataDescriptor.OptionalMaxLODSize = LODGroupInfo.OptionalMaxLODSize;
+										ConstImageNode->SourceDataDescriptor.OptionalLODBias = LODGroupInfo.OptionalLODBias;
+										ConstImageNode->SourceDataDescriptor.NumNonOptionalLODs = GenerationContext.CompilationContext.Options.MinDiskMips;
 
 										const FString TextureName = GetNameSafe(Texture2D).ToLower();
 										ConstImageNode->SourceDataDescriptor.SourceId = CityHash32(reinterpret_cast<const char*>(*TextureName), TextureName.Len() * sizeof(FString::ElementType));
@@ -685,11 +687,13 @@ mu::Ptr<mu::NodeSurface> GenerateMutableSourceSurface(const UEdGraphPin * Pin, F
 
 									CompositedImage->SetNormal(NormalCompositeMipmapImage);
 
-									int32 TotalMips = mu::FImage::GetMipmapCount(ImageConstant->GetSizeX(), ImageConstant->GetSizeY());
-									int32 NumMipsBeyondMin = FMath::Max(0, TotalMips - int32(MipsToSkip) - GenerationContext.CompilationContext.Options.MinDiskMips);
-									int32 HighResMipsForThisImage = FMath::Min(NumMipsBeyondMin, GenerationContext.CompilationContext.Options.NumHighResImageMips);
-									CompositeNormalImage->SourceDataDescriptor.SourceHighResMips = HighResMipsForThisImage;
+									const UTextureLODSettings& TextureLODSettings = GenerationContext.CompilationContext.Options.TargetPlatform->GetTextureLODSettings();;
+									const FTextureLODGroup& LODGroupInfo = TextureLODSettings.GetTextureLODGroup(ReferenceCompositeNormalTexture->LODGroup);
 
+									CompositeNormalImage->SourceDataDescriptor.OptionalMaxLODSize = LODGroupInfo.OptionalMaxLODSize;
+									CompositeNormalImage->SourceDataDescriptor.OptionalLODBias = LODGroupInfo.OptionalLODBias;
+									CompositeNormalImage->SourceDataDescriptor.NumNonOptionalLODs = GenerationContext.CompilationContext.Options.MinDiskMips;
+									
 									const FString TextureName = GetNameSafe(ReferenceCompositeNormalTexture).ToLower();
 									CompositeNormalImage->SourceDataDescriptor.SourceId = CityHash32(reinterpret_cast<const char*>(*TextureName), TextureName.Len() * sizeof(FString::ElementType));
 								}
diff --git a/Engine/Plugins/Mutable/Source/MutableTools/Private/MuT/Compiler.cpp b/Engine/Plugins/Mutable/Source/MutableTools/Private/MuT/Compiler.cpp
--- a/Engine/Plugins/Mutable/Source/MutableTools/Private/MuT/Compiler.cpp
+++ b/Engine/Plugins/Mutable/Source/MutableTools/Private/MuT/Compiler.cpp
@@ -580,9 +580,21 @@ namespace mu
 				{
 					const FImageLODRange& LODRange = Program.ConstantImages[ImageIndex];
 
-					int32 NumHighResMips = FMath::Max(0, AdditionalData.SourceImagePerConstant[ImageIndex].SourceHighResMips);
+	                int32 OptionalMaxLODSize = AdditionalData.SourceImagePerConstant[ImageIndex].OptionalMaxLODSize;
+                    int32 OptionalLODBias = AdditionalData.SourceImagePerConstant[ImageIndex].OptionalLODBias; 
+                    int32 NumNonOptionalLODs = AdditionalData.SourceImagePerConstant[ImageIndex].NumNonOptionalLODs; 
 
-					for (int32 LODRangeIndex = NumHighResMips; LODRangeIndex < LODRange.LODCount; ++LODRangeIndex)
+                    int32 NumOptionalMips = 0; 
+                    if (OptionalMaxLODSize > 0)
+                    {
+                        int32 NumTotalLODs = FMath::CeilLogTwo(FMath::Max(LODRange.ImageSizeX, LODRange.ImageSizeY)) + 1;
+                        int32 FirstOptionalLOD = FMath::Min<int32>(FMath::CeilLogTwo(OptionalMaxLODSize) + 1, NumTotalLODs);
+
+                        NumOptionalMips = FMath::Min<int32>(
+                                NumTotalLODs - (FirstOptionalLOD - OptionalLODBias), FMath::Max<int32>(0, NumTotalLODs - NumNonOptionalLODs));
+                    }
+
+					for (int32 LODRangeIndex = NumOptionalMips; LODRangeIndex < LODRange.LODCount; ++LODRangeIndex)
 					{
 						FConstantResourceIndex LODIndex = Program.ConstantImageLODIndices[LODRange.FirstIndex + LODRangeIndex];
 						check(!LODIndex.Streamable); // Not classified yet
diff --git a/Engine/Plugins/Mutable/Source/MutableTools/Public/MuT/NodeImage.h b/Engine/Plugins/Mutable/Source/MutableTools/Public/MuT/NodeImage.h
--- a/Engine/Plugins/Mutable/Source/MutableTools/Public/MuT/NodeImage.h
+++ b/Engine/Plugins/Mutable/Source/MutableTools/Public/MuT/NodeImage.h
@@ -26,11 +26,23 @@ namespace mu
 			ENeutral = -1,
 		};
 
-		/** Number of mips in the source texture that are considered high resolution.
-		* If -1 it means the entire descriptor is neutral.
-		* If -2 it means the descriptor is invalid (results of an operation that shouldn't happen).
-		*/
-		int32 SourceHighResMips = ENeutral;
+		/** 
+		 * The size of the first mip to be considered optional. Mips with equal or larger size will 
+		 * be considered optional.
+		 * If -1 it means the entire descriptor is neutral.
+		 * If -2 it means the descriptor is invalid (results of an operation that shouldn't happen).
+		 * If less than equal 0 it means all mips are non optional.
+		 */
+		int32 OptionalMaxLODSize = ENeutral;
+
+		/** 
+		 * Bias to the final number of optional mips. Any mip will have at least this number of optional 
+		 * lods if not in the NumNonOptionalLODs range. 
+		 */
+		int32 OptionalLODBias = 0;
+		
+		/** Number of lods from the tail that will never be considered optional */
+		int32 NumNonOptionalLODs = 0;
 
 		/** Source tags that mark this data and prevent it from mixing with other data at compile time. */
 		TArray<FString> Tags;
@@ -40,26 +52,32 @@ namespace mu
 
 		inline bool operator==(const FSourceDataDescriptor& Other) const
 		{
-			return SourceHighResMips == Other.SourceHighResMips && Tags == Other.Tags;
+			return 
+				Tags == Other.Tags &&
+				OptionalMaxLODSize == Other.OptionalMaxLODSize &&
+				OptionalLODBias == Other.OptionalLODBias &&
+				NumNonOptionalLODs == Other.NumNonOptionalLODs;
 		}
 
 		inline bool IsInvalid() const
 		{
-			return SourceHighResMips == EInvalid;
+			return OptionalMaxLODSize == EInvalid;
 		}
 
 		inline bool IsNeutral() const
 		{
-			return SourceHighResMips == ENeutral;
+			return OptionalMaxLODSize == ENeutral;
 		}
 
 		void CombineWith(const FSourceDataDescriptor& Other)
 		{
 			if (IsInvalid() || Other.IsInvalid())
 			{
-				SourceHighResMips = EInvalid;
 				SourceId = MAX_uint32;
 				Tags.Empty();
+				OptionalMaxLODSize = EInvalid;
+				OptionalLODBias = 0;
+				NumNonOptionalLODs = 0;
 				return;
 			}
 
@@ -70,17 +88,21 @@ namespace mu
 
 			if (IsNeutral())
 			{
-				SourceHighResMips = Other.SourceHighResMips;
 				SourceId = Other.SourceId;
 				Tags = Other.Tags;
+		 		OptionalMaxLODSize = Other.OptionalMaxLODSize;
+		 		OptionalLODBias = Other.OptionalLODBias;
+		 		NumNonOptionalLODs = Other.NumNonOptionalLODs;
 				return;
 			}
 
 			if (!(*this == Other))
 			{
-				SourceHighResMips = EInvalid;
 				SourceId = MAX_uint32;
 				Tags.Empty();
+				OptionalMaxLODSize = EInvalid;
+				OptionalLODBias = 0;
+				NumNonOptionalLODs = 0;
 				return;
 			}
 
