Fix static analysis issue.

#rnx

diff --git a/Engine/Plugins/Mutable/Source/CustomizableObjectEditor/Private/MuCOE/GenerateMutableSource/GenerateMutableSourceSurface.cpp b/Engine/Plugins/Mutable/Source/CustomizableObjectEditor/Private/MuCOE/GenerateMutableSource/GenerateMutableSourceSurface.cpp
--- a/Engine/Plugins/Mutable/Source/CustomizableObjectEditor/Private/MuCOE/GenerateMutableSource/GenerateMutableSourceSurface.cpp
+++ b/Engine/Plugins/Mutable/Source/CustomizableObjectEditor/Private/MuCOE/GenerateMutableSource/GenerateMutableSourceSurface.cpp
@@ -687,13 +687,17 @@ mu::Ptr<mu::NodeSurface> GenerateMutableSourceSurface(const UEdGraphPin * Pin, F
 
 									CompositedImage->SetNormal(NormalCompositeMipmapImage);
 
-									const UTextureLODSettings& TextureLODSettings = GenerationContext.CompilationContext.Options.TargetPlatform->GetTextureLODSettings();;
-									const FTextureLODGroup& LODGroupInfo = TextureLODSettings.GetTextureLODGroup(ReferenceCompositeNormalTexture->LODGroup);
+									CompositeNormalImage->SourceDataDescriptor.OptionalMaxLODSize = 0;
+									if (GenerationContext.CompilationContext.Options.TargetPlatform)
+									{
+										const UTextureLODSettings& TextureLODSettings = GenerationContext.CompilationContext.Options.TargetPlatform->GetTextureLODSettings();
+										const FTextureLODGroup& LODGroupInfo = TextureLODSettings.GetTextureLODGroup(ReferenceCompositeNormalTexture->LODGroup);
+
+										CompositeNormalImage->SourceDataDescriptor.OptionalMaxLODSize = LODGroupInfo.OptionalMaxLODSize;
+										CompositeNormalImage->SourceDataDescriptor.OptionalLODBias = LODGroupInfo.OptionalLODBias;
+										CompositeNormalImage->SourceDataDescriptor.NumNonOptionalLODs = GenerationContext.CompilationContext.Options.MinDiskMips;
+									}
 
-									CompositeNormalImage->SourceDataDescriptor.OptionalMaxLODSize = LODGroupInfo.OptionalMaxLODSize;
-									CompositeNormalImage->SourceDataDescriptor.OptionalLODBias = LODGroupInfo.OptionalLODBias;
-									CompositeNormalImage->SourceDataDescriptor.NumNonOptionalLODs = GenerationContext.CompilationContext.Options.MinDiskMips;
-									
 									const FString TextureName = GetNameSafe(ReferenceCompositeNormalTexture).ToLower();
 									CompositeNormalImage->SourceDataDescriptor.SourceId = CityHash32(reinterpret_cast<const char*>(*TextureName), TextureName.Len() * sizeof(FString::ElementType));
 								}
