[Mutable] Move FCustomizableObjectMeshStreamIn::TaskSynchronization decrement logic from DoCancelMeshUpdate to Abort to fix GC hang.

#rnx

diff --git a/Engine/Plugins/Mutable/Source/CustomizableObject/Private/MuCO/CustomizableObjectMeshUpdate.cpp b/Engine/Plugins/Mutable/Source/CustomizableObject/Private/MuCO/CustomizableObjectMeshUpdate.cpp
--- a/Engine/Plugins/Mutable/Source/CustomizableObject/Private/MuCO/CustomizableObjectMeshUpdate.cpp
+++ b/Engine/Plugins/Mutable/Source/CustomizableObject/Private/MuCO/CustomizableObjectMeshUpdate.cpp
@@ -24,6 +24,14 @@ template class TRenderAssetUpdate<FSkelMeshUpdateContext>;
 
 #define UE_MUTABLE_UPDATE_MESH_REGION		TEXT("Task_Mutable_UpdateMesh")
 
+static bool bEnableGCHangFix = true;
+FAutoConsoleVariableRef CVarMutableEnableGCHangFix(
+	TEXT("mutable.EnableGCHangFix"),
+	bEnableGCHangFix,
+	TEXT("Fix hang when FCustomizableObjectMeshStreamIn is canceled and TaskSynchronization is higher than 0.")
+	TEXT("If true, TaskSynchronization decrement will happen in the Abort method instead of DoCancelMeshUpdate.")
+);
+
 FCustomizableObjectMeshStreamIn::FCustomizableObjectMeshStreamIn(const UCustomizableObjectSkeletalMesh* InMesh, bool bInHighPrio, EThreadType CreateResourcesThread)
 	: FSkeletalMeshStreamIn(InMesh, CreateResourcesThread),
 	bHighPriority(bInHighPrio)
@@ -49,29 +57,61 @@ FCustomizableObjectMeshStreamIn::FCustomizableObjectMeshStreamIn(const UCustomiz
 
 void FCustomizableObjectMeshStreamIn::OnUpdateMeshFinished()
 {
-	check(TaskSynchronization.GetValue() > 0)
+	if (!IsCancelled() || !bEnableGCHangFix)
+	{
+		check(TaskSynchronization.GetValue() > 0)
 
-	// At this point task synchronization would hold the number of pending requests.
-	TaskSynchronization.Decrement();
+		// At this point task synchronization would hold the number of pending requests.
+		TaskSynchronization.Decrement();
 
-	// The tick here is intended to schedule the success or cancel callback.
-	// Using TT_None ensure gets which could create a dead lock.
-	Tick(FSkeletalMeshUpdate::TT_None);
+		// The tick here is intended to schedule the success or cancel callback.
+		// Using TT_None ensure gets which could create a dead lock.
+		Tick(FSkeletalMeshUpdate::TT_None);
+	}
 }
 
-void FCustomizableObjectMeshStreamIn::DoInitiate(const FContext& Context)
+void FCustomizableObjectMeshStreamIn::Abort()
 {
-	check(Context.CurrentThread == TT_Async);
+	if (!IsCancelled() && !IsCompleted() && bEnableGCHangFix)
+	{
+		FSkeletalMeshStreamIn::Abort();
 
-	if (IsCancelled())
+		// At this point task synchronization might hold the number of pending requests.
+		TaskSynchronization.Set(0);
+
+		if (MutableTaskId != FMutableTaskGraph::INVALID_ID && UCustomizableObjectSystem::IsCreated())
+		{
+			if (UCustomizableObjectSystemPrivate* CustomizableObjectSystem = UCustomizableObjectSystem::GetInstance()->GetPrivate())
+			{
+				// Cancel task if not launched yet.
+				CustomizableObjectSystem->MutableTaskGraph.CancelMutableThreadTaskLowPriority(MutableTaskId);
+			}
+		}
+
+		// Clear MeshUpdate data
+		OperationData = nullptr;
+	}
+	else
 	{
-		return;
+		FSkeletalMeshStreamIn::Abort();
 	}
+}
+
+void FCustomizableObjectMeshStreamIn::DoInitiate(const FContext& Context)
+{
+	check(Context.CurrentThread == TT_Async);
 
 	// Launch MutableTask
 	RequestMeshUpdate(Context);
 
-	PushTask(Context, TT_Async, SRA_UPDATE_CALLBACK(DoConvertResources), TT_Async, SRA_UPDATE_CALLBACK(DoCancelMeshUpdate));
+	if (bEnableGCHangFix)
+	{
+		PushTask(Context, TT_Async, SRA_UPDATE_CALLBACK(DoConvertResources), TT_Async, SRA_UPDATE_CALLBACK(DoCancel));
+	}
+	else
+	{
+		PushTask(Context, TT_Async, SRA_UPDATE_CALLBACK(DoConvertResources), TT_Async, SRA_UPDATE_CALLBACK(DoCancelMeshUpdate));
+	}
 }
 
 void FCustomizableObjectMeshStreamIn::DoConvertResources(const FContext& Context)
@@ -132,7 +172,7 @@ namespace impl
 	{
 		MUTABLE_CPUPROFILER_SCOPE(Task_Mutable_UpdateMesh_Loop);
 
-		if (LODIndex == OperationData->CurrentFirstLODIdx)
+		if (Task->IsCancelled() || LODIndex == OperationData->CurrentFirstLODIdx)
 		{
 			Task_Mutable_UpdateMesh_End(OperationData, Task, InstanceID);
 			return;
@@ -154,6 +194,12 @@ namespace impl
 	void Task_Mutable_UpdateMesh(const TSharedPtr<FMutableMeshOperationData> OperationData, TRefCountPtr<FCustomizableObjectMeshStreamIn>& Task)
 	{
 		MUTABLE_CPUPROFILER_SCOPE(Task_Mutable_UpdateMesh);
+
+		if (Task->IsCancelled() && bEnableGCHangFix)
+		{
+			return;
+		}
+
 		TRACE_BEGIN_REGION(UE_MUTABLE_UPDATE_MESH_REGION);
 
 		TSharedPtr<mu::FSystem> System = OperationData->System;
@@ -187,6 +233,11 @@ namespace impl
 
 void FCustomizableObjectMeshStreamIn::RequestMeshUpdate(const FContext& Context)
 {
+	if (IsCancelled())
+	{
+		return;
+	}
+
 	if (!UCustomizableObjectSystem::IsActive())
 	{
 		Abort();
@@ -218,6 +269,11 @@ void FCustomizableObjectMeshStreamIn::RequestMeshUpdate(const FContext& Context)
 		{
 			impl::Task_Mutable_UpdateMesh(SharedOperationData, RefThis);
 		});
+
+	if (IsCancelled() && TaskSynchronization.GetValue() > 0 && bEnableGCHangFix)
+	{
+		TaskSynchronization.Set(0);
+	}
 }
 
 
diff --git a/Engine/Plugins/Mutable/Source/CustomizableObject/Private/MuCO/CustomizableObjectMeshUpdate.h b/Engine/Plugins/Mutable/Source/CustomizableObject/Private/MuCO/CustomizableObjectMeshUpdate.h
--- a/Engine/Plugins/Mutable/Source/CustomizableObject/Private/MuCO/CustomizableObjectMeshUpdate.h
+++ b/Engine/Plugins/Mutable/Source/CustomizableObject/Private/MuCO/CustomizableObjectMeshUpdate.h
@@ -55,6 +55,9 @@ public:
 
 	void OnUpdateMeshFinished();
 
+	// Decrement TaskSynchronization counter and cancel pending low-priority tasks.
+	virtual void Abort() override;
+
 private:
 
 	void DoInitiate(const FContext& Context);
