[mutable] Fixed incomplete undo transaction for the creation of CORerouteNode nodes
- We were not marking the graph as modified, only the node.
- Added backwards compatibility fixup for the case where no pins are found.

#rnx

diff --git a/Engine/Plugins/Mutable/Source/CustomizableObject/Internal/MuCO/CustomizableObjectCustomVersion.h b/Engine/Plugins/Mutable/Source/CustomizableObject/Internal/MuCO/CustomizableObjectCustomVersion.h
--- a/Engine/Plugins/Mutable/Source/CustomizableObject/Internal/MuCO/CustomizableObjectCustomVersion.h
+++ b/Engine/Plugins/Mutable/Source/CustomizableObject/Internal/MuCO/CustomizableObjectCustomVersion.h
@@ -166,6 +166,8 @@ struct CUSTOMIZABLEOBJECT_API FCustomizableObjectCustomVersion
 
 		SetDisplayNamePropertyAsPinNameOfTableNodes,
 
+		FixedNoPinsInRerouteNodes,
+
 		// -----<new versions can be added above this line>-------------------------------------------------
 		VersionPlusOne,
 		LatestVersion = VersionPlusOne - 1
diff --git a/Engine/Plugins/Mutable/Source/CustomizableObjectEditor/Private/MuCOE/EdGraphSchema_CustomizableObject.cpp b/Engine/Plugins/Mutable/Source/CustomizableObjectEditor/Private/MuCOE/EdGraphSchema_CustomizableObject.cpp
--- a/Engine/Plugins/Mutable/Source/CustomizableObjectEditor/Private/MuCOE/EdGraphSchema_CustomizableObject.cpp
+++ b/Engine/Plugins/Mutable/Source/CustomizableObjectEditor/Private/MuCOE/EdGraphSchema_CustomizableObject.cpp
@@ -1192,12 +1192,14 @@ void UEdGraphSchema_CustomizableObject::OnPinConnectionDoubleCicked(UEdGraphPin*
 {
 	const FScopedTransaction Transaction(LOCTEXT("CreateRerouteNodeOnWire", "Create Reroute Node"));
 
+	UEdGraph* ParentGraph = PinA->GetOwningNode()->GetGraph();
+	ParentGraph->Modify();
+
 	// This constant is duplicated from inside of SGraphNodeKnot
 	const FVector2D NodeSpacerSize(42.0f, 24.0f);
 	const FVector2D KnotTopLeft = GraphPosition - (NodeSpacerSize * 0.5f);
 
 	// Create a new knot
-	UEdGraph* ParentGraph = PinA->GetOwningNode()->GetGraph();
 	UCustomizableObjectNodeReroute* DefaultNodeReroute = CastChecked<UCustomizableObjectNodeReroute>(UCustomizableObjectNodeReroute::StaticClass()->GetDefaultObject());
 	UCustomizableObjectNodeReroute* NodeReroute = CastChecked<UCustomizableObjectNodeReroute>(FCustomizableObjectSchemaAction_NewNode::CreateNode(ParentGraph, nullptr, KnotTopLeft, DefaultNodeReroute));
 
diff --git a/Engine/Plugins/Mutable/Source/CustomizableObjectEditor/Private/MuCOE/Nodes/CustomizableObjectNode.cpp b/Engine/Plugins/Mutable/Source/CustomizableObjectEditor/Private/MuCOE/Nodes/CustomizableObjectNode.cpp
--- a/Engine/Plugins/Mutable/Source/CustomizableObjectEditor/Private/MuCOE/Nodes/CustomizableObjectNode.cpp
+++ b/Engine/Plugins/Mutable/Source/CustomizableObjectEditor/Private/MuCOE/Nodes/CustomizableObjectNode.cpp
@@ -221,13 +221,13 @@ void UCustomizableObjectNode::FixupReconstructPins(UCustomizableObjectNodeRemapP
 
 	RemoveWarnings();
 
-	// Move the existing orphan and non orphan pins to a saved array
+	// Move the existing orphan and non-orphan pins to a saved array
 	TArray<UEdGraphPin*> OldPins(Pins); // We can not empty Pins at this point since it will break all FEdGraphPinReference during the reconstruction.
 
 	// Recreate the new pins
 	FuncAllocateDefaultPins(this, RemapPinsAction);
 	
-	// Try to remap orphan and non orphan pins.
+	// Try to remap orphan and non-orphan pins.
 	TArray<UEdGraphPin*> NewPins;
 	NewPins.Reset(Pins.Num() - OldPins.Num());
 	for (UEdGraphPin* Pin : Pins)
diff --git a/Engine/Plugins/Mutable/Source/CustomizableObjectEditor/Private/MuCOE/Nodes/CustomizableObjectNodeReroute.cpp b/Engine/Plugins/Mutable/Source/CustomizableObjectEditor/Private/MuCOE/Nodes/CustomizableObjectNodeReroute.cpp
--- a/Engine/Plugins/Mutable/Source/CustomizableObjectEditor/Private/MuCOE/Nodes/CustomizableObjectNodeReroute.cpp
+++ b/Engine/Plugins/Mutable/Source/CustomizableObjectEditor/Private/MuCOE/Nodes/CustomizableObjectNodeReroute.cpp
@@ -5,7 +5,7 @@
 #include "SCustomizableObjectNodeReroute.h"
 #include "SGraphNodeKnot.h"
 #include "MuCOE/EdGraphSchema_CustomizableObject.h"
-#include "MuCOE/GraphTraversal.h"
+#include "MuCO/CustomizableObjectCustomVersion.h"
 #include "MuCOE/RemapPins/CustomizableObjectNodeRemapPinsByPosition.h"
 
 #define LOCTEXT_NAMESPACE "CustomizableObjectEditor"
@@ -38,6 +38,7 @@ FText UCustomizableObjectNodeReroute::GetTooltipText() const
 	return LOCTEXT("KnotTooltip", "Reroute Node (reroutes wires)");
 }
 
+
 FText UCustomizableObjectNodeReroute::GetNodeTitle(ENodeTitleType::Type TitleType) const
 {
 	if (TitleType == ENodeTitleType::EditableTitle)
@@ -101,6 +102,26 @@ UCustomizableObjectNodeRemapPins* UCustomizableObjectNodeReroute::CreateRemapPin
 }
 
 
+void UCustomizableObjectNodeReroute::BackwardsCompatibleFixup(int32 CustomizableObjectCustomVersion)
+{
+	Super::BackwardsCompatibleFixup(CustomizableObjectCustomVersion);
+
+	if (CustomizableObjectCustomVersion == FCustomizableObjectCustomVersion::FixedNoPinsInRerouteNodes)
+	{
+		// Case where the Pins array is empty when it should not be due to an inconmplete undo operation
+		if (Pins.IsEmpty())
+		{
+			// Extracted from UCustomizableObjectNodeReroute::AllocateDefaultPins
+			
+			UEdGraphPin* MyInputPin = CreatePin(EGPD_Input, UEdGraphSchema_CustomizableObject::PC_Wildcard, TEXT("InputPin"));
+			MyInputPin->bDefaultValueIsIgnored = true;
+
+			CreatePin(EGPD_Output, UEdGraphSchema_CustomizableObject::PC_Wildcard, TEXT("OutputPin"));
+		}
+	}
+}
+
+
 UEdGraphPin* UCustomizableObjectNodeReroute::GetInputPin() const
 {
 	return Pins[0];
diff --git a/Engine/Plugins/Mutable/Source/CustomizableObjectEditor/Public/MuCOE/Nodes/CustomizableObjectNode.h b/Engine/Plugins/Mutable/Source/CustomizableObjectEditor/Public/MuCOE/Nodes/CustomizableObjectNode.h
--- a/Engine/Plugins/Mutable/Source/CustomizableObjectEditor/Public/MuCOE/Nodes/CustomizableObjectNode.h
+++ b/Engine/Plugins/Mutable/Source/CustomizableObjectEditor/Public/MuCOE/Nodes/CustomizableObjectNode.h
@@ -237,7 +237,7 @@ public:
 	/**
 	 * Specialization of ReconstructNode UEdGraphNode function.
 	 *
-	 * @param RemapPins pointer to a remap pins action. Remap pins action which will be used once the node has been reconstructed. Can be null.
+	 * @param RemapPinsMode pointer to a remap pins action. Remap pins action which will be used once the node has been reconstructed. Can be null.
 	 */
 	virtual void ReconstructNode(UCustomizableObjectNodeRemapPins* RemapPinsMode);
 
@@ -253,7 +253,7 @@ public:
 	 */
 	void SetPinHidden(const TArray<UEdGraphPin*>& Pins, bool bHidden);
 
-	/** Return true if a pin can be hidden. Override it to have a custom behaivour. */
+	/** Return true if a pin can be hidden. Override it to have a custom behaviour. */
 	virtual bool CanPinBeHidden(const UEdGraphPin& Pin) const;
 
 	/** @return true if the details should spawn a Pin Viewer. */
diff --git a/Engine/Plugins/Mutable/Source/CustomizableObjectEditor/Public/MuCOE/Nodes/CustomizableObjectNodeReroute.h b/Engine/Plugins/Mutable/Source/CustomizableObjectEditor/Public/MuCOE/Nodes/CustomizableObjectNodeReroute.h
--- a/Engine/Plugins/Mutable/Source/CustomizableObjectEditor/Public/MuCOE/Nodes/CustomizableObjectNodeReroute.h
+++ b/Engine/Plugins/Mutable/Source/CustomizableObjectEditor/Public/MuCOE/Nodes/CustomizableObjectNodeReroute.h
@@ -24,7 +24,8 @@ public:
 	virtual TSharedPtr<SGraphNode> CreateVisualWidget() override;
 	virtual void NodeConnectionListChanged() override;
 	virtual UCustomizableObjectNodeRemapPins* CreateRemapPinsDefault() const override;
-
+	virtual void BackwardsCompatibleFixup(int32 CustomizableObjectCustomVersion) override;
+	
 	// UCustomizableObjectNode interface
 	virtual void AllocateDefaultPins(UCustomizableObjectNodeRemapPins* RemapPins) override;
 	virtual void ReconstructNode(UCustomizableObjectNodeRemapPins* RemapPinsMode) override;
