[Mutable] Fix FCustomizableObjectCustomVersion::NodeMaterialTypedImagePins backwards compatibility code

#rnx

diff --git a/Engine/Plugins/Mutable/Source/CustomizableObjectEditor/Private/MuCOE/Nodes/CustomizableObjectNodeMaterial.cpp b/Engine/Plugins/Mutable/Source/CustomizableObjectEditor/Private/MuCOE/Nodes/CustomizableObjectNodeMaterial.cpp
--- a/Engine/Plugins/Mutable/Source/CustomizableObjectEditor/Private/MuCOE/Nodes/CustomizableObjectNodeMaterial.cpp
+++ b/Engine/Plugins/Mutable/Source/CustomizableObjectEditor/Private/MuCOE/Nodes/CustomizableObjectNodeMaterial.cpp
@@ -658,6 +658,10 @@ void UCustomizableObjectNodeMaterial::BackwardsCompatibleFixup(int32 Customizabl
 					{
 						PinData.PinMode = EPinMode::Mutable;
 					}
+					else if (PinData.PinMode == EPinMode::Passthrough)
+					{
+						PinData.PinMode = EPinMode::Mutable;
+					}
 				}
 				else if (LinkedPin->PinType.PinCategory == UEdGraphSchema_CustomizableObject::PC_PassThroughImage)
 				{
@@ -668,6 +672,10 @@ void UCustomizableObjectNodeMaterial::BackwardsCompatibleFixup(int32 Customizabl
 					{
 						PinData.PinMode = EPinMode::Passthrough;
 					}
+					else if (PinData.PinMode == EPinMode::Mutable)
+					{
+						PinData.PinMode = EPinMode::Passthrough;
+					}
 				}
 			}
 			else
