[Mutable] Fix pins breaking due to Editor language settings.

diff --git a/Engine/Plugins/Mutable/Source/CustomizableObjectEditor/Private/MuCOE/EdGraphSchema_CustomizableObject.cpp b/Engine/Plugins/Mutable/Source/CustomizableObjectEditor/Private/MuCOE/EdGraphSchema_CustomizableObject.cpp
--- a/Engine/Plugins/Mutable/Source/CustomizableObjectEditor/Private/MuCOE/EdGraphSchema_CustomizableObject.cpp
+++ b/Engine/Plugins/Mutable/Source/CustomizableObjectEditor/Private/MuCOE/EdGraphSchema_CustomizableObject.cpp
@@ -1352,9 +1352,104 @@ FPinConnectionResponse UEdGraphSchema_CustomizableObject::MovePinLinks(UEdGraphP
 }
 
 
-FText UEdGraphSchema_CustomizableObject::GetPinCategoryName(const FName& PinCategory)
+FName UEdGraphSchema_CustomizableObject::GetPinCategoryName(const FName& PinCategory)
 {
 	if (PinCategory == UEdGraphSchema_CustomizableObject::PC_Object)
+	{
+		return TEXT("Object");
+	}
+	else if (PinCategory == UEdGraphSchema_CustomizableObject::PC_Component)
+	{
+		return TEXT("Component");
+	}
+	else if (PinCategory == UEdGraphSchema_CustomizableObject::PC_Material)
+	{
+		return TEXT("Material");
+	}
+	else if (PinCategory == UEdGraphSchema_CustomizableObject::PC_Modifier)
+	{
+		return TEXT("Modifier");
+	}
+	else if (PinCategory == UEdGraphSchema_CustomizableObject::PC_Mesh)
+	{
+		return TEXT("Mesh");
+	}
+	else if (PinCategory == UEdGraphSchema_CustomizableObject::PC_Layout)
+	{
+		return TEXT("Layout");
+	}
+	else if (PinCategory == UEdGraphSchema_CustomizableObject::PC_Image)
+	{
+		return TEXT("Texture");
+	}
+	else if (PinCategory == UEdGraphSchema_CustomizableObject::PC_PassThroughImage)
+	{
+		return TEXT("PassThrough Texture");
+	}
+	else if (PinCategory == UEdGraphSchema_CustomizableObject::PC_Projector)
+	{
+		return TEXT("Projector");
+	}
+	else if (PinCategory == UEdGraphSchema_CustomizableObject::PC_GroupProjector)
+	{
+		return TEXT("Group Projector");
+	}
+	else if (PinCategory == UEdGraphSchema_CustomizableObject::PC_Color)
+	{
+		return TEXT("Color");
+	}
+	else if (PinCategory == UEdGraphSchema_CustomizableObject::PC_Float)
+	{
+		return TEXT("Float");
+	}
+	else if (PinCategory == UEdGraphSchema_CustomizableObject::PC_Bool)
+	{
+		return TEXT("Bool");
+	}
+	else if (PinCategory == UEdGraphSchema_CustomizableObject::PC_Enum)
+	{
+		return TEXT("Enum");
+	}
+	else if (PinCategory == UEdGraphSchema_CustomizableObject::PC_Stack)
+	{
+		return TEXT("Stack");
+	}
+	else if (PinCategory == UEdGraphSchema_CustomizableObject::PC_MaterialAsset)
+	{
+		return TEXT("Table Material");
+	}
+	else if (PinCategory == UEdGraphSchema_CustomizableObject::PC_Wildcard)
+	{
+		return TEXT("Wildcard");
+	}
+	else if (PinCategory == UEdGraphSchema_CustomizableObject::PC_PoseAsset)
+	{
+		return TEXT("PoseAsset");
+	}
+	else if (PinCategory == UEdGraphSchema_CustomizableObject::PC_Transform)
+	{
+		return TEXT("Transform");
+	}
+	else
+	{
+		for (const FRegisteredCustomizableObjectPinType& PinType : ICustomizableObjectModule::Get().GetExtendedPinTypes())
+		{
+			if (PinType.PinType.Name == PinCategory)
+			{
+				return PinType.PinType.Name;
+			}
+		}
+
+		// Need to fail gracefully here in case a plugin that was active when this graph was
+		// created is no longer loaded.
+		return TEXT("Unknown");
+	}
+}
+
+
+FText UEdGraphSchema_CustomizableObject::GetPinCategoryFriendlyName(const FName& PinCategory)
+{
+if (PinCategory == UEdGraphSchema_CustomizableObject::PC_Object)
 	{
 		return LOCTEXT("Object_Pin_Category", "Object");
 	}
@@ -1364,7 +1459,7 @@ FText UEdGraphSchema_CustomizableObject::GetPinCategoryName(const FName& PinCate
 	}
 	else if (PinCategory == UEdGraphSchema_CustomizableObject::PC_Material)
 	{
-		return LOCTEXT("MeshSection_Pin_Category", "Material");
+		return LOCTEXT("MeshSection_Pin_Category", "Mesh Section");
 	}
 	else if (PinCategory == UEdGraphSchema_CustomizableObject::PC_Modifier)
 	{
@@ -1447,17 +1542,4 @@ FText UEdGraphSchema_CustomizableObject::GetPinCategoryName(const FName& PinCate
 }
 
 
-FText UEdGraphSchema_CustomizableObject::GetPinCategoryFriendlyName(const FName& PinCategory)
-{
-	if (PinCategory == UEdGraphSchema_CustomizableObject::PC_Material)
-	{
-		return LOCTEXT("MeshSection_Pin_Category_Friendly", "Mesh Section");
-	}
-	else
-	{
-		return GetPinCategoryName(PinCategory);
-	}
-}
-
-
 #undef LOCTEXT_NAMESPACE
diff --git a/Engine/Plugins/Mutable/Source/CustomizableObjectEditor/Public/MuCOE/EdGraphSchema_CustomizableObject.h b/Engine/Plugins/Mutable/Source/CustomizableObjectEditor/Public/MuCOE/EdGraphSchema_CustomizableObject.h
--- a/Engine/Plugins/Mutable/Source/CustomizableObjectEditor/Public/MuCOE/EdGraphSchema_CustomizableObject.h
+++ b/Engine/Plugins/Mutable/Source/CustomizableObjectEditor/Public/MuCOE/EdGraphSchema_CustomizableObject.h
@@ -170,7 +170,7 @@ public:
 	UEdGraphNode* AddComment(class UEdGraph* ParentGraph, UEdGraphPin* FromPin, const FVector2D Location, bool bSelectNewNode = true) const;
 	
 	/** Given a pin category name, get its friendly name (user readable). */
-	static FText GetPinCategoryName(const FName& PinCategory);
+	static FName GetPinCategoryName(const FName& PinCategory);
 	static FText GetPinCategoryFriendlyName(const FName& PinCategory);
 
 };
