[Mutable] Fix grooms not attaching when opening an editor.

- Removed FCustomizableObjectInstanceEditor::OnUpdatePreviewInstance duplicated code.

diff --git a/Engine/Plugins/Mutable/Source/CustomizableObject/Private/MuCO/CustomizableObjectSystem.cpp b/Engine/Plugins/Mutable/Source/CustomizableObject/Private/MuCO/CustomizableObjectSystem.cpp
--- a/Engine/Plugins/Mutable/Source/CustomizableObject/Private/MuCO/CustomizableObjectSystem.cpp
+++ b/Engine/Plugins/Mutable/Source/CustomizableObject/Private/MuCO/CustomizableObjectSystem.cpp
@@ -993,6 +993,7 @@ void UpdateSkeletalMesh(const TSharedRef<FUpdateContextPrivate>& Context)
 	for (const TTuple<FName, TObjectPtr<USkeletalMesh>>& Pair : CustomizableObjectInstance->GetPrivate()->SkeletalMeshes)
 	{
 		FPreSetSkeletalMeshParams Params;
+		Params.Instance = CustomizableObjectInstance;
 		Params.SkeletalMesh = Pair.Value;
 		
 		CustomizableObjectInstance->PreSetSkeletalMeshDelegate.Broadcast(Params);
diff --git a/Engine/Plugins/Mutable/Source/CustomizableObject/Public/MuCO/CustomizableObjectInstance.h b/Engine/Plugins/Mutable/Source/CustomizableObject/Public/MuCO/CustomizableObjectInstance.h
--- a/Engine/Plugins/Mutable/Source/CustomizableObject/Public/MuCO/CustomizableObjectInstance.h
+++ b/Engine/Plugins/Mutable/Source/CustomizableObject/Public/MuCO/CustomizableObjectInstance.h
@@ -205,6 +205,9 @@ struct FPreSetSkeletalMeshParams
 {
 	GENERATED_BODY()
 
+	UPROPERTY()
+	TObjectPtr<UCustomizableObjectInstance> Instance;
+	
 	UPROPERTY()
 	TObjectPtr<USkeletalMesh> SkeletalMesh;
 };
diff --git a/Engine/Plugins/Mutable/Source/CustomizableObjectEditor/Private/MuCOE/CustomizableObjectEditorViewportClient.cpp b/Engine/Plugins/Mutable/Source/CustomizableObjectEditor/Private/MuCOE/CustomizableObjectEditorViewportClient.cpp
--- a/Engine/Plugins/Mutable/Source/CustomizableObjectEditor/Private/MuCOE/CustomizableObjectEditorViewportClient.cpp
+++ b/Engine/Plugins/Mutable/Source/CustomizableObjectEditor/Private/MuCOE/CustomizableObjectEditorViewportClient.cpp
@@ -741,6 +741,7 @@ void FCustomizableObjectEditorViewportClient::CreatePreviewActor(const TWeakObje
 
 	bUpdated = false;
 	
+	InInstance->PreSetSkeletalMeshNativeDelegate.AddSP(SharedThis(this), &FCustomizableObjectEditorViewportClient::OnPreSetSkeletalMesh);
 	InInstance->UpdatedNativeDelegate.AddSP(SharedThis(this), &FCustomizableObjectEditorViewportClient::OnInstanceUpdate);
 	
 	Invalidate();
@@ -785,8 +786,9 @@ void FCustomizableObjectEditorViewportClient::SetPreviewAnimationAsset(UAnimatio
 }
 
 
-void FCustomizableObjectEditorViewportClient::OnInstanceUpdate(UCustomizableObjectInstance* Instance)
+void FCustomizableObjectEditorViewportClient::OnPreSetSkeletalMesh(const FPreSetSkeletalMeshParams& Params)
 {
+	TObjectPtr<UCustomizableObjectInstance> Instance = Params.Instance;
 	if (!Instance)
 	{
 		return;
@@ -854,8 +856,11 @@ void FCustomizableObjectEditorViewportClient::OnInstanceUpdate(UCustomizableObje
 			}
 		}
 	}
+}
 
 
+void FCustomizableObjectEditorViewportClient::OnInstanceUpdate(UCustomizableObjectInstance* Instance)
+{
 	if (!bUpdated)
 	{
 		bUpdated = true;
diff --git a/Engine/Plugins/Mutable/Source/CustomizableObjectEditor/Private/MuCOE/CustomizableObjectInstanceEditor.cpp b/Engine/Plugins/Mutable/Source/CustomizableObjectEditor/Private/MuCOE/CustomizableObjectInstanceEditor.cpp
--- a/Engine/Plugins/Mutable/Source/CustomizableObjectEditor/Private/MuCOE/CustomizableObjectInstanceEditor.cpp
+++ b/Engine/Plugins/Mutable/Source/CustomizableObjectEditor/Private/MuCOE/CustomizableObjectInstanceEditor.cpp
@@ -615,23 +615,6 @@ void FCustomizableObjectInstanceEditor::OnInstancePropertySelectionChanged(FProp
 
 void FCustomizableObjectInstanceEditor::OnUpdatePreviewInstance(UCustomizableObjectInstance* Instance)
 {
-	if (!Actor)
-	{
-		return;
-	}
-
-	check(CustomizableObjectInstance)
-	
-	if (Instance->GetPrivate()->SkeletalMeshStatus != ESkeletalMeshStatus::Success)
-	{
-		Actor->GetRootComponent()->SetVisibility(false, true);
-	}
-
-	Actor->GetRootComponent()->UpdateBounds();
-	
-	// Sets which UV material to draw
-	Viewport->SetDrawDefaultUVMaterial();
-
 	if (TextureAnalyzer.IsValid())
 	{
 		TextureAnalyzer.Get()->RefreshTextureAnalyzerTable(CustomizableObjectInstance);
diff --git a/Engine/Plugins/Mutable/Source/CustomizableObjectEditor/Public/MuCOE/CustomizableObjectEditorViewportClient.h b/Engine/Plugins/Mutable/Source/CustomizableObjectEditor/Public/MuCOE/CustomizableObjectEditorViewportClient.h
--- a/Engine/Plugins/Mutable/Source/CustomizableObjectEditor/Public/MuCOE/CustomizableObjectEditorViewportClient.h
+++ b/Engine/Plugins/Mutable/Source/CustomizableObjectEditor/Public/MuCOE/CustomizableObjectEditorViewportClient.h
@@ -324,8 +324,9 @@ private:
 	bool HandleBeginTransform();
 	bool HandleEndTransform();
 
-	void OnInstanceUpdate(UCustomizableObjectInstance* Instance);
+	void OnPreSetSkeletalMesh(const FPreSetSkeletalMeshParams& Params);
 
+	void OnInstanceUpdate(UCustomizableObjectInstance* Instance);
 
 	/** Preview Actor. All preview components are attached to this actor. */
 	TStrongObjectPtr<AActor> Actor;
