[Mutable] Fix crash when using Image Layer Blend Normals with a Mask.

diff --git a/Engine/Plugins/Mutable/Source/MutableRuntime/Private/MuR/CodeRunnerTasked.cpp b/Engine/Plugins/Mutable/Source/MutableRuntime/Private/MuR/CodeRunnerTasked.cpp
--- a/Engine/Plugins/Mutable/Source/MutableRuntime/Private/MuR/CodeRunnerTasked.cpp
+++ b/Engine/Plugins/Mutable/Source/MutableRuntime/Private/MuR/CodeRunnerTasked.cpp
@@ -1023,9 +1023,23 @@ namespace mu
 		{
 			FImageOperator ImOp = MakeImageOperator(Runner);
 
+			if (Mask->GetFormat() != EImageFormat::L_UByte &&
+				static_cast<EBlendType>(Args.blendType) == EBlendType::BT_NORMAL_COMBINE) // TODO Optimize. BT_NORMAL_COMBINE (ImageLayerCombineColour with mask) does not support formats such as RLE. Hence why we are changing the format here.
+			{
+				MUTABLE_CPUPROFILER_SCOPE(EmergencyFix_Format);
+				TSharedPtr<FImage> Formatted = Runner->CreateImage(Mask->GetSizeX(), Mask->GetSizeY(), Mask->GetLODCount(), EImageFormat::L_UByte, EInitializationType::NotInitialized);
+
+				bool bSuccess = false;
+				ImOp.ImagePixelFormat(bSuccess, ImageCompressionQuality, Formatted.Get(), Mask.Get());
+				check(bSuccess); // Decompression cannot fail
+
+				Runner->Release(Mask);
+				Mask = Formatted;
+			}
+			
 			if (Base->GetSize() != Mask->GetSize())
 			{
-				MUTABLE_CPUPROFILER_SCOPE(ImageResize_EmergencyFixSize);
+				MUTABLE_CPUPROFILER_SCOPE(EmergencyFix_Size);
 				TSharedPtr<FImage> Resized = Runner->CreateImage(Base->GetSizeX(), Base->GetSizeY(), 1, Mask->GetFormat(), EInitializationType::NotInitialized);
 				ImOp.ImageResizeLinear(Resized.Get(), ImageCompressionQuality, Mask.Get() );
 				Runner->Release(Mask);
@@ -1034,7 +1048,7 @@ namespace mu
 
 			if (Mask->GetLODCount() < Base->GetLODCount())
 			{
-				MUTABLE_CPUPROFILER_SCOPE(ImageResize_EmergencyFixMips);
+				MUTABLE_CPUPROFILER_SCOPE(EmergencyFix_LOD);
 				int32 StartLevel = Mask->GetLODCount() - 1;
 				int32 LevelCount = Base->GetLODCount();
 
