[Mutable] Avoid BeginUpdate to sync with the Game Thread due to Extension Data.

#rnx

diff --git a/Engine/Plugins/Mutable/Source/CustomizableObject/Internal/MuCO/CustomizableObjectInstancePrivate.h b/Engine/Plugins/Mutable/Source/CustomizableObject/Internal/MuCO/CustomizableObjectInstancePrivate.h
--- a/Engine/Plugins/Mutable/Source/CustomizableObject/Internal/MuCO/CustomizableObjectInstancePrivate.h
+++ b/Engine/Plugins/Mutable/Source/CustomizableObject/Internal/MuCO/CustomizableObjectInstancePrivate.h
@@ -263,7 +263,7 @@ public:
 	 * If no StreamableManager is provided, it will load assets synchronously. */
 	UE::Tasks::FTask LoadAdditionalAssetsAndData(const TSharedRef<FUpdateContextPrivate>& OperationData);
 
-	void AdditionalAssetsAsyncLoaded();
+	void AdditionalAssetsAsyncLoaded(const TSharedRef<FUpdateContextPrivate> OperationData);
 
 	void TickUpdateCloseCustomizableObjects(UCustomizableObjectInstance& Public, FMutableInstanceUpdateMap& InOutRequestedUpdates);
 	void UpdateInstanceIfNotGenerated(UCustomizableObjectInstance& Public, FMutableInstanceUpdateMap& InOutRequestedUpdates);
diff --git a/Engine/Plugins/Mutable/Source/CustomizableObject/Internal/MuCO/CustomizableObjectPrivate.h b/Engine/Plugins/Mutable/Source/CustomizableObject/Internal/MuCO/CustomizableObjectPrivate.h
--- a/Engine/Plugins/Mutable/Source/CustomizableObject/Internal/MuCO/CustomizableObjectPrivate.h
+++ b/Engine/Plugins/Mutable/Source/CustomizableObject/Internal/MuCO/CustomizableObjectPrivate.h
@@ -1579,9 +1579,6 @@ public:
 #if WITH_EDITORONLY_DATA
 	TArray<FProfileParameterDat>& GetInstancePropertiesProfiles();
 #endif
-	
-	const FCustomizableObjectResourceData* GetLoadedStreamedResource(int32 ResourceIndex);
-	void UnloadStreamedResource(int32 ResourceIndex);
 
 #if WITH_EDITORONLY_DATA
 	TObjectPtr<UEdGraph>& GetSource() const;
diff --git a/Engine/Plugins/Mutable/Source/CustomizableObject/Internal/MuCO/CustomizableObjectStreamedResourceData.h b/Engine/Plugins/Mutable/Source/CustomizableObject/Internal/MuCO/CustomizableObjectStreamedResourceData.h
new file mode 100644
--- /dev/null
+++ b/Engine/Plugins/Mutable/Source/CustomizableObject/Internal/MuCO/CustomizableObjectStreamedResourceData.h
@@ -0,0 +1,80 @@
+// Copyright Epic Games, Inc. All Rights Reserved.
+
+#pragma once
+
+#include "MuCO/CustomizableObjectResourceData.h"
+#include "UObject/Object.h"
+#include "UObject/SoftObjectPtr.h"
+
+#include "CustomizableObjectStreamedResourceData.generated.h"
+
+
+
+/**
+ * Used to store resource data that can be streamed in and out from an external package.
+ *
+ * This allows you to load large resource data on demand, instead of having it always loaded.
+ */
+USTRUCT()
+struct CUSTOMIZABLEOBJECT_API FCustomizableObjectStreamedResourceData
+{
+	GENERATED_BODY()
+
+public:
+	FCustomizableObjectStreamedResourceData() = default;
+
+#if WITH_EDITOR
+	/** Streamed data can only be constructed in the editor */
+	FCustomizableObjectStreamedResourceData(UCustomizableObjectResourceDataContainer* InContainer);
+
+	/**
+	 * Clears the hard reference to the container, so that it's not forced to load as soon as this
+	 * struct is loaded, and replaces the soft reference with the provided one.
+	 * 
+	 * Note that the container may not have been moved to its new path yet, so don't try to load
+	 * the container from the given path here.
+	 */
+	void ConvertToSoftReferenceForCooking(const TSoftObjectPtr<UCustomizableObjectResourceDataContainer>& NewContainerPath);
+#endif // WITH_EDITOR
+
+	/**
+	 * If this returns true, GetLoadedData can safely be called.
+	 *
+	 * The data will stay loaded until Unload is called.
+	 */
+	bool IsLoaded() const;
+
+	/**
+	 * Returns the data, as long as it's loaded.
+	 *
+	 * If the data isn't loaded, an assertion will fail.
+	 */
+	const FCustomizableObjectResourceData& GetLoadedData() const;
+
+	/**
+	 * Release this struct's hard reference to the loaded data.
+	 *
+	 * This doesn't immediately remove it from memory, but will allow it to be deleted by the
+	 * garbage collector if there are no other references.
+	 *
+	 * Note that in editor builds streamed Resource Data must always be loaded due to the package
+	 * structure, and therefore this function has no effect in editor builds.
+	 */
+	void Release();
+
+	/**
+	 * To stream this data in, load the object at the path returned by GetPath and then pass it
+	 * into NotifyLoaded.
+	 */
+	const TSoftObjectPtr<UCustomizableObjectResourceDataContainer>& GetPath() const { return ContainerPath; }
+
+	/** Keep a strong reference to the struct. */	
+	void Hold();
+
+private:
+	UPROPERTY()
+	TSoftObjectPtr<UCustomizableObjectResourceDataContainer> ContainerPath;
+
+	UPROPERTY()
+	TObjectPtr<const UCustomizableObjectResourceDataContainer> Container;
+};
diff --git a/Engine/Plugins/Mutable/Source/CustomizableObject/Internal/MuCO/CustomizableObjectSystemPrivate.h b/Engine/Plugins/Mutable/Source/CustomizableObject/Internal/MuCO/CustomizableObjectSystemPrivate.h
--- a/Engine/Plugins/Mutable/Source/CustomizableObject/Internal/MuCO/CustomizableObjectSystemPrivate.h
+++ b/Engine/Plugins/Mutable/Source/CustomizableObject/Internal/MuCO/CustomizableObjectSystemPrivate.h
@@ -625,6 +625,12 @@ public:
 
 	/** List of component names. Index is the ObjectComponentIndex. */
 	TArray<FName> ComponentNames;
+
+	/** Index of the resource in the StreamedResourceData array of the Model Resources. */
+	TArray<int32> StreamedResourceIndex;
+
+	/** Index of the resource in the ExtensionStreamedResourceData array of the Model Resources. */
+	TArray<int32> ExtensionStreamedResourceIndex;
 	
 	TMap<FName, uint8> NumLODsAvailable;
 
@@ -724,9 +730,7 @@ public:
 
 	// This object is responsible for streaming data to the MutableSystem.
 	TSharedPtr<class FUnrealMutableModelBulkReader> Streamer;
-
-	TSharedPtr<class FUnrealExtensionDataStreamer> ExtensionDataStreamer;
-
+	
 	// This object is responsible for providing custom images and meshes to mutable (for image parameters, etc.)
 	// This object is called from the mutable thread, and it should only access data already safely submitted from
 	// the game thread and stored in FUnrealMutableImageProvider::GlobalExternalImages.
diff --git a/Engine/Plugins/Mutable/Source/CustomizableObject/Private/MuCO/CustomizableObject.cpp b/Engine/Plugins/Mutable/Source/CustomizableObject/Private/MuCO/CustomizableObject.cpp
--- a/Engine/Plugins/Mutable/Source/CustomizableObject/Private/MuCO/CustomizableObject.cpp
+++ b/Engine/Plugins/Mutable/Source/CustomizableObject/Private/MuCO/CustomizableObject.cpp
@@ -2601,48 +2601,6 @@ TArray<FProfileParameterDat>& UCustomizableObjectPrivate::GetInstancePropertiesP
 }
 #endif
 
-const FCustomizableObjectResourceData* UCustomizableObjectPrivate::GetLoadedStreamedResource(int32 ResourceIndex)
-{
-#if WITH_EDITORONLY_DATA
-	FModelResources& LocalModelResources = ModelResourcesEditor;
-#else
-	FModelResources& LocalModelResources = ModelResources;
-#endif
-
-	if (LocalModelResources.StreamedResourceData.IsValidIndex(ResourceIndex))
-	{
-		FCustomizableObjectStreamedResourceData& Resource = LocalModelResources.StreamedResourceData[ResourceIndex];
-		if (!Resource.IsLoaded())
-		{
-			Resource.NotifyLoaded(Resource.GetPath().Get());
-		}
-
-		return &Resource.GetLoadedData();
-	}
-
-	return nullptr;
-}
-
-void UCustomizableObjectPrivate::UnloadStreamedResource(int32 ResourceIndex)
-{
-	// Only Unload in cooked builds. Unloading them when in the editor will trigger an assert. 
-	if (FPlatformProperties::RequiresCookedData())
-	{
-		return;
-	}
-
-#if WITH_EDITORONLY_DATA
-	FModelResources& LocalModelResources = ModelResourcesEditor;
-#else
-	FModelResources& LocalModelResources = ModelResources;
-#endif
-
-	if (LocalModelResources.StreamedResourceData.IsValidIndex(ResourceIndex))
-	{
-		LocalModelResources.StreamedResourceData[ResourceIndex].Unload();
-	}
-}
-
 
 #if WITH_EDITORONLY_DATA
 TObjectPtr<UEdGraph>& UCustomizableObjectPrivate::GetSource() const
@@ -4081,7 +4039,6 @@ void FMutableRefSkeletalMeshData::InitResources(UCustomizableObject* InOuter, FM
 
 		Data.AssetUserData = InModelResources.StreamedResourceData[Data.AssetUserDataIndex].GetPath().LoadSynchronous();
 		check(Data.AssetUserData);
-		check(Data.AssetUserData->Data.Type == ECOResourceDataType::AssetUserData);
 	}
 }
 
diff --git a/Engine/Plugins/Mutable/Source/CustomizableObject/Private/MuCO/CustomizableObjectInstance.cpp b/Engine/Plugins/Mutable/Source/CustomizableObject/Private/MuCO/CustomizableObjectInstance.cpp
--- a/Engine/Plugins/Mutable/Source/CustomizableObject/Private/MuCO/CustomizableObjectInstance.cpp
+++ b/Engine/Plugins/Mutable/Source/CustomizableObject/Private/MuCO/CustomizableObjectInstance.cpp
@@ -53,6 +53,8 @@
 #include "PhysicsEngine/SkeletalBodySetup.h"
 #include "Hash/CityHash.h"
 #include "MuCO/CustomizableObjectCustomVersion.h"
+#include "MuCO/CustomizableObjectResourceData.h"
+#include "MuCO/CustomizableObjectStreamedResourceData.h"
 #include "MuCO/CustomizableObjectResourceDataTypes.h"
 #include "Serialization/BulkData.h"
 
@@ -3608,10 +3610,8 @@ void UCustomizableInstancePrivate::InitSkeletalMeshData(const TSharedRef<FUpdate
 	{
 		for (const FMutableRefAssetUserData& MutAssetUserData : RefSkeletalMeshData.AssetUserData)
 		{
-			if (MutAssetUserData.AssetUserData && MutAssetUserData.AssetUserData->Data.Type == ECOResourceDataType::AssetUserData)
+			if (const FCustomizableObjectAssetUserData* DataPtr = MutAssetUserData.AssetUserData->Data.Data.GetPtr<FCustomizableObjectAssetUserData>())
 			{
-				const FCustomizableObjectAssetUserData* DataPtr = MutAssetUserData.AssetUserData->Data.Data.GetPtr<FCustomizableObjectAssetUserData>();
-				check(DataPtr);
 #if WITH_EDITORONLY_DATA
 				SkeletalMesh->AddAssetUserData(DataPtr->AssetUserDataEditor);
 #else
@@ -5472,6 +5472,7 @@ UE::Tasks::FTask UCustomizableInstancePrivate::LoadAdditionalAssetsAndData(const
 						const FCustomizableObjectStreamedResourceData& StreamedResource = ModelResources.StreamedResourceData[ResourceIndex];
 						if (!StreamedResource.IsLoaded())
 						{
+							OperationData->StreamedResourceIndex.Add(ResourceIndex);
 							AssetsToStream.AddUnique(StreamedResource.GetPath().ToSoftObjectPath());
 						}
 
@@ -5650,6 +5651,43 @@ UE::Tasks::FTask UCustomizableInstancePrivate::LoadAdditionalAssetsAndData(const
 		AssetsToStream.Add(MeshRef.ToSoftObjectPath());
 	}
 
+	// Copy FExtensionData Object node input from the Instance to the InstanceUpdateData
+	for (int32 ExtensionDataIndex = 0; ExtensionDataIndex < OperationData->MutableInstance->GetExtensionDataCount(); ExtensionDataIndex++)
+	{
+		TSharedPtr<const mu::FExtensionData> ExtensionData;
+		FName Name;
+		OperationData->MutableInstance->GetExtensionData(ExtensionDataIndex, ExtensionData, Name);
+
+		check(ExtensionData);
+
+		FInstanceUpdateData::FNamedExtensionData& NewEntry = OperationData->InstanceUpdateData.ExtendedInputPins.AddDefaulted_GetRef();
+		NewEntry.Data = ExtensionData;
+		NewEntry.Name = Name;
+		check(NewEntry.Name != NAME_None);
+
+		if (!ModelResources.StreamedExtensionData.IsValidIndex(ExtensionData->Index))
+		{
+			// The compiled data appears to be out of sync with the CO's properties
+
+			UE_LOG(LogMutable, Error, TEXT("Couldn't find streamed Extension Data with index %d in %s. Compiled data may be stale."),
+				ExtensionData->Index, *CustomizableObject->GetFullName());
+
+			continue;
+		}
+
+		const FCustomizableObjectStreamedResourceData& StreamedData = ModelResources.StreamedExtensionData[ExtensionData->Index];
+		if (StreamedData.IsLoaded())
+		{
+			continue;
+		}
+
+		// Note that this just checks if the path is non-null, NOT if the object is loaded
+		check(!StreamedData.GetPath().IsNull());
+
+		OperationData->ExtensionStreamedResourceIndex.Add(ExtensionData->Index);
+		AssetsToStream.Add(StreamedData.GetPath().ToSoftObjectPath());
+	}
+	
 	TArray<UE::Tasks::FTaskEvent> StreamingCompletionEvents;
 	if (AssetsToStream.Num() > 0)
 	{
@@ -5667,7 +5705,7 @@ UE::Tasks::FTask UCustomizableInstancePrivate::LoadAdditionalAssetsAndData(const
 
 		PrivateSystem->StreamableManager->RequestAsyncLoad(
 			AssetsToStream,
-			FStreamableDelegate::CreateUObject(this, &UCustomizableInstancePrivate::AdditionalAssetsAsyncLoaded),
+			FStreamableDelegate::CreateUObject(this, &UCustomizableInstancePrivate::AdditionalAssetsAsyncLoaded, OperationData),
 			bEnableHighPriorityLoading ? FStreamableManager::AsyncLoadHighPriority : FStreamableManager::DefaultAsyncLoadPriority);
 	}
 	
@@ -6003,7 +6041,7 @@ const TArray<TObjectPtr<UMaterialInterface>>* UCustomizableObjectInstance::GetOv
 }
 
 
-void UCustomizableInstancePrivate::AdditionalAssetsAsyncLoaded()
+void UCustomizableInstancePrivate::AdditionalAssetsAsyncLoaded(const TSharedRef<FUpdateContextPrivate> OperationData)
 {
 	MUTABLE_CPUPROFILER_SCOPE(UCustomizableInstancePrivate::AdditionalAssetsAsyncLoaded);
 
@@ -6013,8 +6051,18 @@ void UCustomizableInstancePrivate::AdditionalAssetsAsyncLoaded()
 
 	UCustomizableObjectPrivate* CustomizableObjectPrivate = GetPublic()->GetCustomizableObject()->GetPrivate();
 
-	const FModelResources& ModelResources = CustomizableObjectPrivate->GetModelResources();
+	FModelResources& ModelResources = CustomizableObjectPrivate->GetModelResources();
 
+	for (int32 ResourceIndex : OperationData->StreamedResourceIndex)
+	{
+		ModelResources.StreamedResourceData[ResourceIndex].Hold();
+	}
+	
+	for (int32 ResourceIndex : OperationData->ExtensionStreamedResourceIndex)
+	{
+		ModelResources.StreamedExtensionData[ResourceIndex].Hold();
+	}
+	
 	// Loaded Materials
 	check(ObjectToInstanceIndexMap.Num() == ReferencedMaterials.Num());
 
@@ -6046,31 +6094,16 @@ void UCustomizableInstancePrivate::AdditionalAssetsAsyncLoaded()
 	{
 		for (int32 ResourceIndex : ComponentData.StreamedResourceIndex)
 		{
-			const FCustomizableObjectResourceData* ResourceData = CustomizableObjectPrivate->GetLoadedStreamedResource(ResourceIndex);
-			if (!ResourceData)
-			{
-				check(false); // Invalid resource index?
-				continue;
-			}
+			const FCustomizableObjectResourceData& ResourceData = ModelResources.StreamedExtensionData[ResourceIndex].GetLoadedData();
 
-			switch (ResourceData->Type)
+			if (const FCustomizableObjectAssetUserData* AUDResource = ResourceData.Data.GetPtr<FCustomizableObjectAssetUserData>())
 			{
-				case ECOResourceDataType::AssetUserData:
-				{
-					const FCustomizableObjectAssetUserData* AUDResource = ResourceData->Data.GetPtr<FCustomizableObjectAssetUserData>();
 #if WITH_EDITORONLY_DATA
-					ComponentData.AssetUserDataArray.Add(AUDResource->AssetUserDataEditor);
+				ComponentData.AssetUserDataArray.Add(AUDResource->AssetUserDataEditor);
 #else
-					ComponentData.AssetUserDataArray.Add(AUDResource->AssetUserData);
+				ComponentData.AssetUserDataArray.Add(AUDResource->AssetUserData);
 #endif
-					break;
-				}
-				default:
-					break;
 			}
-
-			// Unload by removing the reference to the container. Only if the platform has cooked data.
-			CustomizableObjectPrivate->UnloadStreamedResource(ResourceIndex);
 		}
 
 		// Loaded Skeletons
@@ -6175,7 +6208,7 @@ void UCustomizableInstancePrivate::AdditionalAssetsAsyncLoaded()
 	}
 
 	PassThroughMeshesToLoad.Empty();
-
+	
 #if WITH_EDITOR
 	// TODO: Remove with UE-217665 when the underlying bug in the ColorPicker is solved
 	// Reenable the throttling which disabled when launching the Async Load
diff --git a/Engine/Plugins/Mutable/Source/CustomizableObject/Private/MuCO/CustomizableObjectStreamedResourceData.cpp b/Engine/Plugins/Mutable/Source/CustomizableObject/Private/MuCO/CustomizableObjectStreamedResourceData.cpp
--- a/Engine/Plugins/Mutable/Source/CustomizableObject/Private/MuCO/CustomizableObjectStreamedResourceData.cpp
+++ b/Engine/Plugins/Mutable/Source/CustomizableObject/Private/MuCO/CustomizableObjectStreamedResourceData.cpp
@@ -41,40 +41,24 @@ bool FCustomizableObjectStreamedResourceData::IsLoaded() const
 const FCustomizableObjectResourceData& FCustomizableObjectStreamedResourceData::GetLoadedData() const
 {
 	check(IsInGameThread());
-
 	check(Container);
+	
 	return Container->Data;
 }
 
-void FCustomizableObjectStreamedResourceData::Unload()
+void FCustomizableObjectStreamedResourceData::Release()
 {
 	check(IsInGameThread());
 
-	if (!Container)
-	{
-		// Already unloaded
-		return;
-	}
-
-	if (FPlatformProperties::RequiresCookedData())
+	if (FPlatformProperties::RequiresCookedData()) // TODO GMT Remove once UE-232022. Editor will no longer use StreamedResources.
 	{
-		// The container should have been cooked out to its own package
-		check(Container->GetOuter()->IsA<UPackage>());
-
 		Container = nullptr;
 	}
-	else
-	{
-		// In an uncooked build the container is stored in the CO package and the CO needs to
-		// maintain a hard reference to it to ensure that it isn't discarded on save.
-		check(Container->GetOuter()->IsA<UCustomizableObject>());
-	}
 }
 
-void FCustomizableObjectStreamedResourceData::NotifyLoaded(const UCustomizableObjectResourceDataContainer* LoadedContainer)
+void FCustomizableObjectStreamedResourceData::Hold()
 {
 	check(IsInGameThread());
-	check(ContainerPath.Get() == LoadedContainer);
 
-	Container = LoadedContainer;
+	Container = GetPath().Get();
 }
diff --git a/Engine/Plugins/Mutable/Source/CustomizableObject/Private/MuCO/CustomizableObjectSystem.cpp b/Engine/Plugins/Mutable/Source/CustomizableObject/Private/MuCO/CustomizableObjectSystem.cpp
--- a/Engine/Plugins/Mutable/Source/CustomizableObject/Private/MuCO/CustomizableObjectSystem.cpp
+++ b/Engine/Plugins/Mutable/Source/CustomizableObject/Private/MuCO/CustomizableObjectSystem.cpp
@@ -18,7 +18,6 @@
 #include "MuCO/CustomizableObjectInstanceUsage.h"
 #include "MuCO/ICustomizableObjectModule.h"
 #include "MuCO/LogBenchmarkUtil.h"
-#include "MuCO/UnrealExtensionDataStreamer.h"
 #include "MuCO/UnrealMutableImageProvider.h"
 #include "MuCO/UnrealMutableModelDiskStreamer.h"
 #include "MuCO/UnrealPortabilityHelpers.h"
@@ -581,8 +580,7 @@ void UCustomizableObjectSystem::InitSystem()
 	mu::FSettings Settings;
 	Settings.SetProfile(false);
 	Settings.SetWorkingMemoryBytes(Private->LastWorkingMemoryBytes);
-	Private->ExtensionDataStreamer = MakeShared<FUnrealExtensionDataStreamer>(Private);
-	Private->MutableSystem = MakeShared<mu::FSystem>(Settings, Private->ExtensionDataStreamer);
+	Private->MutableSystem = MakeShared<mu::FSystem>(Settings);
 	check(Private->MutableSystem);
 
 	Private->Streamer = MakeShared<FUnrealMutableModelBulkReader>();
@@ -926,6 +924,19 @@ void FinishUpdateGlobal(const TSharedRef<FUpdateContextPrivate>& Context)
 	Context->UpdateCallback.ExecuteIfBound(ContextPublic);
 	Context->UpdateNativeCallback.Broadcast(ContextPublic);
 
+	if (Object)
+	{
+		for (const int32 ResourceIndex : Context->StreamedResourceIndex)
+		{
+			Object->GetPrivate()->GetModelResources().StreamedResourceData[ResourceIndex].Release();
+		}
+
+		for (const int32 ResourceIndex : Context->ExtensionStreamedResourceIndex)
+		{
+			Object->GetPrivate()->GetModelResources().StreamedExtensionData[ResourceIndex].Release();
+		}
+	}
+	
 	if (CVarFixLowPriorityTasksOverlap.GetValueOnGameThread())
 	{
 		if (SystemPrivate && Context->bLowPriorityTasksBlocked)
@@ -2156,21 +2167,6 @@ namespace impl
 		{
 			OperationData->RelevantParametersInProgress.Reset();
 		}
-
-		// Copy FExtensionData Object node input from the Instance to the InstanceUpdateData
-		for (int32 ExtensionDataIndex = 0; ExtensionDataIndex < OperationData->MutableInstance->GetExtensionDataCount(); ExtensionDataIndex++)
-		{
-			TSharedPtr<const mu::FExtensionData> ExtensionData;
-			FName Name;
-			OperationData->MutableInstance->GetExtensionData(ExtensionDataIndex, ExtensionData, Name);
-
-			check(ExtensionData);
-
-			FInstanceUpdateData::FNamedExtensionData& NewEntry = OperationData->InstanceUpdateData.ExtendedInputPins.AddDefaulted_GetRef();
-			NewEntry.Data = ExtensionData;
-			NewEntry.Name = Name;
-			check(NewEntry.Name != NAME_None);
-		}
 		
 		OperationData->TaskGetMeshTime = FPlatformTime::Seconds() - StartTime;
 
@@ -3423,10 +3419,7 @@ namespace impl
 		UCustomizableObjectSystemPrivate* SystemPrivateData = System->GetPrivate();
 
 		SystemPrivateData->CurrentInstanceBeingUpdated = CandidateInstance;
-
-		check(SystemPrivateData->ExtensionDataStreamer != nullptr);
-		SystemPrivateData->ExtensionDataStreamer->SetActiveObject(CustomizableObject);
-
+		
 		FString StateName = CandidateInstance->GetCustomizableObject()->GetPrivate()->GetStateName(CandidateInstance->GetPrivate()->GetState());
 		const FMutableStateData* StateData = CandidateInstance->GetCustomizableObject()->GetPrivate()->GetModelResources().StateUIDataMap.Find(StateName);
 
diff --git a/Engine/Plugins/Mutable/Source/CustomizableObject/Private/MuCO/UnrealExtensionDataStreamer.cpp b/Engine/Plugins/Mutable/Source/CustomizableObject/Private/MuCO/UnrealExtensionDataStreamer.cpp
deleted file mode 100644
diff --git a/Engine/Plugins/Mutable/Source/CustomizableObject/Public/MuCO/CustomizableObjectResourceData.h b/Engine/Plugins/Mutable/Source/CustomizableObject/Public/MuCO/CustomizableObjectResourceData.h
--- a/Engine/Plugins/Mutable/Source/CustomizableObject/Public/MuCO/CustomizableObjectResourceData.h
+++ b/Engine/Plugins/Mutable/Source/CustomizableObject/Public/MuCO/CustomizableObjectResourceData.h
@@ -21,8 +21,6 @@ struct FCustomizableObjectResourceData
 {
 	GENERATED_BODY()
 
-public:
-
 	UPROPERTY()
 	ECOResourceDataType Type = ECOResourceDataType::None;
 
diff --git a/Engine/Plugins/Mutable/Source/CustomizableObject/Public/MuCO/CustomizableObjectStreamedResourceData.h b/Engine/Plugins/Mutable/Source/CustomizableObject/Public/MuCO/CustomizableObjectStreamedResourceData.h
deleted file mode 100644
diff --git a/Engine/Plugins/Mutable/Source/CustomizableObject/Public/MuCO/UnrealExtensionDataStreamer.h b/Engine/Plugins/Mutable/Source/CustomizableObject/Public/MuCO/UnrealExtensionDataStreamer.h
deleted file mode 100644
diff --git a/Engine/Plugins/Mutable/Source/MutableRuntime/Internal/MuR/ModelPrivate.h b/Engine/Plugins/Mutable/Source/MutableRuntime/Internal/MuR/ModelPrivate.h
--- a/Engine/Plugins/Mutable/Source/MutableRuntime/Internal/MuR/ModelPrivate.h
+++ b/Engine/Plugins/Mutable/Source/MutableRuntime/Internal/MuR/ModelPrivate.h
@@ -8,7 +8,6 @@
 #include "MuR/SerialisationPrivate.h"
 #include "MuR/Operations.h"
 #include "MuR/ExtensionData.h"
-#include "MuR/ExtensionDataStreamer.h"
 #include "MuR/Image.h"
 #include "MuR/Mesh.h"
 #include "MuR/ParametersPrivate.h"
@@ -110,31 +109,17 @@ namespace mu
 			AlwaysLoaded
 		};
 
-		// This should be initialized to a valid load state when Data is set
-		ELoadState LoadState = ELoadState::Invalid;
-
-		inline void Serialise(FOutputArchive& Arch) const
+		void Serialise(FOutputArchive& Arch) const
 		{
 			Arch << Data;
 		}
 		
-		inline void Unserialise(FInputArchive& Arch)
+		void Unserialise(FInputArchive& Arch)
 		{
 			Arch >> Data;
 
 			check(Data);
-			check(Data->Origin == FExtensionData::EOrigin::ConstantAlwaysLoaded
-				|| Data->Origin == FExtensionData::EOrigin::ConstantStreamed);
-
-			if (Data->Origin == FExtensionData::EOrigin::ConstantAlwaysLoaded)
-			{
-				LoadState = ELoadState::AlwaysLoaded;
-			}
-			else
-			{
-				// Streamed constants are assumed to be unloaded to start with
-				LoadState = ELoadState::Unloaded;
-			}
+			check(Data->Origin == FExtensionData::EOrigin::ConstantAlwaysLoaded || Data->Origin == FExtensionData::EOrigin::ConstantStreamed);
 		}
 	};
 
@@ -474,12 +459,6 @@ namespace mu
 			FExtensionDataConstant& NewConstant = ConstantExtensionData.AddDefaulted_GetRef();
 			NewConstant.Data = Data;
 
-			// The data is assumed to be loaded during compilation
-			NewConstant.LoadState =
-				Data->Origin == FExtensionData::EOrigin::ConstantAlwaysLoaded
-				? FExtensionDataConstant::ELoadState::AlwaysLoaded
-				: FExtensionDataConstant::ELoadState::CurrentlyLoaded;
-
 			return ConstantExtensionData.Num() - 1;
 		}
 
@@ -749,7 +728,6 @@ namespace mu
 		{
 			const FExtensionDataConstant& Constant = ConstantExtensionData[ConstantIndex];
 
-			check(Constant.LoadState != FExtensionDataConstant::ELoadState::Unloaded);
 			check(Constant.Data);
 
 			Result = Constant.Data;
diff --git a/Engine/Plugins/Mutable/Source/MutableRuntime/Internal/MuR/SystemPrivate.h b/Engine/Plugins/Mutable/Source/MutableRuntime/Internal/MuR/SystemPrivate.h
--- a/Engine/Plugins/Mutable/Source/MutableRuntime/Internal/MuR/SystemPrivate.h
+++ b/Engine/Plugins/Mutable/Source/MutableRuntime/Internal/MuR/SystemPrivate.h
@@ -1933,7 +1933,7 @@ namespace mu
 		friend class System;
     public:
 
-        Private(const FSettings&, const TSharedPtr<FExtensionDataStreamer>& );
+        Private(const FSettings&);
         ~Private();
 
         //-----------------------------------------------------------------------------------------
@@ -1954,9 +1954,7 @@ namespace mu
 		MUTABLERUNTIME_API TSharedPtr<const FInstance> BuildInstance(const TSharedPtr<const FModel>&, const FParameters*, OP::ADDRESS);
 		MUTABLERUNTIME_API TSharedPtr<const FLayout> BuildLayout(const TSharedPtr<const FModel>&, const FParameters*, OP::ADDRESS) ;
     	MUTABLERUNTIME_API FProjector BuildProjector(const TSharedPtr<const FModel>&, const FParameters*, OP::ADDRESS) ;
-
-		FExtensionDataStreamer* GetExtensionDataStreamer() const { return ExtensionDataStreamer.Get(); }
-
+    	
         //!
         FSettings Settings;
 
@@ -1993,11 +1991,6 @@ namespace mu
 
 		//! Update some mutable core unreal stats.
 		void UpdateStats();
-
-	private:
-
-		/** Owned by this system. */
-		TSharedPtr<FExtensionDataStreamer> ExtensionDataStreamer;
     };
 
 }
diff --git a/Engine/Plugins/Mutable/Source/MutableRuntime/Private/MuR/CodeRunner.h b/Engine/Plugins/Mutable/Source/MutableRuntime/Private/MuR/CodeRunner.h
--- a/Engine/Plugins/Mutable/Source/MutableRuntime/Private/MuR/CodeRunner.h
+++ b/Engine/Plugins/Mutable/Source/MutableRuntime/Private/MuR/CodeRunner.h
@@ -6,7 +6,6 @@
 #include "Containers/ContainerAllocationPolicies.h"
 #include "HAL/PlatformCrt.h"
 #include "HAL/PlatformMath.h"
-#include "MuR/ExtensionDataStreamer.h"
 #include "MuR/Image.h"
 #include "MuR/Layout.h"
 #include "MuR/Operations.h"
@@ -275,31 +274,6 @@ namespace  mu
 			uint32 RomIndex = 0;
 		};
 
-		class FLoadExtensionDataTask : public CodeRunner::FIssuedTask
-		{
-		public:
-			FLoadExtensionDataTask(const FScheduledOp& InOp, const TSharedPtr<const FExtensionData>& InExtensionData, int32 InModelConstantIndex)
-				: FIssuedTask(InOp)
-				, Data(InExtensionData)
-				, ModelConstantIndex(InModelConstantIndex)
-			{
-				check(Data.Get());
-				check(Data->Origin == FExtensionData::EOrigin::ConstantStreamed);
-			}
-
-			// FIssuedTask interface
-			virtual bool Prepare(CodeRunner*, bool& bOutFailed) override;
-			virtual bool Complete(CodeRunner*) override;
-			virtual bool IsComplete(CodeRunner*) override;
-
-		private:
-			TSharedPtr<const FExtensionData> Data;
-			// The index of this constant in FProgram::ConstantExtensionData
-			int32 ModelConstantIndex;
-
-			TSharedPtr<const FExtensionDataLoadHandle> LoadHandle;
-		};
-
 		class FLoadImageRomsTask : public CodeRunner::FIssuedTask
 		{
 		public:
diff --git a/Engine/Plugins/Mutable/Source/MutableRuntime/Private/MuR/CodeRunnerTasked.cpp b/Engine/Plugins/Mutable/Source/MutableRuntime/Private/MuR/CodeRunnerTasked.cpp
--- a/Engine/Plugins/Mutable/Source/MutableRuntime/Private/MuR/CodeRunnerTasked.cpp
+++ b/Engine/Plugins/Mutable/Source/MutableRuntime/Private/MuR/CodeRunnerTasked.cpp
@@ -13,7 +13,6 @@
 #include "Misc/AssertionMacros.h"
 #include "MuR/CodeRunner.h"
 #include "MuR/System.h"
-#include "MuR/ExtensionDataStreamer.h"
 #include "MuR/Image.h"
 #include "MuR/ImagePrivate.h"
 #include "MuR/Layout.h"
@@ -2246,94 +2245,7 @@ namespace mu
 		return bSuccess;
 	}
 
-
-	//---------------------------------------------------------------------------------------------
-	//---------------------------------------------------------------------------------------------
-	//---------------------------------------------------------------------------------------------
-	bool CodeRunner::FLoadExtensionDataTask::Prepare(CodeRunner* Runner, bool& bOutFailed)
-	{
-		// This runs in the mutable Runner thread
-		MUTABLE_CPUPROFILER_SCOPE(FLoadExtensionDataTask_Prepare);
-		bOutFailed = false;
-
-		FProgram& Program = Runner->Model->GetPrivate()->Program;
-		check(Program.ConstantExtensionData.IsValidIndex(ModelConstantIndex));
-
-		const FExtensionDataConstant& Constant = Program.ConstantExtensionData[ModelConstantIndex];
-		check(Constant.Data == Data);
-
-		if (Constant.LoadState == FExtensionDataConstant::ELoadState::FailedToLoad)
-		{
-			// This already failed before, so don't waste time trying again.
-			//
-			// This behavior could be made an option in future if needed.
-			bOutFailed = true;
-			return false;
-		}
-
-		TArray<TSharedPtr<const FExtensionData>> UnloadedConstants;
-		LoadHandle = Runner->System->GetExtensionDataStreamer()->StartLoad(Data, UnloadedConstants);
-		check(LoadHandle.IsValid());
-
-		// Update the load state of any constants that were unloaded
-		for (const TSharedPtr<const FExtensionData>& Unloaded : UnloadedConstants)
-		{
-			check(Unloaded->Origin == FExtensionData::EOrigin::ConstantStreamed);
-
-			FExtensionDataConstant* FoundConstant = Algo::FindBy(Program.ConstantExtensionData, Unloaded, &FExtensionDataConstant::Data);
-			check(FoundConstant);
-
-			FoundConstant->LoadState = FExtensionDataConstant::ELoadState::Unloaded;
-		}
-
-		// No worker thread work
-		return false;
-	}
-
-
-	//---------------------------------------------------------------------------------------------
-	bool CodeRunner::FLoadExtensionDataTask::Complete(CodeRunner* Runner)
-	{
-		MUTABLE_CPUPROFILER_SCOPE(FLoadExtensionDataTask_Complete);
-
-		bool bSuccess = true;
-
-		// Complete should only be called if the load is finished
-		check(LoadHandle->LoadState != FExtensionDataLoadHandle::ELoadState::Pending);
-
-		FProgram& Program = Runner->Model->GetPrivate()->Program;
-		check(Program.ConstantExtensionData.IsValidIndex(ModelConstantIndex));
-
-		FExtensionDataConstant& Constant = Program.ConstantExtensionData[ModelConstantIndex];
-		check(Constant.Data == Data);
-
-		// Update the load state in the Program
-		if (LoadHandle->LoadState == FExtensionDataLoadHandle::ELoadState::Loaded)
-		{
-			Constant.LoadState = FExtensionDataConstant::ELoadState::CurrentlyLoaded;
-		}
-		else
-		{
-			check(LoadHandle->LoadState == FExtensionDataLoadHandle::ELoadState::FailedToLoad);
-			Constant.LoadState = FExtensionDataConstant::ELoadState::FailedToLoad;
-			bSuccess = false;
-		}
-
-		// Process the constant op normally, now that the data is loaded.
-		Runner->RunCode(Op, Runner->Params, Runner->Model, Runner->LODMask);
-
-		return bSuccess;
-	}
-
-
-	//---------------------------------------------------------------------------------------------
-	bool CodeRunner::FLoadExtensionDataTask::IsComplete(CodeRunner* Runner)
-	{
-		check(LoadHandle.IsValid());
-		return LoadHandle->LoadState != FExtensionDataLoadHandle::ELoadState::Pending;
-	}
-
-
+	
 	//---------------------------------------------------------------------------------------------
 	//---------------------------------------------------------------------------------------------
 	//---------------------------------------------------------------------------------------------
@@ -2677,26 +2589,6 @@ namespace mu
 			}
 			break;
 		}
-		
-		case EOpType::ED_CONSTANT:
-		{
-			OP::ResourceConstantArgs Args = Program.GetOpArgs<OP::ResourceConstantArgs>(item.At);
-
-			const FExtensionDataConstant::ELoadState LoadState = Program.ConstantExtensionData[Args.value].LoadState;
-
-			check(LoadState != FExtensionDataConstant::ELoadState::Invalid);
-			
-			if (LoadState == FExtensionDataConstant::ELoadState::AlwaysLoaded
-				|| LoadState == FExtensionDataConstant::ELoadState::CurrentlyLoaded)
-			{
-				// Already loaded. The rest of the constant code will run right away.
-			}
-			else
-			{
-				Issued = MakeShared<FLoadExtensionDataTask>(item, Program.ConstantExtensionData[Args.value].Data, Args.value);
-			}
-			break;
-		}
 
 		case EOpType::IM_CONSTANT:
 		{
diff --git a/Engine/Plugins/Mutable/Source/MutableRuntime/Private/MuR/NullExtensionDataStreamer.cpp b/Engine/Plugins/Mutable/Source/MutableRuntime/Private/MuR/NullExtensionDataStreamer.cpp
deleted file mode 100644
diff --git a/Engine/Plugins/Mutable/Source/MutableRuntime/Private/MuR/NullExtensionDataStreamer.h b/Engine/Plugins/Mutable/Source/MutableRuntime/Private/MuR/NullExtensionDataStreamer.h
deleted file mode 100644
diff --git a/Engine/Plugins/Mutable/Source/MutableRuntime/Private/MuR/System.cpp b/Engine/Plugins/Mutable/Source/MutableRuntime/Private/MuR/System.cpp
--- a/Engine/Plugins/Mutable/Source/MutableRuntime/Private/MuR/System.cpp
+++ b/Engine/Plugins/Mutable/Source/MutableRuntime/Private/MuR/System.cpp
@@ -22,7 +22,6 @@
 #include "MuR/MutableMath.h"
 #include "MuR/MutableString.h"
 #include "MuR/MutableTrace.h"
-#include "MuR/NullExtensionDataStreamer.h"
 #include "MuR/Operations.h"
 #include "MuR/Parameters.h"
 #include "MuR/ParametersPrivate.h"
@@ -73,12 +72,12 @@ namespace mu
 
 
     //---------------------------------------------------------------------------------------------
-    FSystem::FSystem(const FSettings& InSettings, const TSharedPtr<FExtensionDataStreamer>& DataStreamer)
+    FSystem::FSystem(const FSettings& InSettings)
     {
 		LLM_SCOPE_BYNAME(TEXT("MutableRuntime"));
 
         // Choose the implementation
-        m_pD = new FSystem::Private( InSettings, DataStreamer );
+        m_pD = new FSystem::Private(InSettings);
     }
 
 
@@ -103,7 +102,7 @@ namespace mu
 
 
     //---------------------------------------------------------------------------------------------
-    FSystem::Private::Private( const FSettings& InSettings, const TSharedPtr<mu::FExtensionDataStreamer>& InDataStreamer)
+    FSystem::Private::Private(const FSettings& InSettings)
     {
 		LLM_SCOPE_BYNAME(TEXT("MutableRuntime"));
 
@@ -112,12 +111,6 @@ namespace mu
 		WorkingMemoryManager.BudgetBytes = Settings.WorkingMemoryBytes;
 		WorkingMemoryManager.GeneratedResources.Reserve(WorkingMemoryManager.MaxGeneratedResourceCacheSize);
 
-		ExtensionDataStreamer = InDataStreamer;
-		if (!ExtensionDataStreamer)
-		{
-			ExtensionDataStreamer = MakeShared<NullExtensionDataStreamer>();
-		}
-
 		UpdateStats();
 	}
 
@@ -130,7 +123,6 @@ namespace mu
 
 		// Make it explicit to try to capture metrics
         StreamInterface = nullptr;
-		ExtensionDataStreamer = nullptr;
 		ExternalResourceProvider = nullptr;
     }
 
diff --git a/Engine/Plugins/Mutable/Source/MutableRuntime/Public/MuR/ExtensionDataStreamer.h b/Engine/Plugins/Mutable/Source/MutableRuntime/Public/MuR/ExtensionDataStreamer.h
deleted file mode 100644
diff --git a/Engine/Plugins/Mutable/Source/MutableRuntime/Public/MuR/System.h b/Engine/Plugins/Mutable/Source/MutableRuntime/Public/MuR/System.h
--- a/Engine/Plugins/Mutable/Source/MutableRuntime/Public/MuR/System.h
+++ b/Engine/Plugins/Mutable/Source/MutableRuntime/Public/MuR/System.h
@@ -100,8 +100,7 @@ namespace mu
 		//! Constructor of a system object to build data.
         //! \param Settings Optional class with the settings to use in this system. The default
         //! value configures a production-ready system.
-		//! \param DataStreamer Optional interface to allow the Model to stream in ExtensionData from disk
-        FSystem( const FSettings& Settings = FSettings(), const TSharedPtr<FExtensionDataStreamer>& DataStreamer = nullptr);
+		FSystem( const FSettings& Settings = FSettings());
 
         //! Set a new provider for model data. 
         void SetStreamingInterface(const TSharedPtr<FModelReader>& );
