[Backout] - CL39148332
#rnx
Original CL Desc
-----------------------------------------------------------------
[Mutable] Avoid BeginUpdate to sync with the Game Thread due to Extension Data.

#rnx

diff --git a/Engine/Plugins/Mutable/Source/CustomizableObject/Internal/MuCO/CustomizableObjectInstancePrivate.h b/Engine/Plugins/Mutable/Source/CustomizableObject/Internal/MuCO/CustomizableObjectInstancePrivate.h
--- a/Engine/Plugins/Mutable/Source/CustomizableObject/Internal/MuCO/CustomizableObjectInstancePrivate.h
+++ b/Engine/Plugins/Mutable/Source/CustomizableObject/Internal/MuCO/CustomizableObjectInstancePrivate.h
@@ -263,7 +263,7 @@ public:
 	 * If no StreamableManager is provided, it will load assets synchronously. */
 	UE::Tasks::FTask LoadAdditionalAssetsAndData(const TSharedRef<FUpdateContextPrivate>& OperationData);
 
-	void AdditionalAssetsAsyncLoaded(const TSharedRef<FUpdateContextPrivate> OperationData);
+	void AdditionalAssetsAsyncLoaded();
 
 	void TickUpdateCloseCustomizableObjects(UCustomizableObjectInstance& Public, FMutableInstanceUpdateMap& InOutRequestedUpdates);
 	void UpdateInstanceIfNotGenerated(UCustomizableObjectInstance& Public, FMutableInstanceUpdateMap& InOutRequestedUpdates);
diff --git a/Engine/Plugins/Mutable/Source/CustomizableObject/Internal/MuCO/CustomizableObjectPrivate.h b/Engine/Plugins/Mutable/Source/CustomizableObject/Internal/MuCO/CustomizableObjectPrivate.h
--- a/Engine/Plugins/Mutable/Source/CustomizableObject/Internal/MuCO/CustomizableObjectPrivate.h
+++ b/Engine/Plugins/Mutable/Source/CustomizableObject/Internal/MuCO/CustomizableObjectPrivate.h
@@ -1579,6 +1579,9 @@ public:
 #if WITH_EDITORONLY_DATA
 	TArray<FProfileParameterDat>& GetInstancePropertiesProfiles();
 #endif
+	
+	const FCustomizableObjectResourceData* GetLoadedStreamedResource(int32 ResourceIndex);
+	void UnloadStreamedResource(int32 ResourceIndex);
 
 #if WITH_EDITORONLY_DATA
 	TObjectPtr<UEdGraph>& GetSource() const;
diff --git a/Engine/Plugins/Mutable/Source/CustomizableObject/Internal/MuCO/CustomizableObjectStreamedResourceData.h b/Engine/Plugins/Mutable/Source/CustomizableObject/Internal/MuCO/CustomizableObjectStreamedResourceData.h
deleted file mode 100644
diff --git a/Engine/Plugins/Mutable/Source/CustomizableObject/Internal/MuCO/CustomizableObjectSystemPrivate.h b/Engine/Plugins/Mutable/Source/CustomizableObject/Internal/MuCO/CustomizableObjectSystemPrivate.h
--- a/Engine/Plugins/Mutable/Source/CustomizableObject/Internal/MuCO/CustomizableObjectSystemPrivate.h
+++ b/Engine/Plugins/Mutable/Source/CustomizableObject/Internal/MuCO/CustomizableObjectSystemPrivate.h
@@ -625,12 +625,6 @@ public:
 
 	/** List of component names. Index is the ObjectComponentIndex. */
 	TArray<FName> ComponentNames;
-
-	/** Index of the resource in the StreamedResourceData array of the Model Resources. */
-	TArray<int32> StreamedResourceIndex;
-
-	/** Index of the resource in the ExtensionStreamedResourceData array of the Model Resources. */
-	TArray<int32> ExtensionStreamedResourceIndex;
 	
 	TMap<FName, uint8> NumLODsAvailable;
 
@@ -730,7 +724,9 @@ public:
 
 	// This object is responsible for streaming data to the MutableSystem.
 	TSharedPtr<class FUnrealMutableModelBulkReader> Streamer;
-	
+
+	TSharedPtr<class FUnrealExtensionDataStreamer> ExtensionDataStreamer;
+
 	// This object is responsible for providing custom images and meshes to mutable (for image parameters, etc.)
 	// This object is called from the mutable thread, and it should only access data already safely submitted from
 	// the game thread and stored in FUnrealMutableImageProvider::GlobalExternalImages.
diff --git a/Engine/Plugins/Mutable/Source/CustomizableObject/Private/MuCO/CustomizableObject.cpp b/Engine/Plugins/Mutable/Source/CustomizableObject/Private/MuCO/CustomizableObject.cpp
--- a/Engine/Plugins/Mutable/Source/CustomizableObject/Private/MuCO/CustomizableObject.cpp
+++ b/Engine/Plugins/Mutable/Source/CustomizableObject/Private/MuCO/CustomizableObject.cpp
@@ -2601,6 +2601,48 @@ TArray<FProfileParameterDat>& UCustomizableObjectPrivate::GetInstancePropertiesP
 }
 #endif
 
+const FCustomizableObjectResourceData* UCustomizableObjectPrivate::GetLoadedStreamedResource(int32 ResourceIndex)
+{
+#if WITH_EDITORONLY_DATA
+	FModelResources& LocalModelResources = ModelResourcesEditor;
+#else
+	FModelResources& LocalModelResources = ModelResources;
+#endif
+
+	if (LocalModelResources.StreamedResourceData.IsValidIndex(ResourceIndex))
+	{
+		FCustomizableObjectStreamedResourceData& Resource = LocalModelResources.StreamedResourceData[ResourceIndex];
+		if (!Resource.IsLoaded())
+		{
+			Resource.NotifyLoaded(Resource.GetPath().Get());
+		}
+
+		return &Resource.GetLoadedData();
+	}
+
+	return nullptr;
+}
+
+void UCustomizableObjectPrivate::UnloadStreamedResource(int32 ResourceIndex)
+{
+	// Only Unload in cooked builds. Unloading them when in the editor will trigger an assert. 
+	if (FPlatformProperties::RequiresCookedData())
+	{
+		return;
+	}
+
+#if WITH_EDITORONLY_DATA
+	FModelResources& LocalModelResources = ModelResourcesEditor;
+#else
+	FModelResources& LocalModelResources = ModelResources;
+#endif
+
+	if (LocalModelResources.StreamedResourceData.IsValidIndex(ResourceIndex))
+	{
+		LocalModelResources.StreamedResourceData[ResourceIndex].Unload();
+	}
+}
+
 
 #if WITH_EDITORONLY_DATA
 TObjectPtr<UEdGraph>& UCustomizableObjectPrivate::GetSource() const
@@ -4039,6 +4081,7 @@ void FMutableRefSkeletalMeshData::InitResources(UCustomizableObject* InOuter, FM
 
 		Data.AssetUserData = InModelResources.StreamedResourceData[Data.AssetUserDataIndex].GetPath().LoadSynchronous();
 		check(Data.AssetUserData);
+		check(Data.AssetUserData->Data.Type == ECOResourceDataType::AssetUserData);
 	}
 }
 
diff --git a/Engine/Plugins/Mutable/Source/CustomizableObject/Private/MuCO/CustomizableObjectInstance.cpp b/Engine/Plugins/Mutable/Source/CustomizableObject/Private/MuCO/CustomizableObjectInstance.cpp
--- a/Engine/Plugins/Mutable/Source/CustomizableObject/Private/MuCO/CustomizableObjectInstance.cpp
+++ b/Engine/Plugins/Mutable/Source/CustomizableObject/Private/MuCO/CustomizableObjectInstance.cpp
@@ -53,8 +53,6 @@
 #include "PhysicsEngine/SkeletalBodySetup.h"
 #include "Hash/CityHash.h"
 #include "MuCO/CustomizableObjectCustomVersion.h"
-#include "MuCO/CustomizableObjectResourceData.h"
-#include "MuCO/CustomizableObjectStreamedResourceData.h"
 #include "MuCO/CustomizableObjectResourceDataTypes.h"
 #include "Serialization/BulkData.h"
 
@@ -3610,8 +3608,10 @@ void UCustomizableInstancePrivate::InitSkeletalMeshData(const TSharedRef<FUpdate
 	{
 		for (const FMutableRefAssetUserData& MutAssetUserData : RefSkeletalMeshData.AssetUserData)
 		{
-			if (const FCustomizableObjectAssetUserData* DataPtr = MutAssetUserData.AssetUserData->Data.Data.GetPtr<FCustomizableObjectAssetUserData>())
+			if (MutAssetUserData.AssetUserData && MutAssetUserData.AssetUserData->Data.Type == ECOResourceDataType::AssetUserData)
 			{
+				const FCustomizableObjectAssetUserData* DataPtr = MutAssetUserData.AssetUserData->Data.Data.GetPtr<FCustomizableObjectAssetUserData>();
+				check(DataPtr);
 #if WITH_EDITORONLY_DATA
 				SkeletalMesh->AddAssetUserData(DataPtr->AssetUserDataEditor);
 #else
@@ -5472,7 +5472,6 @@ UE::Tasks::FTask UCustomizableInstancePrivate::LoadAdditionalAssetsAndData(const
 						const FCustomizableObjectStreamedResourceData& StreamedResource = ModelResources.StreamedResourceData[ResourceIndex];
 						if (!StreamedResource.IsLoaded())
 						{
-							OperationData->StreamedResourceIndex.Add(ResourceIndex);
 							AssetsToStream.AddUnique(StreamedResource.GetPath().ToSoftObjectPath());
 						}
 
@@ -5651,43 +5650,6 @@ UE::Tasks::FTask UCustomizableInstancePrivate::LoadAdditionalAssetsAndData(const
 		AssetsToStream.Add(MeshRef.ToSoftObjectPath());
 	}
 
-	// Copy FExtensionData Object node input from the Instance to the InstanceUpdateData
-	for (int32 ExtensionDataIndex = 0; ExtensionDataIndex < OperationData->MutableInstance->GetExtensionDataCount(); ExtensionDataIndex++)
-	{
-		TSharedPtr<const mu::FExtensionData> ExtensionData;
-		FName Name;
-		OperationData->MutableInstance->GetExtensionData(ExtensionDataIndex, ExtensionData, Name);
-
-		check(ExtensionData);
-
-		FInstanceUpdateData::FNamedExtensionData& NewEntry = OperationData->InstanceUpdateData.ExtendedInputPins.AddDefaulted_GetRef();
-		NewEntry.Data = ExtensionData;
-		NewEntry.Name = Name;
-		check(NewEntry.Name != NAME_None);
-
-		if (!ModelResources.StreamedExtensionData.IsValidIndex(ExtensionData->Index))
-		{
-			// The compiled data appears to be out of sync with the CO's properties
-
-			UE_LOG(LogMutable, Error, TEXT("Couldn't find streamed Extension Data with index %d in %s. Compiled data may be stale."),
-				ExtensionData->Index, *CustomizableObject->GetFullName());
-
-			continue;
-		}
-
-		const FCustomizableObjectStreamedResourceData& StreamedData = ModelResources.StreamedExtensionData[ExtensionData->Index];
-		if (StreamedData.IsLoaded())
-		{
-			continue;
-		}
-
-		// Note that this just checks if the path is non-null, NOT if the object is loaded
-		check(!StreamedData.GetPath().IsNull());
-
-		OperationData->ExtensionStreamedResourceIndex.Add(ExtensionData->Index);
-		AssetsToStream.Add(StreamedData.GetPath().ToSoftObjectPath());
-	}
-	
 	TArray<UE::Tasks::FTaskEvent> StreamingCompletionEvents;
 	if (AssetsToStream.Num() > 0)
 	{
@@ -5705,7 +5667,7 @@ UE::Tasks::FTask UCustomizableInstancePrivate::LoadAdditionalAssetsAndData(const
 
 		PrivateSystem->StreamableManager->RequestAsyncLoad(
 			AssetsToStream,
-			FStreamableDelegate::CreateUObject(this, &UCustomizableInstancePrivate::AdditionalAssetsAsyncLoaded, OperationData),
+			FStreamableDelegate::CreateUObject(this, &UCustomizableInstancePrivate::AdditionalAssetsAsyncLoaded),
 			bEnableHighPriorityLoading ? FStreamableManager::AsyncLoadHighPriority : FStreamableManager::DefaultAsyncLoadPriority);
 	}
 	
@@ -6041,7 +6003,7 @@ const TArray<TObjectPtr<UMaterialInterface>>* UCustomizableObjectInstance::GetOv
 }
 
 
-void UCustomizableInstancePrivate::AdditionalAssetsAsyncLoaded(const TSharedRef<FUpdateContextPrivate> OperationData)
+void UCustomizableInstancePrivate::AdditionalAssetsAsyncLoaded()
 {
 	MUTABLE_CPUPROFILER_SCOPE(UCustomizableInstancePrivate::AdditionalAssetsAsyncLoaded);
 
@@ -6051,18 +6013,8 @@ void UCustomizableInstancePrivate::AdditionalAssetsAsyncLoaded(const TSharedRef<
 
 	UCustomizableObjectPrivate* CustomizableObjectPrivate = GetPublic()->GetCustomizableObject()->GetPrivate();
 
-	FModelResources& ModelResources = CustomizableObjectPrivate->GetModelResources();
+	const FModelResources& ModelResources = CustomizableObjectPrivate->GetModelResources();
 
-	for (int32 ResourceIndex : OperationData->StreamedResourceIndex)
-	{
-		ModelResources.StreamedResourceData[ResourceIndex].Hold();
-	}
-	
-	for (int32 ResourceIndex : OperationData->ExtensionStreamedResourceIndex)
-	{
-		ModelResources.StreamedExtensionData[ResourceIndex].Hold();
-	}
-	
 	// Loaded Materials
 	check(ObjectToInstanceIndexMap.Num() == ReferencedMaterials.Num());
 
@@ -6094,16 +6046,31 @@ void UCustomizableInstancePrivate::AdditionalAssetsAsyncLoaded(const TSharedRef<
 	{
 		for (int32 ResourceIndex : ComponentData.StreamedResourceIndex)
 		{
-			const FCustomizableObjectResourceData& ResourceData = ModelResources.StreamedExtensionData[ResourceIndex].GetLoadedData();
+			const FCustomizableObjectResourceData* ResourceData = CustomizableObjectPrivate->GetLoadedStreamedResource(ResourceIndex);
+			if (!ResourceData)
+			{
+				check(false); // Invalid resource index?
+				continue;
+			}
 
-			if (const FCustomizableObjectAssetUserData* AUDResource = ResourceData.Data.GetPtr<FCustomizableObjectAssetUserData>())
+			switch (ResourceData->Type)
 			{
+				case ECOResourceDataType::AssetUserData:
+				{
+					const FCustomizableObjectAssetUserData* AUDResource = ResourceData->Data.GetPtr<FCustomizableObjectAssetUserData>();
 #if WITH_EDITORONLY_DATA
-				ComponentData.AssetUserDataArray.Add(AUDResource->AssetUserDataEditor);
+					ComponentData.AssetUserDataArray.Add(AUDResource->AssetUserDataEditor);
 #else
-				ComponentData.AssetUserDataArray.Add(AUDResource->AssetUserData);
+					ComponentData.AssetUserDataArray.Add(AUDResource->AssetUserData);
 #endif
+					break;
+				}
+				default:
+					break;
 			}
+
+			// Unload by removing the reference to the container. Only if the platform has cooked data.
+			CustomizableObjectPrivate->UnloadStreamedResource(ResourceIndex);
 		}
 
 		// Loaded Skeletons
@@ -6208,7 +6175,7 @@ void UCustomizableInstancePrivate::AdditionalAssetsAsyncLoaded(const TSharedRef<
 	}
 
 	PassThroughMeshesToLoad.Empty();
-	
+
 #if WITH_EDITOR
 	// TODO: Remove with UE-217665 when the underlying bug in the ColorPicker is solved
 	// Reenable the throttling which disabled when launching the Async Load
diff --git a/Engine/Plugins/Mutable/Source/CustomizableObject/Private/MuCO/CustomizableObjectStreamedResourceData.cpp b/Engine/Plugins/Mutable/Source/CustomizableObject/Private/MuCO/CustomizableObjectStreamedResourceData.cpp
--- a/Engine/Plugins/Mutable/Source/CustomizableObject/Private/MuCO/CustomizableObjectStreamedResourceData.cpp
+++ b/Engine/Plugins/Mutable/Source/CustomizableObject/Private/MuCO/CustomizableObjectStreamedResourceData.cpp
@@ -41,24 +41,40 @@ bool FCustomizableObjectStreamedResourceData::IsLoaded() const
 const FCustomizableObjectResourceData& FCustomizableObjectStreamedResourceData::GetLoadedData() const
 {
 	check(IsInGameThread());
+
 	check(Container);
-	
 	return Container->Data;
 }
 
-void FCustomizableObjectStreamedResourceData::Release()
+void FCustomizableObjectStreamedResourceData::Unload()
 {
 	check(IsInGameThread());
 
-	if (FPlatformProperties::RequiresCookedData()) // TODO GMT Remove once UE-232022. Editor will no longer use StreamedResources.
+	if (!Container)
+	{
+		// Already unloaded
+		return;
+	}
+
+	if (FPlatformProperties::RequiresCookedData())
 	{
+		// The container should have been cooked out to its own package
+		check(Container->GetOuter()->IsA<UPackage>());
+
 		Container = nullptr;
 	}
+	else
+	{
+		// In an uncooked build the container is stored in the CO package and the CO needs to
+		// maintain a hard reference to it to ensure that it isn't discarded on save.
+		check(Container->GetOuter()->IsA<UCustomizableObject>());
+	}
 }
 
-void FCustomizableObjectStreamedResourceData::Hold()
+void FCustomizableObjectStreamedResourceData::NotifyLoaded(const UCustomizableObjectResourceDataContainer* LoadedContainer)
 {
 	check(IsInGameThread());
+	check(ContainerPath.Get() == LoadedContainer);
 
-	Container = GetPath().Get();
+	Container = LoadedContainer;
 }
diff --git a/Engine/Plugins/Mutable/Source/CustomizableObject/Private/MuCO/CustomizableObjectSystem.cpp b/Engine/Plugins/Mutable/Source/CustomizableObject/Private/MuCO/CustomizableObjectSystem.cpp
--- a/Engine/Plugins/Mutable/Source/CustomizableObject/Private/MuCO/CustomizableObjectSystem.cpp
+++ b/Engine/Plugins/Mutable/Source/CustomizableObject/Private/MuCO/CustomizableObjectSystem.cpp
@@ -18,6 +18,7 @@
 #include "MuCO/CustomizableObjectInstanceUsage.h"
 #include "MuCO/ICustomizableObjectModule.h"
 #include "MuCO/LogBenchmarkUtil.h"
+#include "MuCO/UnrealExtensionDataStreamer.h"
 #include "MuCO/UnrealMutableImageProvider.h"
 #include "MuCO/UnrealMutableModelDiskStreamer.h"
 #include "MuCO/UnrealPortabilityHelpers.h"
@@ -580,7 +581,8 @@ void UCustomizableObjectSystem::InitSystem()
 	mu::FSettings Settings;
 	Settings.SetProfile(false);
 	Settings.SetWorkingMemoryBytes(Private->LastWorkingMemoryBytes);
-	Private->MutableSystem = MakeShared<mu::FSystem>(Settings);
+	Private->ExtensionDataStreamer = MakeShared<FUnrealExtensionDataStreamer>(Private);
+	Private->MutableSystem = MakeShared<mu::FSystem>(Settings, Private->ExtensionDataStreamer);
 	check(Private->MutableSystem);
 
 	Private->Streamer = MakeShared<FUnrealMutableModelBulkReader>();
@@ -924,19 +926,6 @@ void FinishUpdateGlobal(const TSharedRef<FUpdateContextPrivate>& Context)
 	Context->UpdateCallback.ExecuteIfBound(ContextPublic);
 	Context->UpdateNativeCallback.Broadcast(ContextPublic);
 
-	if (Object)
-	{
-		for (const int32 ResourceIndex : Context->StreamedResourceIndex)
-		{
-			Object->GetPrivate()->GetModelResources().StreamedResourceData[ResourceIndex].Release();
-		}
-
-		for (const int32 ResourceIndex : Context->ExtensionStreamedResourceIndex)
-		{
-			Object->GetPrivate()->GetModelResources().StreamedExtensionData[ResourceIndex].Release();
-		}
-	}
-	
 	if (CVarFixLowPriorityTasksOverlap.GetValueOnGameThread())
 	{
 		if (SystemPrivate && Context->bLowPriorityTasksBlocked)
@@ -2167,6 +2156,21 @@ namespace impl
 		{
 			OperationData->RelevantParametersInProgress.Reset();
 		}
+
+		// Copy FExtensionData Object node input from the Instance to the InstanceUpdateData
+		for (int32 ExtensionDataIndex = 0; ExtensionDataIndex < OperationData->MutableInstance->GetExtensionDataCount(); ExtensionDataIndex++)
+		{
+			TSharedPtr<const mu::FExtensionData> ExtensionData;
+			FName Name;
+			OperationData->MutableInstance->GetExtensionData(ExtensionDataIndex, ExtensionData, Name);
+
+			check(ExtensionData);
+
+			FInstanceUpdateData::FNamedExtensionData& NewEntry = OperationData->InstanceUpdateData.ExtendedInputPins.AddDefaulted_GetRef();
+			NewEntry.Data = ExtensionData;
+			NewEntry.Name = Name;
+			check(NewEntry.Name != NAME_None);
+		}
 		
 		OperationData->TaskGetMeshTime = FPlatformTime::Seconds() - StartTime;
 
@@ -3419,7 +3423,10 @@ namespace impl
 		UCustomizableObjectSystemPrivate* SystemPrivateData = System->GetPrivate();
 
 		SystemPrivateData->CurrentInstanceBeingUpdated = CandidateInstance;
-		
+
+		check(SystemPrivateData->ExtensionDataStreamer != nullptr);
+		SystemPrivateData->ExtensionDataStreamer->SetActiveObject(CustomizableObject);
+
 		FString StateName = CandidateInstance->GetCustomizableObject()->GetPrivate()->GetStateName(CandidateInstance->GetPrivate()->GetState());
 		const FMutableStateData* StateData = CandidateInstance->GetCustomizableObject()->GetPrivate()->GetModelResources().StateUIDataMap.Find(StateName);
 
diff --git a/Engine/Plugins/Mutable/Source/CustomizableObject/Private/MuCO/UnrealExtensionDataStreamer.cpp b/Engine/Plugins/Mutable/Source/CustomizableObject/Private/MuCO/UnrealExtensionDataStreamer.cpp
new file mode 100644
--- /dev/null
+++ b/Engine/Plugins/Mutable/Source/CustomizableObject/Private/MuCO/UnrealExtensionDataStreamer.cpp
@@ -0,0 +1,289 @@
+// Copyright Epic Games, Inc. All Rights Reserved.
+
+#include "MuCO/UnrealExtensionDataStreamer.h"
+
+#include "Algo/Find.h"
+#include "Algo/IndexOf.h"
+#include "Async/Async.h"
+#include "MuCO/CustomizableObject.h"
+#include "MuCO/CustomizableObjectSystem.h"
+#include "MuCO/CustomizableObjectSystemPrivate.h"
+#include "Templates/SharedPointer.h"
+
+FUnrealExtensionDataStreamer::FUnrealExtensionDataStreamer(const TWeakObjectPtr<UCustomizableObjectSystemPrivate>& InSystemPrivateWeak)
+{
+	check(IsInGameThread());
+
+	SystemPrivate = InSystemPrivateWeak;
+	Mutex = new FCriticalSection();
+}
+
+FUnrealExtensionDataStreamer::~FUnrealExtensionDataStreamer()
+{
+	check(IsInGameThread());
+
+	CancelPendingLoads();
+
+	delete Mutex;
+}
+
+void FUnrealExtensionDataStreamer::SetActiveObject(UCustomizableObject* InObject)
+{
+	check(IsInGameThread());
+	check(InObject);
+
+	FScopeLock Lock(Mutex);
+
+	ActiveObject = InObject;
+}
+
+void FUnrealExtensionDataStreamer::ClearActiveObject()
+{
+	check(IsInGameThread());
+
+	FScopeLock Lock(Mutex);
+
+	ActiveObject = nullptr;
+}
+
+bool FUnrealExtensionDataStreamer::AreAnyLoadsPending() const
+{
+	FScopeLock Lock(Mutex);
+
+	return PendingLoads.Num() > 0;
+}
+
+void FUnrealExtensionDataStreamer::CancelPendingLoads()
+{
+	// Needs to be called from the Game thread, as FStreamableHandle::CancelHandle requires this.
+	//
+	// Restricting this to the Game thread also guarantees that any tasks cancelled by this
+	// function will definitely not have started yet, since they only run on the Game thread.
+	check(IsInGameThread());
+
+	FScopeLock Lock(Mutex);
+
+	bool bAnyTasksNeedCancelling = false;
+	for (const FPendingLoad& PendingLoad : PendingLoads)
+	{
+		if (!PendingLoad.TaskFuture.IsReady())
+		{
+			// This task was queued to be run on the Game thread but hasn't run yet
+			bAnyTasksNeedCancelling = true;
+			continue;
+		}
+
+		TSharedPtr<FStreamableHandle> Handle = PendingLoad.TaskFuture.Get();
+		if (Handle.IsValid())
+		{
+			Handle->CancelHandle();
+		}
+	}
+
+	if (bAnyTasksNeedCancelling)
+	{
+		// Set the shared cancel value to true so that whenever the tasks actually execute they
+		// will cancel immediately.
+		check(ShouldCancelPtr.IsValid());
+		*ShouldCancelPtr = true;
+
+		// The next task to be queued should allocate a new cancel value to ensure it doesn't
+		// interfere with the previous one.
+		ShouldCancelPtr = nullptr;
+	}
+
+	PendingLoads.Empty();
+}
+
+TSharedPtr<mu::FExtensionData> FUnrealExtensionDataStreamer::CloneExtensionData(const TSharedPtr<const mu::FExtensionData>& Source)
+{
+	unimplemented();
+	return nullptr;
+}
+
+TSharedRef<const mu::FExtensionDataLoadHandle> FUnrealExtensionDataStreamer::StartLoad(
+	const TSharedPtr<const mu::FExtensionData>& Data,
+	TArray<TSharedPtr<const mu::FExtensionData>>& OutUnloadedConstants)
+{
+	// Note that this can be called from any thread
+
+	OutUnloadedConstants.Reset();
+
+	TSharedRef<mu::FExtensionDataLoadHandle> LoadHandle = MakeShared<mu::FExtensionDataLoadHandle>();
+
+	LoadHandle->Data = Data;
+	LoadHandle->LoadState = mu::FExtensionDataLoadHandle::ELoadState::Pending;
+
+	{
+		FScopeLock Lock(Mutex);
+
+		if (!ShouldCancelPtr.IsValid())
+		{
+			ShouldCancelPtr = MakeShared<bool>(false);
+		}
+
+		TFuture<TSharedPtr<FStreamableHandle>> LoadHandleFuture = Async(EAsyncExecution::TaskGraphMainThread,
+			[
+				ShouldCancelRef = ShouldCancelPtr.ToSharedRef(),
+				LoadHandle,
+				ObjectToLoadFor = ActiveObject,
+				SystemPrivate = SystemPrivate
+			]()
+			{
+				if (*ShouldCancelRef)
+				{
+					// Load was cancelled
+					LoadHandle->LoadState = mu::FExtensionDataLoadHandle::ELoadState::FailedToLoad;
+					return TSharedPtr<FStreamableHandle>();
+				}
+
+				return StartLoadOnGameThread(SystemPrivate, ObjectToLoadFor, LoadHandle);
+			});
+
+
+		PendingLoads.Emplace(LoadHandle, MoveTemp(LoadHandleFuture));
+	}
+
+	return LoadHandle;
+}
+
+// Note that this function only returns an FStreamableHandle if it needed to start a load.
+//
+// If the requested object is already loaded, this function can return nullptr and still be
+// considered successful.
+TSharedPtr<FStreamableHandle> FUnrealExtensionDataStreamer::StartLoadOnGameThread(
+	const TWeakObjectPtr<UCustomizableObjectSystemPrivate>& SystemPrivateWeak,
+	const TWeakObjectPtr<UCustomizableObject>& ObjectToLoadFor,
+	const TSharedRef<mu::FExtensionDataLoadHandle>& LoadHandle)
+{
+	check(IsInGameThread());
+	check(LoadHandle->Data->Origin == mu::FExtensionData::EOrigin::ConstantStreamed);
+
+	UCustomizableObjectSystemPrivate* SystemPrivate = SystemPrivateWeak.Get();
+	if (!SystemPrivate)
+	{
+		LoadHandle->LoadState = mu::FExtensionDataLoadHandle::ELoadState::FailedToLoad;
+		return nullptr;
+	}
+	
+	UCustomizableObject* Object = ObjectToLoadFor.Get();
+	if (!Object)
+	{
+		// CO doesn't exist anymore
+		LoadHandle->LoadState = mu::FExtensionDataLoadHandle::ELoadState::FailedToLoad;
+		return nullptr;
+	}
+
+	FModelResources& ModelResources = Object->GetPrivate()->GetModelResources();
+	
+	if (!ModelResources.StreamedExtensionData.IsValidIndex(LoadHandle->Data->Index))
+	{
+		// The compiled data appears to be out of sync with the CO's properties
+
+		UE_LOG(LogMutable, Error, TEXT("Couldn't find streamed Extension Data with index %d in %s. Compiled data may be stale."),
+			LoadHandle->Data->Index, *Object->GetFullName());
+
+		LoadHandle->LoadState = mu::FExtensionDataLoadHandle::ELoadState::FailedToLoad;
+		return nullptr;
+	}
+
+	FCustomizableObjectStreamedResourceData& StreamedData = ModelResources.StreamedExtensionData[LoadHandle->Data->Index];
+	if (StreamedData.IsLoaded())
+	{
+		// Already loaded
+
+		// Need to call NotifyLoadCompleted to mark the request as complete and remove it from
+		// PendingLoads.
+		SystemPrivate->ExtensionDataStreamer->NotifyLoadCompleted(Object, LoadHandle);
+
+		return nullptr;
+	}
+
+	// Note that this just checks if the path is non-null, NOT if the object is loaded
+	check(!StreamedData.GetPath().IsNull());
+
+	TArray<FSoftObjectPath> TargetsToStream;
+	TargetsToStream.Add(StreamedData.GetPath().ToSoftObjectPath());
+
+	check(SystemPrivate->ExtensionDataStreamer);
+	TFunction<void()> OnLoadComplete = [FExtensionDataStreamer = SystemPrivate->ExtensionDataStreamer, ObjectToLoadFor, LoadHandle]()
+	{
+		UCustomizableObject* Object = ObjectToLoadFor.Get();
+		if (!Object)
+		{
+			// CO doesn't exist anymore
+			return;
+		}
+
+		FExtensionDataStreamer->NotifyLoadCompleted(Object, LoadHandle);
+	};
+
+	const bool bManageActiveHandle = false;
+	const bool bStartStalled = false;
+
+	FString DebugString;
+#if UE_BUILD_DEBUG || UE_BUILD_DEVELOPMENT
+	DebugString = FString::Printf(TEXT("UnrealExtensionDataStreamer for %s"), *Object->GetPathName());
+#endif
+
+	const TSharedRef<FMutableStreamableManager>& Manager = UCustomizableObjectSystem::GetInstance()->GetPrivate()->StreamableManager;
+
+	return Manager->RequestAsyncLoad(
+		TargetsToStream,
+		MoveTemp(OnLoadComplete),
+		FStreamableManager::DefaultAsyncLoadPriority,
+		bManageActiveHandle,
+		bStartStalled,
+		DebugString);
+}
+
+void FUnrealExtensionDataStreamer::NotifyLoadCompleted(
+	UCustomizableObject* Object,
+	const TSharedRef<mu::FExtensionDataLoadHandle>& LoadHandle)
+{
+	check(IsInGameThread());
+	check(Object);
+
+	FModelResources& ModelResources = Object->GetPrivate()->GetModelResources();
+
+	if (!ModelResources.StreamedExtensionData.IsValidIndex(LoadHandle->Data->Index))
+	{
+		// The compiled data appears to be out of sync with the CO's properties
+
+		UE_LOG(LogMutable, Error, TEXT("Couldn't find streamed Extension Data with index %d in %s. Compiled data may be stale."),
+			LoadHandle->Data->Index, *Object->GetFullName());
+
+		LoadHandle->LoadState = mu::FExtensionDataLoadHandle::ELoadState::FailedToLoad;
+		return;
+	}
+
+	FCustomizableObjectStreamedResourceData& StreamedData = ModelResources.StreamedExtensionData[LoadHandle->Data->Index];
+
+	// The object could have been loaded by another request, in which case we can skip updating
+	// the StreamedData.
+	if (!StreamedData.IsLoaded())
+	{
+		const UCustomizableObjectResourceDataContainer* LoadedObject = StreamedData.GetPath().Get();
+		if (!LoadedObject)
+		{
+			// Object wasn't loaded for some reason
+			LoadHandle->LoadState = mu::FExtensionDataLoadHandle::ELoadState::FailedToLoad;
+			return;
+		}
+
+		StreamedData.NotifyLoaded(LoadedObject);
+	}
+
+	LoadHandle->LoadState = mu::FExtensionDataLoadHandle::ELoadState::Loaded;
+
+	{
+		FScopeLock Lock(Mutex);
+
+		const int32 Index = Algo::IndexOfBy(PendingLoads, LoadHandle, &FPendingLoad::LoadHandle);
+		// If the load was previously cancelled, CancelHandle should have prevented this callback from
+		// running, so at this point the load should still be in PendingLoads.
+		check(Index != INDEX_NONE);
+
+		PendingLoads.RemoveAtSwap(Index);
+	}
+}
diff --git a/Engine/Plugins/Mutable/Source/CustomizableObject/Public/MuCO/CustomizableObjectResourceData.h b/Engine/Plugins/Mutable/Source/CustomizableObject/Public/MuCO/CustomizableObjectResourceData.h
--- a/Engine/Plugins/Mutable/Source/CustomizableObject/Public/MuCO/CustomizableObjectResourceData.h
+++ b/Engine/Plugins/Mutable/Source/CustomizableObject/Public/MuCO/CustomizableObjectResourceData.h
@@ -21,6 +21,8 @@ struct FCustomizableObjectResourceData
 {
 	GENERATED_BODY()
 
+public:
+
 	UPROPERTY()
 	ECOResourceDataType Type = ECOResourceDataType::None;
 
diff --git a/Engine/Plugins/Mutable/Source/CustomizableObject/Public/MuCO/CustomizableObjectStreamedResourceData.h b/Engine/Plugins/Mutable/Source/CustomizableObject/Public/MuCO/CustomizableObjectStreamedResourceData.h
new file mode 100644
--- /dev/null
+++ b/Engine/Plugins/Mutable/Source/CustomizableObject/Public/MuCO/CustomizableObjectStreamedResourceData.h
@@ -0,0 +1,78 @@
+// Copyright Epic Games, Inc. All Rights Reserved.
+
+#pragma once
+
+#include "MuCO/CustomizableObjectResourceData.h"
+#include "UObject/Object.h"
+#include "UObject/SoftObjectPtr.h"
+
+#include "CustomizableObjectStreamedResourceData.generated.h"
+
+
+
+/**
+ * Used to store resource data that can be streamed in and out from an external package.
+ *
+ * This allows you to load large resource data on demand, instead of having it always loaded.
+ */
+USTRUCT()
+struct CUSTOMIZABLEOBJECT_API FCustomizableObjectStreamedResourceData
+{
+	GENERATED_BODY()
+
+public:
+	FCustomizableObjectStreamedResourceData() = default;
+
+#if WITH_EDITOR
+	/** Streamed data can only be constructed in the editor */
+	FCustomizableObjectStreamedResourceData(UCustomizableObjectResourceDataContainer* InContainer);
+
+	/**
+	 * Clears the hard reference to the container, so that it's not forced to load as soon as this
+	 * struct is loaded, and replaces the soft reference with the provided one.
+	 * 
+	 * Note that the container may not have been moved to its new path yet, so don't try to load
+	 * the container from the given path here.
+	 */
+	void ConvertToSoftReferenceForCooking(const TSoftObjectPtr<UCustomizableObjectResourceDataContainer>& NewContainerPath);
+#endif // WITH_EDITOR
+
+	/**
+	 * If this returns true, GetLoadedData can safely be called.
+	 *
+	 * The data will stay loaded until Unload is called.
+	 */
+	bool IsLoaded() const;
+
+	/**
+	 * Returns the data, as long as it's loaded.
+	 *
+	 * If the data isn't loaded, an assertion will fail.
+	 */
+	const FCustomizableObjectResourceData& GetLoadedData() const;
+
+	/**
+	 * Release this struct's hard reference to the loaded data.
+	 *
+	 * This doesn't immediately remove it from memory, but will allow it to be deleted by the
+	 * garbage collector if there are no other references.
+	 *
+	 * Note that in editor builds streamed Resource Data must always be loaded due to the package
+	 * structure, and therefore this function has no effect in editor builds.
+	 */
+	void Unload();
+
+	/**
+	 * To stream this data in, load the object at the path returned by GetPath and then pass it
+	 * into NotifyLoaded.
+	 */
+	const TSoftObjectPtr<UCustomizableObjectResourceDataContainer>& GetPath() const { return ContainerPath; }
+	void NotifyLoaded(const UCustomizableObjectResourceDataContainer* LoadedContainer);
+
+private:
+	UPROPERTY()
+	TSoftObjectPtr<UCustomizableObjectResourceDataContainer> ContainerPath;
+
+	UPROPERTY()
+	TObjectPtr<const UCustomizableObjectResourceDataContainer> Container;
+};
diff --git a/Engine/Plugins/Mutable/Source/CustomizableObject/Public/MuCO/UnrealExtensionDataStreamer.h b/Engine/Plugins/Mutable/Source/CustomizableObject/Public/MuCO/UnrealExtensionDataStreamer.h
new file mode 100644
--- /dev/null
+++ b/Engine/Plugins/Mutable/Source/CustomizableObject/Public/MuCO/UnrealExtensionDataStreamer.h
@@ -0,0 +1,99 @@
+// Copyright Epic Games, Inc. All Rights Reserved.
+
+#pragma once
+
+#include "Async/Future.h"
+#include "Engine/StreamableManager.h"
+#include "MuR/ExtensionData.h"
+#include "MuR/ExtensionDataStreamer.h"
+#include "Templates/SharedPointerFwd.h"
+
+class UCustomizableObjectSystemPrivate;
+class UCustomizableObject;
+
+/** An implementation of mu::FExtensionDataStreamer designed to work with the Customizable Object integration */
+class CUSTOMIZABLEOBJECT_API FUnrealExtensionDataStreamer : public mu::FExtensionDataStreamer
+{
+public:
+	FUnrealExtensionDataStreamer(const TWeakObjectPtr<UCustomizableObjectSystemPrivate>& InSystemPrivateWeak);
+	virtual ~FUnrealExtensionDataStreamer();
+
+	/** Not copyable because TFuture isn't copyable. This shouldn't be limiting in practise. */
+	FUnrealExtensionDataStreamer(const FUnrealExtensionDataStreamer& Other) = delete;
+	FUnrealExtensionDataStreamer& operator=(const FUnrealExtensionDataStreamer& Other) = delete;
+
+	/**
+	 * Note that *only* the functions from the mu::FExtensionDataStreamer interface may be called
+	 * from other threads. All other functions may only be called from the Game thread.
+	 */
+
+	/**
+	 * Set the Customizable Object that is currently running.
+	 *
+	 * StartLoad calls are assumed to come from the object passed in here.
+	 */
+	void SetActiveObject(UCustomizableObject* InObject);
+
+	/**
+	 * Clear the active object
+	 * 
+	 * Call this after a Customizable Object has finished running to help detect cases where
+	 * SetActiveObject wasn't called before running a different CO.
+	 */
+	void ClearActiveObject();
+
+	/** Returns true if there are any loads requested by StartLoad that haven't yet completed. */
+	bool AreAnyLoadsPending() const;
+
+	/** Cancels all loads requested by StartLoad that haven't yet completed. */
+	void CancelPendingLoads();
+
+	/** mu::FExtensionDataStreamer interface */
+	virtual TSharedPtr<mu::FExtensionData> CloneExtensionData(const TSharedPtr<const mu::FExtensionData>& Source) override;
+	
+	virtual TSharedRef<const mu::FExtensionDataLoadHandle> StartLoad(
+		const TSharedPtr<const mu::FExtensionData>& Data,
+		TArray<TSharedPtr<const mu::FExtensionData>>& OutUnloadedConstants) override;
+
+private:
+	/** Queued up on the Game thread to start the async load */
+	static TSharedPtr<FStreamableHandle> StartLoadOnGameThread(
+		const TWeakObjectPtr<UCustomizableObjectSystemPrivate>& SystemPrivateWeak,
+		const TWeakObjectPtr<UCustomizableObject>& ObjectToLoadFor,
+		const TSharedRef<mu::FExtensionDataLoadHandle>& LoadHandle);
+
+	/** Called when the async load has finished loading the Extension Data */
+	void NotifyLoadCompleted(UCustomizableObject* Object, const TSharedRef<mu::FExtensionDataLoadHandle>& LoadHandle);
+
+	TWeakObjectPtr<UCustomizableObjectSystemPrivate> SystemPrivate; // Can not be a TStrongObjectPtr since then it will create a cycle b TSharedPtr and UObjects.
+
+	/** Guards PendingLoads, ShouldCancelPtr and ActiveObject */
+	FCriticalSection* Mutex = nullptr;
+
+	TWeakObjectPtr<UCustomizableObject> ActiveObject;
+	
+	struct FPendingLoad
+	{
+		FPendingLoad(
+			const TSharedRef<mu::FExtensionDataLoadHandle>& InLoadHandle,
+			TFuture<TSharedPtr<FStreamableHandle>>&& InTaskFuture)
+			: LoadHandle(InLoadHandle)
+			, TaskFuture(MoveTemp(InTaskFuture))
+		{
+		}
+
+		FPendingLoad(const FPendingLoad& Other) = delete;
+		FPendingLoad& operator=(const FPendingLoad& Other) = delete;
+
+		TSharedRef<mu::FExtensionDataLoadHandle> LoadHandle;
+		TFuture<TSharedPtr<FStreamableHandle>> TaskFuture;
+	};
+	/** All outstanding loads that haven't called NotifyLoadCompleted yet */
+	TArray<FPendingLoad> PendingLoads;
+
+	/**
+	 * This is shared by all outstanding loads and can be used to cancel them if their Game thread
+	 * task hasn't run yet.
+	 */
+	TSharedPtr<bool> ShouldCancelPtr;
+};
diff --git a/Engine/Plugins/Mutable/Source/MutableRuntime/Internal/MuR/ModelPrivate.h b/Engine/Plugins/Mutable/Source/MutableRuntime/Internal/MuR/ModelPrivate.h
--- a/Engine/Plugins/Mutable/Source/MutableRuntime/Internal/MuR/ModelPrivate.h
+++ b/Engine/Plugins/Mutable/Source/MutableRuntime/Internal/MuR/ModelPrivate.h
@@ -8,6 +8,7 @@
 #include "MuR/SerialisationPrivate.h"
 #include "MuR/Operations.h"
 #include "MuR/ExtensionData.h"
+#include "MuR/ExtensionDataStreamer.h"
 #include "MuR/Image.h"
 #include "MuR/Mesh.h"
 #include "MuR/ParametersPrivate.h"
@@ -109,17 +110,31 @@ namespace mu
 			AlwaysLoaded
 		};
 
-		void Serialise(FOutputArchive& Arch) const
+		// This should be initialized to a valid load state when Data is set
+		ELoadState LoadState = ELoadState::Invalid;
+
+		inline void Serialise(FOutputArchive& Arch) const
 		{
 			Arch << Data;
 		}
 		
-		void Unserialise(FInputArchive& Arch)
+		inline void Unserialise(FInputArchive& Arch)
 		{
 			Arch >> Data;
 
 			check(Data);
-			check(Data->Origin == FExtensionData::EOrigin::ConstantAlwaysLoaded || Data->Origin == FExtensionData::EOrigin::ConstantStreamed);
+			check(Data->Origin == FExtensionData::EOrigin::ConstantAlwaysLoaded
+				|| Data->Origin == FExtensionData::EOrigin::ConstantStreamed);
+
+			if (Data->Origin == FExtensionData::EOrigin::ConstantAlwaysLoaded)
+			{
+				LoadState = ELoadState::AlwaysLoaded;
+			}
+			else
+			{
+				// Streamed constants are assumed to be unloaded to start with
+				LoadState = ELoadState::Unloaded;
+			}
 		}
 	};
 
@@ -459,6 +474,12 @@ namespace mu
 			FExtensionDataConstant& NewConstant = ConstantExtensionData.AddDefaulted_GetRef();
 			NewConstant.Data = Data;
 
+			// The data is assumed to be loaded during compilation
+			NewConstant.LoadState =
+				Data->Origin == FExtensionData::EOrigin::ConstantAlwaysLoaded
+				? FExtensionDataConstant::ELoadState::AlwaysLoaded
+				: FExtensionDataConstant::ELoadState::CurrentlyLoaded;
+
 			return ConstantExtensionData.Num() - 1;
 		}
 
@@ -728,6 +749,7 @@ namespace mu
 		{
 			const FExtensionDataConstant& Constant = ConstantExtensionData[ConstantIndex];
 
+			check(Constant.LoadState != FExtensionDataConstant::ELoadState::Unloaded);
 			check(Constant.Data);
 
 			Result = Constant.Data;
diff --git a/Engine/Plugins/Mutable/Source/MutableRuntime/Internal/MuR/SystemPrivate.h b/Engine/Plugins/Mutable/Source/MutableRuntime/Internal/MuR/SystemPrivate.h
--- a/Engine/Plugins/Mutable/Source/MutableRuntime/Internal/MuR/SystemPrivate.h
+++ b/Engine/Plugins/Mutable/Source/MutableRuntime/Internal/MuR/SystemPrivate.h
@@ -1933,7 +1933,7 @@ namespace mu
 		friend class System;
     public:
 
-        Private(const FSettings&);
+        Private(const FSettings&, const TSharedPtr<FExtensionDataStreamer>& );
         ~Private();
 
         //-----------------------------------------------------------------------------------------
@@ -1954,7 +1954,9 @@ namespace mu
 		MUTABLERUNTIME_API TSharedPtr<const FInstance> BuildInstance(const TSharedPtr<const FModel>&, const FParameters*, OP::ADDRESS);
 		MUTABLERUNTIME_API TSharedPtr<const FLayout> BuildLayout(const TSharedPtr<const FModel>&, const FParameters*, OP::ADDRESS) ;
     	MUTABLERUNTIME_API FProjector BuildProjector(const TSharedPtr<const FModel>&, const FParameters*, OP::ADDRESS) ;
-    	
+
+		FExtensionDataStreamer* GetExtensionDataStreamer() const { return ExtensionDataStreamer.Get(); }
+
         //!
         FSettings Settings;
 
@@ -1991,6 +1993,11 @@ namespace mu
 
 		//! Update some mutable core unreal stats.
 		void UpdateStats();
+
+	private:
+
+		/** Owned by this system. */
+		TSharedPtr<FExtensionDataStreamer> ExtensionDataStreamer;
     };
 
 }
diff --git a/Engine/Plugins/Mutable/Source/MutableRuntime/Private/MuR/CodeRunner.h b/Engine/Plugins/Mutable/Source/MutableRuntime/Private/MuR/CodeRunner.h
--- a/Engine/Plugins/Mutable/Source/MutableRuntime/Private/MuR/CodeRunner.h
+++ b/Engine/Plugins/Mutable/Source/MutableRuntime/Private/MuR/CodeRunner.h
@@ -6,6 +6,7 @@
 #include "Containers/ContainerAllocationPolicies.h"
 #include "HAL/PlatformCrt.h"
 #include "HAL/PlatformMath.h"
+#include "MuR/ExtensionDataStreamer.h"
 #include "MuR/Image.h"
 #include "MuR/Layout.h"
 #include "MuR/Operations.h"
@@ -274,6 +275,31 @@ namespace  mu
 			uint32 RomIndex = 0;
 		};
 
+		class FLoadExtensionDataTask : public CodeRunner::FIssuedTask
+		{
+		public:
+			FLoadExtensionDataTask(const FScheduledOp& InOp, const TSharedPtr<const FExtensionData>& InExtensionData, int32 InModelConstantIndex)
+				: FIssuedTask(InOp)
+				, Data(InExtensionData)
+				, ModelConstantIndex(InModelConstantIndex)
+			{
+				check(Data.Get());
+				check(Data->Origin == FExtensionData::EOrigin::ConstantStreamed);
+			}
+
+			// FIssuedTask interface
+			virtual bool Prepare(CodeRunner*, bool& bOutFailed) override;
+			virtual bool Complete(CodeRunner*) override;
+			virtual bool IsComplete(CodeRunner*) override;
+
+		private:
+			TSharedPtr<const FExtensionData> Data;
+			// The index of this constant in FProgram::ConstantExtensionData
+			int32 ModelConstantIndex;
+
+			TSharedPtr<const FExtensionDataLoadHandle> LoadHandle;
+		};
+
 		class FLoadImageRomsTask : public CodeRunner::FIssuedTask
 		{
 		public:
diff --git a/Engine/Plugins/Mutable/Source/MutableRuntime/Private/MuR/CodeRunnerTasked.cpp b/Engine/Plugins/Mutable/Source/MutableRuntime/Private/MuR/CodeRunnerTasked.cpp
--- a/Engine/Plugins/Mutable/Source/MutableRuntime/Private/MuR/CodeRunnerTasked.cpp
+++ b/Engine/Plugins/Mutable/Source/MutableRuntime/Private/MuR/CodeRunnerTasked.cpp
@@ -13,6 +13,7 @@
 #include "Misc/AssertionMacros.h"
 #include "MuR/CodeRunner.h"
 #include "MuR/System.h"
+#include "MuR/ExtensionDataStreamer.h"
 #include "MuR/Image.h"
 #include "MuR/ImagePrivate.h"
 #include "MuR/Layout.h"
@@ -2245,7 +2246,94 @@ namespace mu
 		return bSuccess;
 	}
 
-	
+
+	//---------------------------------------------------------------------------------------------
+	//---------------------------------------------------------------------------------------------
+	//---------------------------------------------------------------------------------------------
+	bool CodeRunner::FLoadExtensionDataTask::Prepare(CodeRunner* Runner, bool& bOutFailed)
+	{
+		// This runs in the mutable Runner thread
+		MUTABLE_CPUPROFILER_SCOPE(FLoadExtensionDataTask_Prepare);
+		bOutFailed = false;
+
+		FProgram& Program = Runner->Model->GetPrivate()->Program;
+		check(Program.ConstantExtensionData.IsValidIndex(ModelConstantIndex));
+
+		const FExtensionDataConstant& Constant = Program.ConstantExtensionData[ModelConstantIndex];
+		check(Constant.Data == Data);
+
+		if (Constant.LoadState == FExtensionDataConstant::ELoadState::FailedToLoad)
+		{
+			// This already failed before, so don't waste time trying again.
+			//
+			// This behavior could be made an option in future if needed.
+			bOutFailed = true;
+			return false;
+		}
+
+		TArray<TSharedPtr<const FExtensionData>> UnloadedConstants;
+		LoadHandle = Runner->System->GetExtensionDataStreamer()->StartLoad(Data, UnloadedConstants);
+		check(LoadHandle.IsValid());
+
+		// Update the load state of any constants that were unloaded
+		for (const TSharedPtr<const FExtensionData>& Unloaded : UnloadedConstants)
+		{
+			check(Unloaded->Origin == FExtensionData::EOrigin::ConstantStreamed);
+
+			FExtensionDataConstant* FoundConstant = Algo::FindBy(Program.ConstantExtensionData, Unloaded, &FExtensionDataConstant::Data);
+			check(FoundConstant);
+
+			FoundConstant->LoadState = FExtensionDataConstant::ELoadState::Unloaded;
+		}
+
+		// No worker thread work
+		return false;
+	}
+
+
+	//---------------------------------------------------------------------------------------------
+	bool CodeRunner::FLoadExtensionDataTask::Complete(CodeRunner* Runner)
+	{
+		MUTABLE_CPUPROFILER_SCOPE(FLoadExtensionDataTask_Complete);
+
+		bool bSuccess = true;
+
+		// Complete should only be called if the load is finished
+		check(LoadHandle->LoadState != FExtensionDataLoadHandle::ELoadState::Pending);
+
+		FProgram& Program = Runner->Model->GetPrivate()->Program;
+		check(Program.ConstantExtensionData.IsValidIndex(ModelConstantIndex));
+
+		FExtensionDataConstant& Constant = Program.ConstantExtensionData[ModelConstantIndex];
+		check(Constant.Data == Data);
+
+		// Update the load state in the Program
+		if (LoadHandle->LoadState == FExtensionDataLoadHandle::ELoadState::Loaded)
+		{
+			Constant.LoadState = FExtensionDataConstant::ELoadState::CurrentlyLoaded;
+		}
+		else
+		{
+			check(LoadHandle->LoadState == FExtensionDataLoadHandle::ELoadState::FailedToLoad);
+			Constant.LoadState = FExtensionDataConstant::ELoadState::FailedToLoad;
+			bSuccess = false;
+		}
+
+		// Process the constant op normally, now that the data is loaded.
+		Runner->RunCode(Op, Runner->Params, Runner->Model, Runner->LODMask);
+
+		return bSuccess;
+	}
+
+
+	//---------------------------------------------------------------------------------------------
+	bool CodeRunner::FLoadExtensionDataTask::IsComplete(CodeRunner* Runner)
+	{
+		check(LoadHandle.IsValid());
+		return LoadHandle->LoadState != FExtensionDataLoadHandle::ELoadState::Pending;
+	}
+
+
 	//---------------------------------------------------------------------------------------------
 	//---------------------------------------------------------------------------------------------
 	//---------------------------------------------------------------------------------------------
@@ -2589,6 +2677,26 @@ namespace mu
 			}
 			break;
 		}
+		
+		case EOpType::ED_CONSTANT:
+		{
+			OP::ResourceConstantArgs Args = Program.GetOpArgs<OP::ResourceConstantArgs>(item.At);
+
+			const FExtensionDataConstant::ELoadState LoadState = Program.ConstantExtensionData[Args.value].LoadState;
+
+			check(LoadState != FExtensionDataConstant::ELoadState::Invalid);
+			
+			if (LoadState == FExtensionDataConstant::ELoadState::AlwaysLoaded
+				|| LoadState == FExtensionDataConstant::ELoadState::CurrentlyLoaded)
+			{
+				// Already loaded. The rest of the constant code will run right away.
+			}
+			else
+			{
+				Issued = MakeShared<FLoadExtensionDataTask>(item, Program.ConstantExtensionData[Args.value].Data, Args.value);
+			}
+			break;
+		}
 
 		case EOpType::IM_CONSTANT:
 		{
diff --git a/Engine/Plugins/Mutable/Source/MutableRuntime/Private/MuR/NullExtensionDataStreamer.cpp b/Engine/Plugins/Mutable/Source/MutableRuntime/Private/MuR/NullExtensionDataStreamer.cpp
new file mode 100644
--- /dev/null
+++ b/Engine/Plugins/Mutable/Source/MutableRuntime/Private/MuR/NullExtensionDataStreamer.cpp
@@ -0,0 +1,35 @@
+// Copyright Epic Games, Inc. All Rights Reserved.
+
+#include "MuR/NullExtensionDataStreamer.h"
+
+#include "MuR/ExtensionData.h"
+#include "Templates/SharedPointer.h"
+
+namespace mu
+{
+
+TSharedPtr<FExtensionData> NullExtensionDataStreamer::CloneExtensionData(const TSharedPtr<const FExtensionData>& Source)
+{
+	// Note that these appear to be valid and should be good enough for testing, but they won't be
+	// usable since we're not copying the actual data pointed to by Source.
+
+	TSharedPtr<FExtensionData> Result = MakeShared<FExtensionData>();
+	Result->Index = NextIndex;
+	NextIndex++;
+	Result->Origin = FExtensionData::EOrigin::Runtime;
+
+	return Result;
+}
+
+TSharedRef<const FExtensionDataLoadHandle> NullExtensionDataStreamer::StartLoad(const TSharedPtr<const FExtensionData>& Data, TArray<TSharedPtr<const FExtensionData>>& OutUnloadedConstants)
+{
+	OutUnloadedConstants.Reset();
+
+	TSharedRef<FExtensionDataLoadHandle> Result = MakeShared<FExtensionDataLoadHandle>();
+	Result->Data = Data;
+	Result->LoadState = FExtensionDataLoadHandle::ELoadState::FailedToLoad;
+
+	return Result;
+}
+
+}
diff --git a/Engine/Plugins/Mutable/Source/MutableRuntime/Private/MuR/NullExtensionDataStreamer.h b/Engine/Plugins/Mutable/Source/MutableRuntime/Private/MuR/NullExtensionDataStreamer.h
new file mode 100644
--- /dev/null
+++ b/Engine/Plugins/Mutable/Source/MutableRuntime/Private/MuR/NullExtensionDataStreamer.h
@@ -0,0 +1,21 @@
+// Copyright Epic Games, Inc. All Rights Reserved.
+
+#pragma once
+
+#include "MuR/ExtensionDataStreamer.h"
+
+namespace mu
+{
+
+/** A stub implementation that can be used if there's no valid FExtensionDataStreamer */
+class NullExtensionDataStreamer : public FExtensionDataStreamer
+{
+public:
+	virtual TSharedPtr<FExtensionData> CloneExtensionData(const TSharedPtr<const FExtensionData>& Source) override;
+	virtual TSharedRef<const FExtensionDataLoadHandle> StartLoad(const TSharedPtr<const FExtensionData>& Data, TArray<TSharedPtr<const FExtensionData>>& OutUnloadedConstants) override;
+
+private:
+	int32 NextIndex = 0;
+};
+
+}
diff --git a/Engine/Plugins/Mutable/Source/MutableRuntime/Private/MuR/System.cpp b/Engine/Plugins/Mutable/Source/MutableRuntime/Private/MuR/System.cpp
--- a/Engine/Plugins/Mutable/Source/MutableRuntime/Private/MuR/System.cpp
+++ b/Engine/Plugins/Mutable/Source/MutableRuntime/Private/MuR/System.cpp
@@ -22,6 +22,7 @@
 #include "MuR/MutableMath.h"
 #include "MuR/MutableString.h"
 #include "MuR/MutableTrace.h"
+#include "MuR/NullExtensionDataStreamer.h"
 #include "MuR/Operations.h"
 #include "MuR/Parameters.h"
 #include "MuR/ParametersPrivate.h"
@@ -72,12 +73,12 @@ namespace mu
 
 
     //---------------------------------------------------------------------------------------------
-    FSystem::FSystem(const FSettings& InSettings)
+    FSystem::FSystem(const FSettings& InSettings, const TSharedPtr<FExtensionDataStreamer>& DataStreamer)
     {
 		LLM_SCOPE_BYNAME(TEXT("MutableRuntime"));
 
         // Choose the implementation
-        m_pD = new FSystem::Private(InSettings);
+        m_pD = new FSystem::Private( InSettings, DataStreamer );
     }
 
 
@@ -102,7 +103,7 @@ namespace mu
 
 
     //---------------------------------------------------------------------------------------------
-    FSystem::Private::Private(const FSettings& InSettings)
+    FSystem::Private::Private( const FSettings& InSettings, const TSharedPtr<mu::FExtensionDataStreamer>& InDataStreamer)
     {
 		LLM_SCOPE_BYNAME(TEXT("MutableRuntime"));
 
@@ -111,6 +112,12 @@ namespace mu
 		WorkingMemoryManager.BudgetBytes = Settings.WorkingMemoryBytes;
 		WorkingMemoryManager.GeneratedResources.Reserve(WorkingMemoryManager.MaxGeneratedResourceCacheSize);
 
+		ExtensionDataStreamer = InDataStreamer;
+		if (!ExtensionDataStreamer)
+		{
+			ExtensionDataStreamer = MakeShared<NullExtensionDataStreamer>();
+		}
+
 		UpdateStats();
 	}
 
@@ -123,6 +130,7 @@ namespace mu
 
 		// Make it explicit to try to capture metrics
         StreamInterface = nullptr;
+		ExtensionDataStreamer = nullptr;
 		ExternalResourceProvider = nullptr;
     }
 
diff --git a/Engine/Plugins/Mutable/Source/MutableRuntime/Public/MuR/ExtensionDataStreamer.h b/Engine/Plugins/Mutable/Source/MutableRuntime/Public/MuR/ExtensionDataStreamer.h
new file mode 100644
--- /dev/null
+++ b/Engine/Plugins/Mutable/Source/MutableRuntime/Public/MuR/ExtensionDataStreamer.h
@@ -0,0 +1,67 @@
+// Copyright Epic Games, Inc. All Rights Reserved.
+
+#pragma once
+
+#include "MuR/ExtensionData.h"
+#include "MuR/Ptr.h"
+#include "MuR/RefCounted.h"
+#include "MuR/Serialisation.h"
+#include "MuR/Types.h"
+#include "Templates/SharedPointerFwd.h"
+
+namespace mu
+{
+
+	//! Allows runtime code to track the progress of an Extension Data load
+	class FExtensionDataLoadHandle
+	{
+	public:
+		TSharedPtr<const FExtensionData> Data;
+
+		enum class ELoadState : uint8
+		{
+			Pending,
+			FailedToLoad,
+			Loaded
+		};
+		ELoadState LoadState;
+	};
+
+	/** Interface to allow the runtime to load ExtensionData for use
+	* An object implementing this interface must be provided to the System for any Models that use
+	* ExtensionData.
+	* All functions must be safe to call from the Mutable thread.
+	*/
+	class MUTABLERUNTIME_API FExtensionDataStreamer
+	{
+	public:
+
+		//! Ensure virtual destruction.
+		virtual ~FExtensionDataStreamer() = default;
+
+		//! Make a copy of the actual data referenced by the ExtensionData, and return a new
+		//! runtime-generated ExtensionData that references it.
+		//! 
+		//! Source can either be a constant or another runtime-generated ExtensionData.
+		//! 
+		//! The new ExtensionData should be completely independent of the source ExtensionData, such
+		//! that if the source instance is deleted, the new instance should still be valid.
+		virtual TSharedPtr<FExtensionData> CloneExtensionData(const TSharedPtr<const FExtensionData>& Source) = 0;
+	
+		//! Start asynchronously loading the streamed constant data referenced by the given
+		//! ExtensionData.
+		//! 
+		//! \param Data The ExtensionData to load. Its Origin must be ConstantStreamed.
+		//! 
+		//! \param OutUnloadedConstants An array of ConstantStreamed ExtensionData that was unloaded
+		//! to make space for the requested data. Guaranteed to only contain valid pointers to
+		//! ConstantStreamed data from this program.
+		//! 
+		//! \return A handle that allows the caller to poll the load state. When the state is no longer
+		//! Pending, the load is either complete or has failed.
+		virtual TSharedRef<const FExtensionDataLoadHandle> StartLoad(
+			const TSharedPtr<const FExtensionData>& Data,
+			TArray<TSharedPtr<const FExtensionData>>& OutUnloadedConstants) = 0;
+	};
+
+}
diff --git a/Engine/Plugins/Mutable/Source/MutableRuntime/Public/MuR/System.h b/Engine/Plugins/Mutable/Source/MutableRuntime/Public/MuR/System.h
--- a/Engine/Plugins/Mutable/Source/MutableRuntime/Public/MuR/System.h
+++ b/Engine/Plugins/Mutable/Source/MutableRuntime/Public/MuR/System.h
@@ -100,7 +100,8 @@ namespace mu
 		//! Constructor of a system object to build data.
         //! \param Settings Optional class with the settings to use in this system. The default
         //! value configures a production-ready system.
-		FSystem( const FSettings& Settings = FSettings());
+		//! \param DataStreamer Optional interface to allow the Model to stream in ExtensionData from disk
+        FSystem( const FSettings& Settings = FSettings(), const TSharedPtr<FExtensionDataStreamer>& DataStreamer = nullptr);
 
         //! Set a new provider for model data. 
         void SetStreamingInterface(const TSharedPtr<FModelReader>& );
