[Mutable] Fix FCustomizableObjectMeshStreamIn::OperationData being nulled out while converting meshes if the streamer cancels the update task after the IsCancelled checks.

#rnx

diff --git a/Engine/Plugins/Mutable/Source/CustomizableObject/Private/MuCO/CustomizableObjectMeshUpdate.cpp b/Engine/Plugins/Mutable/Source/CustomizableObject/Private/MuCO/CustomizableObjectMeshUpdate.cpp
--- a/Engine/Plugins/Mutable/Source/CustomizableObject/Private/MuCO/CustomizableObjectMeshUpdate.cpp
+++ b/Engine/Plugins/Mutable/Source/CustomizableObject/Private/MuCO/CustomizableObjectMeshUpdate.cpp
@@ -87,9 +87,6 @@ void FCustomizableObjectMeshStreamIn::Abort()
 				CustomizableObjectSystem->MutableTaskGraph.CancelMutableThreadTaskLowPriority(MutableTaskId);
 			}
 		}
-
-		// Clear MeshUpdate data
-		OperationData = nullptr;
 	}
 	else
 	{
@@ -259,13 +256,9 @@ void FCustomizableObjectMeshStreamIn::RequestMeshUpdate(const FContext& Context)
 	
 	TaskSynchronization.Increment();
 
-	TRefCountPtr<FCustomizableObjectMeshStreamIn> RefThis = this;
-	TSharedPtr<FMutableMeshOperationData> SharedOperationData = OperationData;
-
-
 	MutableTaskId = CustomizableObjectSystem->MutableTaskGraph.AddMutableThreadTaskLowPriority(
 		TEXT("Mutable_MeshUpdate"),
-		[SharedOperationData, RefThis]() mutable
+		[SharedOperationData = OperationData, RefThis = TRefCountPtr<FCustomizableObjectMeshStreamIn>(this)]() mutable
 		{
 			impl::Task_Mutable_UpdateMesh(SharedOperationData, RefThis);
 		});
