[Mutable] Exposed section index for Meshes and Materials part of a Skeletal Mesh Node slate.
- Previously, the section was hidden as part of the tooltip, which prevented the user from having a clear view of what section was being manipulated.
- Removed pin tooltip as it is now redundant.
- Removed unnecessary include.

#rnx

diff --git a/Engine/Plugins/Mutable/Source/CustomizableObjectEditor/Private/MuCOE/Nodes/CustomizableObjectNodeSkeletalMesh.cpp b/Engine/Plugins/Mutable/Source/CustomizableObjectEditor/Private/MuCOE/Nodes/CustomizableObjectNodeSkeletalMesh.cpp
--- a/Engine/Plugins/Mutable/Source/CustomizableObjectEditor/Private/MuCOE/Nodes/CustomizableObjectNodeSkeletalMesh.cpp
+++ b/Engine/Plugins/Mutable/Source/CustomizableObjectEditor/Private/MuCOE/Nodes/CustomizableObjectNodeSkeletalMesh.cpp
@@ -13,7 +13,6 @@
 #include "MuCOE/CustomizableObjectLayout.h"
 #include "MuCOE/EdGraphSchema_CustomizableObject.h"
 #include "MuCOE/GraphTraversal.h"
-#include "MuCOE/MutableUtils.h"
 #include "MuCOE/Nodes/CustomizableObjectNodeLayoutBlocks.h"
 #include "MuCOE/RemapPins/CustomizableObjectNodeRemapPinsByNameDefaultPin.h"
 #include "MuCOE/UnrealEditorPortabilityHelpers.h"
@@ -94,26 +93,28 @@ void UCustomizableObjectNodeSkeletalMesh::AllocateDefaultPins(UCustomizableObjec
 			for (int32 SectionIndex = 0; SectionIndex < NumSections; ++SectionIndex)
 			{
 				// Ignore disabled sections.
-				const bool bIsSectionDisabled = ImportedModel ->LODModels[LODIndex].Sections[SectionIndex].bDisabled;
-				if (bIsSectionDisabled)
+				if (ImportedModel->LODModels[LODIndex].Sections[SectionIndex].bDisabled)
 				{
 					continue;
 				}
 
 				UMaterialInterface* MaterialInterface = GetMaterialInterfaceFor(LODIndex, SectionIndex);
 				
-				FString SectionFriendlyName = MaterialInterface ? MaterialInterface->GetName() : FString::Printf(TEXT("Section %i"), SectionIndex);
+				FString Section = FString::Printf(TEXT("Section %i"), SectionIndex);
+				if (MaterialInterface)
+				{
+					Section.Append(FString::Printf(TEXT(" : %s"), *MaterialInterface->GetName()));
+				}
 
 				// Mesh
 				{
-					FString MeshName = FString::Printf(TEXT("LOD %i - Section %i - Mesh"), LODIndex, SectionIndex);
-
 					UCustomizableObjectNodeSkeletalMeshPinDataMesh* PinData = NewObject<UCustomizableObjectNodeSkeletalMeshPinDataMesh>(this);
 					PinData->Init(LODIndex, SectionIndex, ImportedModel->LODModels[LODIndex].NumTexCoords);
 					
-					UEdGraphPin* Pin = CustomCreatePin(EGPD_Output, Schema->PC_Mesh, FName(*MeshName), PinData);
-					Pin->PinFriendlyName = FText::FromString(FString::Printf(TEXT("LOD %i - %s - Mesh"), LODIndex, *SectionFriendlyName));
-					Pin->PinToolTip = MeshName;
+					const FString MeshPinName = FString::Printf(TEXT("LOD %i - Section %i - Mesh"), LODIndex, SectionIndex);
+					
+					UEdGraphPin* Pin = CustomCreatePin(EGPD_Output, Schema->PC_Mesh, FName(*MeshPinName), PinData);
+					Pin->PinFriendlyName = FText::FromString(FString::Printf(TEXT("LOD %i - %s"), LODIndex, *Section));
 				}
 				
 				
@@ -129,18 +130,17 @@ void UCustomizableObjectNodeSkeletalMesh::AllocateDefaultPins(UCustomizableObjec
 					check(ImageInfos.Num() == ImageIds.Num())
 					for (int32 Index = 0; Index < ImageInfos.Num(); ++Index)
 					{
-						const FMaterialParameterInfo& ImageInfo = ImageInfos[Index];
 						FGuid& ImageId = ImageIds[Index];
 						
-						FString ImageNameStr = *ImageInfo.Name.ToString();
-						FString PinName = FString::Printf(TEXT("LOD %i - Section %i - Texture Parameter %s"), LODIndex, SectionIndex, *ImageNameStr);
-						
 						UCustomizableObjectNodeSkeletalMeshPinDataImage* PinData = NewObject<UCustomizableObjectNodeSkeletalMeshPinDataImage>(this);
 						PinData->Init(LODIndex, SectionIndex, ImageId);
 						
-						UEdGraphPin* Pin = CustomCreatePin(EGPD_Output, Schema->PC_Image, FName(*PinName), PinData);
-						Pin->PinFriendlyName = FText::FromString(FString::Printf(TEXT("LOD %i - %s - %s"), LODIndex, *SectionFriendlyName, *ImageNameStr));
-						Pin->PinToolTip = PinName;
+						const FMaterialParameterInfo& ImageInfo = ImageInfos[Index];
+						const FString ImageNameStr = *ImageInfo.Name.ToString();
+						const FString TexturePinName = FString::Printf(TEXT("LOD %i - Section %i - Texture Parameter %s"), LODIndex, SectionIndex, *ImageNameStr);
+						
+						UEdGraphPin* Pin = CustomCreatePin(EGPD_Output, Schema->PC_Image, FName(*TexturePinName), PinData);
+						Pin->PinFriendlyName = FText::FromString(FString::Printf(TEXT("LOD %i - %s - %s"), LODIndex, *Section, *ImageNameStr));
 						Pin->bHidden = true;
 					}
 				}
