[Mutable] Allow Game Thread tasks to lock the Mutable Thread.

If a Game Thread task asks for a Mutable Thread lock, the task will be delayed until all previous Mutable Tasks are executed. While executing the Game Thread task, no further Mutable Thread tasks will be executed (ensured using a precondition trigger).

- Moved GameThread task management to FMutableTaskGraph.
- Moved AdvanceCurrentOperation. It will no longer start Initial generation updates by itself, only manages tasks.
- Moved UpdateMemoryLimit.
- Renamed FTask to FMutableThreadLowPriorityTask.
- StartUpdateSkeletalMesh now adds a Game Thread task.
- Added debug names to Game Thread tasks.
- Added new parameter to AddGameThreadTask which makes sure that the Mutable Thread is locked before executing it. Notice that this new feature is currently not used.

diff --git a/Engine/Plugins/Mutable/Source/CustomizableObject/Internal/MuCO/CustomizableObjectSystemPrivate.h b/Engine/Plugins/Mutable/Source/CustomizableObject/Internal/MuCO/CustomizableObjectSystemPrivate.h
--- a/Engine/Plugins/Mutable/Source/CustomizableObject/Internal/MuCO/CustomizableObjectSystemPrivate.h
+++ b/Engine/Plugins/Mutable/Source/CustomizableObject/Internal/MuCO/CustomizableObjectSystemPrivate.h
@@ -312,36 +312,6 @@ struct FGeneratedMaterial
 };
 
 
-DECLARE_DELEGATE(FMutableTaskDelegate);
-struct FMutableTask
-{
-	/** Actual function to perform in this task. */
-	FMutableTaskDelegate Function;
-
-	TArray<UE::Tasks::FTask, TFixedAllocator<2>> Dependencies; 
-	
-	/** Check if all the dependencies of this task have been completed. */
-	FORCEINLINE bool AreDependenciesComplete() const
-	{
-		for (const UE::Tasks::FTask& Dependency : Dependencies)
-		{
-			if (!Dependency.IsCompleted())
-			{
-				return false;
-			}
-		}
-
-		return true;
-	}
-
-	/** Free the handles for any dependency of this task. */
-	FORCEINLINE void ClearDependencies()
-	{
-		Dependencies.Empty();
-	}
-};
-
-
 // Mutable data generated during the update steps.
 // We keep it from begin to end update, and it is used in several steps.
 struct FInstanceUpdateData
@@ -746,10 +716,7 @@ public:
 
 	// The pending instance updates, discards or releases
 	FMutablePendingInstanceWork MutablePendingInstanceWork;
-
-	// Queue of game-thread tasks that need to be executed for the current operation. TQueue is thread safe
-	TQueue<FMutableTask> PendingTasks;
-
+	
 	static int32 EnableMutableProgressiveMipStreaming;
 	static int32 EnableMutableLiveUpdate;
 	static int32 EnableReuseInstanceTextures;
@@ -775,8 +742,6 @@ public:
 
 	UCustomizableObjectSystem* GetPublic() const;
 	
-	void AddGameThreadTask(const FMutableTask& Task);
-
 	// Remove references to cached objects that have been deleted in the unreal
 	// side, and cannot be cached anyway.
 	// This should only happen in the game thread
diff --git a/Engine/Plugins/Mutable/Source/CustomizableObject/Internal/MuCO/FMutableTaskGraph.h b/Engine/Plugins/Mutable/Source/CustomizableObject/Internal/MuCO/FMutableTaskGraph.h
--- a/Engine/Plugins/Mutable/Source/CustomizableObject/Internal/MuCO/FMutableTaskGraph.h
+++ b/Engine/Plugins/Mutable/Source/CustomizableObject/Internal/MuCO/FMutableTaskGraph.h
@@ -6,13 +6,14 @@
 
 #include "Tasks/Task.h"
 #include "Async/TaskGraphInterfaces.h"
+#include "Containers/Queue.h"
 
 
 /** Mutable Tasks System.
  *
  * Allows launching tasks in different threads:
  * - Mutable Thread
- * - Game Thread (code currently in UCustomizableObjectSystem)
+ * - Game Thread
  * - Any Thread
  *
  * Concurrency between tasks in the Mutable Thread is forbidden by chaining all tasks through their prerequisites.
@@ -31,16 +32,32 @@
  */
 class FMutableTaskGraph
 {
-	struct FTask
+	struct FMutableThreadLowPriorityTask
 	{
 		uint32 Id;
 		FString DebugName;
 		TFunction<void()> Body;
 		double CreationTime = FPlatformTime::Seconds();
 	};
+
+	struct FGameThreadTask
+	{
+		FString DebugName;
+		TUniqueFunction<void()> Body;
+		TArray<UE::Tasks::FTask> Prerequisites;
+		bool bLockMutableThread = false;
+	
+		/** Check if all the dependencies of this task have been completed. */
+		bool AreDependenciesComplete() const;
+	};
 	
 public:
+	FMutableTaskGraph();
+
+	~FMutableTaskGraph();
 
+	void AddGameThreadTask(const TCHAR* DebugName, TUniqueFunction<void()>&& TaskBody, bool bLockMutableThread = false, const TArray<UE::Tasks::FTask>& Prerequisites = {});
+	
 	/** Create and launch a task on the Mutable Thread with Normal priority. */
 	UE::Tasks::FTask AddMutableThreadTask(const TCHAR* DebugName, TUniqueFunction<void()>&& TaskBody, const TArray<UE::Tasks::FTask>& Prerequisites = {});
 
@@ -67,8 +84,17 @@ public:
 
 	/** Returns the number of remaining tasks. */
 	int32 Tick();
-	
+
 private:
+	/** Lock, as soon as possible, the Mutable Thread. */
+	void AsyncLockMutableThread();
+
+	/** Check if the Mutable Thread is locked. */
+	bool IsMutableThreadLocked() const;
+
+	/** Unlock the Mutable Thread. Releases all blocked tasks. */
+	void UnlockMutableThread();
+	
 	/** A Mutable Task Low Priority will only be launched if:
 	 * - No other low priority task is running.
 	 * - Is allowed to launch low priority tasks.
@@ -76,6 +102,8 @@ private:
 	 * @param bFromMutableTask true if called from a Mutable Task. */
 	void TryLaunchMutableTaskLowPriority(bool bFromMutableTask);
 
+	void TryLaunchGameThreadTask();
+	
 	/** Return true if the task is completed (or is no longer valid). */
 	bool IsTaskCompleted(const UE::Tasks::FTask& Task) const;
 
@@ -85,10 +113,13 @@ private:
 	bool bAllowLaunchMutableTaskLowPriority = true;
 
 	/** Queue of low priority tasks. FIFO. */
-	TArray<FTask> QueueMutableTasksLowPriority;
+	TArray<FMutableThreadLowPriorityTask> QueueMutableTasksLowPriority;
+	
+	/** Queue of Game Thread tasks that need to be executed. */
+	TQueue<FGameThreadTask> GameThreadTasks;
 
 public:
-	inline static constexpr uint32 INVALID_ID = 0;
+	static constexpr uint32 INVALID_ID = 0;
 
 private:
 	/** Incremental task ID generator. */
@@ -98,9 +129,17 @@ private:
 	uint32 LastMutableTaskLowPriorityID = INVALID_ID;
 
 	/** Last Mutable Task Low Priority launched to the TaskGraph system. */
-	UE::Tasks::FTask LastMutableTaskLowPriority = {};
+	UE::Tasks::FTask LastMutableTaskLowPriority;
 
 	/** Last Mutable Task launched to the TaskGraph system. Low and normal priority. */
-	UE::Tasks::FTask LastMutableTask = {};
+	UE::Tasks::FTask LastMutableTask;
+
+	/** Last Mutable Task launched before locking the Mutable Thread. Once completed the Mutable Thread will be considered locked. */
+	UE::Tasks::FTask LastMutableTaskBeforeLock;
+	
+	/** Event to block upcoming tasks when the Mutable Thread is locked.
+	 *  Completed = Mutable Thread lock has been requested. May not be locked due to tasks still running (see LastMutableTaskBeforeLock).
+	 *  Not Completed = Mutable Thread unlocked. */
+	UE::Tasks::FTaskEvent MutableThreadUnlockEvent { UE_SOURCE_LOCATION };
 };
 
diff --git a/Engine/Plugins/Mutable/Source/CustomizableObject/Private/MuCO/CustomizableObjectSystem.cpp b/Engine/Plugins/Mutable/Source/CustomizableObject/Private/MuCO/CustomizableObjectSystem.cpp
--- a/Engine/Plugins/Mutable/Source/CustomizableObject/Private/MuCO/CustomizableObjectSystem.cpp
+++ b/Engine/Plugins/Mutable/Source/CustomizableObject/Private/MuCO/CustomizableObjectSystem.cpp
@@ -636,10 +636,7 @@ void UCustomizableObjectSystem::BeginDestroy()
 			FTSTicker::GetCoreTicker().RemoveTicker(Private->TickDelegateHandle);			
 		}
 #endif // !UE_SERVER
-
-		// Discard pending game thread tasks
-		Private->PendingTasks.Empty();
-
+		
 		// Complete pending taskgraph tasks
 		Private->MutableTaskGraph.AllowLaunchingMutableTaskLowPriority(false, false);
 		check(Private->Streamer);
@@ -686,12 +683,6 @@ UCustomizableObjectSystem* UCustomizableObjectSystemPrivate::GetPublic() const
 }
 
 
-void UCustomizableObjectSystemPrivate::AddGameThreadTask(const FMutableTask& Task)
-{
-	PendingTasks.Enqueue(Task);
-}
-
-
 TAutoConsoleVariable<bool> CVarCleanupTextureCache(
 	TEXT("mutable.EnableCleanupCache"),
 	true,
@@ -3074,21 +3065,20 @@ namespace impl
 
 			// Next Task: Callbacks
 			//-------------------------------------------------------------
-			TArray<UE::Tasks::FTask, TFixedAllocator<2>> Dependencies;
+			TArray<UE::Tasks::FTask> Dependencies;
 			if (CVarFixLowPriorityTasksOverlap.GetValueOnGameThread())
 			{
 				Dependencies.Add(Mutable_ReleaseInstanceTask);
 			}
 			
-			CustomizableObjectSystemPrivateData->AddGameThreadTask(
+			CustomizableObjectSystemPrivateData->MutableTaskGraph.AddGameThreadTask(
+				TEXT("Task_Game_Callbacks"),
+				[OperationData]()
 				{
-				FMutableTaskDelegate::CreateLambda(
-					[OperationData]()
-					{
-						Task_Game_Callbacks(OperationData);
-					}),
-					 Dependencies 
-				});
+					Task_Game_Callbacks(OperationData);
+				},
+				false,
+				Dependencies);
 		}
 		else
 		{
@@ -3211,16 +3201,14 @@ namespace impl
 
 		// Next-next Task: Convert Resources
 		//-------------------------------------------------------------
-		SystemPrivateData->AddGameThreadTask(
-			FMutableTask 
+		SystemPrivateData->MutableTaskGraph.AddGameThreadTask(
+			TEXT("Task_Game_ConvertResources"),
+			[OperationData]()
 			{
-				FMutableTaskDelegate::CreateLambda(
-				[OperationData]()
-				{
-					Task_Game_ConvertResources(OperationData);
-				}),
-				{ Game_LoadUnrealAssets, Mutable_GetImagesTask }
-			});
+				Task_Game_ConvertResources(OperationData);
+			},
+			false,
+			{ Game_LoadUnrealAssets, Mutable_GetImagesTask });
 	}
 
 
@@ -3273,16 +3261,14 @@ namespace impl
 				Task_Mutable_GetMeshes(Operation);
 			}, { Requisite });
 
-		SystemPrivate->AddGameThreadTask(
-			FMutableTask
+		SystemPrivate->MutableTaskGraph.AddGameThreadTask(
+			TEXT("Task_Game_LockCache"),
+			[Operation]()
 			{
-				FMutableTaskDelegate::CreateLambda(
-				[Operation]()
-				{
-					impl::Task_Game_LockCache(Operation);
-				}),
-				{ Dependency }
-			});
+				impl::Task_Game_LockCache(Operation);
+			},
+			false,
+			{ Dependency });
 	}
 	
 
@@ -3327,6 +3313,14 @@ namespace impl
 	{
 		MUTABLE_CPUPROFILER_SCOPE(Task_Game_StartUpdate)
 
+#if WITH_EDITOR
+if (CVarMutableLogObjectMemoryOnUpdate.GetValueOnAnyThread())
+{
+		UCustomizableObject* CO = Operation->Instance->GetCustomizableObject();
+		CO->GetPrivate()->LogMemory();
+}
+#endif
+		
 		// Check if a level has been loaded
 		if (FLogBenchmarkUtil::IsBenchmarkingReportingEnabled() && GWorld)
 		{
@@ -3581,17 +3575,15 @@ namespace impl
 					impl::Task_Mutable_GetMeshID(Operation);
 				});
 
-			SystemPrivateData->AddGameThreadTask(
-				FMutableTask 
+			SystemPrivateData->MutableTaskGraph.AddGameThreadTask(
+				TEXT("Task_Game_LockMeshCache"),
+				 [Operation, CacheRuntimeTexturesEvent]()
 				{
-					FMutableTaskDelegate::CreateLambda(
-					[Operation, CacheRuntimeTexturesEvent]()
-					{
-						impl::Task_Game_LockMeshCache(Operation, CacheRuntimeTexturesEvent);
-					}),
-					{ Mutable_GetMeshTask }
-				});
-		}
+					impl::Task_Game_LockMeshCache(Operation, CacheRuntimeTexturesEvent);
+				},
+				false,
+				{ Mutable_GetMeshTask });
+	}
 		else
 		{
 			Mutable_GetMeshTask = SystemPrivateData->MutableTaskGraph.AddMutableThreadTask(
@@ -3605,73 +3597,20 @@ namespace impl
 			//-------------------------------------------------------------
 			{
 				// Task inputs
-				SystemPrivateData->AddGameThreadTask(
-					FMutableTask 
+				SystemPrivateData->MutableTaskGraph.AddGameThreadTask(
+					TEXT("Task_Game_LockCache"),
+					[Operation]()
 					{
-						FMutableTaskDelegate::CreateLambda(
-						[Operation]()
-						{
-							impl::Task_Game_LockCache(Operation);
-						}),
-						{ Mutable_GetMeshTask }
-					});
+						impl::Task_Game_LockCache(Operation);
+					},
+					false,
+					{ Mutable_GetMeshTask });
 			}
 		}
 	}
 } // namespace impl
 
 
-void UCustomizableObjectSystem::AdvanceCurrentOperation() 
-{
-	MUTABLE_CPUPROFILER_SCOPE(AdvanceCurrentOperation);
-
-	check(Private != nullptr);
-
-	// See if we have a game-thread task to process
-	FMutableTask* PendingTask = Private->PendingTasks.Peek();
-	if (PendingTask)
-	{
-		if (PendingTask->AreDependenciesComplete())
-		{
-			PendingTask->ClearDependencies();
-			PendingTask->Function.Execute();
-			Private->PendingTasks.Pop();
-		}
-
-		// Don't do anything else until the pending work is completed.
-		return;
-	}
-
-	// It is safe to do this now.
-	Private->UpdateMemoryLimit();
-
-	// If we don't have an ongoing operation, don't do anything.
-	if (!Private->CurrentMutableOperation.IsValid())
-	{
-		return;
-	}
-
-	// If we reach here it means:
-	// - we have an ongoing operations
-	// - we have no pending work for the ongoing operation
-	// - so we are starting it.
-	{
-		MUTABLE_CPUPROFILER_SCOPE(OperationUpdate);
-
-#if WITH_EDITOR
-		if (CVarMutableLogObjectMemoryOnUpdate.GetValueOnAnyThread())
-		{
-			UCustomizableObject* CO = Private->CurrentMutableOperation->Instance->GetCustomizableObject();
-			CO->GetPrivate()->LogMemory();
-		}
-#endif
-
-		// Start the first task of the update process. See namespace impl comments above.
-		impl::Task_Game_StartUpdate(Private->CurrentMutableOperation.ToSharedRef());
-	}
-}
-
-
 bool UCustomizableObjectSystem::Tick(float DeltaTime)
 {
 	TickInternal(false);
@@ -3943,12 +3882,6 @@ int32 UCustomizableObjectSystem::TickInternal(const bool bBlocking)
 		ReleaseInstanceIDs();
 	}
 	
-	// Advance the current operation
-	if (Private->CurrentMutableOperation)
-	{
-		AdvanceCurrentOperation();
-	}
-	
 	const int32 RemainingTasks = Private->MutableTaskGraph.Tick();
 
 	Private->LogBenchmarkUtil.UpdateStats(); // Must to be the last thing to perform
@@ -4461,8 +4394,20 @@ void UCustomizableObjectSystemPrivate::StartUpdateSkeletalMesh(const TSharedRef<
 			LogStartedUpdateUnmute = CurrentTime + MuteTime;
 		}
 	}
+
+	// It is safe to do this now.
+	UpdateMemoryLimit();
 	
+	check(!CurrentMutableOperation);
 	CurrentMutableOperation = Context;
+
+	MutableTaskGraph.AddGameThreadTask(
+		TEXT("Task_Game_StartUpdate"),
+		[Context]
+		{
+			impl::Task_Game_StartUpdate(Context);
+		},
+		true);
 }
 
 
diff --git a/Engine/Plugins/Mutable/Source/CustomizableObject/Private/MuCO/FMutableTaskGraph.cpp b/Engine/Plugins/Mutable/Source/CustomizableObject/Private/MuCO/FMutableTaskGraph.cpp
--- a/Engine/Plugins/Mutable/Source/CustomizableObject/Private/MuCO/FMutableTaskGraph.cpp
+++ b/Engine/Plugins/Mutable/Source/CustomizableObject/Private/MuCO/FMutableTaskGraph.cpp
@@ -7,6 +7,10 @@
 #include "MuR/MutableTrace.h"
 
 
+#define UE_MUTABLE_THREAD_REQUEST_LOCK TEXT("Mutable Thread Request Lock")
+#define UE_MUTABLE_THREAD_LOCK TEXT("Mutable Thread Lock")
+
+
 constexpr LowLevelTasks::ETaskPriority TASKGRAPH_PRIORITY = LowLevelTasks::ETaskPriority::BackgroundHigh;
 
 static TAutoConsoleVariable<float> CVarMutableTaskLowPriorityMaxWaitTime(
@@ -22,18 +26,62 @@ static TAutoConsoleVariable<bool> CVarEnableMutableTaskLowPriority(
 	ECVF_Scalability);
 
 
+/** Status of the Mutable Thread lock. Only used for Insights Regions. */
+enum class ERegionLockStatus
+{
+	Requested,
+	Locked,
+	Unlocked // No region.
+};
+
+
+/** What Insights Region we are currently showing. Used only for debug. */
+static ERegionLockStatus RegionLockStatus;
+
+
+bool FMutableTaskGraph::FGameThreadTask::AreDependenciesComplete() const
+{
+	for (const UE::Tasks::FTask& Dependency : Prerequisites)
+	{
+		if (!Dependency.IsCompleted())
+		{
+			return false;
+		}
+	}
+
+	return true;
+}
+
+
+FMutableTaskGraph::FMutableTaskGraph()
+{
+	MutableThreadUnlockEvent.Trigger();
+}
+
+
+FMutableTaskGraph::~FMutableTaskGraph()
+{
+	GameThreadTasks.Empty();
+}
+
+
+void FMutableTaskGraph::AddGameThreadTask(const TCHAR* DebugName, TUniqueFunction<void()>&& TaskBody, bool bLockMutableThread, const TArray<UE::Tasks::FTask>& Prerequisites)
+{
+	GameThreadTasks.Enqueue({ DebugName, Forward<TUniqueFunction<void()>>(TaskBody), Prerequisites, bLockMutableThread });
+}
+
 
 UE::Tasks::FTask FMutableTaskGraph::AddMutableThreadTask(const TCHAR* DebugName, TUniqueFunction<void()>&& TaskBody, const TArray<UE::Tasks::FTask>& InPrerequisites)
 {
 	FScopeLock Lock(&MutableTaskLock);
 
-	TArray<UE::Tasks::FTask, TInlineAllocator<1>> Prerequisites = LastMutableTask.IsValid() 
-		? TArray<UE::Tasks::FTask, TInlineAllocator<1>>{LastMutableTask} 
-		: TArray<UE::Tasks::FTask, TInlineAllocator<1>>{};
+	TArray<UE::Tasks::FTask, TInlineAllocator<2>> Prerequisites = LastMutableTask.IsValid() 
+		? TArray<UE::Tasks::FTask, TInlineAllocator<2>>{LastMutableTask, MutableThreadUnlockEvent } 
+		: TArray<UE::Tasks::FTask, TInlineAllocator<2>>{};
 
 	Prerequisites.Append(InPrerequisites);
 	
-	LastMutableTask = UE::Tasks::Launch(DebugName, MoveTemp(TaskBody), Prerequisites, TASKGRAPH_PRIORITY);	
+	LastMutableTask = UE::Tasks::Launch(DebugName, Forward<TUniqueFunction<void()>>(TaskBody), Prerequisites, TASKGRAPH_PRIORITY);
 	
 	return LastMutableTask;
 }
@@ -46,7 +94,7 @@ uint32 FMutableTaskGraph::AddMutableThreadTaskLowPriority(const TCHAR* DebugName
 	if (CVarEnableMutableTaskLowPriority.GetValueOnAnyThread())
 	{
 		const uint32 Id = ++TaskIdGenerator;
-		const FTask Task { Id, DebugName, MoveTemp(TaskBody)};
+		const FMutableThreadLowPriorityTask Task { Id, DebugName, MoveTemp(TaskBody)};
 		QueueMutableTasksLowPriority.Add(Task);		
 
 		TryLaunchMutableTaskLowPriority(false);
@@ -148,7 +196,7 @@ void FMutableTaskGraph::TryLaunchMutableTaskLowPriority(bool bFromMutableTask)
 			return;
 		}
 
-		FTask& NextTask = QueueMutableTasksLowPriority[0];
+		FMutableThreadLowPriorityTask& NextTask = QueueMutableTasksLowPriority[0];
 
 		TimeLimit = CVarMutableTaskLowPriorityMaxWaitTime.GetValueOnAnyThread();
 		TimeElapsed = FPlatformTime::Seconds() - NextTask.CreationTime;
@@ -230,9 +278,47 @@ void FMutableTaskGraph::AllowLaunchingMutableTaskLowPriority(bool bAllow, bool b
 }
 
 
+void FMutableTaskGraph::TryLaunchGameThreadTask()
+{
+	MUTABLE_CPUPROFILER_SCOPE(AdvanceCurrentOperation);
+
+	FGameThreadTask* PendingTask = GameThreadTasks.Peek();
+	if (!PendingTask)
+	{
+		return;
+	}
+
+	if (!PendingTask->AreDependenciesComplete())
+	{
+		return;
+	}
+	
+	if (PendingTask->bLockMutableThread)
+	{
+		AsyncLockMutableThread();
+
+		if (!IsMutableThreadLocked())
+		{
+			return;
+		}
+	}
+
+	PendingTask->Body();
+	
+	if (PendingTask->bLockMutableThread)
+	{
+		UnlockMutableThread();
+	}
+
+	GameThreadTasks.Pop();
+}
+
+
 int32 FMutableTaskGraph::Tick()
 {
 	check(IsInGameThread())
+
+	TryLaunchGameThreadTask();
 	
 	TryLaunchMutableTaskLowPriority(false);
 
@@ -256,6 +342,92 @@ int32 FMutableTaskGraph::Tick()
 		}
 	}
 	
-	return !IsTaskCompleted(LastMutableTask) + !IsTaskCompleted(LastMutableTaskLowPriority) + QueueMutableTasksLowPriority.Num();
+	return !IsTaskCompleted(LastMutableTask) + !IsTaskCompleted(LastMutableTaskLowPriority) + QueueMutableTasksLowPriority.Num() + !GameThreadTasks.IsEmpty();
+}
+
+
+void UpdateRegions(ERegionLockStatus NewRegionLockStatus)
+{
+	if (RegionLockStatus == NewRegionLockStatus)
+	{
+		return;
+	}
+
+	// End current region.
+	switch (RegionLockStatus)
+	{
+	case ERegionLockStatus::Requested:
+		TRACE_END_REGION(UE_MUTABLE_THREAD_REQUEST_LOCK);
+		break;
+	case ERegionLockStatus::Locked:
+		TRACE_END_REGION(UE_MUTABLE_THREAD_LOCK);
+		break;
+	case ERegionLockStatus::Unlocked:
+		break;
+	default:
+		unimplemented();
+	}
+
+	// Start new region.
+	switch (NewRegionLockStatus)
+	{
+	case ERegionLockStatus::Requested:
+		TRACE_BEGIN_REGION(UE_MUTABLE_THREAD_REQUEST_LOCK);
+		break;
+	case ERegionLockStatus::Locked:
+		TRACE_BEGIN_REGION(UE_MUTABLE_THREAD_LOCK);
+		break;
+	case ERegionLockStatus::Unlocked:
+		break;
+	default:
+		unimplemented();
+	}
+
+	RegionLockStatus = NewRegionLockStatus;
+}
+
+
+void FMutableTaskGraph::AsyncLockMutableThread()
+{
+	MUTABLE_CPUPROFILER_SCOPE(FMutableTaskGraph::AsyncLockMutableThread)
+
+	FScopeLock Lock(&MutableTaskLock);
+	
+	if (!MutableThreadUnlockEvent.IsCompleted()) // Locking already in progress.
+	{
+		return;
+	}
+
+	UpdateRegions(ERegionLockStatus::Requested);
+
+	LastMutableTaskBeforeLock = LastMutableTask;
+	MutableThreadUnlockEvent = UE::Tasks::FTaskEvent(UE_SOURCE_LOCATION);
+}
+
+
+bool FMutableTaskGraph::IsMutableThreadLocked() const
+{
+	MUTABLE_CPUPROFILER_SCOPE(FMutableTaskGraph::IsMutableThreadLocked)
+
+	FScopeLock Lock(&MutableTaskLock);
+
+	const bool bLocked = !MutableThreadUnlockEvent.IsCompleted() && LastMutableTaskBeforeLock.IsCompleted();
+
+	UpdateRegions(bLocked ? ERegionLockStatus::Locked : RegionLockStatus);
+	
+	return bLocked;
+}
+
+
+void FMutableTaskGraph::UnlockMutableThread()
+{
+	MUTABLE_CPUPROFILER_SCOPE(FMutableTaskGraph::UnlockMutableThread)
+
+	FScopeLock Lock(&MutableTaskLock);
+
+	UpdateRegions(ERegionLockStatus::Unlocked);
+	
+	LastMutableTaskBeforeLock = {};
+	MutableThreadUnlockEvent.Trigger();
 }
 
diff --git a/Engine/Plugins/Mutable/Source/CustomizableObject/Public/MuCO/CustomizableObjectSystem.h b/Engine/Plugins/Mutable/Source/CustomizableObject/Public/MuCO/CustomizableObjectSystem.h
--- a/Engine/Plugins/Mutable/Source/CustomizableObject/Public/MuCO/CustomizableObjectSystem.h
+++ b/Engine/Plugins/Mutable/Source/CustomizableObject/Public/MuCO/CustomizableObjectSystem.h
@@ -358,10 +358,7 @@ private:
 
 	/** Returns the number of remaining operations. */
 	int32 TickInternal(bool bBlocking);
-
-	// If there is an on-going operation, advance it.
-	void AdvanceCurrentOperation();
-
+	
 	void DiscardInstances();
 	void ReleaseInstanceIDs();
 
