[Mutable] Added old CO warning.

Soon we are dropping support to old serialized Customizable Objects. If you want to maintain those please save them as soon as possible. Upcoming changes will remove the compatibility code and prevent loading old Customizable Objects.

#rnx

diff --git a/Engine/Plugins/Mutable/Source/CustomizableObject/Private/MuCO/CustomizableObject.cpp b/Engine/Plugins/Mutable/Source/CustomizableObject/Private/MuCO/CustomizableObject.cpp
--- a/Engine/Plugins/Mutable/Source/CustomizableObject/Private/MuCO/CustomizableObject.cpp
+++ b/Engine/Plugins/Mutable/Source/CustomizableObject/Private/MuCO/CustomizableObject.cpp
@@ -41,6 +41,7 @@
 #include "UObject/ObjectSaveContext.h"
 #include "MuCO/ICustomizableObjectEditorModule.h"
 #include "MuCO/CustomizableObjectSystemPrivate.h"
+#include "Misc/MessageDialog.h"
 
 #if WITH_EDITOR
 #include "Editor.h"
@@ -319,17 +320,33 @@ bool UCustomizableObjectPrivate::TryLoadCompiledCookDataForPlatform(const ITarge
 #endif // End WITH_EDITOR
 
 
+bool bShowOldCustomizableObjectWarning = true;
+
+
 void UCustomizableObject::PostLoad()
 {
 	Super::PostLoad();
 
+	const int32 CustomVersion = GetLinkerCustomVersion(FCustomizableObjectCustomVersion::GUID);
+
+	if (CustomVersion < FCustomizableObjectCustomVersion::UseUVRects) // A bit older than 5.5
+	{
+		if (bShowOldCustomizableObjectWarning)
+		{
+			bShowOldCustomizableObjectWarning = false;
+			FMessageDialog::Open(EAppMsgType::Ok, LOCTEXT("OldCustomizableObject", "Old Customizable Objects found. Please resave them to avoid future incompatibilities. See output log for more information."));
+		}
+
+		UE_LOG(LogMutable, Warning, TEXT("Unsupported old Customizable Object. Please resave it: %s"), *GetPackage()->GetPathName());
+	}
+	
 #if WITH_EDITOR
 	if (Source)
 	{
 		Source->ConditionalPostLoad();
 	}
 	
-	for (int32 Version = GetLinkerCustomVersion(FCustomizableObjectCustomVersion::GUID) + 1; Version <= FCustomizableObjectCustomVersion::LatestVersion; ++Version)
+	for (int32 Version = CustomVersion + 1; Version <= FCustomizableObjectCustomVersion::LatestVersion; ++Version)
 	{
 		GetPrivate()->BackwardsCompatibleFixup(Version);
 		
